Fix the static mesh import to order the material array as it is order in the fbx file.

diff --git a/Engine/Config/BaseEditorPerProjectUserSettings.ini b/Engine/Config/BaseEditorPerProjectUserSettings.ini
--- a/Engine/Config/BaseEditorPerProjectUserSettings.ini
+++ b/Engine/Config/BaseEditorPerProjectUserSettings.ini
@@ -603,7 +603,7 @@ bImportTextures=True
 bOverrideFullName=True
 bConvertScene=True
 bForceFrontXAxis=False
-bConvertSceneUnit=true
+bConvertSceneUnit=False
 bRemoveNamespaces=True
 bImportAnimations=False
 bResampleAnimations=False
diff --git a/Engine/Source/Editor/UnrealEd/Private/Fbx/FbxAssetImportData.cpp b/Engine/Source/Editor/UnrealEd/Private/Fbx/FbxAssetImportData.cpp
--- a/Engine/Source/Editor/UnrealEd/Private/Fbx/FbxAssetImportData.cpp
+++ b/Engine/Source/Editor/UnrealEd/Private/Fbx/FbxAssetImportData.cpp
@@ -9,7 +9,7 @@ UFbxAssetImportData::UFbxAssetImportData(const FObjectInitializer& ObjectInitial
 	, ImportUniformScale(1.0f)
 	, bConvertScene(true)
 	, bForceFrontXAxis(false)
-	, bConvertSceneUnit(true)
+	, bConvertSceneUnit(false)
 	, bImportAsScene(false)
 	, FbxSceneImportDataReference(nullptr)
 {
diff --git a/Engine/Source/Editor/UnrealEd/Private/Fbx/FbxFactory.cpp b/Engine/Source/Editor/UnrealEd/Private/Fbx/FbxFactory.cpp
--- a/Engine/Source/Editor/UnrealEd/Private/Fbx/FbxFactory.cpp
+++ b/Engine/Source/Editor/UnrealEd/Private/Fbx/FbxFactory.cpp
@@ -650,6 +650,17 @@ UObject* UFbxFactory::RecursiveImportNode(void* VoidFbxImporter, void* VoidNode,
 				}
 			}
 		}
+		
+		if (NewObject)
+		{
+			//Reorder the material
+			TArray<FbxNode*> Nodes;
+			FbxImporter->FindAllLODGroupNode(Nodes, Node, 0);
+			if (Nodes.Num() > 0)
+			{
+				FbxImporter->ReorderMaterialToFbxOrder(Cast<UStaticMesh>(NewObject), Nodes);
+			}
+		}
 	}
 	else
 	{
@@ -661,6 +672,11 @@ UObject* UFbxFactory::RecursiveImportNode(void* VoidFbxImporter, void* VoidNode,
 
 			if ( NewObject )
 			{
+				//Reorder the material
+				TArray<FbxNode*> Nodes;
+				Nodes.Add(Node);
+				FbxImporter->ReorderMaterialToFbxOrder(Cast<UStaticMesh>(NewObject), Nodes);
+
 				OutNewAssets.Add(NewObject);
 			}
 		}
diff --git a/Engine/Source/Editor/UnrealEd/Private/Fbx/FbxStaticMeshImport.cpp b/Engine/Source/Editor/UnrealEd/Private/Fbx/FbxStaticMeshImport.cpp
--- a/Engine/Source/Editor/UnrealEd/Private/Fbx/FbxStaticMeshImport.cpp
+++ b/Engine/Source/Editor/UnrealEd/Private/Fbx/FbxStaticMeshImport.cpp
@@ -944,6 +944,7 @@ UStaticMesh* UnFbx::FFbxImporter::ReimportStaticMesh(UStaticMesh* Mesh, UFbxStat
 	{
 		FbxNode* NodeParent = RecursiveFindParentLodGroup(Node->GetParent());
 
+		TArray<FbxNode*> LodZeroNodes;
 		// if the Fbx mesh is a part of LODGroup, update LOD
 		if (NodeParent && NodeParent->GetNodeAttribute() && NodeParent->GetNodeAttribute()->GetAttributeType() == FbxNodeAttribute::eLODGroup)
 		{
@@ -951,6 +952,7 @@ UStaticMesh* UnFbx::FFbxImporter::ReimportStaticMesh(UStaticMesh* Mesh, UFbxStat
 			FindAllLODGroupNode(AllNodeInLod, NodeParent, 0);
 			if (AllNodeInLod.Num() > 0)
 			{
+				LodZeroNodes = AllNodeInLod;
 				NewMesh = ImportStaticMeshAsSingle(Mesh->GetOuter(), AllNodeInLod, *Mesh->GetName(), RF_Public | RF_Standalone, TemplateImportData, Mesh, 0, ExistMeshDataPtr);
 			}
 
@@ -972,8 +974,10 @@ UStaticMesh* UnFbx::FFbxImporter::ReimportStaticMesh(UStaticMesh* Mesh, UFbxStat
 		}
 		else
 		{
+			LodZeroNodes.Add(Node);
 			NewMesh = ImportStaticMesh(Mesh->GetOuter(), Node, *Mesh->GetName(), RF_Public|RF_Standalone, TemplateImportData, Mesh, 0, ExistMeshDataPtr);
 		}
+		ReorderMaterialToFbxOrder(NewMesh, LodZeroNodes);
 	}
 	else
 	{
@@ -981,6 +985,7 @@ UStaticMesh* UnFbx::FFbxImporter::ReimportStaticMesh(UStaticMesh* Mesh, UFbxStat
 		if (FbxMeshArray.Num() > 0)
 		{
 			NewMesh = ImportStaticMeshAsSingle(Mesh->GetOuter(), FbxMeshArray, *Mesh->GetName(), RF_Public|RF_Standalone, TemplateImportData, Mesh, 0, ExistMeshDataPtr);
+			ReorderMaterialToFbxOrder(NewMesh, FbxMeshArray);
 		}
 		else // no mesh found in the FBX file
 		{
@@ -1223,10 +1228,10 @@ UStaticMesh* UnFbx::FFbxImporter::ImportStaticMeshAsSingle(UObject* InParent, TA
 
 	if (bBuildStatus)
 	{
-		UE_LOG(LogFbx,Verbose,TEXT("== Initial material list:"));
+		UE_LOG(LogFbx, Verbose, TEXT("== Initial material list:"));
 		for (int32 MaterialIndex = 0; MaterialIndex < MeshMaterials.Num(); ++MaterialIndex)
 		{
-			UE_LOG(LogFbx,Verbose,TEXT("%d: %s"),MaterialIndex,*MeshMaterials[MaterialIndex].GetName() );
+			UE_LOG(LogFbx, Verbose, TEXT("%d: %s"), MaterialIndex, *MeshMaterials[MaterialIndex].GetName());
 		}
 
 		// Compress the materials array by removing any duplicates.
@@ -1252,11 +1257,11 @@ UStaticMesh* UnFbx::FFbxImporter::ImportStaticMeshAsSingle(UObject* InParent, TA
 			{
 				int32 UniqueIndex = UniqueMaterials.Add(MeshMaterials[MaterialIndex]);
 
-				MaterialMap.Add( UniqueIndex );
+				MaterialMap.Add(UniqueIndex);
 			}
 			else
 			{
-				UE_LOG(LogFbx,Verbose,TEXT("  remap %d -> %d"), MaterialIndex, MaterialMap[MaterialIndex]);
+				UE_LOG(LogFbx, Verbose, TEXT("  remap %d -> %d"), MaterialIndex, MaterialMap[MaterialIndex]);
 			}
 		}
 
@@ -1264,13 +1269,13 @@ UStaticMesh* UnFbx::FFbxImporter::ImportStaticMeshAsSingle(UObject* InParent, TA
 		{
 			AddTokenizedErrorMessage(
 				FTokenizedMessage::Create(
-				EMessageSeverity::Warning,
-				FText::Format(LOCTEXT("Error_TooManyMaterials", "StaticMesh has a large number({0}) of materials and may render inefficently.  Consider breaking up the mesh into multiple Static Mesh Assets"),
-				FText::AsNumber(UniqueMaterials.Num())
-				)), 
+					EMessageSeverity::Warning,
+					FText::Format(LOCTEXT("Error_TooManyMaterials", "StaticMesh has a large number({0}) of materials and may render inefficently.  Consider breaking up the mesh into multiple Static Mesh Assets"),
+						FText::AsNumber(UniqueMaterials.Num())
+						)),
 				FFbxErrors::StaticMesh_TooManyMaterials);
 		}
-		
+
 		//The fix is required for blender file. The static mesh build have change and now required that
 		//the sections (face declaration) must be declare in the same order as the material index
 		TArray<uint32> SortedMaterialIndex;
@@ -1313,9 +1318,9 @@ UStaticMesh* UnFbx::FFbxImporter::ImportStaticMeshAsSingle(UObject* InParent, TA
 		{
 			int32 RemappedIndex = SortedMaterialIndex[SortedIndex] & 0xffff;
 			SortedMaterials.Add(UniqueMaterials[RemappedIndex]);
-			UE_LOG(LogFbx,Verbose,TEXT("%d: %s"),SortedIndex,*UniqueMaterials[RemappedIndex].GetName());
+			UE_LOG(LogFbx, Verbose, TEXT("%d: %s"), SortedIndex, *UniqueMaterials[RemappedIndex].GetName());
 		}
-		UE_LOG(LogFbx,Verbose,TEXT("== Mapping table:"));
+		UE_LOG(LogFbx, Verbose, TEXT("== Mapping table:"));
 		for (int32 MaterialIndex = 0; MaterialIndex < MaterialMap.Num(); ++MaterialIndex)
 		{
 			for (int32 SortedIndex = 0; SortedIndex < SortedMaterialIndex.Num(); ++SortedIndex)
@@ -1323,13 +1328,13 @@ UStaticMesh* UnFbx::FFbxImporter::ImportStaticMeshAsSingle(UObject* InParent, TA
 				int32 RemappedIndex = SortedMaterialIndex[SortedIndex] & 0xffff;
 				if (MaterialMap[MaterialIndex] == RemappedIndex)
 				{
-					UE_LOG(LogFbx,Verbose,TEXT("  sort %d -> %d"), MaterialIndex, SortedIndex);
+					UE_LOG(LogFbx, Verbose, TEXT("  sort %d -> %d"), MaterialIndex, SortedIndex);
 					MaterialMap[MaterialIndex] = SortedIndex;
 					break;
 				}
 			}
 		}
-		
+
 		// Remap material indices.
 		int32 MaxMaterialIndex = 0;
 		int32 FirstOpenUVChannel = 1;
@@ -1353,7 +1358,7 @@ UStaticMesh* UnFbx::FFbxImporter::ImportStaticMeshAsSingle(UObject* InParent, TA
 			{
 				TArray<FFbxMaterial> OldSortedMaterials;
 
-				Exchange(OldSortedMaterials,SortedMaterials);
+				Exchange(OldSortedMaterials, SortedMaterials);
 				SortedMaterials.Empty(LocalRawMesh.MaterialIndexToImportIndex.Num());
 				for (int32 MaterialIndex = 0; MaterialIndex < LocalRawMesh.MaterialIndexToImportIndex.Num(); ++MaterialIndex)
 				{
@@ -1369,12 +1374,12 @@ UStaticMesh* UnFbx::FFbxImporter::ImportStaticMeshAsSingle(UObject* InParent, TA
 
 			for (int32 TriIndex = 0; TriIndex < LocalRawMesh.FaceMaterialIndices.Num(); ++TriIndex)
 			{
-				MaxMaterialIndex = FMath::Max<int32>(MaxMaterialIndex,LocalRawMesh.FaceMaterialIndices[TriIndex]);
+				MaxMaterialIndex = FMath::Max<int32>(MaxMaterialIndex, LocalRawMesh.FaceMaterialIndices[TriIndex]);
 			}
 
-			for( int32 i = 0; i < MAX_MESH_TEXTURE_COORDS; i++ )
+			for (int32 i = 0; i < MAX_MESH_TEXTURE_COORDS; i++)
 			{
-				if( LocalRawMesh.WedgeTexCoords[i].Num() == 0 )
+				if (LocalRawMesh.WedgeTexCoords[i].Num() == 0)
 				{
 					FirstOpenUVChannel = i;
 					break;
@@ -1389,15 +1394,15 @@ UStaticMesh* UnFbx::FFbxImporter::ImportStaticMeshAsSingle(UObject* InParent, TA
 		{
 			StaticMesh->StaticMaterials.Empty();
 		}
-		
+
 		// Replace map of sections with the unique material set
-		int32 NumMaterials = FMath::Min(SortedMaterials.Num(),MaxMaterialIndex+1);
+		int32 NumMaterials = FMath::Min(SortedMaterials.Num(), MaxMaterialIndex + 1);
 		for (int32 MaterialIndex = 0; MaterialIndex < NumMaterials; ++MaterialIndex)
 		{
 			FMeshSectionInfo Info = StaticMesh->SectionInfoMap.Get(LODIndex, MaterialIndex);
 
 			int32 Index = 0;
-			
+
 			FName MaterialFName = FName(*(SortedMaterials[MaterialIndex].GetName()));
 			FString CleanMaterialSlotName = MaterialFName.ToString();
 			int32 SkinOffset = CleanMaterialSlotName.Find(TEXT("_skin"));
@@ -1412,7 +1417,7 @@ UStaticMesh* UnFbx::FFbxImporter::ImportStaticMeshAsSingle(UObject* InParent, TA
 				StaticMesh->StaticMaterials[Index].MaterialInterface = SortedMaterials[MaterialIndex].Material;
 				OldMaterialIndex.RemoveAt(0);
 			}
-			else if (InStaticMesh )
+			else if (InStaticMesh)
 			{
 				Index = INDEX_NONE;
 				FStaticMaterial StaticMaterialImported(SortedMaterials[MaterialIndex].Material, FName(*CleanMaterialSlotName), MaterialFName);
@@ -1431,7 +1436,7 @@ UStaticMesh* UnFbx::FFbxImporter::ImportStaticMeshAsSingle(UObject* InParent, TA
 			}
 			else
 			{
-				
+
 				Index = StaticMesh->StaticMaterials.Add(FStaticMaterial(SortedMaterials[MaterialIndex].Material, FName(*CleanMaterialSlotName), MaterialFName));
 			}
 
@@ -1462,7 +1467,7 @@ UStaticMesh* UnFbx::FFbxImporter::ImportStaticMeshAsSingle(UObject* InParent, TA
 		}
 
 		UFbxStaticMeshImportData* ImportData = UFbxStaticMeshImportData::GetImportDataForStaticMesh(StaticMesh, TemplateImportData);
-		
+
 		//@third party BEGIN SIMPLYGON
 		/* ImportData->Update(UFactory::GetCurrentFilename());
 		Developer Note: Update method above computed Hash internally. Hash is calculated based on the file size.
@@ -1471,7 +1476,7 @@ UStaticMesh* UnFbx::FFbxImporter::ImportStaticMeshAsSingle(UObject* InParent, TA
 		*/
 		ImportData->Update(UFactory::GetCurrentFilename(), UFactory::GetFileHash());
 		//@third party END SIMPLYGON
-		
+
 
 		// @todo This overrides restored values currently but we need to be able to import over the existing settings if the user chose to do so.
 		SrcModel.BuildSettings.bRemoveDegenerates = ImportOptions->bRemoveDegenerates;
@@ -1480,11 +1485,11 @@ UStaticMesh* UnFbx::FFbxImporter::ImportStaticMeshAsSingle(UObject* InParent, TA
 		SrcModel.BuildSettings.bRecomputeNormals = ImportOptions->NormalImportMethod == FBXNIM_ComputeNormals;
 		SrcModel.BuildSettings.bRecomputeTangents = ImportOptions->NormalImportMethod != FBXNIM_ImportNormalsAndTangents;
 		SrcModel.BuildSettings.bUseMikkTSpace = (ImportOptions->NormalGenerationMethod == EFBXNormalGenerationMethod::MikkTSpace) && (!ImportOptions->ShouldImportNormals() || !ImportOptions->ShouldImportTangents());
-		if( ImportOptions->bGenerateLightmapUVs )
+		if (ImportOptions->bGenerateLightmapUVs)
 		{
 			SrcModel.BuildSettings.bGenerateLightmapUVs = true;
 			SrcModel.BuildSettings.DstLightmapIndex = FirstOpenUVChannel;
-			StaticMesh->LightMapCoordinateIndex = FirstOpenUVChannel;
+StaticMesh->LightMapCoordinateIndex = FirstOpenUVChannel;
 		}
 		else
 		{
@@ -1505,21 +1510,21 @@ UStaticMesh* UnFbx::FFbxImporter::ImportStaticMeshAsSingle(UObject* InParent, TA
 
 		StaticMesh->Build(false, &BuildErrors);
 
-		for( FText& Error : BuildErrors )
+		for (FText& Error : BuildErrors)
 		{
-			AddTokenizedErrorMessage(FTokenizedMessage::Create(EMessageSeverity::Warning, Error), FFbxErrors::StaticMesh_BuildError );
+			AddTokenizedErrorMessage(FTokenizedMessage::Create(EMessageSeverity::Warning, Error), FFbxErrors::StaticMesh_BuildError);
 		}
-		
+
 		// this is damage control. After build, we'd like to absolutely sure that 
 		// all index is pointing correctly and they're all used. Otherwise we remove them
 		FMeshSectionInfoMap OldSectionInfoMap = StaticMesh->SectionInfoMap;
 		StaticMesh->SectionInfoMap.Clear();
 		// fix up section data
-		for (int32 LODResoureceIndex = 0; LODResoureceIndex<StaticMesh->RenderData->LODResources.Num(); ++LODResoureceIndex)
+		for (int32 LODResoureceIndex = 0; LODResoureceIndex < StaticMesh->RenderData->LODResources.Num(); ++LODResoureceIndex)
 		{
 			FStaticMeshLODResources& LOD = StaticMesh->RenderData->LODResources[LODResoureceIndex];
 			int32 NumSections = LOD.Sections.Num();
-			for(int32 SectionIndex = 0; SectionIndex < NumSections; ++SectionIndex)
+			for (int32 SectionIndex = 0; SectionIndex < NumSections; ++SectionIndex)
 			{
 				FMeshSectionInfo Info = OldSectionInfoMap.Get(LODResoureceIndex, SectionIndex);
 				if (StaticMesh->StaticMaterials.IsValidIndex(Info.MaterialIndex))
@@ -1537,10 +1542,10 @@ UStaticMesh* UnFbx::FFbxImporter::ImportStaticMeshAsSingle(UObject* InParent, TA
 			TArray< FString > MissingUVSets;
 			TArray< FString > BadUVSets;
 			TArray< FString > ValidUVSets;
-			UStaticMesh::CheckLightMapUVs( StaticMesh, MissingUVSets, BadUVSets, ValidUVSets );
+			UStaticMesh::CheckLightMapUVs(StaticMesh, MissingUVSets, BadUVSets, ValidUVSets);
 
 			// NOTE: We don't care about missing UV sets here, just bad ones!
-			if( BadUVSets.Num() > 0 )
+			if (BadUVSets.Num() > 0)
 			{
 				AddTokenizedErrorMessage(FTokenizedMessage::Create(EMessageSeverity::Warning, FText::Format(LOCTEXT("Error_UVSetLayoutProblem", "Warning: The light map UV set for static mesh '{0}' appears to have layout problems.  Either the triangle UVs are overlapping one another or the UV are out of bounds (0.0 - 1.0 range.)"), FText::FromString(MeshName))), FFbxErrors::StaticMesh_UVSetLayoutProblem);
 			}
@@ -1573,7 +1578,7 @@ UStaticMesh* UnFbx::FFbxImporter::ImportStaticMeshAsSingle(UObject* InParent, TA
 			for (int32 DirIdx = 0; DirIdx < NumDirs; ++DirIdx) { Dirs[DirIdx] = KDopDir18[DirIdx]; }
 			GenerateKDopAsSimpleCollision(StaticMesh, Dirs);
 		}
-		
+
 		//warnings based on geometry
 		VerifyGeometry(StaticMesh);
 
@@ -1583,6 +1588,78 @@ UStaticMesh* UnFbx::FFbxImporter::ImportStaticMeshAsSingle(UObject* InParent, TA
 	return StaticMesh;
 }
 
+void UnFbx::FFbxImporter::ReorderMaterialToFbxOrder(UStaticMesh* StaticMesh, TArray<FbxNode*>& MeshNodeArray)
+{
+	TArray<FFbxMaterial> MeshMaterials;
+	for (int32 MeshIndex = 0; MeshIndex < MeshNodeArray.Num(); MeshIndex++)
+	{
+		FbxNode* Node = MeshNodeArray[MeshIndex];
+		if (Node->GetMesh())
+		{
+			int32 MaterialCount = Node->GetMaterialCount();
+
+			for (int32 MaterialIndex = 0; MaterialIndex < MaterialCount; MaterialIndex++)
+			{
+				FFbxMaterial* NewMaterial = new(MeshMaterials) FFbxMaterial;
+				FbxSurfaceMaterial *FbxMaterial = Node->GetMaterial(MaterialIndex);
+				NewMaterial->FbxMaterial = FbxMaterial;
+				FString MaterialFullName = GetMaterialFullName(*FbxMaterial);
+				FString BasePackageName = PackageTools::SanitizePackageName(FPackageName::GetLongPackagePath(StaticMesh->GetOutermost()->GetName()) / MaterialFullName);
+				UMaterialInterface* UnrealMaterialInterface = FindObject<UMaterialInterface>(NULL, *(BasePackageName + TEXT(".") + MaterialFullName));
+				if (UnrealMaterialInterface == NULL)
+				{
+					UnrealMaterialInterface = UMaterial::GetDefaultMaterial(MD_Surface);
+				}
+				NewMaterial->Material = UnrealMaterialInterface;
+			}
+		}
+	}
+
+	//Reorder the StaticMaterials array to reflect the order in the fbx file
+	//So we make sure the order reflect the material ID in the DCCs
+	FMeshSectionInfoMap OldSectionInfoMap = StaticMesh->SectionInfoMap;
+	TArray<int32> FbxRemapMaterials;
+	TArray<FStaticMaterial> NewStaticMaterials;
+	for (int32 FbxMaterialIndex = 0; FbxMaterialIndex < MeshMaterials.Num(); ++FbxMaterialIndex)
+	{
+		FFbxMaterial &FbxMaterial = MeshMaterials[FbxMaterialIndex];
+		int32 FoundMaterialIndex = INDEX_NONE;
+		for (int32 BuildMaterialIndex = 0; BuildMaterialIndex < StaticMesh->StaticMaterials.Num(); ++BuildMaterialIndex)
+		{
+			FStaticMaterial &BuildMaterial = StaticMesh->StaticMaterials[BuildMaterialIndex];
+			if (FbxMaterial.Material == BuildMaterial.MaterialInterface)
+			{
+				FoundMaterialIndex = BuildMaterialIndex;
+				break;
+			}
+		}
+		FbxRemapMaterials.Add(FoundMaterialIndex);
+		NewStaticMaterials.Add(StaticMesh->StaticMaterials[FoundMaterialIndex]);
+	}
+	StaticMesh->StaticMaterials.Empty();
+	for (const FStaticMaterial &BuildMaterial : NewStaticMaterials)
+	{
+		StaticMesh->StaticMaterials.Add(BuildMaterial);
+	}
+
+	//Remap the material instance of the staticmaterial array and remap the material index of all sections
+	for (int32 LODResoureceIndex = 0; LODResoureceIndex < StaticMesh->RenderData->LODResources.Num(); ++LODResoureceIndex)
+	{
+		FStaticMeshLODResources& LOD = StaticMesh->RenderData->LODResources[LODResoureceIndex];
+		int32 NumSections = LOD.Sections.Num();
+		for (int32 SectionIndex = 0; SectionIndex < NumSections; ++SectionIndex)
+		{
+			FMeshSectionInfo Info = OldSectionInfoMap.Get(LODResoureceIndex, SectionIndex);
+			int32 RemapIndex = FbxRemapMaterials.Find(Info.MaterialIndex);
+			if (StaticMesh->StaticMaterials.IsValidIndex(RemapIndex))
+			{
+				Info.MaterialIndex = RemapIndex;
+				StaticMesh->SectionInfoMap.Set(LODResoureceIndex, SectionIndex, Info);
+			}
+		}
+	}
+}
+
 bool UnFbx::FFbxImporter::ImportSubDSurface(USubDSurface* Out, UObject* InParent, TArray<FbxNode*>& MeshNodeArray, const FName InName, EObjectFlags Flags, UFbxStaticMeshImportData* TemplateImportData)
 {
 	if(!MeshNodeArray.Num())
diff --git a/Engine/Source/Editor/UnrealEd/Private/FbxImporter.h b/Engine/Source/Editor/UnrealEd/Private/FbxImporter.h
--- a/Engine/Source/Editor/UnrealEd/Private/FbxImporter.h
+++ b/Engine/Source/Editor/UnrealEd/Private/FbxImporter.h
@@ -606,6 +606,13 @@ public:
 	*/
 	UNREALED_API UStaticMesh* ImportStaticMeshAsSingle(UObject* InParent, TArray<FbxNode*>& MeshNodeArray, const FName InName, EObjectFlags Flags, UFbxStaticMeshImportData* TemplateImportData, UStaticMesh* InStaticMesh, int LODIndex = 0, void *ExistMeshDataPtr = nullptr);
 
+	/**
+	 * Helper function to reorder the material array after we build the staticmesh.
+	 * It order the material like it is in the fbx file and reassign section material index properly.
+	 * This must be call once all LOD are imported.
+	 */
+	void ReorderMaterialToFbxOrder(UStaticMesh* StaticMesh, TArray<FbxNode*>& MeshNodeArray);
+
 	/**
 	* Creates a SubDSurface mesh from all the meshes in FBX scene with the given name and flags.
 	*
