when doing a fbx import, compare the imported material name instead of the material instance when reordering the material array.

diff --git a/Engine/Source/Editor/UnrealEd/Private/Fbx/FbxFactory.cpp b/Engine/Source/Editor/UnrealEd/Private/Fbx/FbxFactory.cpp
--- a/Engine/Source/Editor/UnrealEd/Private/Fbx/FbxFactory.cpp
+++ b/Engine/Source/Editor/UnrealEd/Private/Fbx/FbxFactory.cpp
@@ -304,6 +304,7 @@ UObject* UFbxFactory::FactoryCreateBinary
 						if (FbxMeshArray.Num() > 0)
 						{
 							NewStaticMesh = FbxImporter->ImportStaticMeshAsSingle(InParent, FbxMeshArray, Name, Flags, ImportUI->StaticMeshImportData, NULL, 0);
+							FbxImporter->ReorderMaterialToFbxOrder(NewStaticMesh, FbxMeshArray);
 						}
 
 						ImportedMeshCount = NewStaticMesh ? 1 : 0;
diff --git a/Engine/Source/Editor/UnrealEd/Private/Fbx/FbxStaticMeshImport.cpp b/Engine/Source/Editor/UnrealEd/Private/Fbx/FbxStaticMeshImport.cpp
--- a/Engine/Source/Editor/UnrealEd/Private/Fbx/FbxStaticMeshImport.cpp
+++ b/Engine/Source/Editor/UnrealEd/Private/Fbx/FbxStaticMeshImport.cpp
@@ -1590,7 +1590,7 @@ StaticMesh->LightMapCoordinateIndex = FirstOpenUVChannel;
 
 void UnFbx::FFbxImporter::ReorderMaterialToFbxOrder(UStaticMesh* StaticMesh, TArray<FbxNode*>& MeshNodeArray)
 {
-	TArray<FFbxMaterial> MeshMaterials;
+	TArray<FString> MeshMaterials;
 	for (int32 MeshIndex = 0; MeshIndex < MeshNodeArray.Num(); MeshIndex++)
 	{
 		FbxNode* Node = MeshNodeArray[MeshIndex];
@@ -1600,17 +1600,13 @@ void UnFbx::FFbxImporter::ReorderMaterialToFbxOrder(UStaticMesh* StaticMesh, TAr
 
 			for (int32 MaterialIndex = 0; MaterialIndex < MaterialCount; MaterialIndex++)
 			{
-				FFbxMaterial* NewMaterial = new(MeshMaterials) FFbxMaterial;
+				//Get the original fbx import name
 				FbxSurfaceMaterial *FbxMaterial = Node->GetMaterial(MaterialIndex);
-				NewMaterial->FbxMaterial = FbxMaterial;
-				FString MaterialFullName = GetMaterialFullName(*FbxMaterial);
-				FString BasePackageName = PackageTools::SanitizePackageName(FPackageName::GetLongPackagePath(StaticMesh->GetOutermost()->GetName()) / MaterialFullName);
-				UMaterialInterface* UnrealMaterialInterface = FindObject<UMaterialInterface>(NULL, *(BasePackageName + TEXT(".") + MaterialFullName));
-				if (UnrealMaterialInterface == NULL)
+				FString FbxMaterialName = FbxMaterial ? ANSI_TO_TCHAR(FbxMaterial->GetName()) : TEXT("None");
+				if (!MeshMaterials.Contains(FbxMaterialName))
 				{
-					UnrealMaterialInterface = UMaterial::GetDefaultMaterial(MD_Surface);
+					MeshMaterials.Add(FbxMaterialName);
 				}
-				NewMaterial->Material = UnrealMaterialInterface;
 			}
 		}
 	}
@@ -1622,12 +1618,12 @@ void UnFbx::FFbxImporter::ReorderMaterialToFbxOrder(UStaticMesh* StaticMesh, TAr
 	TArray<FStaticMaterial> NewStaticMaterials;
 	for (int32 FbxMaterialIndex = 0; FbxMaterialIndex < MeshMaterials.Num(); ++FbxMaterialIndex)
 	{
-		FFbxMaterial &FbxMaterial = MeshMaterials[FbxMaterialIndex];
+		const FString &FbxMaterial = MeshMaterials[FbxMaterialIndex];
 		int32 FoundMaterialIndex = INDEX_NONE;
 		for (int32 BuildMaterialIndex = 0; BuildMaterialIndex < StaticMesh->StaticMaterials.Num(); ++BuildMaterialIndex)
 		{
 			FStaticMaterial &BuildMaterial = StaticMesh->StaticMaterials[BuildMaterialIndex];
-			if (FbxMaterial.Material == BuildMaterial.MaterialInterface)
+			if (FbxMaterial.Compare(BuildMaterial.ImportedMaterialSlotName.ToString()) == 0)
 			{
 				FoundMaterialIndex = BuildMaterialIndex;
 				break;
