Removing an assert that can fire when processing an enum entry with a literal name of "None"

This fired while re-saving some test content.

diff --git a/Engine/Source/Editor/UnrealEd/Private/Kismet2/EnumEditorUtils.cpp b/Engine/Source/Editor/UnrealEd/Private/Kismet2/EnumEditorUtils.cpp
--- a/Engine/Source/Editor/UnrealEd/Private/Kismet2/EnumEditorUtils.cpp
+++ b/Engine/Source/Editor/UnrealEd/Private/Kismet2/EnumEditorUtils.cpp
@@ -446,13 +446,10 @@ FText FEnumEditorUtils::GetEnumeratorDisplayName(const UUserDefinedEnum* Enum, c
 	{
 		const FName EnumEntryName = *Enum->GetEnumName(EnumeratorIndex);
 
-		if (!EnumEntryName.IsNone())
+		const FText* EnumEntryDisplayName = Enum->DisplayNameMap.Find(EnumEntryName);
+		if (EnumEntryDisplayName)
 		{
-			const FText* EnumEntryDisplayName = Enum->DisplayNameMap.Find(EnumEntryName);
-			if (EnumEntryDisplayName)
-			{
-				return *EnumEntryDisplayName;
-			}
+			return *EnumEntryDisplayName;
 		}
 	}
 
@@ -468,7 +465,6 @@ bool FEnumEditorUtils::SetEnumeratorDisplayName(UUserDefinedEnum* Enum, int32 En
 			PrepareForChange(Enum);
 			
 			const FName EnumEntryName = *Enum->GetEnumName(EnumeratorIndex);
-			check(!EnumEntryName.IsNone());
 
 			FText DisplayNameToSet = NewDisplayName;
 
@@ -531,8 +527,6 @@ void FEnumEditorUtils::EnsureAllDisplayNamesExist(UUserDefinedEnum* Enum)
 			for (int32 Index = 0; Index < EnumeratorsToEnsure; ++Index)
 			{
 				const FName EnumEntryName = *Enum->GetEnumName(Index);
-				check(!EnumEntryName.IsNone());
-
 				KnownEnumEntryNames.Add(EnumEntryName);
 			}
 
@@ -551,7 +545,6 @@ void FEnumEditorUtils::EnsureAllDisplayNamesExist(UUserDefinedEnum* Enum)
 		for	(int32 Index = 0; Index < EnumeratorsToEnsure; ++Index)
 		{
 			const FName EnumEntryName = *Enum->GetEnumName(Index);
-			check(!EnumEntryName.IsNone());
 
 			if (!Enum->DisplayNameMap.Contains(EnumEntryName))
 			{
@@ -592,7 +585,6 @@ void FEnumEditorUtils::UpgradeDisplayNamesFromMetaData(UUserDefinedEnum* Enum)
 				bDidUpgradeDisplayNames = true;
 
 				const FName EnumEntryName = *Enum->GetEnumName(Index);
-				check(!EnumEntryName.IsNone());
 
 				FText DisplayNameToSet;
 
diff --git a/Engine/Source/Runtime/Engine/Private/UserDefinedEnum.cpp b/Engine/Source/Runtime/Engine/Private/UserDefinedEnum.cpp
--- a/Engine/Source/Runtime/Engine/Private/UserDefinedEnum.cpp
+++ b/Engine/Source/Runtime/Engine/Private/UserDefinedEnum.cpp
@@ -113,13 +113,10 @@ FText UUserDefinedEnum::GetEnumText(int32 InIndex) const
 {
 	const FName EnumEntryName = *GetEnumName(InIndex);
 
-	if (!EnumEntryName.IsNone())
+	const FText* EnumEntryDisplayName = DisplayNameMap.Find(EnumEntryName);
+	if (EnumEntryDisplayName)
 	{
-		const FText* EnumEntryDisplayName = DisplayNameMap.Find(EnumEntryName);
-		if (EnumEntryDisplayName)
-		{
-			return *EnumEntryDisplayName;
-		}
+		return *EnumEntryDisplayName;
 	}
 
 	return FText::GetEmpty();
