Fix a crash when reordering material with a fbx containing unused materials

diff --git a/Engine/Source/Editor/UnrealEd/Private/Fbx/FbxStaticMeshImport.cpp b/Engine/Source/Editor/UnrealEd/Private/Fbx/FbxStaticMeshImport.cpp
--- a/Engine/Source/Editor/UnrealEd/Private/Fbx/FbxStaticMeshImport.cpp
+++ b/Engine/Source/Editor/UnrealEd/Private/Fbx/FbxStaticMeshImport.cpp
@@ -1629,8 +1629,12 @@ void UnFbx::FFbxImporter::ReorderMaterialToFbxOrder(UStaticMesh* StaticMesh, TAr
 				break;
 			}
 		}
-		FbxRemapMaterials.Add(FoundMaterialIndex);
-		NewStaticMaterials.Add(StaticMesh->StaticMaterials[FoundMaterialIndex]);
+		
+		if (FoundMaterialIndex != INDEX_NONE)
+		{
+			FbxRemapMaterials.Add(FoundMaterialIndex);
+			NewStaticMaterials.Add(StaticMesh->StaticMaterials[FoundMaterialIndex]);
+		}
 	}
 	StaticMesh->StaticMaterials.Empty();
 	for (const FStaticMaterial &BuildMaterial : NewStaticMaterials)
