Copying //UE4/Dev-Blueprints to //UE4/Dev-Main (Source: //UE4/Dev-Blueprints @ 3235800)

diff --git a/Engine/Source/Editor/UnrealEd/Private/CookOnTheFlyServer.cpp b/Engine/Source/Editor/UnrealEd/Private/CookOnTheFlyServer.cpp
--- a/Engine/Source/Editor/UnrealEd/Private/CookOnTheFlyServer.cpp
+++ b/Engine/Source/Editor/UnrealEd/Private/CookOnTheFlyServer.cpp
@@ -5845,11 +5845,11 @@ void UCookOnTheFlyServer::StartCookByTheBook( const FCookByTheBookStartupOptions
 
 	GenerateAssetRegistry();
 
+	const UProjectPackagingSettings* const PackagingSettings = GetDefault<UProjectPackagingSettings>();
 
 	NeverCookPackageList.Empty();
 	{
 		const FString AbsoluteGameContentDir = FPaths::ConvertRelativePathToFull(FPaths::GameContentDir());
-		const UProjectPackagingSettings* const PackagingSettings = GetDefault<UProjectPackagingSettings>();
 
 		TArray<FString> NeverCookDirectories = CookByTheBookStartupOptions.NeverCookDirectories;
 
@@ -5923,14 +5923,29 @@ void UCookOnTheFlyServer::StartCookByTheBook( const FCookByTheBookStartupOptions
 	if (bRunConversion)
 	{
 		FNativeCodeGenInitData CodeGenData;
-		for (auto Entry : CookByTheBookStartupOptions.TargetPlatforms)
+		for (const ITargetPlatform* Entry : CookByTheBookStartupOptions.TargetPlatforms)
 		{
+			FPlatformNativizationDetails PlatformNativizationDetails;
+			PlatformNativizationDetails.PlatformName = Entry->PlatformName();
 			// If you change this target path you must also update logic in CookCommand.Automation.CS. Passing a single directory around is cumbersome for testing, so I have hard coded it.
-			CodeGenData.CodegenTargets.Push(TPairInitializer<FString, FString>(Entry->PlatformName(), FString(FPaths::Combine( *FPaths::GameIntermediateDir(), *(Entry->PlatformName())))));
+			PlatformNativizationDetails.PlatformTargetDirectory = FString(FPaths::Combine(*FPaths::GameIntermediateDir(), *(Entry->PlatformName())));
+			PlatformNativizationDetails.CompilerNativizationOptions.ClientOnlyPlatform = Entry->IsClientOnly();
+			PlatformNativizationDetails.CompilerNativizationOptions.ServerOnlyPlatform = Entry->IsServerOnly();
+
+			CodeGenData.CodegenTargets.Push(PlatformNativizationDetails);
 		}
 		CodeGenData.ManifestIdentifier = CookByTheBookStartupOptions.ChildCookIdentifier;
 		IBlueprintNativeCodeGenModule::InitializeModule(CodeGenData);
 	}
+	else if(PackagingSettings->bWarnIfPackagedWithoutNativizationFlag && PackagingSettings->BlueprintNativizationMethod != EProjectPackagingBlueprintNativizationMethod::Disabled)
+	{
+		// Warn if we're cooking without the -nativizeAssets flag, when the project settings specify a nativization method.
+		// If the "exclusive" (whitelist) method is set, we only warn if at least one asset has been selected for conversion.
+		if (PackagingSettings->BlueprintNativizationMethod != EProjectPackagingBlueprintNativizationMethod::Exclusive || PackagingSettings->NativizeBlueprintAssets.Num() > 0)
+		{
+			UE_LOG(LogCook, Warning, TEXT("Project is configured for Blueprint nativization, but the conversion flag (-nativizeAssets) has been omitted from the command line. No assets will be converted as a result."));
+		}
+	}
 
 	// need to test this out
 	/*if (RecompileChangedShaders(TargetPlatformNames))
diff --git a/Engine/Source/Editor/UnrealEd/Private/EditorEngine.cpp b/Engine/Source/Editor/UnrealEd/Private/EditorEngine.cpp
--- a/Engine/Source/Editor/UnrealEd/Private/EditorEngine.cpp
+++ b/Engine/Source/Editor/UnrealEd/Private/EditorEngine.cpp
@@ -1556,6 +1556,12 @@ void UEditorEngine::Tick( float DeltaSeconds, bool bIdleMode )
 			// Use the PlayWorld as the GWorld, because who knows what will happen in the Tick.
 			OldGWorld = SetPlayInEditorWorld( PlayWorld );
 
+			// Transfer debug references to ensure debugging ref's are valid for this tick in case of multiple game instances.
+			if (OldGWorld && OldGWorld != PlayWorld)
+			{
+				OldGWorld->TransferBlueprintDebugReferences(PlayWorld);
+			}
+
 			// Tick all travel and Pending NetGames (Seamless, server, client)
 			TickWorldTravel(PieContext, DeltaSeconds);
 
diff --git a/Engine/Source/Editor/UnrealEd/Private/Kismet2/BlueprintEditorUtils.cpp b/Engine/Source/Editor/UnrealEd/Private/Kismet2/BlueprintEditorUtils.cpp
--- a/Engine/Source/Editor/UnrealEd/Private/Kismet2/BlueprintEditorUtils.cpp
+++ b/Engine/Source/Editor/UnrealEd/Private/Kismet2/BlueprintEditorUtils.cpp
@@ -7088,7 +7088,10 @@ TSharedRef<SWidget> FBlueprintEditorUtils::ConstructBlueprintParentClassPicker(
 		bIsAnimBlueprint |= Blueprint->IsA(UAnimBlueprint::StaticClass());
 		bIsLevelScriptActor |= Blueprint->ParentClass->IsChildOf( ALevelScriptActor::StaticClass() );
 		bIsComponentBlueprint |= Blueprint->ParentClass->IsChildOf( UActorComponent::StaticClass() );
-		BlueprintClasses.Add(Blueprint->GeneratedClass);
+		if(Blueprint->GeneratedClass)
+		{
+			BlueprintClasses.Add(Blueprint->GeneratedClass);
+		}
 	}
 
 	// Fill in options
@@ -7152,7 +7155,10 @@ TSharedRef<SWidget> FBlueprintEditorUtils::ConstructBlueprintParentClassPicker(
 	for( const UBlueprint* Blueprint : Blueprints)
 	{
 		// don't allow making me my own parent!
-		Filter->DisallowedClasses.Add(Blueprint->GeneratedClass);
+		if(Blueprint->GeneratedClass)
+		{
+			Filter->DisallowedClasses.Add(Blueprint->GeneratedClass);
+		}
 	}
 
 	return FModuleManager::LoadModuleChecked<FClassViewerModule>("ClassViewer").CreateClassViewer(Options, OnPicked);
@@ -7237,7 +7243,10 @@ TSharedRef<SWidget> FBlueprintEditorUtils::ConstructBlueprintInterfaceClassPicke
 	TArray<UClass*> BlueprintClasses;
 	for (const UBlueprint* Blueprint : Blueprints)
 	{
-		BlueprintClasses.Add(Blueprint->GeneratedClass);
+		if(Blueprint->GeneratedClass)
+		{
+			BlueprintClasses.Add(Blueprint->GeneratedClass);
+		}
 	}
 
 	// Fill in options
@@ -7251,8 +7260,11 @@ TSharedRef<SWidget> FBlueprintEditorUtils::ConstructBlueprintInterfaceClassPicke
 	for (const UBlueprint* Blueprint : Blueprints)
 	{
 		// don't allow making me my own parent!
-		Filter->DisallowedClasses.Add(Blueprint->GeneratedClass);
-
+		if(Blueprint->GeneratedClass)
+		{
+			Filter->DisallowedClasses.Add(Blueprint->GeneratedClass);
+		}
+		
 		UClass const* const ParentClass = Blueprint->ParentClass;
 		// see if the parent class has any prohibited interfaces
 		if ((ParentClass != nullptr) && ParentClass->HasMetaData(FBlueprintMetadata::MD_ProhibitedInterfaces))
@@ -7514,12 +7526,12 @@ void FBlueprintEditorUtils::AnalyticsTrackNewNode( UEdGraphNode *NewNode )
 bool FBlueprintEditorUtils::IsObjectADebugCandidate( AActor* InActorObject, UBlueprint* InBlueprint, bool bInDisallowDerivedBlueprints )
 {
 	const bool bPassesFlags = !InActorObject->HasAnyFlags(RF_ClassDefaultObject) && !InActorObject->IsPendingKill();
-	bool bCanDebugThisObject;
+	bool bCanDebugThisObject = false;
 	if( bInDisallowDerivedBlueprints == true )
 	{
 		bCanDebugThisObject = InActorObject->GetClass()->ClassGeneratedBy == InBlueprint;
 	}
-	else
+	else if(InBlueprint->GeneratedClass)
 	{
 		bCanDebugThisObject = InActorObject->IsA( InBlueprint->GeneratedClass );
 	}
@@ -8325,12 +8337,11 @@ bool FBlueprintEditorUtils::HasGetTypeHash(const FEdGraphPinType& PinType)
 	}
 
 	const UScriptStruct* StructType = Cast<const UScriptStruct>(PinType.PinSubCategoryObject.Get());
-	return StructType && 
-		(	// we can always hash non-native structs (see UScriptStruct::GetStructTypeHash)
-			!StructType->IsNative() || 
-			// but native structs need to provide a GetTypeHash:
-			( StructType->GetCppStructOps() && StructType->GetCppStructOps()->HasGetTypeHash() )
-		); 
+	if( StructType )
+	{
+		return StructHasGetTypeHash(StructType);
+	}
+	return false;
 }
 
 bool FBlueprintEditorUtils::PropertyHasGetTypeHash(const UProperty* PropertyType)
@@ -8574,11 +8585,5 @@ void FBlueprintEditorUtils::BuildComponentInstancingData(UActorComponent* Compon
 	}
 }
 
-bool FBlueprintEditorUtils::ShouldEnableAdvancedContainers()
-{
-	UBlueprintEditorSettings const* Settings = GetDefault<UBlueprintEditorSettings>();
-	return Settings->bEnableAdvancedContainers;
-}
-
 #undef LOCTEXT_NAMESPACE
 
diff --git a/Engine/Source/Editor/UnrealEd/Private/Kismet2/StructureEditorUtils.cpp b/Engine/Source/Editor/UnrealEd/Private/Kismet2/StructureEditorUtils.cpp
--- a/Engine/Source/Editor/UnrealEd/Private/Kismet2/StructureEditorUtils.cpp
+++ b/Engine/Source/Editor/UnrealEd/Private/Kismet2/StructureEditorUtils.cpp
@@ -231,7 +231,11 @@ struct FMemberVariableNameHelper
 		FString Result;
 		if (!NameBase.IsEmpty())
 		{
-			if (ensure(FName::IsValidXName(NameBase, INVALID_OBJECTNAME_CHARACTERS)))
+			if (!FName::IsValidXName(NameBase, INVALID_OBJECTNAME_CHARACTERS))
+			{
+				Result = MakeObjectNameFromDisplayLabel(NameBase, NAME_None).GetPlainNameString();
+			}
+			else
 			{
 				Result = NameBase;
 			}
@@ -342,7 +346,6 @@ bool FStructureEditorUtils::RenameVariable(UUserDefinedStruct* Struct, FGuid Var
 		auto VarDesc = GetVarDescByGuid(Struct, VarGuid);
 		if (VarDesc 
 			&& !NewDisplayNameStr.IsEmpty()
-			&& FName::IsValidXName(NewDisplayNameStr, INVALID_OBJECTNAME_CHARACTERS) 
 			&& IsUniqueVariableDisplayName(Struct, NewDisplayNameStr))
 		{
 			const FScopedTransaction Transaction(LOCTEXT("RenameVariable", "Rename Variable"));
diff --git a/Engine/Source/Editor/UnrealEd/Private/Settings/SettingsClasses.cpp b/Engine/Source/Editor/UnrealEd/Private/Settings/SettingsClasses.cpp
--- a/Engine/Source/Editor/UnrealEd/Private/Settings/SettingsClasses.cpp
+++ b/Engine/Source/Editor/UnrealEd/Private/Settings/SettingsClasses.cpp
@@ -495,6 +495,31 @@ void ULevelEditorViewportSettings::PostEditChangeProperty( struct FPropertyChang
 UProjectPackagingSettings::UProjectPackagingSettings( const FObjectInitializer& ObjectInitializer )
 	: Super(ObjectInitializer)
 {
+	bWarnIfPackagedWithoutNativizationFlag = true;
+}
+
+
+void UProjectPackagingSettings::PostInitProperties()
+{
+	// Migrate from deprecated Blueprint nativization packaging flags.
+	// Note: This assumes that LoadConfig() has been called before getting here.
+	FString NewValue;
+	const FString ConfigFileName = UProjectPackagingSettings::StaticClass()->GetConfigName();
+	const FString ClassSectionName = UProjectPackagingSettings::StaticClass()->GetPathName();
+	const bool bIgnoreOldFlags = GConfig->GetString(*ClassSectionName, GET_MEMBER_NAME_STRING_CHECKED(UProjectPackagingSettings, BlueprintNativizationMethod), NewValue, ConfigFileName);
+	if (!bIgnoreOldFlags && bNativizeBlueprintAssets_DEPRECATED)
+	{
+		BlueprintNativizationMethod = bNativizeOnlySelectedBlueprints_DEPRECATED ? EProjectPackagingBlueprintNativizationMethod::Exclusive : EProjectPackagingBlueprintNativizationMethod::Inclusive;
+	}
+
+	// Reset deprecated settings to defaults.
+	bNativizeBlueprintAssets_DEPRECATED = false;
+	bNativizeOnlySelectedBlueprints_DEPRECATED = false;
+
+	// Cache the current set of Blueprint assets selected for nativization.
+	CachedNativizeBlueprintAssets = NativizeBlueprintAssets;
+
+	Super::PostInitProperties();
 }
 
 
@@ -588,6 +613,107 @@ void UProjectPackagingSettings::PostEditChangeProperty( FPropertyChangedEvent& P
 			}
 		}
 	}
+	else if (Name == FName((TEXT("NativizeBlueprintAssets"))))
+	{
+		int32 AssetIndex;
+		auto OnSelectBlueprintForExclusiveNativizationLambda = [](const FString& PackageName, bool bSelect)
+		{
+			if (!PackageName.IsEmpty())
+			{
+				// This should only apply to loaded packages. Any unloaded packages defer setting the transient flag to when they're loaded.
+				if (UPackage* Package = FindPackage(nullptr, *PackageName))
+				{
+					// Find the Blueprint asset within the package.
+					if (UBlueprint* Blueprint = FindObject<UBlueprint>(Package, *FPaths::GetBaseFilename(PackageName)))
+					{
+						// We're toggling the transient flag on or off.
+						if (Blueprint->bSelectedForNativization != bSelect)
+						{
+							// Get a reference to the property.
+							UProperty* PropertyToChange = UBlueprint::StaticClass()->FindPropertyByName(GET_MEMBER_NAME_CHECKED(UBlueprint, bSelectedForNativization));
+							check(PropertyToChange != nullptr);
+
+							// Update the transient property's value.
+							Blueprint->PreEditChange(PropertyToChange);
+							Blueprint->bSelectedForNativization = bSelect;
+
+							// If selecting for nativization, this will propagate the transient flag to any dependencies. Note that we don't reverse-propagate here
+							// back to the config settings until/unless Blueprint assets are saved, so this will only set the transient flag and mark dependencies dirty.
+							FPropertyChangedEvent BlueprintPropertyChangedEvent(PropertyToChange);
+							Blueprint->PostEditChangeProperty(BlueprintPropertyChangedEvent);
+						}
+					}
+				}
+			}
+		};
+
+		if (NativizeBlueprintAssets.Num() > 0)
+		{
+			for (AssetIndex = 0; AssetIndex < NativizeBlueprintAssets.Num(); ++AssetIndex)
+			{
+				const FString& PackageName = NativizeBlueprintAssets[AssetIndex].FilePath;
+				if (AssetIndex >= CachedNativizeBlueprintAssets.Num())
+				{
+					// A new entry was added; toggle the exclusive flag on the corresponding Blueprint asset (if loaded).
+					OnSelectBlueprintForExclusiveNativizationLambda(PackageName, true);
+
+					// Add an entry to the end of the cached list.
+					CachedNativizeBlueprintAssets.Add(NativizeBlueprintAssets[AssetIndex]);
+				}
+				else if (!PackageName.Equals(CachedNativizeBlueprintAssets[AssetIndex].FilePath))
+				{
+					if (NativizeBlueprintAssets.Num() < CachedNativizeBlueprintAssets.Num())
+					{
+						// An entry was removed; toggle the exclusive flag on the corresponding Blueprint asset (if loaded).
+						OnSelectBlueprintForExclusiveNativizationLambda(CachedNativizeBlueprintAssets[AssetIndex].FilePath, false);
+
+						// Remove this entry from the cached list.
+						CachedNativizeBlueprintAssets.RemoveAt(AssetIndex);
+					}
+					else if(NativizeBlueprintAssets.Num() > CachedNativizeBlueprintAssets.Num())
+					{
+						// A new entry was inserted; toggle the exclusive flag on the corresponding Blueprint asset (if loaded).
+						OnSelectBlueprintForExclusiveNativizationLambda(PackageName, true);
+
+						// Insert the new entry into the cached list.
+						CachedNativizeBlueprintAssets.Insert(NativizeBlueprintAssets[AssetIndex], AssetIndex);
+					}
+					else
+					{
+						// An entry was changed; toggle the exclusive flag on the corresponding Blueprint assets (if loaded).
+						OnSelectBlueprintForExclusiveNativizationLambda(CachedNativizeBlueprintAssets[AssetIndex].FilePath, false);
+						OnSelectBlueprintForExclusiveNativizationLambda(PackageName, true);
+
+						// Update the cached entry.
+						CachedNativizeBlueprintAssets[AssetIndex].FilePath = PackageName;
+					}
+				}
+			}
+
+			if (CachedNativizeBlueprintAssets.Num() > NativizeBlueprintAssets.Num())
+			{
+				// Removed entries at the end of the list; toggle the exclusive flag on the corresponding Blueprint asset(s) (if loaded).
+				for (AssetIndex = NativizeBlueprintAssets.Num(); AssetIndex < CachedNativizeBlueprintAssets.Num(); ++AssetIndex)
+				{
+					OnSelectBlueprintForExclusiveNativizationLambda(CachedNativizeBlueprintAssets[AssetIndex].FilePath, false);
+				}
+
+				// Remove entries from the end of the cached list.
+				CachedNativizeBlueprintAssets.RemoveAt(NativizeBlueprintAssets.Num(), CachedNativizeBlueprintAssets.Num() - NativizeBlueprintAssets.Num());
+			}
+		}
+		else if(CachedNativizeBlueprintAssets.Num() > 0)
+		{
+			// Removed all entries; toggle the exclusive flag on the corresponding Blueprint asset(s) (if loaded).
+			for (AssetIndex = 0; AssetIndex < CachedNativizeBlueprintAssets.Num(); ++AssetIndex)
+			{
+				OnSelectBlueprintForExclusiveNativizationLambda(CachedNativizeBlueprintAssets[AssetIndex].FilePath, false);
+			}
+
+			// Clear the cached list.
+			CachedNativizeBlueprintAssets.Empty();
+		}
+	}
 }
 
 bool UProjectPackagingSettings::CanEditChange( const UProperty* InProperty ) const
@@ -596,10 +722,73 @@ bool UProjectPackagingSettings::CanEditChange( const UProperty* InProperty ) con
 	{
 		return false;
 	}
+	else if (InProperty->GetFName() == FName(TEXT("NativizeBlueprintAssets")))
+	{
+		return BlueprintNativizationMethod == EProjectPackagingBlueprintNativizationMethod::Exclusive;
+	}
 
 	return Super::CanEditChange(InProperty);
 }
 
+bool UProjectPackagingSettings::AddBlueprintAssetToNativizationList(const class UBlueprint* InBlueprint)
+{
+	if (InBlueprint)
+	{
+		const FString PackageName = InBlueprint->GetOutermost()->GetName();
+
+		// Make sure it's not already in the exclusive list. This can happen if the user previously added this asset in the Project Settings editor.
+		bool bFound = false;
+		for (int AssetIndex = 0; AssetIndex < NativizeBlueprintAssets.Num(); ++AssetIndex)
+		{
+			if (NativizeBlueprintAssets[AssetIndex].FilePath.Equals(PackageName, ESearchCase::IgnoreCase))
+			{
+				bFound = true;
+				break;
+			}
+		}
+
+		if (!bFound)
+		{
+			// Add this Blueprint asset to the exclusive list.
+			FFilePath FileInfo;
+			FileInfo.FilePath = PackageName;
+			NativizeBlueprintAssets.Add(FileInfo);
+
+			// Also add it to the mirrored list for tracking edits.
+			CachedNativizeBlueprintAssets.Add(FileInfo);
+
+			return true;
+		}
+	}
+
+	return false;
+}
+
+bool UProjectPackagingSettings::RemoveBlueprintAssetFromNativizationList(const class UBlueprint* InBlueprint)
+{
+	if (InBlueprint)
+	{
+		const FString PackageName = InBlueprint->GetOutermost()->GetName();
+
+		// Remove this Blueprint asset from the exclusive list (if found).
+		for (int AssetIndex = 0; AssetIndex < NativizeBlueprintAssets.Num(); ++AssetIndex)
+		{
+			if (NativizeBlueprintAssets[AssetIndex].FilePath.Equals(PackageName, ESearchCase::IgnoreCase))
+			{
+				// Note: Intentionally not using RemoveAtSwap() here, so that the order is preserved.
+				NativizeBlueprintAssets.RemoveAt(AssetIndex);
+
+				// Also remove it from the mirrored list (for tracking edits).
+				CachedNativizeBlueprintAssets.RemoveAt(AssetIndex);
+				
+				return true;
+			}
+		}
+	}
+
+	return false;
+}
+
 /* UCrashReporterSettings interface
 *****************************************************************************/
 UCrashReporterSettings::UCrashReporterSettings(const FObjectInitializer& ObjectInitializer)
diff --git a/Engine/Source/Editor/UnrealEd/Public/Kismet2/BlueprintEditorUtils.h b/Engine/Source/Editor/UnrealEd/Public/Kismet2/BlueprintEditorUtils.h
--- a/Engine/Source/Editor/UnrealEd/Public/Kismet2/BlueprintEditorUtils.h
+++ b/Engine/Source/Editor/UnrealEd/Public/Kismet2/BlueprintEditorUtils.h
@@ -1361,9 +1361,6 @@ public:
 	 */
 	static void BuildComponentInstancingData(UActorComponent* ComponentTemplate, FBlueprintCookedComponentInstancingData& OutData);
 
-	/** Switch for hiding TMap/TSet support until stablized */
-	static bool ShouldEnableAdvancedContainers();
-
 protected:
 	// Removes all NULL graph references from the SubGraphs array and recurses thru the non-NULL ones
 	static void CleanNullGraphReferencesRecursive(UEdGraph* Graph);
diff --git a/Engine/Source/Programs/AutomationTool/Scripts/BuildCookRun.Automation.cs b/Engine/Source/Programs/AutomationTool/Scripts/BuildCookRun.Automation.cs
--- a/Engine/Source/Programs/AutomationTool/Scripts/BuildCookRun.Automation.cs
+++ b/Engine/Source/Programs/AutomationTool/Scripts/BuildCookRun.Automation.cs
@@ -209,10 +209,33 @@ public class BuildCookRun : BuildCommand
 
         Project.Build(this, Params, WorkingCL, bGenerateNativeScripts ? (ProjectBuildTargets.All & ~ClientTargets) : ProjectBuildTargets.All);
 		Project.Cook(Params);
-		if (bGenerateNativeScripts)
-		{
+        if (bGenerateNativeScripts)
+        {
             Project.Build(this, Params, WorkingCL, ClientTargets);
-		}
+        }
+        else
+        {
+            ConfigHierarchy GameIni = ConfigCache.ReadHierarchy(ConfigHierarchyType.Game, Params.RawProjectPath.Directory, HostPlatform.Current.HostEditorPlatform);
+            if (GameIni != null)
+            {
+                List<string> NativizeBlueprintAssets = null;
+                string BlueprintNativizationMethod = "Disabled";
+                bool bWarnIfPackagedWithoutNativizationFlag = true;
+                GameIni.GetString("/Script/UnrealEd.ProjectPackagingSettings", "BlueprintNativizationMethod", out BlueprintNativizationMethod);
+                GameIni.GetBool("/Script/UnrealEd.ProjectPackagingSettings", "bWarnIfPackagedWithoutNativizationFlag", out bWarnIfPackagedWithoutNativizationFlag);
+                GameIni.GetArray("/Script/UnrealEd.ProjectPackagingSettings", "NativizeBlueprintAssets", out NativizeBlueprintAssets);
+
+                if (bWarnIfPackagedWithoutNativizationFlag && BlueprintNativizationMethod != "Disabled")
+                {
+                    // Warn if we're cooking without the -nativizeAssets flag, when the project settings specify a nativization method.
+                    // If the "exclusive" (whitelist) method is set, we only warn if at least one asset has been selected for conversion.
+                    if (BlueprintNativizationMethod != "Exclusive" || (NativizeBlueprintAssets != null && NativizeBlueprintAssets.Count > 0))
+                    {
+                        LogWarning("Project is configured for Blueprint nativization, but the conversion flag (-nativizeAssets) was omitted from the command line. No nativized assets have been built as a result.");
+                    }
+                }
+            }
+        }
 		Project.CopyBuildToStagingDirectory(Params);
 		Project.Package(Params, WorkingCL);
 		Project.Archive(Params);
diff --git a/Engine/Source/Programs/UnrealHeaderTool/Private/CodeGenerator.cpp b/Engine/Source/Programs/UnrealHeaderTool/Private/CodeGenerator.cpp
--- a/Engine/Source/Programs/UnrealHeaderTool/Private/CodeGenerator.cpp
+++ b/Engine/Source/Programs/UnrealHeaderTool/Private/CodeGenerator.cpp
@@ -3277,6 +3277,12 @@ void FNativeClassHeaderGenerator::ExportGeneratedEnumsInitCode(const TArray<UEnu
 			GeneratedEnumRegisterFunctionText.Logf(TEXT("\t\t\tReturnEnum->SetEnums(EnumNames, %s%s);\r\n"), *EnumTypeStr, *ParamAddMaxKeyIfMissing);
 			GeneratedEnumRegisterFunctionText.Logf(TEXT("\t\t\tReturnEnum->CppType = TEXT(\"%s\");\r\n"), *Enum->CppType);
 
+			const FString& EnumDisplayNameFn = Enum->GetMetaData(TEXT("EnumDisplayNameFn"));
+			if( !EnumDisplayNameFn.IsEmpty() )
+			{
+				GeneratedEnumRegisterFunctionText.Logf(TEXT("\t\t\tReturnEnum->SetEnumDisplayNameFn(&%s);\r\n"), *EnumDisplayNameFn);
+			}
+
 			FString Meta = GetMetaDataCodeForObject(Enum, TEXT("ReturnEnum"), TEXT("\t\t\t"));
 			if (Meta.Len())
 			{
diff --git a/Engine/Source/Programs/UnrealHeaderTool/Private/HeaderParser.cpp b/Engine/Source/Programs/UnrealHeaderTool/Private/HeaderParser.cpp
--- a/Engine/Source/Programs/UnrealHeaderTool/Private/HeaderParser.cpp
+++ b/Engine/Source/Programs/UnrealHeaderTool/Private/HeaderParser.cpp
@@ -3606,11 +3606,6 @@ FIndexRange*                    ParsedVarIndexRange
 			FError::Throwf(TEXT("Nested containers are not supported.") );
 		}
 
-		if (MapKeyType.Type == CPT_Struct)
-		{
-			FError::Throwf(TEXT("USTRUCTs are not currently supported as key types."));
-		}
-
 		if (MapKeyType.Type == CPT_Interface)
 		{
 			FError::Throwf(TEXT("UINTERFACEs are not currently supported as key types."));
@@ -3680,11 +3675,6 @@ FIndexRange*                    ParsedVarIndexRange
 			FError::Throwf(TEXT("Nested containers are not supported.") );
 		}
 
-		if (VarProperty.Type == CPT_Struct)
-		{
-			FError::Throwf(TEXT("USTRUCTs are not currently supported as element types."));
-		}
-
 		if (VarProperty.Type == CPT_Interface)
 		{
 			FError::Throwf(TEXT("UINTERFACEs are not currently supported as element types."));
@@ -6013,7 +6003,17 @@ void FHeaderParser::CompileFunctionDeclaration(FClasses& AllClasses)
 	{ 
 		if (!bHasMenuCategory && !bInternalOnly && !bDeprecated) 
 		{ 
-			FError::Throwf(TEXT("Blueprint accessible functions must have a category specified")); 
+			const bool bModuleIsGame = CurrentlyParsedModule && (
+				CurrentlyParsedModule->ModuleType == EBuildModuleType::GameDeveloper ||
+				CurrentlyParsedModule->ModuleType == EBuildModuleType::GameEditor ||
+				CurrentlyParsedModule->ModuleType == EBuildModuleType::GameRuntime ||
+				CurrentlyParsedModule->ModuleType == EBuildModuleType::GameThirdParty);
+
+			// To allow for quick iteration, don't enforce the requirement that game functions have to be categorized
+			if (!bModuleIsGame)
+			{
+				FError::Throwf(TEXT("Blueprint accessible functions must have a category specified"));
+			}
 		} 
 	}
 
diff --git a/Engine/Source/Runtime/Core/Public/Templates/HasGetTypeHash.h b/Engine/Source/Runtime/Core/Public/Templates/HasGetTypeHash.h
--- a/Engine/Source/Runtime/Core/Public/Templates/HasGetTypeHash.h
+++ b/Engine/Source/Runtime/Core/Public/Templates/HasGetTypeHash.h
@@ -3,6 +3,10 @@
 #pragma once
 
 #include "CoreTypes.h"
+#include "Templates/AndOrNot.h"
+#include "Templates/IsArithmetic.h"
+#include "Templates/IsPointer.h"
+#include "Templates/IsEnum.h"
 
 namespace UE4GetTypeHashExists_Private
 {
@@ -26,8 +30,15 @@ namespace UE4GetTypeHashExists_Private
 	template <typename T>
 	const T& Make();
 
-	template <typename T>
+	template <typename T, bool bIsHashableScalarType = TOr<TIsArithmetic<T>, TIsPointer<T>, TIsEnum<T>>::Value>
 	struct GetTypeHashQuery
+	{
+		// All arithmetic, pointer and enums types are hashable
+		enum { Value = true };
+	};
+
+	template <typename T>
+	struct GetTypeHashQuery<T, false>
 	{
 		enum { Value = sizeof(FReturnValueCheck<decltype(GetTypeHash(Make<T>()))>::Func()) == sizeof(char[2]) };
 	};
diff --git a/Engine/Source/Runtime/Core/Public/Templates/IsEnum.h b/Engine/Source/Runtime/Core/Public/Templates/IsEnum.h
new file mode 100644
--- /dev/null
+++ b/Engine/Source/Runtime/Core/Public/Templates/IsEnum.h
@@ -0,0 +1,9 @@
+// Copyright 1998-2016 Epic Games, Inc. All Rights Reserved.
+
+#pragma once
+
+template <typename T>
+struct TIsEnum
+{
+	enum { Value = __is_enum(T) };
+};
diff --git a/Engine/Source/Runtime/Core/Public/Templates/TypeHash.h b/Engine/Source/Runtime/Core/Public/Templates/TypeHash.h
--- a/Engine/Source/Runtime/Core/Public/Templates/TypeHash.h
+++ b/Engine/Source/Runtime/Core/Public/Templates/TypeHash.h
@@ -3,7 +3,7 @@
 #pragma once
 
 #include "CoreTypes.h"
-#include "Templates/UnrealTypeTraits.h"
+#include "Templates/IsEnum.h"
 #include "Misc/Crc.h"
 
 /**
diff --git a/Engine/Source/Runtime/Core/Public/Templates/UnrealTypeTraits.h b/Engine/Source/Runtime/Core/Public/Templates/UnrealTypeTraits.h
--- a/Engine/Source/Runtime/Core/Public/Templates/UnrealTypeTraits.h
+++ b/Engine/Source/Runtime/Core/Public/Templates/UnrealTypeTraits.h
@@ -14,6 +14,7 @@
 #include "Templates/AndOrNot.h"
 #include "Templates/AreTypesEqual.h"
 #include "Templates/IsArithmetic.h"
+#include "Templates/IsEnum.h"
 #include "Templates/RemoveCV.h"
 
 #include "Templates/IsPODType.h"
@@ -25,7 +26,6 @@
  -----------------------------------------------------------------------------*/
 #define HAS_TRIVIAL_CONSTRUCTOR(T) __has_trivial_constructor(T)
 #define IS_POD(T) __is_pod(T)
-#define IS_ENUM(T) __is_enum(T)
 #define IS_EMPTY(T) __is_empty(T)
 
 
@@ -173,7 +173,7 @@ struct TIsFunction<RetType(Params...)>
 template<typename T> 
 struct TIsZeroConstructType 
 { 
-	enum { Value = TOrValue<IS_ENUM(T), TIsArithmetic<T>, TIsPointer<T>>::Value };
+	enum { Value = TOr<TIsEnum<T>, TIsArithmetic<T>, TIsPointer<T>>::Value };
 };
 
 /**
@@ -350,14 +350,15 @@ public:
  * Helper for array traits. Provides a common base to more easily refine a portion of the traits
  * when specializing. Mainly used by MemoryOps.h which is used by the contiguous storage containers like TArray.
  */
-template<typename T> struct TTypeTraitsBase
+template<typename T>
+struct TTypeTraitsBase
 {
 	typedef typename TCallTraits<T>::ParamType ConstInitType;
 	typedef typename TCallTraits<T>::ConstPointerType ConstPointerType;
 
 	// There's no good way of detecting this so we'll just assume it to be true for certain known types and expect
 	// users to customize it for their custom types.
-	enum { IsBytewiseComparable = TOrValue<IS_ENUM(T), TIsArithmetic<T>, TIsPointer<T>>::Value };
+	enum { IsBytewiseComparable = TOr<TIsEnum<T>, TIsArithmetic<T>, TIsPointer<T>>::Value };
 };
 
 /**
@@ -497,16 +498,9 @@ public:																											\
 	enum { Value = sizeof(MemberTest<T>(nullptr)) == sizeof(char) };											\
 };
 
-template <typename T>
-struct TIsEnum
-{
-	enum { Value = IS_ENUM(T) };
-};
-
 /*-----------------------------------------------------------------------------
  * Undef Macros abstracting the presence of certain compiler intrinsic type traits
  -----------------------------------------------------------------------------*/
 #undef IS_EMPTY
-#undef IS_ENUM
 #undef IS_POD
 #undef HAS_TRIVIAL_CONSTRUCTOR
diff --git a/Engine/Source/Runtime/Core/Public/UObject/BlueprintsObjectVersion.h b/Engine/Source/Runtime/Core/Public/UObject/BlueprintsObjectVersion.h
--- a/Engine/Source/Runtime/Core/Public/UObject/BlueprintsObjectVersion.h
+++ b/Engine/Source/Runtime/Core/Public/UObject/BlueprintsObjectVersion.h
@@ -18,6 +18,7 @@ struct CORE_API FBlueprintsObjectVersion
 		AllowDeletionConformed,
 		AdvancedContainerSupport,
 		SCSHasComponentTemplateClass,
+		ComponentTemplateClassSupport,
 
 		// -----<new versions can be added above this line>-------------------------------------------------
 		VersionPlusOne,
diff --git a/Engine/Source/Runtime/Core/Public/UObject/PropertyPortFlags.h b/Engine/Source/Runtime/Core/Public/UObject/PropertyPortFlags.h
--- a/Engine/Source/Runtime/Core/Public/UObject/PropertyPortFlags.h
+++ b/Engine/Source/Runtime/Core/Public/UObject/PropertyPortFlags.h
@@ -98,4 +98,7 @@ enum EPropertyPortFlags
 
 	/** Export in C++ form */
 	PPF_ExportCpp					= 0x10000000,
-};
+
+	/** Ignores CPF_SkipSerialization flag when using tagged serialization */
+	PPF_ForceTaggedSerialization	= 0x20000000,
+};
\ No newline at end of file
diff --git a/Engine/Source/Runtime/CoreUObject/Private/Blueprint/BlueprintSupport.cpp b/Engine/Source/Runtime/CoreUObject/Private/Blueprint/BlueprintSupport.cpp
--- a/Engine/Source/Runtime/CoreUObject/Private/Blueprint/BlueprintSupport.cpp
+++ b/Engine/Source/Runtime/CoreUObject/Private/Blueprint/BlueprintSupport.cpp
@@ -1760,8 +1760,10 @@ void FLinkerLoad::CreateDynamicTypeLoader()
 			{
 				// The class object is created here. The class is not fully constructed yet (no CLASS_Constructed flag), ZConstructor will do that later.
 				// The class object is needed to resolve circular dependencies. Regular native classes use deferred initialization/registration to avoid them.
-				ObjectImport->XObject = ClassConstructFn->StaticClassFn();
-				OuterImport->XObject = ObjectImport->XObject->GetOutermost();
+
+				ClassConstructFn->StaticClassFn();
+
+				//We don't fill the ObjectImport->XObject and OuterImport->XObject, because the class still must be created as export.
 			}
 		}
 	}
@@ -1890,10 +1892,8 @@ void FLinkerLoad::CreateDynamicTypeLoader()
 		}
 	}
 	}
-	else
-	{
+
 	LinkerRoot->SetPackageFlags(LinkerRoot->GetPackageFlags() | PKG_CompiledIn);
-	}
 }
 
 /*******************************************************************************
@@ -2242,13 +2242,18 @@ FConvertedBlueprintsDependencies& FConvertedBlueprintsDependencies::Get()
 	return ConvertedBlueprintsDependencies;
 }
 
-void FConvertedBlueprintsDependencies::RegisterClass(FName PackageName, GetDependenciesNamesFunc GetAssets)
+void FConvertedBlueprintsDependencies::RegisterConvertedClass(FName PackageName, GetDependenciesNamesFunc GetAssets)
 {
 	check(!PackageNameToGetter.Contains(PackageName));
 	ensure(GetAssets);
 	PackageNameToGetter.Add(PackageName, GetAssets);
 }
 
+static bool IsBlueprintDependencyDataNull(const FBlueprintDependencyData& Dependency)
+{
+	return Dependency.ObjectRef.ObjectName == NAME_None;
+}
+
 void FConvertedBlueprintsDependencies::GetAssets(FName PackageName, TArray<FBlueprintDependencyData>& OutDependencies) const
 {
 	auto FuncPtr = PackageNameToGetter.Find(PackageName);
@@ -2257,6 +2262,7 @@ void FConvertedBlueprintsDependencies::GetAssets(FName PackageName, TArray<FBlue
 	if (Func)
 	{
 		Func(OutDependencies);
+		OutDependencies.RemoveAll(IsBlueprintDependencyDataNull);
 	}
 }
 
@@ -2276,12 +2282,19 @@ void FConvertedBlueprintsDependencies::FillUsedAssetsInDynamicClass(UDynamicClas
 			int32 ImportIndex = 0;
 			for (FBlueprintDependencyData& ItData : UsedAssetdData)
 			{
-				FObjectImport& Import = Linker->Imp(FPackageIndex::FromImport(ImportIndex));
-				check(Import.ObjectName == ItData.ObjectRef.ObjectName);
-				UObject* TheAsset = Import.XObject;
-				UE_CLOG(!TheAsset, LogBlueprintSupport, Error, TEXT("Could not find UDynamicClass dependent asset (EDL) %s in %s"), *ItData.ObjectRef.ObjectName.ToString(), *ItData.ObjectRef.PackageName.ToString());
-				DynamicClass->UsedAssets.Add(TheAsset);
-				ImportIndex += 2;
+				if (!IsBlueprintDependencyDataNull(ItData))
+				{
+					FObjectImport& Import = Linker->Imp(FPackageIndex::FromImport(ImportIndex));
+					check(Import.ObjectName == ItData.ObjectRef.ObjectName);
+					UObject* TheAsset = Import.XObject;
+					UE_CLOG(!TheAsset, LogBlueprintSupport, Error, TEXT("Could not find UDynamicClass dependent asset (EDL) %s in %s"), *ItData.ObjectRef.ObjectName.ToString(), *ItData.ObjectRef.PackageName.ToString());
+					DynamicClass->UsedAssets.Add(TheAsset);
+					ImportIndex += 2;
+				}
+				else
+				{
+					DynamicClass->UsedAssets.Add(nullptr);
+				}
 			}
 			return;
 		}
@@ -2290,10 +2303,17 @@ void FConvertedBlueprintsDependencies::FillUsedAssetsInDynamicClass(UDynamicClas
 
 	for (FBlueprintDependencyData& ItData : UsedAssetdData)
 	{
-		const FString PathToObj = FString::Printf(TEXT("%s.%s"), *ItData.ObjectRef.PackageName.ToString(), *ItData.ObjectRef.ObjectName.ToString());
-		UObject* TheAsset = LoadObject<UObject>(nullptr, *PathToObj);
-		UE_CLOG(!TheAsset, LogBlueprintSupport, Error, TEXT("Could not find UDynamicClass dependent asset (non-EDL) %s in %s"), *ItData.ObjectRef.ObjectName.ToString(), *ItData.ObjectRef.PackageName.ToString());
-		DynamicClass->UsedAssets.Add(TheAsset);
+		if (ItData.ObjectRef.ObjectName != NAME_None)
+		{
+			const FString PathToObj = FString::Printf(TEXT("%s.%s"), *ItData.ObjectRef.PackageName.ToString(), *ItData.ObjectRef.ObjectName.ToString());
+			UObject* TheAsset = LoadObject<UObject>(nullptr, *PathToObj);
+			UE_CLOG(!TheAsset, LogBlueprintSupport, Error, TEXT("Could not find UDynamicClass dependent asset (non-EDL) %s in %s"), *ItData.ObjectRef.ObjectName.ToString(), *ItData.ObjectRef.PackageName.ToString());
+			DynamicClass->UsedAssets.Add(TheAsset);
+		}
+		else
+		{
+			DynamicClass->UsedAssets.Add(nullptr);
+		}
 	}
 }
 
diff --git a/Engine/Source/Runtime/CoreUObject/Private/Serialization/AsyncLoading.cpp b/Engine/Source/Runtime/CoreUObject/Private/Serialization/AsyncLoading.cpp
--- a/Engine/Source/Runtime/CoreUObject/Private/Serialization/AsyncLoading.cpp
+++ b/Engine/Source/Runtime/CoreUObject/Private/Serialization/AsyncLoading.cpp
@@ -1636,17 +1636,18 @@ EAsyncPackageState::Type FAsyncPackage::LoadImports_Event()
 			ExistingPackage = CastChecked<UPackage>(Import->XObject);
 			PendingPackage = ExistingPackage->LinkerLoad ? ExistingPackage->LinkerLoad->AsyncRoot : nullptr;
 		}
-		bool bCompiledIn = ExistingPackage && ExistingPackage->HasAnyPackageFlags(PKG_CompiledIn);
+		const bool bDynamicPackage = ExistingPackage && GetConvertedDynamicPackageNameToTypeName().Contains(ExistingPackage->GetFName());
+		const bool bCompiledInNotDynamic = ExistingPackage && ExistingPackage->HasAnyPackageFlags(PKG_CompiledIn) && !bDynamicPackage;
 		// Our import package name is the import name
 		const FName ImportPackageFName(Import->ObjectName);
-		check(!PendingPackage || !bCompiledIn); // we should never have a pending package for something that is compiled in
-		if (!PendingPackage && !bCompiledIn)
+		check(!PendingPackage || !bCompiledInNotDynamic); // we should never have a pending package for something that is compiled in
+		if (!PendingPackage && !bCompiledInNotDynamic)
 		{
 			PendingPackage = FAsyncLoadingThread::Get().FindAsyncPackage(ImportPackageFName);
 		}
 		if (!PendingPackage)
 		{
-			if (bCompiledIn)
+			if (bCompiledInNotDynamic)
 			{
 				// This can happen with editor only classes, not sure if this should be a warning or a silent continue
 				UE_LOG(LogStreaming, Warning, TEXT("FAsyncPackage::LoadImports for %s: Skipping import %s, depends on missing native class"), *Desc.NameToLoad.ToString(), *OriginalImport->ObjectName.ToString());
@@ -1759,17 +1760,26 @@ static bool IsFullyLoadedObj(UObject* Obj)
 	}
 //native blueprint 
 	UDynamicClass* UD = Cast<UDynamicClass>(Obj);
-	if (!UD || 0 != (UD->ClassFlags & CLASS_Constructed))
+	if (!UD)
 	{
 		return true;
 	}
-	/*
-	if (UD->GetDefaultObject(false))
+
+	if (GEventDrivenLoaderEnabled)
 	{
-		UE_CLOG(!UD->HasAnyClassFlags(CLASS_TokenStreamAssembled), LogStreaming, Fatal, TEXT("Class %s is fully loaded, but does not have its token stream assembled."), *UD->GetFullName());
-		return true;
+		if (0 != (UD->ClassFlags & CLASS_Constructed))
+		{
+			return true;
+		}
+	}
+	else
+	{
+		if (UD->GetDefaultObject(false))
+		{
+			UE_CLOG(!UD->HasAnyClassFlags(CLASS_TokenStreamAssembled), LogStreaming, Fatal, TEXT("Class %s is fully loaded, but does not have its token stream assembled."), *UD->GetFullName());
+			return true;
+		}
 	}
-	*/
 	return false;
 }
 
@@ -1815,7 +1825,6 @@ EAsyncPackageState::Type FAsyncPackage::SetupImports_Event()
 			if (ImportPackage)
 			{
 				FLinkerLoad* ImportLinker = ImportPackage->LinkerLoad;
-				check(!ImportLinker || !ImportPackage->HasAnyPackageFlags(PKG_CompiledIn)); // we should never have a linker on a compiled in package
 				if (ImportLinker && ImportLinker->AsyncRoot)
 				{
 					check(ImportLinker->AsyncRoot != this);
@@ -1859,7 +1868,6 @@ EAsyncPackageState::Type FAsyncPackage::SetupImports_Event()
 			if (ImportPackage)
 			{
 				FLinkerLoad* ImportLinker = ImportPackage->LinkerLoad;
-				check(!ImportLinker || !ImportPackage->HasAnyPackageFlags(PKG_CompiledIn)); // we should never have a linker on a compiled in package
 				bool bDynamicImport = ImportLinker && ImportLinker->bDynamicClassLinker;
 				if (!ImportLinker || !ImportLinker->AsyncRoot)
 				{
@@ -3405,7 +3413,7 @@ void FAsyncPackage::Event_StartPostload()
 			FObjectExport& Export = Linker->ExportMap[LocalExportIndex];
 			UObject* Object = Export.Object;
 			checkSlow(!(Object && !ReferencedObjects.Contains(Object)));
-			if (Object && (Object->HasAnyFlags(RF_NeedPostLoad) || Object->HasAnyInternalFlags(EInternalObjectFlags::AsyncLoading)))
+			if (Object && (Object->HasAnyFlags(RF_NeedPostLoad) || Linker->bDynamicClassLinker || Object->HasAnyInternalFlags(EInternalObjectFlags::AsyncLoading)))
 			{
 				ObjLoaded.Add(Object);
 			}
diff --git a/Engine/Source/Runtime/CoreUObject/Private/UObject/Class.cpp b/Engine/Source/Runtime/CoreUObject/Private/UObject/Class.cpp
--- a/Engine/Source/Runtime/CoreUObject/Private/UObject/Class.cpp
+++ b/Engine/Source/Runtime/CoreUObject/Private/UObject/Class.cpp
@@ -33,6 +33,7 @@
 #include "UObject/LinkerPlaceholderFunction.h"
 #include "UObject/StructScriptLoader.h"
 #include "UObject/PropertyHelper.h"
+#include "Serialization/ArchiveScriptReferenceCollector.h"
 
 // This flag enables some expensive class tree validation that is meant to catch mutations of 
 // the class tree outside of SetSuperStruct. It has been disabled because loading blueprints 
@@ -1335,6 +1336,18 @@ void UStruct::AddReferencedObjects(UObject* InThis, FReferenceCollector& Collect
 		// Required by the unified GC when running in the editor
 		Collector.AddReferencedObject( This->SuperStruct, This );
 		Collector.AddReferencedObject( This->Children, This );
+
+		TArray<UObject*> ScriptObjectReferences;
+		FArchiveScriptReferenceCollector ObjectReferenceCollector( ScriptObjectReferences );
+		int32 iCode = 0;
+		while( iCode < This->Script.Num() )
+		{	
+			This->SerializeExpr( iCode, ObjectReferenceCollector );
+		}
+		for( int32 Index = 0; Index < ScriptObjectReferences.Num(); Index++ )
+		{
+			Collector.AddReferencedObject( ScriptObjectReferences[ Index ], This );
+		}
 	}
 
 	//@todo NickW, temp hack to make stale property chains less crashy
@@ -2486,6 +2499,11 @@ FGuid UScriptStruct::GetCustomGuid() const
 	return FGuid();
 }
 
+FString UScriptStruct::GetStructCPPName() const
+{
+	return FString::Printf(TEXT("F%s"), *GetName());
+}
+
 IMPLEMENT_CORE_INTRINSIC_CLASS(UScriptStruct, UStruct,
 	{
 	}
diff --git a/Engine/Source/Runtime/CoreUObject/Private/UObject/Enum.cpp b/Engine/Source/Runtime/CoreUObject/Private/UObject/Enum.cpp
--- a/Engine/Source/Runtime/CoreUObject/Private/UObject/Enum.cpp
+++ b/Engine/Source/Runtime/CoreUObject/Private/UObject/Enum.cpp
@@ -22,6 +22,14 @@ TMap<FName, UEnum*> UEnum::AllEnumNames;
 TMap<FName,TMap<FName,FName> > UEnum::EnumRedirects;
 TMap<FName,TMap<FString,FString> > UEnum::EnumSubstringRedirects;
 
+UEnum::UEnum(const FObjectInitializer& ObjectInitializer)
+	: UField(ObjectInitializer)
+	, CppType()
+	, CppForm(ECppForm::Regular)
+	, EnumDisplayNameFn( nullptr )
+{
+}
+
 void UEnum::Serialize( FArchive& Ar )
 {
 	Ar.UsingCustomVersion(FCoreObjectVersion::GUID);
@@ -398,6 +406,11 @@ FText UEnum::GetEnumText(int32 InIndex) const
 	}
 #endif
 
+	if(EnumDisplayNameFn)
+	{
+		return (*EnumDisplayNameFn)(InIndex);
+	}
+
 	return FText::FromString( GetEnumName(InIndex) );
 }
 
diff --git a/Engine/Source/Runtime/CoreUObject/Private/UObject/EnumProperty.cpp b/Engine/Source/Runtime/CoreUObject/Private/UObject/EnumProperty.cpp
--- a/Engine/Source/Runtime/CoreUObject/Private/UObject/EnumProperty.cpp
+++ b/Engine/Source/Runtime/CoreUObject/Private/UObject/EnumProperty.cpp
@@ -51,7 +51,7 @@ UEnumProperty::UEnumProperty(const FObjectInitializer& ObjectInitializer, UEnum*
 }
 
 UEnumProperty::UEnumProperty(const FObjectInitializer& ObjectInitializer, ECppProperty, int32 InOffset, uint64 InFlags, UEnum* InEnum)
-	: UProperty(ObjectInitializer, EC_CppProperty, InOffset, InFlags)
+	: UProperty(ObjectInitializer, EC_CppProperty, InOffset, InFlags | CPF_HasGetValueTypeHash)
 	, Enum(InEnum)
 {
 	// This is expected to be set post-construction by AddCppProperty
@@ -61,7 +61,6 @@ UEnumProperty::UEnumProperty(const FObjectInitializer& ObjectInitializer, ECppPr
 void UEnumProperty::AddCppProperty(UProperty* Inner)
 {
 	check(!UnderlyingProp);
-
 	UnderlyingProp = CastChecked<UNumericProperty>(Inner);
 }
 
@@ -390,6 +389,13 @@ bool UEnumProperty::ConvertFromType(const FPropertyTag& Tag, FArchive& Ar, uint8
 	return bOutAdvanceProperty;
 }
 
+uint32 UEnumProperty::GetValueTypeHashInternal(const void* Src) const
+{
+	check(UnderlyingProp);
+	return UnderlyingProp->GetValueTypeHash(Src);
+}
+
+
 IMPLEMENT_CORE_INTRINSIC_CLASS(UEnumProperty, UProperty,
 	{
 		Class->EmitObjectReference(UE4EnumProperty_Private::FEnumPropertyFriend::EnumOffset, TEXT("Enum"));
diff --git a/Engine/Source/Runtime/CoreUObject/Private/UObject/LinkerLoad.cpp b/Engine/Source/Runtime/CoreUObject/Private/UObject/LinkerLoad.cpp
--- a/Engine/Source/Runtime/CoreUObject/Private/UObject/LinkerLoad.cpp
+++ b/Engine/Source/Runtime/CoreUObject/Private/UObject/LinkerLoad.cpp
@@ -3022,16 +3022,18 @@ void FLinkerLoad::LoadAllObjects( bool bForcePreload )
 #endif // USE_CIRCULAR_DEPENDENCY_LOAD_DEFERRING
 
 		UObject* LoadedObject = CreateExportAndPreload(ExportIndex, bForcePreload);
-		// DynamicClass could be created without calling CreateImport. The imported objects will be required later when a CDO is created.
-		// intentional merge conflict, talk to Gil
-		if (UDynamicClass* DynamicClass = Cast<UDynamicClass>(LoadedObject))
+
+		if(!GEventDrivenLoaderEnabled || !EVENT_DRIVEN_ASYNC_LOAD_ACTIVE_AT_RUNTIME)
 		{
-			for (int32 ImportIndex = 0; ImportIndex < ImportMap.Num(); ++ImportIndex)
+			// DynamicClass could be created without calling CreateImport. The imported objects will be required later when a CDO is created.
+			if (UDynamicClass* DynamicClass = Cast<UDynamicClass>(LoadedObject))
 			{
-				CreateImport(ImportIndex);
+				for (int32 ImportIndex = 0; ImportIndex < ImportMap.Num(); ++ImportIndex)
+				{
+					CreateImport(ImportIndex);
+				}
 			}
 		}
-		// intentional merge conflict, talk to Gil
 
 		// If needed send a heartbeat, but no need to do it too often
 		if (bShouldTickHeartBeat && (ExportIndex % 10) == 0)
diff --git a/Engine/Source/Runtime/CoreUObject/Private/UObject/Property.cpp b/Engine/Source/Runtime/CoreUObject/Private/UObject/Property.cpp
--- a/Engine/Source/Runtime/CoreUObject/Private/UObject/Property.cpp
+++ b/Engine/Source/Runtime/CoreUObject/Private/UObject/Property.cpp
@@ -638,7 +638,7 @@ bool UProperty::ShouldSerializeValue( FArchive& Ar ) const
 		||	((PropertyFlags & CPF_NonPIEDuplicateTransient) && !(Ar.GetPortFlags() & PPF_DuplicateForPIE) && (Ar.GetPortFlags() & PPF_Duplicate))
 		||	((PropertyFlags & CPF_NonTransactional) && Ar.IsTransacting())
 		||	((PropertyFlags & CPF_Deprecated) && !Ar.HasAllPortFlags(PPF_UseDeprecatedProperties) && (Ar.IsSaving() || Ar.IsTransacting() || Ar.WantBinaryPropertySerialization()))
-		||  ((PropertyFlags & CPF_SkipSerialization))
+		||  ((PropertyFlags & CPF_SkipSerialization) && (Ar.WantBinaryPropertySerialization() || !Ar.HasAllPortFlags(PPF_ForceTaggedSerialization)))
 		||  (IsEditorOnlyProperty() && Ar.IsFilterEditorOnly());
 
 	return !Skip;
diff --git a/Engine/Source/Runtime/CoreUObject/Private/UObject/PropertyMap.cpp b/Engine/Source/Runtime/CoreUObject/Private/UObject/PropertyMap.cpp
--- a/Engine/Source/Runtime/CoreUObject/Private/UObject/PropertyMap.cpp
+++ b/Engine/Source/Runtime/CoreUObject/Private/UObject/PropertyMap.cpp
@@ -663,6 +663,7 @@ void UMapProperty::AddCppProperty( UProperty* Property )
 	{
 		// If it's at offset zero, assume it's the key
 		check(!KeyProp);
+		ensureAlwaysMsgf(Property->HasAllPropertyFlags(CPF_HasGetValueTypeHash), TEXT("Attempting to create Map Property with unhashable key type: %s - Provide a GetTypeHash function!"), *Property->GetName());
 		KeyProp = Property;
 	}
 	else
diff --git a/Engine/Source/Runtime/CoreUObject/Private/UObject/PropertySet.cpp b/Engine/Source/Runtime/CoreUObject/Private/UObject/PropertySet.cpp
--- a/Engine/Source/Runtime/CoreUObject/Private/UObject/PropertySet.cpp
+++ b/Engine/Source/Runtime/CoreUObject/Private/UObject/PropertySet.cpp
@@ -597,6 +597,7 @@ void USetProperty::AddCppProperty( UProperty* Property )
 {
 	check(!ElementProp);
 	check(Property);
+	ensureAlwaysMsgf(Property->HasAllPropertyFlags(CPF_HasGetValueTypeHash), TEXT("Attempting to create Set Property with unhashable element type: %s - Provide a GetTypeHash function!"), *Property->GetName());
 
 	ElementProp = Property;
 }
diff --git a/Engine/Source/Runtime/CoreUObject/Private/UObject/PropertyStruct.cpp b/Engine/Source/Runtime/CoreUObject/Private/UObject/PropertyStruct.cpp
--- a/Engine/Source/Runtime/CoreUObject/Private/UObject/PropertyStruct.cpp
+++ b/Engine/Source/Runtime/CoreUObject/Private/UObject/PropertyStruct.cpp
@@ -35,14 +35,14 @@ static inline void PreloadInnerStructMembers(UStructProperty* StructProperty)
 -----------------------------------------------------------------------------*/
 
 UStructProperty::UStructProperty(ECppProperty, int32 InOffset, uint64 InFlags, UScriptStruct* InStruct)
-	: UProperty(FObjectInitializer::Get(), EC_CppProperty, InOffset, InFlags)
+	: UProperty(FObjectInitializer::Get(), EC_CppProperty, InOffset, InStruct->GetCppStructOps() ? InStruct->GetCppStructOps()->GetComputedPropertyFlags() | InFlags : InFlags)
 	, Struct(InStruct)
 {
 	ElementSize = Struct->PropertiesSize;
 }
 
 UStructProperty::UStructProperty( const FObjectInitializer& ObjectInitializer, ECppProperty, int32 InOffset, uint64 InFlags, UScriptStruct* InStruct )
-	:	UProperty( ObjectInitializer, EC_CppProperty, InOffset, InFlags )
+	:	UProperty( ObjectInitializer, EC_CppProperty, InOffset, InStruct->GetCppStructOps() ? InStruct->GetCppStructOps()->GetComputedPropertyFlags() | InFlags : InFlags )
 	,	Struct( InStruct )
 {
 	ElementSize = Struct->PropertiesSize;
@@ -77,17 +77,14 @@ void UStructProperty::LinkInternal(FArchive& Ar)
 	PreloadInnerStructMembers(this);
 	
 	ElementSize = Align(Struct->PropertiesSize, Struct->GetMinAlignment());
-	if (Struct->StructFlags & STRUCT_IsPlainOldData) // if there is nothing to construct or the struct is known to be memcpy-able, then allow memcpy
+	if(UScriptStruct::ICppStructOps* Ops = Struct->GetCppStructOps())
 	{
-		PropertyFlags |= CPF_IsPlainOldData;
+		PropertyFlags |= Ops->GetComputedPropertyFlags();
 	}
-	if (Struct->StructFlags & STRUCT_NoDestructor)
-	{
-		PropertyFlags |= CPF_NoDestructor;
-	}
-	if (Struct->StructFlags & STRUCT_ZeroConstructor)
+	else
 	{
-		PropertyFlags |= CPF_ZeroConstructor;
+		// User Defined structs won't have UScriptStruct::ICppStructOps. Setting their flags here.
+		PropertyFlags |= CPF_HasGetValueTypeHash;
 	}
 }
 
@@ -211,7 +208,7 @@ bool UStructProperty::HasNoOpConstructor() const
 
 FString UStructProperty::GetCPPType( FString* ExtendedTypeText/*=NULL*/, uint32 CPPExportFlags/*=0*/ ) const
 {
-	return FString::Printf(TEXT("F%s"), *Struct->GetName());
+	return Struct->GetStructCPPName();
 }
 
 FString UStructProperty::GetCPPTypeForwardDeclaration() const
diff --git a/Engine/Source/Runtime/CoreUObject/Private/UObject/SavePackage.cpp b/Engine/Source/Runtime/CoreUObject/Private/UObject/SavePackage.cpp
--- a/Engine/Source/Runtime/CoreUObject/Private/UObject/SavePackage.cpp
+++ b/Engine/Source/Runtime/CoreUObject/Private/UObject/SavePackage.cpp
@@ -640,6 +640,13 @@ public:
 			{
 				MarkNameAsReferenced(ReplObjClass->GetFName());
 			}
+
+			FName ReplacedName;
+			Coordinator->FindReplacedNameAndOuter(Obj, ReplacedName); //TODO: should we care about replaced outer ?
+			if (ReplacedName != NAME_None)
+			{
+				MarkNameAsReferenced(ReplacedName);
+			}
 		}
 	}
 #endif //WITH_EDITOR
@@ -4673,9 +4680,9 @@ ESavePackageResult UPackage::Save(UPackage* InOuter, UObject* Base, EObjectFlags
 					TArray<UObject*> TagImpObjects;
 					GetObjectsWithAnyMarks(TagImpObjects, OBJECTMARK_TagImp);
 
+					const EObjectMark ObjectMarks = UPackage::GetObjectMarksForTargetPlatform(TargetPlatform, Linker->IsCooking());
 					if (Linker->IsCooking())
 					{
-						const EObjectMark ObjectMarks = UPackage::GetObjectMarksForTargetPlatform(TargetPlatform, Linker->IsCooking());
 						for (UObject* ObjImport : TagImpObjects)
 						{
 							if (ObjImport->HasAllMarks(ObjectMarks))
@@ -4700,6 +4707,11 @@ ESavePackageResult UPackage::Save(UPackage* InOuter, UObject* Base, EObjectFlags
 							if (UObject* ReplacedOuter = Coordinator->FindReplacedNameAndOuter(Obj, /*out*/ReplacedName))
 							{
 								ReplacedImportOuters.Add(Obj, ReplacedOuter);
+								if (ReplacedOuter->HasAllMarks(ObjectMarks))
+								{
+									// The actual outer is not for the current platform, so the object shouldn't be cooked.
+									Obj->UnMark(OBJECTMARK_TagImp);
+								}
 							}
 						}
 #endif //WITH_EDITOR
diff --git a/Engine/Source/Runtime/CoreUObject/Private/UObject/ScriptCore.cpp b/Engine/Source/Runtime/CoreUObject/Private/UObject/ScriptCore.cpp
--- a/Engine/Source/Runtime/CoreUObject/Private/UObject/ScriptCore.cpp
+++ b/Engine/Source/Runtime/CoreUObject/Private/UObject/ScriptCore.cpp
@@ -2615,6 +2615,57 @@ void UObject::execSetArray( FFrame& Stack, RESULT_DECL )
 }
 IMPLEMENT_VM_FUNCTION( EX_SetArray, execSetArray );
 
+void UObject::execSetSet( FFrame& Stack, RESULT_DECL )
+{
+	// Get the set address
+	Stack.MostRecentPropertyAddress = NULL;
+	Stack.MostRecentProperty = NULL;
+	Stack.Step( Stack.Object, NULL ); // Set to set
+	int32 Num = Stack.ReadInt<int32>();
+
+	USetProperty* SetProperty = CastChecked<USetProperty>(Stack.MostRecentProperty);
+ 	FScriptSetHelper SetHelper(SetProperty, Stack.MostRecentPropertyAddress);
+ 	SetHelper.EmptyElements(Num);
+ 
+ 	// Read in the parameters one at a time
+ 	int32 i = 0;
+ 	while(*Stack.Code != EX_EndSet)
+ 	{
+ 		SetHelper.AddUninitializedValue();
+ 		Stack.Step(Stack.Object, SetHelper.GetElementPtr(i++));
+ 	}
+	SetHelper.Rehash();
+ 
+ 	P_FINISH;
+}
+IMPLEMENT_VM_FUNCTION( EX_SetSet, execSetSet );
+
+void UObject::execSetMap( FFrame& Stack, RESULT_DECL )
+{
+	// Get the map address
+	Stack.MostRecentPropertyAddress = NULL;
+	Stack.MostRecentProperty = NULL;
+	Stack.Step( Stack.Object, NULL ); // Map to set
+	int32 Num = Stack.ReadInt<int32>();
+	
+	UMapProperty* MapProperty = CastChecked<UMapProperty>(Stack.MostRecentProperty);
+ 	FScriptMapHelper MapHelper(MapProperty, Stack.MostRecentPropertyAddress);
+ 	MapHelper.EmptyValues(Num);
+ 
+ 	// Read in the parameters one at a time
+ 	int32 i = 0;
+ 	while(*Stack.Code != EX_EndMap)
+ 	{
+ 		MapHelper.AddUninitializedValue();
+ 		Stack.Step(Stack.Object, MapHelper.GetKeyPtr(i));
+ 		Stack.Step(Stack.Object, MapHelper.GetValuePtr(i++));
+ 	}
+	MapHelper.Rehash();
+ 
+ 	P_FINISH;
+}
+IMPLEMENT_VM_FUNCTION( EX_SetMap, execSetMap );
+
 void UObject::execArrayConst(FFrame& Stack, RESULT_DECL)
 {
 	UProperty* InnerProperty = CastChecked<UProperty>(Stack.ReadObject());
diff --git a/Engine/Source/Runtime/CoreUObject/Public/Blueprint/BlueprintSupport.h b/Engine/Source/Runtime/CoreUObject/Public/Blueprint/BlueprintSupport.h
--- a/Engine/Source/Runtime/CoreUObject/Public/Blueprint/BlueprintSupport.h
+++ b/Engine/Source/Runtime/CoreUObject/Public/Blueprint/BlueprintSupport.h
@@ -274,6 +274,8 @@ struct COREUOBJECT_API FBlueprintDependencyObjectRef
 	FName ClassPackageName;
 	FName ClassName;
 
+	FBlueprintDependencyObjectRef() {}
+
 	FORCENOINLINE FBlueprintDependencyObjectRef(const TCHAR* InPackageFolder
 		, const TCHAR* InShortPackageName
 		, const TCHAR* InObjectName
@@ -312,7 +314,7 @@ private:
 public:
 	static FConvertedBlueprintsDependencies& Get();
 
-	void RegisterClass(FName PackageName, GetDependenciesNamesFunc GetAssets);
+	void RegisterConvertedClass(FName PackageName, GetDependenciesNamesFunc GetAssets);
 
 	/** Get all assets paths necessary for the class with the given class name and all converted classes that dependencies. */
 	void GetAssets(FName PackageName, TArray<FBlueprintDependencyData>& OutDependencies) const;
diff --git a/Engine/Source/Runtime/CoreUObject/Public/UObject/AssetPtr.h b/Engine/Source/Runtime/CoreUObject/Public/UObject/AssetPtr.h
--- a/Engine/Source/Runtime/CoreUObject/Public/UObject/AssetPtr.h
+++ b/Engine/Source/Runtime/CoreUObject/Public/UObject/AssetPtr.h
@@ -86,6 +86,7 @@ public:
 
 template <> struct TIsPODType<FAssetPtr> { enum { Value = TIsPODType<TPersistentObjectPtr<FStringAssetReference> >::Value }; };
 template <> struct TIsWeakPointerType<FAssetPtr> { enum { Value = TIsWeakPointerType<TPersistentObjectPtr<FStringAssetReference> >::Value }; };
+template <> struct THasGetTypeHash<FAssetPtr> { enum { Value = THasGetTypeHash<TPersistentObjectPtr<FStringAssetReference> >::Value }; };
 
 /**
  * TAssetPtr is templatized wrapper of the generic FAssetPtr, it can be used in UProperties
diff --git a/Engine/Source/Runtime/CoreUObject/Public/UObject/Class.h b/Engine/Source/Runtime/CoreUObject/Public/UObject/Class.h
--- a/Engine/Source/Runtime/CoreUObject/Public/UObject/Class.h
+++ b/Engine/Source/Runtime/CoreUObject/Public/UObject/Class.h
@@ -16,6 +16,7 @@
 #include "UObject/GarbageCollection.h"
 #include "UObject/CoreNative.h"
 #include "Templates/HasGetTypeHash.h"
+#include "Templates/IsEnum.h"
 
 struct FCustomPropertyListNode;
 struct FFrame;
@@ -903,6 +904,8 @@ public:
 		/** Calls GetTypeHash if enabled */
 		virtual uint32 GetTypeHash(const void* Src) = 0;
 
+		/** Returns property flag values that can be computed at compile time */
+		virtual uint64 GetComputedPropertyFlags() const = 0;
 	private:
 		/** sizeof() of the structure **/
 		const int32 Size;
@@ -1047,6 +1050,14 @@ public:
 			ensure(HasGetTypeHash());
 			return GetTypeHashOrNot((const CPPSTRUCT*)Src);
 		}
+		virtual uint64 GetComputedPropertyFlags() const override
+		{
+			return 
+				(TIsPODType<CPPSTRUCT>::Value ? CPF_IsPlainOldData : 0) 
+				| (TIsTriviallyDestructible<CPPSTRUCT>::Value ? CPF_NoDestructor : 0) 
+				| (TIsZeroConstructType<CPPSTRUCT>::Value ? CPF_ZeroConstructor : 0)
+				| (THasGetTypeHash<CPPSTRUCT>::Value ? CPF_HasGetValueTypeHash : 0);;
+		}
 	};
 
 	/** Template for noexport classes to autoregister before main starts **/
@@ -1234,6 +1245,9 @@ public:
 	virtual COREUOBJECT_API void RecursivelyPreload();
 
 	virtual COREUOBJECT_API FGuid GetCustomGuid() const;
+	
+	/** Returns the (native, c++) name of the struct */
+	virtual COREUOBJECT_API FString GetStructCPPName() const;
 
 #if WITH_EDITOR
 	/**
@@ -1365,7 +1379,7 @@ public:
 	{
 		//@TODO: UCREMOVAL: CPF_ConstParm added as a hack to get blueprints compiling with a const DamageType parameter.
 		const uint64 IgnoreFlags = CPF_PersistentInstance | CPF_ExportObject | CPF_InstancedReference 
-			| CPF_ContainsInstancedReference | CPF_ComputedFlags | CPF_ConstParm | CPF_HasGetValueTypeHash | CPF_UObjectWrapper
+			| CPF_ContainsInstancedReference | CPF_ComputedFlags | CPF_ConstParm | CPF_UObjectWrapper
 			| CPF_NativeAccessSpecifiers | CPF_AdvancedDisplay;
 		return IgnoreFlags;
 	}
@@ -1407,12 +1421,21 @@ public:
 	UEnum.
 -----------------------------------------------------------------------------*/
 
+typedef FText(*FEnumDisplayNameFn)(int32);
+
 //
 // Reflection data for an enumeration.
 //
 class COREUOBJECT_API UEnum : public UField
 {
-	DECLARE_CASTED_CLASS_INTRINSIC_WITH_API(UEnum, UField, 0, TEXT("/Script/CoreUObject"), CASTCLASS_UEnum, NO_API)
+	DECLARE_CASTED_CLASS_INTRINSIC_NO_CTOR(UEnum, UField, 0, TEXT("/Script/CoreUObject"), CASTCLASS_UEnum, NO_API)
+	UEnum(const FObjectInitializer& ObjectInitialzer);
+
+	/** Associate a function for looking up Enum display names by index, only intended for use by generated code: */
+	void SetEnumDisplayNameFn(FEnumDisplayNameFn InEnumDisplayNameFn)
+	{ 
+		EnumDisplayNameFn = InEnumDisplayNameFn; 
+	}
 
 public:
 	enum class ECppForm
@@ -1457,6 +1480,9 @@ protected:
 	/** How the enum was originally defined. */
 	ECppForm CppForm;
 
+	/** pointer to function used to look up the enum's display name. Currently only assigned for UEnums generated for nativized blueprints */
+	FEnumDisplayNameFn EnumDisplayNameFn;
+
 	/** global list of all value names used by all enums in memory, used for property text import */
 	static TMap<FName, UEnum*> AllEnumNames;
 
diff --git a/Engine/Source/Runtime/CoreUObject/Public/UObject/EnumProperty.h b/Engine/Source/Runtime/CoreUObject/Public/UObject/EnumProperty.h
--- a/Engine/Source/Runtime/CoreUObject/Public/UObject/EnumProperty.h
+++ b/Engine/Source/Runtime/CoreUObject/Public/UObject/EnumProperty.h
@@ -64,8 +64,10 @@ public:
 	}
 
 private:
-	friend struct UE4EnumProperty_Private::FEnumPropertyFriend;
+	virtual uint32 GetValueTypeHashInternal(const void* Src) const override;
 
+	friend struct UE4EnumProperty_Private::FEnumPropertyFriend;
+	
 #if HACK_HEADER_GENERATOR
 public:
 #endif
diff --git a/Engine/Source/Runtime/CoreUObject/Public/UObject/Object.h b/Engine/Source/Runtime/CoreUObject/Public/UObject/Object.h
--- a/Engine/Source/Runtime/CoreUObject/Public/UObject/Object.h
+++ b/Engine/Source/Runtime/CoreUObject/Public/UObject/Object.h
@@ -1165,6 +1165,8 @@ public:
 	DECLARE_FUNCTION(execTransformConst);
 	DECLARE_FUNCTION(execStructConst);
 	DECLARE_FUNCTION(execSetArray);
+	DECLARE_FUNCTION(execSetSet);
+	DECLARE_FUNCTION(execSetMap);
 	DECLARE_FUNCTION(execArrayConst);
 
 	// Object construction
diff --git a/Engine/Source/Runtime/CoreUObject/Public/UObject/ObjectMacros.h b/Engine/Source/Runtime/CoreUObject/Public/UObject/ObjectMacros.h
--- a/Engine/Source/Runtime/CoreUObject/Public/UObject/ObjectMacros.h
+++ b/Engine/Source/Runtime/CoreUObject/Public/UObject/ObjectMacros.h
@@ -436,7 +436,7 @@ typedef uint64 EClassCastFlags;
 #define CPF_DevelopmentAssets		(CPF_EditorOnly)
 
 /** all the properties that should never be loaded or saved */
-#define CPF_ComputedFlags			(CPF_IsPlainOldData | CPF_NoDestructor | CPF_ZeroConstructor)
+#define CPF_ComputedFlags			(CPF_IsPlainOldData | CPF_NoDestructor | CPF_ZeroConstructor | CPF_HasGetValueTypeHash)
 
 #define CPF_AllFlags				DECLARE_UINT64(0xFFFFFFFFFFFFFFFF)
 
diff --git a/Engine/Source/Runtime/CoreUObject/Public/UObject/Script.h b/Engine/Source/Runtime/CoreUObject/Public/UObject/Script.h
--- a/Engine/Source/Runtime/CoreUObject/Public/UObject/Script.h
+++ b/Engine/Source/Runtime/CoreUObject/Public/UObject/Script.h
@@ -208,10 +208,10 @@ enum EExprToken
 	EX_UInt64Const			= 0x36,	// 64-bit unsigned integer constant.
 	//						= 0x37,
 	EX_PrimitiveCast		= 0x38,	// A casting operator for primitives which reads the type as the subsequent byte
-	//						= 0x39,
-	//						= 0x3A,
-	//						= 0x3B,
-	//						= 0x3C,
+	EX_SetSet				= 0x39,
+	EX_EndSet				= 0x3A,
+	EX_SetMap				= 0x3B,
+	EX_EndMap				= 0x3C,
 	//						= 0x3D,
 	//						= 0x3E,
 	//						= 0x3F,
diff --git a/Engine/Source/Runtime/CoreUObject/Public/UObject/ScriptSerialization.h b/Engine/Source/Runtime/CoreUObject/Public/UObject/ScriptSerialization.h
--- a/Engine/Source/Runtime/CoreUObject/Public/UObject/ScriptSerialization.h
+++ b/Engine/Source/Runtime/CoreUObject/Public/UObject/ScriptSerialization.h
@@ -416,6 +416,16 @@
 			while( SerializeExpr( iCode, Ar) != EX_EndArray );
 			break;
 		}
+		case EX_SetSet:
+			SerializeExpr( iCode, Ar ); // set property
+			XFER(int32);			// Number of elements
+			while( SerializeExpr( iCode, Ar) != EX_EndSet );
+			break;
+		case EX_SetMap:
+			SerializeExpr( iCode, Ar ); // map property
+			XFER(int32);			// Number of elements
+			while( SerializeExpr( iCode, Ar) != EX_EndMap );
+			break;
 		case EX_ArrayConst:
 		{
 			XFERPTR(UProperty*);	// Inner property
diff --git a/Engine/Source/Runtime/CoreUObject/Public/UObject/UnrealType.h b/Engine/Source/Runtime/CoreUObject/Public/UObject/UnrealType.h
--- a/Engine/Source/Runtime/CoreUObject/Public/UObject/UnrealType.h
+++ b/Engine/Source/Runtime/CoreUObject/Public/UObject/UnrealType.h
@@ -950,7 +950,8 @@ protected:
 		return 
 			(TIsPODType<TCppType>::Value ? CPF_IsPlainOldData : 0) 
 			| (TIsTriviallyDestructible<TCppType>::Value ? CPF_NoDestructor : 0) 
-			| (TIsZeroConstructType<TCppType>::Value ? CPF_ZeroConstructor : 0);
+			| (TIsZeroConstructType<TCppType>::Value ? CPF_ZeroConstructor : 0)
+			| (THasGetTypeHash<TCppType>::Value ? CPF_HasGetValueTypeHash : 0);
 
 	}
 };
@@ -1239,12 +1240,12 @@ public:
 	}
 
 	TProperty_Numeric(ECppProperty, int32 InOffset, uint64 InFlags)
-		: Super(FObjectInitializer::Get(), EC_CppProperty, InOffset, InFlags | CPF_HasGetValueTypeHash)
+		: Super(FObjectInitializer::Get(), EC_CppProperty, InOffset, InFlags)
 	{
 	}
 
 	TProperty_Numeric( const FObjectInitializer& ObjectInitializer, ECppProperty, int32 InOffset, uint64 InFlags )
-		:	Super( ObjectInitializer, EC_CppProperty, InOffset, InFlags | CPF_HasGetValueTypeHash)
+		:	Super( ObjectInitializer, EC_CppProperty, InOffset, InFlags)
 	{
 	}
 
@@ -2027,12 +2028,12 @@ class COREUOBJECT_API UObjectProperty : public TUObjectPropertyBase<UObject*>
 	DECLARE_CASTED_CLASS_INTRINSIC(UObjectProperty, TUObjectPropertyBase<UObject*>, 0, TEXT("/Script/CoreUObject"), CASTCLASS_UObjectProperty)
 
 	UObjectProperty(ECppProperty, int32 InOffset, uint64 InFlags, UClass* InClass)
-		: TUObjectPropertyBase(FObjectInitializer::Get(), EC_CppProperty, InOffset, InFlags | CPF_HasGetValueTypeHash, InClass)
+		: TUObjectPropertyBase(FObjectInitializer::Get(), EC_CppProperty, InOffset, InFlags, InClass)
 	{
 	}
 
 	UObjectProperty( const FObjectInitializer& ObjectInitializer, ECppProperty, int32 InOffset, uint64 InFlags, UClass* InClass )
-		: TUObjectPropertyBase(ObjectInitializer, EC_CppProperty, InOffset, InFlags | CPF_HasGetValueTypeHash, InClass)
+		: TUObjectPropertyBase(ObjectInitializer, EC_CppProperty, InOffset, InFlags, InClass)
 	{
 	}
 
@@ -2100,12 +2101,12 @@ class COREUOBJECT_API ULazyObjectProperty : public TUObjectPropertyBase<FLazyObj
 	DECLARE_CASTED_CLASS_INTRINSIC(ULazyObjectProperty, TUObjectPropertyBase<FLazyObjectPtr>, 0, TEXT("/Script/CoreUObject"), CASTCLASS_ULazyObjectProperty)
 
 	ULazyObjectProperty(ECppProperty, int32 InOffset, uint64 InFlags, UClass* InClass)
-		: TUObjectPropertyBase(FObjectInitializer::Get(), EC_CppProperty, InOffset, InFlags | CPF_HasGetValueTypeHash, InClass)
+		: TUObjectPropertyBase(FObjectInitializer::Get(), EC_CppProperty, InOffset, InFlags, InClass)
 	{
 	}
 
 	ULazyObjectProperty( const FObjectInitializer& ObjectInitializer, ECppProperty, int32 InOffset, uint64 InFlags, UClass* InClass )
-		: TUObjectPropertyBase(ObjectInitializer, EC_CppProperty, InOffset, InFlags | CPF_HasGetValueTypeHash, InClass)
+		: TUObjectPropertyBase(ObjectInitializer, EC_CppProperty, InOffset, InFlags, InClass)
 	{
 	}
 
@@ -2138,11 +2139,11 @@ class COREUOBJECT_API UAssetObjectProperty : public TUObjectPropertyBase<FAssetP
 	DECLARE_CASTED_CLASS_INTRINSIC(UAssetObjectProperty, TUObjectPropertyBase<FAssetPtr>, 0, TEXT("/Script/CoreUObject"), CASTCLASS_UAssetObjectProperty)
 
 	UAssetObjectProperty(ECppProperty, int32 InOffset, uint64 InFlags, UClass* InClass)
-		: TUObjectPropertyBase(FObjectInitializer::Get(), EC_CppProperty, InOffset, InFlags | CPF_HasGetValueTypeHash, InClass)
+		: TUObjectPropertyBase(FObjectInitializer::Get(), EC_CppProperty, InOffset, InFlags, InClass)
 	{}
 
 	UAssetObjectProperty( const FObjectInitializer& ObjectInitializer, ECppProperty, int32 InOffset, uint64 InFlags, UClass* InClass )
-		: TUObjectPropertyBase(ObjectInitializer, EC_CppProperty, InOffset, InFlags | CPF_HasGetValueTypeHash, InClass)
+		: TUObjectPropertyBase(ObjectInitializer, EC_CppProperty, InOffset, InFlags, InClass)
 	{}
 
 	// UHT interface
@@ -2390,12 +2391,12 @@ public:
 	typedef TTypeFundamentals::TCppType TCppType;
 
 	UNameProperty(ECppProperty, int32 InOffset, uint64 InFlags)
-		: UNameProperty_Super(FObjectInitializer::Get(), EC_CppProperty, InOffset, InFlags | CPF_HasGetValueTypeHash)
+		: UNameProperty_Super(FObjectInitializer::Get(), EC_CppProperty, InOffset, InFlags)
 	{
 	}
 
 	UNameProperty( const FObjectInitializer& ObjectInitializer, ECppProperty, int32 InOffset, uint64 InFlags )
-	:	UNameProperty_Super( ObjectInitializer, EC_CppProperty, InOffset, InFlags | CPF_HasGetValueTypeHash )
+	:	UNameProperty_Super( ObjectInitializer, EC_CppProperty, InOffset, InFlags )
 	{
 	}
 
@@ -2427,12 +2428,12 @@ public:
 	typedef TTypeFundamentals::TCppType TCppType;
 
 	UStrProperty(ECppProperty, int32 InOffset, uint64 InFlags)
-		: UStrProperty_Super(FObjectInitializer::Get(), EC_CppProperty, InOffset, InFlags | CPF_HasGetValueTypeHash)
+		: UStrProperty_Super(FObjectInitializer::Get(), EC_CppProperty, InOffset, InFlags)
 	{
 	}
 
 	UStrProperty( const FObjectInitializer& ObjectInitializer, ECppProperty, int32 InOffset, uint64 InFlags )
-	:	UStrProperty_Super( ObjectInitializer, EC_CppProperty, InOffset, InFlags | CPF_HasGetValueTypeHash)
+	:	UStrProperty_Super( ObjectInitializer, EC_CppProperty, InOffset, InFlags)
 	{
 	}
 
diff --git a/Engine/Source/Runtime/Engine/Classes/EdGraph/EdGraphPin.h b/Engine/Source/Runtime/Engine/Classes/EdGraph/EdGraphPin.h
--- a/Engine/Source/Runtime/Engine/Classes/EdGraph/EdGraphPin.h
+++ b/Engine/Source/Runtime/Engine/Classes/EdGraph/EdGraphPin.h
@@ -240,6 +240,9 @@ public:
 	}
 
 	ENGINE_API bool Serialize(FArchive& Ar);
+
+	static ENGINE_API FEdGraphPinType GetPinTypeForTerminalType( const FEdGraphTerminalType& TerminalType );
+	static ENGINE_API FEdGraphPinType GetTerminalTypeForContainer( const FEdGraphPinType& ContainerType );
 };
 
 template<>
diff --git a/Engine/Source/Runtime/Engine/Classes/Engine/Blueprint.h b/Engine/Source/Runtime/Engine/Classes/Engine/Blueprint.h
--- a/Engine/Source/Runtime/Engine/Classes/Engine/Blueprint.h
+++ b/Engine/Source/Runtime/Engine/Classes/Engine/Blueprint.h
@@ -84,6 +84,23 @@ enum class EBlueprintCompileMode : uint8
 	FinalRelease UMETA(ToolTip="Always compile in final release mode.")
 };
 
+USTRUCT()
+struct FCompilerNativizationOptions
+{
+	GENERATED_USTRUCT_BODY()
+
+	UPROPERTY()
+	bool ServerOnlyPlatform;
+
+	UPROPERTY()
+	bool ClientOnlyPlatform;
+
+	FCompilerNativizationOptions()
+		: ServerOnlyPlatform(false)
+		, ClientOnlyPlatform(false)
+	{}
+};
+
 struct FKismetCompilerOptions
 {
 public:
@@ -104,6 +121,7 @@ public:
 
 	TSharedPtr<FString> OutHeaderSourceCode;
 	TSharedPtr<FString> OutCppSourceCode;
+	FCompilerNativizationOptions NativizationOptions;
 
 	bool DoesRequireCppCodeGeneration() const
 	{
@@ -355,9 +373,9 @@ class ENGINE_API UBlueprint : public UBlueprintCore
 	UPROPERTY(EditAnywhere, Category=BlueprintOptions)
 	TArray<FString> HideCategories;
 	 
-	/** When nativization is enabled, and "NativizeOnlySelectedBlueprints" is set true, then this asset will be nativized. All super classes must be also nativized. */
-	UPROPERTY(EditAnywhere, Category=Experimental)
-	bool bNativize;
+	/** When exclusive nativization is enabled, then this asset will be nativized. All super classes must be also nativized. */
+	UPROPERTY(transient, DisplayName = "Nativize", EditAnywhere, Category= Packaging)
+	bool bSelectedForNativization;
 
 	/** TRUE to show a warning when attempting to start in PIE and there is a compiler error on this Blueprint */
 	UPROPERTY(transient)
@@ -568,6 +586,9 @@ public:
 	/** Find the object in the TemplateObjects array with the supplied name */
 	UActorComponent* FindTemplateByName(const FName& TemplateName) const;
 
+	/** Rename the component template in the TemplateObjects array with the supplied name */
+	bool RenameComponentTemplate(const FName& OldTemplateName, const FName& NewTemplateName);
+
 	/** Find a timeline by name */
 	class UTimelineTemplate* FindTimelineTemplateByVariableName(const FName& TimelineName);	
 
@@ -643,6 +664,9 @@ public:
 	virtual void BeginCacheForCookedPlatformData(const ITargetPlatform *TargetPlatform) override;
 	virtual bool IsCachedCookedPlatformDataLoaded(const ITargetPlatform* TargetPlatform) override;
 	virtual void ClearAllCachedCookedPlatformData() override;
+#if WITH_EDITOR
+	virtual void PostEditChangeProperty(struct FPropertyChangedEvent& PropertyChangedEvent) override;
+#endif
 	//~ End UObject Interface
 
 	/** Consigns the GeneratedClass and the SkeletonGeneratedClass to oblivion, and nulls their references */
@@ -776,6 +800,10 @@ private:
 	/** Broadcasts a notification whenever the blueprint is compiled. */
 	FCompiledEvent CompiledEvent;
 
+	/** Deprecated properties. */
+	UPROPERTY()
+	bool bNativize_DEPRECATED;
+
 #if WITH_EDITOR
 public:
 	/** If this blueprint is currently being compiled, the CurrentMessageLog will be the log currently being used to send logs to. */
diff --git a/Engine/Source/Runtime/Engine/Classes/Engine/BlueprintGeneratedClass.h b/Engine/Source/Runtime/Engine/Classes/Engine/BlueprintGeneratedClass.h
--- a/Engine/Source/Runtime/Engine/Classes/Engine/BlueprintGeneratedClass.h
+++ b/Engine/Source/Runtime/Engine/Classes/Engine/BlueprintGeneratedClass.h
@@ -210,7 +210,7 @@ protected:
 	TArray<struct FNodeToCodeAssociation> DebugNodeLineNumbers;
 
 	// List of entry points that contributed to the ubergraph
-	TMap<FName, FEdGraphPinReference> DelegatePins;
+	TMap<int32, FName> EntryPoints;
 
 	// Acceleration structure for execution wire highlighting at runtime
 	TMap<TWeakObjectPtr<UFunction>, FDebuggingInfoForSingleFunction> PerFunctionLineNumbers;
@@ -445,11 +445,13 @@ public:
 		PerFuncInfo.SourcePinToLineNumbersMap.Add(SourcePin, CodeOffset);
 	}
 
-	const TMap<FName, FEdGraphPinReference>& GetCompilerGeneratedEvents() const { return DelegatePins; }
+	const TMap<int32, FName>& GetEntryPoints() const { return EntryPoints; }
 
-	void RegisterPinToEventName(UEdGraphPin const* SourcePin, const FName FunctionName)
+	bool IsValidEntryPoint(const int32 LinkId) const { return EntryPoints.Contains(LinkId); }
+
+	void RegisterEntryPoint(const int32 ScriptOffset, const FName FunctionName)
 	{
-		DelegatePins.Add(FunctionName) = SourcePin;
+		EntryPoints.Add(ScriptOffset, FunctionName);
 	}
 
 	// Registers an association between an object (pin or node typically) and an associated class member property
diff --git a/Engine/Source/Runtime/Engine/Classes/Engine/InheritableComponentHandler.h b/Engine/Source/Runtime/Engine/Classes/Engine/InheritableComponentHandler.h
--- a/Engine/Source/Runtime/Engine/Classes/Engine/InheritableComponentHandler.h
+++ b/Engine/Source/Runtime/Engine/Classes/Engine/InheritableComponentHandler.h
@@ -125,6 +125,7 @@ public:
 public:
 
 	//~ Begin UObject Interface
+	virtual void Serialize(FArchive& Ar) override;
 	virtual void PostLoad() override;
 	virtual void GetPreloadDependencies(TArray<UObject*>& OutDeps) override;
 	//~ End UObject Interface
@@ -152,6 +153,9 @@ public:
 
 private:
 	const FComponentOverrideRecord* FindRecord(const FComponentKey Key) const;
+
+	/** Helper method used to assist with fixing up component template names at load time. */
+	void FixComponentTemplateName(UActorComponent* ComponentTemplate, const FString& NewName);
 	
 	/** All component records */
 	UPROPERTY()
diff --git a/Engine/Source/Runtime/Engine/Classes/Engine/MemberReference.h b/Engine/Source/Runtime/Engine/Classes/Engine/MemberReference.h
--- a/Engine/Source/Runtime/Engine/Classes/Engine/MemberReference.h
+++ b/Engine/Source/Runtime/Engine/Classes/Engine/MemberReference.h
@@ -372,9 +372,15 @@ public:
 				ReturnField = FindField<TFieldType>(TargetScope, MemberName);
 
 #if WITH_EDITOR
+
+				// If the reference variable is valid we need to make sure that our GUID matches
+				if (ReturnField != NULL)
+				{
+					UBlueprint::GetGuidFromClassByFieldName<TFieldType>(TargetScope, MemberName, MemberGuid);
+				}
 				// If we have a GUID find the reference variable and make sure the name is up to date and find the field again
 				// For now only variable references will have valid GUIDs.  Will have to deal with finding other names subsequently
-				if (ReturnField == NULL && MemberGuid.IsValid())
+				else if (MemberGuid.IsValid())
 				{
 					const FName RenamedMemberName = UBlueprint::GetFieldNameFromClassByGuid<TFieldType>(TargetScope, MemberGuid);
 					if (RenamedMemberName != NAME_None)
diff --git a/Engine/Source/Runtime/Engine/Classes/Engine/UserDefinedStruct.h b/Engine/Source/Runtime/Engine/Classes/Engine/UserDefinedStruct.h
--- a/Engine/Source/Runtime/Engine/Classes/Engine/UserDefinedStruct.h
+++ b/Engine/Source/Runtime/Engine/Classes/Engine/UserDefinedStruct.h
@@ -70,8 +70,11 @@ public:
 	virtual uint32 GetStructTypeHash(const void* Src) const override;
 	virtual void RecursivelyPreload() override;
 	virtual FGuid GetCustomGuid() const override;
+	virtual FString GetStructCPPName() const override;
 	// End of  UScriptStruct interface.
 
+	static uint32 GetUserDefinedStructTypeHash(const void* Src, const UScriptStruct* Type);
+
 #if WITH_EDITOR
 public:
 	// UStruct
diff --git a/Engine/Source/Runtime/Engine/Classes/Kismet/KismetArrayLibrary.h b/Engine/Source/Runtime/Engine/Classes/Kismet/KismetArrayLibrary.h
--- a/Engine/Source/Runtime/Engine/Classes/Kismet/KismetArrayLibrary.h
+++ b/Engine/Source/Runtime/Engine/Classes/Kismet/KismetArrayLibrary.h
@@ -122,11 +122,11 @@ class ENGINE_API UKismetArrayLibrary : public UBlueprintFunctionLibrary
 	static int32 Array_LastIndex(const TArray<int32>& TargetArray);
 
 	/*
-	 *Given an array and an index, returns the item found at that index
+	 *Given an array and an index, returns a copy of the item found at that index
 	 *
 	 *@param	TargetArray		The array to get an item from
 	 *@param	Index			The index in the array to get an item from
-	 *@return	The item stored at the index
+	 *@return	A copy of the item stored at the index
 	*/
 	UFUNCTION(BlueprintPure, CustomThunk, meta=(DisplayName = "Get", CompactNodeTitle = "GET", ArrayParm = "TargetArray", ArrayTypeDependentParams = "Item", BlueprintThreadSafe), Category="Utilities|Array")
 	static void Array_Get(const TArray<int32>& TargetArray, int32 Index, int32& Item);
@@ -185,7 +185,7 @@ class ENGINE_API UKismetArrayLibrary : public UBlueprintFunctionLibrary
 	 *@param	IndexToTest		The Index, that we want to test for being valid
 	 *@return	True if the Index is Valid, i.e. greater than or equal to zero, and less than the number of elements in TargetArray.
 	*/
-	UFUNCTION(BlueprintPure, CustomThunk, meta = (DisplayName = "Is Valid Index", CompactNodeTitle = "IS VALID INDEX", ArrayParm = "TargetArray", BlueprintThreadSafe), Category = "Utilities/Array")
+	UFUNCTION(BlueprintPure, CustomThunk, meta = (DisplayName = "Is Valid Index", CompactNodeTitle = "IS VALID INDEX", ArrayParm = "TargetArray", BlueprintThreadSafe), Category = "Utilities|Array")
 	static bool Array_IsValidIndex(const TArray<int32>& TargetArray, int32 IndexToTest);
 
 	// Native functions that will be called by the below custom thunk layers, which read off the property address, and call the appropriate native handler
diff --git a/Engine/Source/Runtime/Engine/Classes/Kismet/KismetMathLibrary.h b/Engine/Source/Runtime/Engine/Classes/Kismet/KismetMathLibrary.h
--- a/Engine/Source/Runtime/Engine/Classes/Kismet/KismetMathLibrary.h
+++ b/Engine/Source/Runtime/Engine/Classes/Kismet/KismetMathLibrary.h
@@ -852,6 +852,18 @@ class ENGINE_API UKismetMathLibrary : public UBlueprintFunctionLibrary
 	UFUNCTION(BlueprintPure, Category="Math|Vector", meta=(Keywords = "ProjectOnTo"))
 	static FVector ProjectVectorOnToVector(FVector V, FVector Target);
 
+	/** 
+	 * Given a direction vector and a surface normal, returns the vector reflected across the surface normal.
+	 * Produces a result like shining a laser at a mirror!
+	 *
+	 * @param Direction Direction vector the ray is coming from.
+	 * @param SurfaceNormal A normal of the surface the ray should be reflected on.
+	 *
+	 * @returns Reflected vector.
+	 */
+	UFUNCTION(BlueprintPure, Category="Math|Vector", meta=(Keywords = "Reflection"))
+	static FVector GetReflectionVector(FVector Direction, FVector SurfaceNormal);
+
 	/**
 	 * Find closest points between 2 segments.
 	 *
diff --git a/Engine/Source/Runtime/Engine/Private/ActorConstruction.cpp b/Engine/Source/Runtime/Engine/Private/ActorConstruction.cpp
--- a/Engine/Source/Runtime/Engine/Private/ActorConstruction.cpp
+++ b/Engine/Source/Runtime/Engine/Private/ActorConstruction.cpp
@@ -672,28 +672,30 @@ bool AActor::ExecuteConstruction(const FTransform& Transform, const FComponentIn
 	{
 		if (bErrorFree)
 		{
-			// Get all scene components owned by the given actor prior to SCS execution
+			// Get all components owned by the given actor prior to SCS execution.
 			// Note: GetComponents() internally does a NULL check, so we can assume here that all entries are valid.
-			TInlineComponentArray<USceneComponent*> AllSceneComponents;
-			GetComponents(AllSceneComponents);
+			TInlineComponentArray<UActorComponent*> PreSCSComponents;
+			GetComponents(PreSCSComponents);
 
 			// Determine the set of native scene components that SCS nodes can attach to.
 			TInlineComponentArray<USceneComponent*> NativeSceneComponents;
-			NativeSceneComponents.Reserve(AllSceneComponents.Num());
-			for (USceneComponent* SceneComponent : AllSceneComponents)
+			for (UActorComponent* ActorComponent : PreSCSComponents)
 			{
-				// Exclude subcomponents of native components, as these could unintentionally be matched by name during SCS execution. Also exclude instance-only components.
-				if (SceneComponent->CreationMethod == EComponentCreationMethod::Native && SceneComponent->GetOuter()->IsA<AActor>())
+				if (USceneComponent* SceneComponent = Cast<USceneComponent>(ActorComponent))
 				{
-					// If RootComponent is not set, the first unattached native scene component will be used as root. This matches what's done in FixupNativeActorComponents().
-					// @TODO - consider removing this; keeping here as a fallback just in case it wasn't set prior to SCS execution, but in most cases now this should be valid. 
-					if (RootComponent == nullptr && SceneComponent->GetAttachParent() == nullptr)
+					// Exclude subcomponents of native components, as these could unintentionally be matched by name during SCS execution. Also exclude instance-only components.
+					if (SceneComponent->CreationMethod == EComponentCreationMethod::Native && SceneComponent->GetOuter()->IsA<AActor>())
 					{
-						// Note: All native scene components should already have been registered at this point, so we don't need to register the component here.
-						SetRootComponent(SceneComponent);
-					}
+						// If RootComponent is not set, the first unattached native scene component will be used as root. This matches what's done in FixupNativeActorComponents().
+						// @TODO - consider removing this; keeping here as a fallback just in case it wasn't set prior to SCS execution, but in most cases now this should be valid. 
+						if (RootComponent == nullptr && SceneComponent->GetAttachParent() == nullptr)
+						{
+							// Note: All native scene components should already have been registered at this point, so we don't need to register the component here.
+							SetRootComponent(SceneComponent);
+						}
 
-					NativeSceneComponents.Add(SceneComponent);
+						NativeSceneComponents.Add(SceneComponent);
+					}
 				}
 			}
 
@@ -719,13 +721,14 @@ bool AActor::ExecuteConstruction(const FTransform& Transform, const FComponentIn
 				RegisterAllComponents();
 			}
 
-			// Once SCS execution has finished, we do a final pass to register any components that may have been deferred or were otherwise left unregistered after SCS execution.
+			// Once SCS execution has finished, we do a final pass to register any new components that may have been deferred or were otherwise left unregistered after SCS execution.
 			TInlineComponentArray<UActorComponent*> PostSCSComponents;
 			GetComponents(PostSCSComponents);
 			for (UActorComponent* ActorComponent : PostSCSComponents)
 			{
 				// Limit registration to components that are known to have been created during SCS execution
-				if (!ActorComponent->IsRegistered() && (ActorComponent->CreationMethod == EComponentCreationMethod::SimpleConstructionScript || !AllSceneComponents.Contains(ActorComponent)))
+				if (!ActorComponent->IsRegistered() && ActorComponent->bAutoRegister && !ActorComponent->IsPendingKill()
+					&& (ActorComponent->CreationMethod == EComponentCreationMethod::SimpleConstructionScript || !PreSCSComponents.Contains(ActorComponent)))
 				{
 					USimpleConstructionScript::RegisterInstancedComponent(ActorComponent);
 				}
diff --git a/Engine/Source/Runtime/Engine/Private/Blueprint.cpp b/Engine/Source/Runtime/Engine/Private/Blueprint.cpp
--- a/Engine/Source/Runtime/Engine/Private/Blueprint.cpp
+++ b/Engine/Source/Runtime/Engine/Private/Blueprint.cpp
@@ -17,6 +17,7 @@
 #include "Components/TimelineComponent.h"
 
 #if WITH_EDITOR
+#include "Editor/UnrealEd/Classes/Settings/ProjectPackagingSettings.h"
 #include "Engine/SimpleConstructionScript.h"
 #include "Engine/SCS_Node.h"
 #include "Kismet2/BlueprintEditorUtils.h"
@@ -427,6 +428,65 @@ void UBlueprint::Serialize(FArchive& Ar)
 		}
 	}
 
+	if (Ar.IsPersistent())
+	{
+		bool bSettingsChanged = false;
+		const FString PackageName = GetOutermost()->GetName();
+		UProjectPackagingSettings* PackagingSettings = GetMutableDefault<UProjectPackagingSettings>();
+
+		if (Ar.IsLoading())
+		{
+			if (bNativize_DEPRECATED)
+			{
+				// Migrate to the new transient flag.
+				bNativize_DEPRECATED = false;
+				bSelectedForNativization = true;
+
+				// Add this Blueprint asset to the exclusive list in the Project Settings (in case it doesn't exist).
+				bSettingsChanged |= PackagingSettings->AddBlueprintAssetToNativizationList(this);
+			}
+			else
+			{
+				// Cache whether or not this Blueprint asset was selected for exclusive nativization in the Project Settings.
+				for (int AssetIndex = 0; AssetIndex < PackagingSettings->NativizeBlueprintAssets.Num(); ++AssetIndex)
+				{
+					if (PackagingSettings->NativizeBlueprintAssets[AssetIndex].FilePath.Equals(PackageName, ESearchCase::IgnoreCase))
+					{
+						bSelectedForNativization = true;
+						break;
+					}
+				}
+			}
+		}
+		else if (Ar.IsSaving())
+		{
+			if (bSelectedForNativization)
+			{
+				// Auto-enable the exclusive method, but only if Blueprint nativization is not already turned on, and only if this will be the first asset selected for nativization.
+				if (PackagingSettings->BlueprintNativizationMethod == EProjectPackagingBlueprintNativizationMethod::Disabled && !PackagingSettings->NativizeBlueprintAssets.Num())
+				{
+					PackagingSettings->BlueprintNativizationMethod = EProjectPackagingBlueprintNativizationMethod::Exclusive;
+
+					bSettingsChanged = true;
+				}
+
+				// Add this Blueprint asset to the exclusive list in the Project Settings.
+				bSettingsChanged |= PackagingSettings->AddBlueprintAssetToNativizationList(this);
+			}
+			else
+			{
+				// Remove this Blueprint asset from the exclusive list in the Project Settings.
+				bSettingsChanged |= PackagingSettings->RemoveBlueprintAssetFromNativizationList(this);
+			}
+		}
+
+		if (bSettingsChanged)
+		{
+			// Update cached config settings and save.
+			PackagingSettings->SaveConfig();
+			PackagingSettings->UpdateDefaultConfigFile();
+		}
+	}
 #endif // WITH_EDITORONLY_DATA
 }
 
@@ -1246,7 +1306,7 @@ bool UBlueprint::CanEditChange(const UProperty* InProperty) const
 	if (bIsEditable && InProperty)
 	{
 		const FName PropertyName = InProperty->GetFName();
-		if (PropertyName == GET_MEMBER_NAME_CHECKED(UBlueprint, bNativize))
+		if (PropertyName == GET_MEMBER_NAME_CHECKED(UBlueprint, bSelectedForNativization))
 		{
 			return (BlueprintType != BPTYPE_LevelScript) && (BlueprintType != BPTYPE_MacroLibrary);
 		}
@@ -1540,6 +1600,38 @@ UInheritableComponentHandler* UBlueprint::GetInheritableComponentHandler(bool bC
 #endif
 
 #if WITH_EDITOR
+void UBlueprint::PostEditChangeProperty(struct FPropertyChangedEvent& PropertyChangedEvent)
+{
+	Super::PostEditChangeProperty(PropertyChangedEvent);
+
+	UProperty* ChangedProperty = PropertyChangedEvent.Property;
+	if (ChangedProperty != nullptr)
+	{
+		if (ChangedProperty->GetFName() == GET_MEMBER_NAME_CHECKED(UBlueprint, bSelectedForNativization) && bSelectedForNativization)
+		{
+			// Toggle the 'bSelectedForNativization' flag on for all parent class and interface dependencies.
+			TArray<UClass*> NativizeClassDependencies;
+			FBlueprintEditorUtils::FindImplementedInterfaces(this, false, NativizeClassDependencies);
+			NativizeClassDependencies.AddUnique(ParentClass);
+
+			for (UClass* NativizeClassDependency : NativizeClassDependencies)
+			{
+				if (UBlueprint* BP = GetBlueprintFromClass(NativizeClassDependency))
+				{
+					// Only toggle if not already turned on.
+					if (!BP->bSelectedForNativization)
+					{
+						BP->PreEditChange(ChangedProperty);
+						BP->bSelectedForNativization = true;
+
+						// This will recursively handle parent classes (as well as any interface classes they might implement).
+						BP->PostEditChangeProperty(PropertyChangedEvent);
+					}
+				}
+			}
+		}
+	}
+}
 
 FName UBlueprint::GetFunctionNameFromClassByGuid(const UClass* InClass, const FGuid FunctionGuid)
 {
diff --git a/Engine/Source/Runtime/Engine/Private/BlueprintMapLibrary.cpp b/Engine/Source/Runtime/Engine/Private/BlueprintMapLibrary.cpp
--- a/Engine/Source/Runtime/Engine/Private/BlueprintMapLibrary.cpp
+++ b/Engine/Source/Runtime/Engine/Private/BlueprintMapLibrary.cpp
@@ -71,7 +71,7 @@ void UBlueprintMapLibrary::GenericMap_Values(const void* TargetMap, const UMapPr
 		UProperty* InnerProp = ArrayProperty->Inner;
 		
 		int32 Size = MapHelper.Num();
-		for( int32 I = 0; I < Size; ++I )
+		for( int32 I = 0; Size; ++I )
 		{
 			if(MapHelper.IsValidIndex(I))
 			{
diff --git a/Engine/Source/Runtime/Engine/Private/ComponentInstanceDataCache.cpp b/Engine/Source/Runtime/Engine/Private/ComponentInstanceDataCache.cpp
--- a/Engine/Source/Runtime/Engine/Private/ComponentInstanceDataCache.cpp
+++ b/Engine/Source/Runtime/Engine/Private/ComponentInstanceDataCache.cpp
@@ -61,11 +61,26 @@ FActorComponentInstanceData::FActorComponentInstanceData(const UActorComponent*
 			FComponentPropertyWriter(const UActorComponent* Component, TArray<uint8>& InBytes)
 				: FObjectWriter(InBytes)
 			{
+				// Include properties that would normally skip tagged serialization (e.g. bulk serialization of array properties).
+				ArPortFlags |= PPF_ForceTaggedSerialization;
+
 				if (Component)
 				{
+					UClass* ComponentClass = Component->GetClass();
+
 					Component->GetUCSModifiedProperties(PropertiesToSkip);
 
-					UClass* ComponentClass = Component->GetClass();
+					if (AActor* ComponentOwner = Component->GetOwner())
+					{
+						// If this is the owning Actor's root scene component, don't include relative transform properties. This is handled elsewhere.
+						if (Component == ComponentOwner->GetRootComponent())
+						{
+							PropertiesToSkip.Add(ComponentClass->FindPropertyByName(GET_MEMBER_NAME_CHECKED(USceneComponent, RelativeLocation)));
+							PropertiesToSkip.Add(ComponentClass->FindPropertyByName(GET_MEMBER_NAME_CHECKED(USceneComponent, RelativeRotation)));
+							PropertiesToSkip.Add(ComponentClass->FindPropertyByName(GET_MEMBER_NAME_CHECKED(USceneComponent, RelativeScale3D)));
+						}
+					}
+
 					ComponentClass->SerializeTaggedProperties(*this, (uint8*)Component, ComponentClass, (uint8*)Component->GetArchetype());
 				}
 			}
@@ -158,6 +173,9 @@ void FActorComponentInstanceData::ApplyToComponent(UActorComponent* Component, c
 			FComponentPropertyReader(UActorComponent* InComponent, TArray<uint8>& InBytes)
 				: FObjectReader(InBytes)
 			{
+				// Include properties that would normally skip tagged serialization (e.g. bulk serialization of array properties).
+				ArPortFlags |= PPF_ForceTaggedSerialization;
+
 				InComponent->GetUCSModifiedProperties(PropertiesToSkip);
 
 				UClass* Class = InComponent->GetClass();
diff --git a/Engine/Source/Runtime/Engine/Private/Components/ActorComponent.cpp b/Engine/Source/Runtime/Engine/Private/Components/ActorComponent.cpp
--- a/Engine/Source/Runtime/Engine/Private/Components/ActorComponent.cpp
+++ b/Engine/Source/Runtime/Engine/Private/Components/ActorComponent.cpp
@@ -1660,6 +1660,9 @@ void UActorComponent::DetermineUCSModifiedProperties()
 				: FArchive()
 			{
 				ArIsSaving = true;
+
+				// Include properties that would normally skip tagged serialization (e.g. bulk serialization of array properties).
+				ArPortFlags |= PPF_ForceTaggedSerialization;
 			}
 
 			virtual bool ShouldSkipProperty(const UProperty* InProperty) const override
diff --git a/Engine/Source/Runtime/Engine/Private/Components/ChildActorComponent.cpp b/Engine/Source/Runtime/Engine/Private/Components/ChildActorComponent.cpp
--- a/Engine/Source/Runtime/Engine/Private/Components/ChildActorComponent.cpp
+++ b/Engine/Source/Runtime/Engine/Private/Components/ChildActorComponent.cpp
@@ -89,7 +89,7 @@ void UChildActorComponent::Serialize(FArchive& Ar)
 #if WITH_EDITOR
 	// Since we sometimes serialize properties in instead of using duplication if we are a template
 	// and are not pointing at a component we own we'll need to fix that
-	if (ChildActorTemplate && ChildActorTemplate->GetOuter() != this && IsTemplate())
+	if (!Ar.IsPersistent() && ChildActorTemplate && ChildActorTemplate->GetOuter() != this && IsTemplate())
 	{
 		const FString TemplateName = FString::Printf(TEXT("%s_%s_CAT"), *GetName(), *ChildActorClass->GetName());
 		ChildActorTemplate = CastChecked<AActor>(StaticDuplicateObject(ChildActorTemplate, this, *TemplateName));
diff --git a/Engine/Source/Runtime/Engine/Private/EdGraph/EdGraphPin.cpp b/Engine/Source/Runtime/Engine/Private/EdGraph/EdGraphPin.cpp
--- a/Engine/Source/Runtime/Engine/Private/EdGraph/EdGraphPin.cpp
+++ b/Engine/Source/Runtime/Engine/Private/EdGraph/EdGraphPin.cpp
@@ -153,6 +153,27 @@ bool FEdGraphPinType::Serialize(FArchive& Ar)
 	return true;
 }
 
+FEdGraphPinType FEdGraphPinType::GetPinTypeForTerminalType( const FEdGraphTerminalType& TerminalType )
+{
+	FEdGraphPinType TerminalTypeAsPin;
+	TerminalTypeAsPin.PinCategory = TerminalType.TerminalCategory;
+	TerminalTypeAsPin.PinSubCategory = TerminalType.TerminalSubCategory;
+	TerminalTypeAsPin.PinSubCategoryObject = TerminalType.TerminalSubCategoryObject;
+
+	return TerminalTypeAsPin;
+}
+
+FEdGraphPinType FEdGraphPinType::GetTerminalTypeForContainer( const FEdGraphPinType& ContainerType )
+{
+	ensure(ContainerType.IsContainer());
+	
+	FEdGraphPinType TerminalType = ContainerType;
+	TerminalType.bIsArray = false;
+	TerminalType.bIsMap = false;
+	TerminalType.bIsSet = false;
+	return TerminalType;
+}
+
 /////////////////////////////////////////////////////
 // FEdGraphPinReference
 void FEdGraphPinReference::SetPin(const UEdGraphPin* NewPin)
diff --git a/Engine/Source/Runtime/Engine/Private/InheritableComponentHandler.cpp b/Engine/Source/Runtime/Engine/Private/InheritableComponentHandler.cpp
--- a/Engine/Source/Runtime/Engine/Private/InheritableComponentHandler.cpp
+++ b/Engine/Source/Runtime/Engine/Private/InheritableComponentHandler.cpp
@@ -6,6 +6,7 @@
 #include "Engine/SCS_Node.h"
 #include "UObject/PropertyPortFlags.h"
 #include "UObject/LinkerLoad.h"
+#include "UObject/BlueprintsObjectVersion.h"
 
 #if WITH_EDITOR
 #include "Kismet2/BlueprintEditorUtils.h"
@@ -15,6 +16,13 @@
 
 const FString UInheritableComponentHandler::SCSDefaultSceneRootOverrideNamePrefix(TEXT("ICH-"));
 
+void UInheritableComponentHandler::Serialize(FArchive& Ar)
+{
+	Super::Serialize(Ar);
+
+	Ar.UsingCustomVersion(FBlueprintsObjectVersion::GUID);
+}
+
 void UInheritableComponentHandler::PostLoad()
 {
 	Super::PostLoad();
@@ -26,29 +34,32 @@ void UInheritableComponentHandler::PostLoad()
 			FComponentOverrideRecord& Record = Records[Index];
 			if (Record.ComponentTemplate)
 			{
-				// Fix up component class on load, if it's not already set.
-				if (Record.ComponentClass == nullptr)
+				if (GetLinkerCustomVersion(FBlueprintsObjectVersion::GUID) < FBlueprintsObjectVersion::SCSHasComponentTemplateClass)
 				{
-					Record.ComponentClass = Record.ComponentTemplate->GetClass();
-				}
+					// Fix up component class on load, if it's not already set.
+					if (Record.ComponentClass == nullptr)
+					{
+						Record.ComponentClass = Record.ComponentTemplate->GetClass();
+					}
 
-				// Fix up component template name on load, if it doesn't match the original template name. Otherwise, archetype lookups will fail for this template.
-				// For example, this can occur after a component variable rename in a parent BP class, but before a child BP class with an override template is loaded.
-				if (UActorComponent* OriginalTemplate = Record.ComponentKey.GetOriginalTemplate())
-				{
-					FString ExpectedTemplateName = OriginalTemplate->GetName();
-					if (USCS_Node* SCSNode = Record.ComponentKey.FindSCSNode())
+					// Fix up component template name on load, if it doesn't match the original template name. Otherwise, archetype lookups will fail for this template.
+					// For example, this can occur after a component variable rename in a parent BP class, but before a child BP class with an override template is loaded.
+					if (UActorComponent* OriginalTemplate = Record.ComponentKey.GetOriginalTemplate())
 					{
-						// We append a prefix onto SCS default scene root node overrides. This is done to ensure that the override template does not collide with our owner's own SCS default scene root node template.
-						if (SCSNode == SCSNode->GetSCS()->GetDefaultSceneRootNode())
+						FString ExpectedTemplateName = OriginalTemplate->GetName();
+						if (USCS_Node* SCSNode = Record.ComponentKey.FindSCSNode())
 						{
-							ExpectedTemplateName = SCSDefaultSceneRootOverrideNamePrefix + ExpectedTemplateName;
+							// We append a prefix onto SCS default scene root node overrides. This is done to ensure that the override template does not collide with our owner's own SCS default scene root node template.
+							if (SCSNode == SCSNode->GetSCS()->GetDefaultSceneRootNode())
+							{
+								ExpectedTemplateName = SCSDefaultSceneRootOverrideNamePrefix + ExpectedTemplateName;
+							}
 						}
-					}
 
-					if (ExpectedTemplateName != Record.ComponentTemplate->GetName())
-					{
-						Record.ComponentTemplate->Rename(*ExpectedTemplateName, nullptr, REN_DontCreateRedirectors | REN_ForceNoResetLoaders);
+						if (ExpectedTemplateName != Record.ComponentTemplate->GetName())
+						{
+							FixComponentTemplateName(Record.ComponentTemplate, ExpectedTemplateName);
+						}
 					}
 				}
 
@@ -466,6 +477,36 @@ const FComponentOverrideRecord* UInheritableComponentHandler::FindRecord(const F
 	return nullptr;
 }
 
+void UInheritableComponentHandler::FixComponentTemplateName(UActorComponent* ComponentTemplate, const FString& NewName)
+{
+	// Override template names were not previously kept in sync w/ past node rename operations. Thus, we need to check for
+	// and correct other (stale) template names. Otherwise, these could collide with the one we're trying to correct here.
+	for (int32 Index = 0; Index < Records.Num(); ++Index)
+	{
+		FComponentOverrideRecord& Record = Records[Index];
+		if (Record.ComponentTemplate && Record.ComponentTemplate != ComponentTemplate && Record.ComponentTemplate->GetName() == NewName)
+		{
+			if (UActorComponent* OriginalTemplate = Record.ComponentKey.GetOriginalTemplate())
+			{
+				if (OriginalTemplate->GetName() != Record.ComponentTemplate->GetName())
+				{
+					// Recursively fix up this record's component template name first to also match its original template, which will then free up the name.
+					FixComponentTemplateName(Record.ComponentTemplate, OriginalTemplate->GetName());
+				}
+			}
+
+			// There should only be at most one collision, so we'll stop looking now.
+			break;
+		}
+	}
+
+	// Precondition: There are no other objects in the same scope with this name.
+	check(!FindObjectWithOuter(ComponentTemplate->GetOuter(), nullptr, FName(*NewName)));
+
+	// Now that we're sure there are no collisions with other records, we can safely rename this one to its new name.
+	ComponentTemplate->Rename(*NewName, nullptr, REN_DontCreateRedirectors | REN_ForceNoResetLoaders);
+}
+
 // FComponentOverrideRecord
 
 FComponentKey::FComponentKey(const USCS_Node* SCSNode) : OwnerClass(nullptr)
diff --git a/Engine/Source/Runtime/Engine/Private/KismetMathLibrary.cpp b/Engine/Source/Runtime/Engine/Private/KismetMathLibrary.cpp
--- a/Engine/Source/Runtime/Engine/Private/KismetMathLibrary.cpp
+++ b/Engine/Source/Runtime/Engine/Private/KismetMathLibrary.cpp
@@ -393,6 +393,12 @@ FRotator UKismetMathLibrary::RandomRotator(bool bRoll)
 	return RRot;
 }
 
+
+FVector UKismetMathLibrary::GetReflectionVector(FVector Direction, FVector SurfaceNormal)
+{
+	return FMath::GetReflectionVector(Direction, SurfaceNormal);
+}
+
 FVector UKismetMathLibrary::FindClosestPointOnLine(FVector Point, FVector LineOrigin, FVector LineDirection)
 {
 	const FVector SafeDir = LineDirection.GetSafeNormal();
diff --git a/Engine/Source/Runtime/Engine/Private/SimpleConstructionScript.cpp b/Engine/Source/Runtime/Engine/Private/SimpleConstructionScript.cpp
--- a/Engine/Source/Runtime/Engine/Private/SimpleConstructionScript.cpp
+++ b/Engine/Source/Runtime/Engine/Private/SimpleConstructionScript.cpp
@@ -41,6 +41,8 @@ void USimpleConstructionScript::Serialize(FArchive& Ar)
 {
 	Super::Serialize(Ar);
 
+	Ar.UsingCustomVersion(FBlueprintsObjectVersion::GUID);
+
 	if(Ar.IsLoading())
 	{
 		if(Ar.UE4Ver() < VER_UE4_REMOVE_NATIVE_COMPONENTS_FROM_BLUEPRINT_SCS)
