ViewportInteraction: Fixed assert when trying to use editor extensions in Simulate mode
- Editor world extensions could not be bound to certain types of worlds; removed this restriction

diff --git a/Engine/Source/Editor/UnrealEd/Private/EditorWorldExtension.cpp b/Engine/Source/Editor/UnrealEd/Private/EditorWorldExtension.cpp
--- a/Engine/Source/Editor/UnrealEd/Private/EditorWorldExtension.cpp
+++ b/Engine/Source/Editor/UnrealEd/Private/EditorWorldExtension.cpp
@@ -227,12 +227,11 @@ TSharedPtr<FEditorWorldExtensionCollection> FEditorWorldExtensionManager::OnWorl
 	//Only add editor type world to the map
 	UWorld* World = InWorldContext.World();
 	TSharedPtr<FEditorWorldExtensionCollection> Result;
-	if( World && ( InWorldContext.WorldType == EWorldType::Editor || InWorldContext.WorldType == EWorldType::EditorPreview || InWorldContext.WorldType == EWorldType::GamePreview ) )
+	if( World != nullptr )
 	{
 		TSharedPtr<FEditorWorldExtensionCollection> ExtensionCollection( new FEditorWorldExtensionCollection( InWorldContext ) );
 		Result = ExtensionCollection;
 		EditorWorldExtensionCollection.Add( World->GetUniqueID(), Result );
-
 	}
 	return Result;
 }
