Adding alignment guidelines to transforming actors in the viewport. Use "EditorViewportSnapping.EnableGuides 1" to turn on.

diff --git a/Engine/Source/Editor/UnrealEd/Private/SnappingUtils.cpp b/Engine/Source/Editor/UnrealEd/Private/SnappingUtils.cpp
--- a/Engine/Source/Editor/UnrealEd/Private/SnappingUtils.cpp
+++ b/Engine/Source/Editor/UnrealEd/Private/SnappingUtils.cpp
@@ -20,6 +20,11 @@
 #include "ISnappingPolicy.h"
 #include "ViewportSnappingModule.h"
 
+namespace EditorViewportSnapping
+{
+	static FAutoConsoleVariable EnableGuides(TEXT("EditorViewportSnapping.EnableGuides"), 0, TEXT("Whether or not guidelines should be enabled. Off by default, set to 1 to enable."));
+}
+
 //////////////////////////////////////////////////////////////////////////
 // FEditorViewportSnapping
 
@@ -366,14 +371,194 @@ bool FEditorViewportSnapping::SnapToBSPVertex(FVector& Location, FVector GridBas
 
 void FEditorViewportSnapping::ClearSnappingHelpers( bool bClearImmediately )
 {
-	VertexSnappingImpl.ClearSnappingHelpers( bClearImmediately );
+	VertexSnappingImpl.ClearSnappingHelpers(bClearImmediately);
 }
 
-void FEditorViewportSnapping::DrawSnappingHelpers(const FSceneView* View,FPrimitiveDrawInterface* PDI)
+void FEditorViewportSnapping::DrawSnappingHelpers(const FSceneView* View, FPrimitiveDrawInterface* PDI)
 {
+	struct Local
+	{
+		struct FGuideData
+		{
+			bool bGuideShouldDraw;
+			FVector SnapPoint;
+			FVector GuideStart;
+			FVector GuideEnd;
+			FLinearColor GuideColor;
+			float GuideLength;
+		};
+
+		static  TArray<FVector> FindPotentialSnapPoints_ActorSpace(const AActor* SnapActor)
+		{
+			TArray<FVector> PotentialSnapPoints;
+
+			FBox ActorSpaceBoundingBox = SnapActor->CalculateComponentsBoundingBoxInLocalSpace();
+			FVector BoxCenter;
+			FVector BoxExtents;
+
+			ActorSpaceBoundingBox.GetCenterAndExtents(BoxCenter, BoxExtents);
+
+			FVector PotentialSnapPoint = FVector::ZeroVector;
+			
+			// Potential snap points are:
+			// The center of each face
+			for (int32 X = -1; X < 2; ++X)
+			{
+				PotentialSnapPoint[0] = X*BoxExtents[0];
+				PotentialSnapPoints.Add(PotentialSnapPoint);
+				// The center of each edge
+				for (int32 Y = -1; Y < 2; ++Y)
+				{
+					PotentialSnapPoint[1] = Y*BoxExtents[1];
+					PotentialSnapPoints.Add(PotentialSnapPoint);
+					// Each corner
+					for (int32 Z = -1; Z < 2; ++Z)
+					{
+						PotentialSnapPoint[2] = Z*BoxExtents[2];
+						PotentialSnapPoints.Add(PotentialSnapPoint);
+					}
+				}
+			}
+			return PotentialSnapPoints;
+		}
+	};
+
+
+	if (EditorViewportSnapping::EnableGuides->GetInt() == 1)
+	{
+		// Guidelines draw differently at different ranges
+		const float LargeGuideDistance = 10.0f;
+		const float MediumGuideDistance = 1.0f;
+		const float SmallGuideDistance = 0.25f;
+
+		// Use the last selected actor as the basis for drawing guidelines. 
+		AActor* TransformingActor = GEditor->GetSelectedActors()->GetBottom<AActor>();
+		if (TransformingActor)
+		{
+			// Don't let the guide lines be shorter than the local bounding box extent in any direction. 
+			// This helps when objects are close together
+			const FBox ActorSpaceBoundingBox = TransformingActor->CalculateComponentsBoundingBoxInLocalSpace();
+			FVector ActorBoxExtents;
+			FVector ActorBoxCenter;
+			ActorSpaceBoundingBox.GetCenterAndExtents(ActorBoxCenter, ActorBoxExtents);
+			const FVector MinGuideLength = ActorBoxExtents - ActorBoxCenter;
+
+			// Get all the potential snap points on the transforming actor
+			const TArray<FVector> ActorTransformingSnapPoints = Local::FindPotentialSnapPoints_ActorSpace(TransformingActor);
+
+			// Find all possible candidates for alignment
+			// TODO: add the world grid
+			// TODO: remove any actors it might not make sense to align to
+			TArray<const AActor*> CandidateActors;
+			for (FActorIterator It(TransformingActor->GetWorld()); It; ++It)
+			{
+				const AActor* PossibleCandidateActor = *It;
+
+				// Don't align to yourself, the entire world, or any actors hidden in the editor
+				if (PossibleCandidateActor != TransformingActor
+					&&	PossibleCandidateActor != TransformingActor->GetWorld()->GetDefaultBrush()
+					&& PossibleCandidateActor->IsHiddenEd() == false)
+				{
+					{
+						CandidateActors.Add(PossibleCandidateActor);
+					}
+				}
+
+			}
+
+			for (const AActor* CandidateActor : CandidateActors)
+			{
+				// Get the actor space snap points for the stationary actor
+				const TArray<FVector> ActorCandidateSnapPoints = Local::FindPotentialSnapPoints_ActorSpace(CandidateActor);
+
+				// Set up the initial guide information for X, Y, and Z guide line
+				TArray<Local::FGuideData> GuideHelpers;
+				Local::FGuideData InitialGuideHelper;
+				InitialGuideHelper.bGuideShouldDraw = false;
+				InitialGuideHelper.SnapPoint = InitialGuideHelper.GuideStart = InitialGuideHelper.GuideEnd = FVector::ZeroVector;
+				InitialGuideHelper.GuideColor = FLinearColor::Black;
+				InitialGuideHelper.GuideLength = 10000000.0f;
+				GuideHelpers.Init(InitialGuideHelper, 3);
+
+				for (const FVector ActorCandidateSnapPoint : ActorCandidateSnapPoints)
+				{
+					// Transform the stationary snap point into world space
+					const FVector WorldCandidateSnapPoint = CandidateActor->ActorToWorld().TransformPosition(ActorCandidateSnapPoint);
+
+					// Check it against each moving snap point
+					for (const FVector ActorTransformingSnapPoint : ActorTransformingSnapPoints)
+					{
+						int32 NumberOfMatchingAxes = 0;
+						// Transform the moving snap point into world space
+						const FVector WorldTransformingSnapPoint = TransformingActor->ActorToWorld().TransformPosition(ActorTransformingSnapPoint);
+						// Transform the moving snap point into stationary local space
+						FVector CandidateActorTransformingSnapPoint = CandidateActor->ActorToWorld().Inverse().TransformPosition(WorldTransformingSnapPoint);
+
+						FLinearColor GuideColor;
+						for (int32 PointAxis = 0; PointAxis < 3; ++PointAxis)
+						{
+							if (FMath::IsNearlyEqual(CandidateActorTransformingSnapPoint[PointAxis], ActorCandidateSnapPoint[PointAxis], LargeGuideDistance))
+							{
+								GuideColor = FLinearColor::Blue;
+								NumberOfMatchingAxes++;
+								if (FMath::IsNearlyEqual(CandidateActorTransformingSnapPoint[PointAxis], ActorCandidateSnapPoint[PointAxis], SmallGuideDistance))
+								{
+									GuideColor = FLinearColor::Green;
+								}
+							}
+						}
+						// Points must match in at least two axes to create a guide line
+						if (NumberOfMatchingAxes >= 2)
+						{
+							for (int32 PointAxis = 0; PointAxis < 3; ++PointAxis)
+							{
+								// If the two points have a component that is not within the LargeGuideDistance of each other, that is the axis we are drawing along
+								if (!FMath::IsNearlyEqual(CandidateActorTransformingSnapPoint[PointAxis], ActorCandidateSnapPoint[PointAxis], LargeGuideDistance))
+								{
+									// The guide line is drawn from where the actor will be aligned to the stationary actor it is aligning to
+									FVector ActorGuideStart = ActorCandidateSnapPoint;
+									ActorGuideStart[PointAxis] = CandidateActorTransformingSnapPoint[PointAxis];
+									const FVector WorldGuideStart = CandidateActor->ActorToWorld().TransformPosition(ActorGuideStart);
+									GuideHelpers[PointAxis].bGuideShouldDraw = true;
+
+									// Try to find the shortest guide line, but also don't switch between two nearly-equal guides (prevents jitter between snap points)
+									// Also, don't let the new guide length drop below our minimum guide length
+									if ((WorldCandidateSnapPoint - WorldGuideStart).Size() < GuideHelpers[PointAxis].GuideLength &&
+										!FMath::IsNearlyEqual(GuideHelpers[PointAxis].GuideLength, (WorldCandidateSnapPoint - WorldGuideStart).Size(), 1.0f)
+										&& (WorldCandidateSnapPoint - WorldGuideStart).Size() >= MinGuideLength[PointAxis])
+									{
+										GuideHelpers[PointAxis].GuideLength = (WorldCandidateSnapPoint - WorldGuideStart).Size();
+										GuideHelpers[PointAxis].GuideStart = WorldGuideStart;
+										GuideHelpers[PointAxis].GuideEnd = WorldCandidateSnapPoint;
+										GuideHelpers[PointAxis].SnapPoint = WorldTransformingSnapPoint;
+										GuideHelpers[PointAxis].GuideColor = GuideColor;
+									}
+								}
+							}
+						}
+
+					}
+				}
+
+				// Draw any of the guide lines in X, Y, or Z that should draw
+				for (int32 PossibleGuides = 0; PossibleGuides < 3; ++PossibleGuides)
+				{
+					if (GuideHelpers[PossibleGuides].bGuideShouldDraw)
+					{
+						PDI->DrawLine(GuideHelpers[PossibleGuides].GuideStart, GuideHelpers[PossibleGuides].GuideEnd, GuideHelpers[PossibleGuides].GuideColor, SDPG_World, 2.0f, 1.0f);
+						PDI->DrawPoint(GuideHelpers[PossibleGuides].SnapPoint, GuideHelpers[PossibleGuides].GuideColor, 10.0f, SDPG_World);
+						PDI->DrawPoint(GuideHelpers[PossibleGuides].GuideEnd, GuideHelpers[PossibleGuides].GuideColor, 10.0f, SDPG_World);
+					}
+				}
+			}
+		}
+	}
+	
 	VertexSnappingImpl.DrawSnappingHelpers( View, PDI );
 }
 
+
+
 bool FEditorViewportSnapping::SnapLocationToNearestVertex( FVector& Location, const FVector2D& MouseLocation, FLevelEditorViewportClient* ViewportClient, FVector& OutVertexNormal, bool bDrawVertHelpers )
 {
 	bool bSnapped = false;
