Merging //UE4/Private-Geometry to Dev-VREditor (//UE4/Dev-VREditor) 3339870

diff --git a/Engine/Config/BaseEditorPerProjectUserSettings.ini b/Engine/Config/BaseEditorPerProjectUserSettings.ini
--- a/Engine/Config/BaseEditorPerProjectUserSettings.ini
+++ b/Engine/Config/BaseEditorPerProjectUserSettings.ini
@@ -153,12 +153,6 @@ bBreakOnExceptions=False
 bToolbarCustomization=False
 bBehaviorTreeEditor=False
 bEnableFindAndReplaceReferences =False
-; Enables a preview of the Unreal Editor in VR.  This adds a new tool bar button that allows you to toggle into "VR Mode" instantly.  This feature is still in development, but your feedback is appreciated!
-bEnableVREditing=false
-; If true, wearing a Vive or Oculus Rift headset will automatically enter VR Editing mode if Enable VR Editing is true. 
-bEnableAutoVREditMode=true
-; If true, the VR Tutorial window will always be shown at startup
-bAlwaysShowVRTutorial=true
 
 [/Script/UnrealEd.EditorLoadingSavingSettings]
 ; True if we should automatically load a default level at start up
diff --git a/Engine/Plugins/MovieScene/LevelSequenceEditor/Source/LevelSequenceEditor/Private/LevelSequenceEditorToolkit.cpp b/Engine/Plugins/MovieScene/LevelSequenceEditor/Source/LevelSequenceEditor/Private/LevelSequenceEditorToolkit.cpp
--- a/Engine/Plugins/MovieScene/LevelSequenceEditor/Source/LevelSequenceEditor/Private/LevelSequenceEditorToolkit.cpp
+++ b/Engine/Plugins/MovieScene/LevelSequenceEditor/Source/LevelSequenceEditor/Private/LevelSequenceEditorToolkit.cpp
@@ -279,12 +279,13 @@ void FLevelSequenceEditorToolkit::Initialize(const EToolkitMode::Type Mode, cons
 	FLevelSequenceEditorToolkit::OnOpened().Broadcast(*this);
 
 	{
-		const UWorld* World = CastChecked<UWorld>( GetLevelSequenceEditorPlaybackContext() );
+		UWorld* World = CastChecked<UWorld>( GetLevelSequenceEditorPlaybackContext() );
 		UVREditorMode* VRMode = Cast<UVREditorMode>( GEditor->GetEditorWorldExtensionsManager()->GetEditorWorldExtensions( World )->FindExtension( UVREditorMode::StaticClass() ) );
 		if (VRMode != nullptr)
 		{
 			VRMode->OnVREditingModeExit_Handler.BindSP(this, &FLevelSequenceEditorToolkit::HandleVREditorModeExit);
-			VRMode->SaveSequencerSettings(Sequencer->GetKeyAllEnabled(), Sequencer->GetAutoKeyMode());
+			USequencerSettings& SavedSequencerSettings = *Sequencer->GetSequencerSettings();
+			VRMode->SaveSequencerSettings(Sequencer->GetKeyAllEnabled(), Sequencer->GetAutoKeyMode(), SavedSequencerSettings);
 			// Override currently set autokey behavior to always autokey all
 			Sequencer->SetAutoKeyMode(EAutoKeyMode::KeyAll);
 			Sequencer->SetKeyAllEnabled(true);
@@ -725,13 +726,12 @@ void FLevelSequenceEditorToolkit::HandleActorAddedToSequencer(AActor* Actor, con
 
 void FLevelSequenceEditorToolkit::HandleVREditorModeExit()
 {
-	const UWorld* World = Cast<UWorld>(GetLevelSequenceEditorPlaybackContext());
+	UWorld* World = Cast<UWorld>(GetLevelSequenceEditorPlaybackContext());
 	UVREditorMode* VRMode = CastChecked<UVREditorMode>( GEditor->GetEditorWorldExtensionsManager()->GetEditorWorldExtensions( World )->FindExtension( UVREditorMode::StaticClass() ) );
 
-	// Reset auto key settings
+	// Reset sequencer settings
 	Sequencer->SetAutoKeyMode(VRMode->GetSavedEditorState().AutoKeyMode);
 	Sequencer->SetKeyAllEnabled(VRMode->GetSavedEditorState().bKeyAllEnabled);
-
 	VRMode->OnVREditingModeExit_Handler.Unbind();
 }
 
@@ -967,7 +967,7 @@ void FLevelSequenceEditorToolkit::HandleTrackMenuExtensionAddTrack(FMenuBuilder&
 
 bool FLevelSequenceEditorToolkit::OnRequestClose()
 {
-	const UWorld* World = CastChecked<UWorld>(GetLevelSequenceEditorPlaybackContext());
+	UWorld* World = CastChecked<UWorld>(GetLevelSequenceEditorPlaybackContext());
 	UVREditorMode* VRMode = Cast<UVREditorMode>(GEditor->GetEditorWorldExtensionsManager()->GetEditorWorldExtensions(World)->FindExtension(UVREditorMode::StaticClass()));
 	if (VRMode != nullptr)
 	{
diff --git a/Engine/Plugins/Runtime/OculusRift/Source/OculusRift/Private/OculusRiftHMD.cpp b/Engine/Plugins/Runtime/OculusRift/Source/OculusRift/Private/OculusRiftHMD.cpp
--- a/Engine/Plugins/Runtime/OculusRift/Source/OculusRift/Private/OculusRiftHMD.cpp
+++ b/Engine/Plugins/Runtime/OculusRift/Source/OculusRift/Private/OculusRiftHMD.cpp
@@ -26,6 +26,9 @@
 
 static TAutoConsoleVariable<int32> CStartInVRVar(TEXT("vr.bStartInVR"), 0, TEXT("Start in VR flag"));
 static TAutoConsoleVariable<int32> CGracefulExitVar(TEXT("vr.bExitGracefully"), 0, TEXT("Exit gracefully when forced by Universal Menu."));
+static TAutoConsoleVariable<float> CRightEyeResolutionScaling( TEXT( "rift.RightEyeResolutionScaling" ), 1.0f, TEXT( "Scales the resolution of the right eye only" ) );
+static TAutoConsoleVariable<int32> CLeftEyeOverscan( TEXT( "rift.LeftEyeOverscan" ), 0, TEXT( "Whether to render the left eye at a larger widescreen resolution, then crop a column from the center for the HMD.  This is a presentation mode feature." ) );
+static TAutoConsoleVariable<float> CLeftEyeOverscanAspect( TEXT( "rift.LeftEyeOverscanAspect" ), 16.0f / 9.0f, TEXT( "Aspect ratio when rift.LeftEyeOverscan is enabled" ) );
 
 //---------------------------------------------------
 // Oculus Rift Plugin Implementation
@@ -580,17 +583,17 @@ bool FOculusRiftHMD::OnStartGameFrame( FWorldContext& InWorldContext )
 		const OVR::Vector3f newRight(CurrentSettings->EyeRenderDesc[1].HmdToEyeOffset);
 		const OVR::Vector3f prevLeft(MasterSettings->EyeRenderDesc[0].HmdToEyeOffset);
 		const OVR::Vector3f prevRight(MasterSettings->EyeRenderDesc[1].HmdToEyeOffset);
-		if (newLeft != prevLeft || newRight != prevRight)
+		// for debugging purposes only: overriding IPD
+		if( CurrentSettings->Flags.bOverrideIPD )
 		{
-			const float newIAD = newRight.Distance(newLeft);
-			UE_LOG(LogHMD, Log, TEXT("IAD has changed, new IAD is %.4f meters"), newIAD);
+			check( CurrentSettings->InterpupillaryDistance >= 0 );
+			CurrentSettings->EyeRenderDesc[ 0 ].HmdToEyeOffset.x = -CurrentSettings->InterpupillaryDistance * 0.5f;
+			CurrentSettings->EyeRenderDesc[ 1 ].HmdToEyeOffset.x = CurrentSettings->InterpupillaryDistance * 0.5f;
 		}
-		// for debugging purposes only: overriding IPD
-		if (CurrentSettings->Flags.bOverrideIPD)
+		else if (newLeft != prevLeft || newRight != prevRight)
 		{
-			check(CurrentSettings->InterpupillaryDistance >= 0);
-			CurrentSettings->EyeRenderDesc[0].HmdToEyeOffset.x = -CurrentSettings->InterpupillaryDistance * 0.5f;
-			CurrentSettings->EyeRenderDesc[1].HmdToEyeOffset.x = CurrentSettings->InterpupillaryDistance * 0.5f;
+			const float newIAD = newRight.Distance(newLeft);
+			UE_LOG(LogHMD, Log, TEXT("IAD has changed, new IAD is %.4f meters"), newIAD);
 		}
 #endif // #if !UE_BUILD_SHIPPING
 		// Save EyeRenderDesc in main settings.
@@ -1752,7 +1755,9 @@ void FOculusRiftHMD::SetupViewFamily(FSceneViewFamily& InViewFamily)
 	auto frame = GetFrame();
 	check(frame);
 
-	InViewFamily.EngineShowFlags.MotionBlur = 0;
+	extern RENDERER_API TAutoConsoleVariable<int32> CVarAllowMotionBlurInVR;
+
+	InViewFamily.EngineShowFlags.MotionBlur = CVarAllowMotionBlurInVR->GetInt() != 0;
 	InViewFamily.EngineShowFlags.HMDDistortion = false;
 	InViewFamily.EngineShowFlags.StereoRendering = IsStereoEnabled();
 	if (frame->Settings->Flags.bPauseRendering)
@@ -2137,10 +2142,6 @@ void FOculusRiftHMD::UpdateHmdRenderInfo()
 		// Calc FOV, symmetrical, for each eye. 
 		CurrentSettings->EyeFov[0] = HmdDesc.DefaultEyeFov[0];
 		CurrentSettings->EyeFov[1] = HmdDesc.DefaultEyeFov[1];
-
-		// Calc FOV in radians
-		CurrentSettings->VFOVInRadians = GetVerticalFovRadians(CurrentSettings->EyeFov[0], CurrentSettings->EyeFov[1]);
-		CurrentSettings->HFOVInRadians = GetHorizontalFovRadians(CurrentSettings->EyeFov[0], CurrentSettings->EyeFov[1]);
 	}
 
 	const ovrSizei recommenedTex0Size = ovr_GetFovTextureSize(OvrSession, ovrEye_Left, CurrentSettings->EyeFov[0], 1.0f);
@@ -2153,6 +2154,27 @@ void FOculusRiftHMD::UpdateHmdRenderInfo()
 	CurrentSettings->IdealScreenPercentage = FMath::Max(float(idealRenderTargetSize.w) / float(HmdDesc.Resolution.w) * 100.f,
 														float(idealRenderTargetSize.h) / float(HmdDesc.Resolution.h) * 100.f);
 
+	if (!CurrentSettings->Flags.bOverrideFOV)
+	{
+		if( CLeftEyeOverscan->GetInt() != 0 )
+		{
+			// Compute a new FOV for the left eye that will "overscan" the HMD's viewport to allow for a
+			// much wider mirror image.  The user of the HMD will see the horizontally cropped section of
+			// view that would normally be rendered without the overscan, but viewers of the mirror display
+			// will see the entire view.
+			const ovrSizei OriginalEyeRenderTargetSize = ovr_GetFovTextureSize( OvrSession, ovrEye_Left, CurrentSettings->EyeFov[ 0 ], 1.0f );
+			const float OriginalEyeAspect = (float)OriginalEyeRenderTargetSize.w / (float)OriginalEyeRenderTargetSize.h;
+			const float NewEyeAspect = CLeftEyeOverscanAspect->GetFloat();
+
+			CurrentSettings->EyeFov[ 0 ].LeftTan *= NewEyeAspect / OriginalEyeAspect;
+			CurrentSettings->EyeFov[ 0 ].RightTan *= NewEyeAspect / OriginalEyeAspect;
+		}
+
+		// Calc FOV in radians
+		CurrentSettings->VFOVInRadians = GetVerticalFovRadians(CurrentSettings->EyeFov[0], CurrentSettings->EyeFov[1]);
+		CurrentSettings->HFOVInRadians = GetHorizontalFovRadians(CurrentSettings->EyeFov[0], CurrentSettings->EyeFov[1]);
+	}
+
 	// Override eye distance by the value from HMDInfo (stored in Profile).
 	if (!CurrentSettings->Flags.bOverrideIPD)
 	{
@@ -2242,7 +2264,13 @@ void FOculusRiftHMD::UpdateStereoRenderingParams()
 
 		for(int eyeIndex = 0; eyeIndex < ovrEye_Count; eyeIndex++)
 		{
-			vpSize[eyeIndex] = ovr_GetFovTextureSize(OvrSession, (ovrEyeType) eyeIndex, CurrentSettings->EyeFov[eyeIndex], pixelDensity);
+			float eyePixelDensity = pixelDensity;
+			if( eyeIndex == ovrEye_Right )
+			{
+				eyePixelDensity *= CRightEyeResolutionScaling->GetFloat();
+			}
+
+			vpSize[eyeIndex] = ovr_GetFovTextureSize(OvrSession, (ovrEyeType) eyeIndex, CurrentSettings->EyeFov[eyeIndex], eyePixelDensity);
 			vpSize[eyeIndex].w = FMath::Max(vpSize[eyeIndex].w, 1);
 			vpSize[eyeIndex].h = FMath::Max(vpSize[eyeIndex].h, 1);
 
@@ -2512,32 +2540,35 @@ void FOculusRiftHMD::OnBeginPlay(FWorldContext& InWorldContext)
 	CachedViewportWidget.Reset();
 	CachedWindow.Reset();
 
-#if WITH_EDITOR
-	// @TODO: add more values here.
-	// This call make sense when 'Play' is used from the Editor;
-	if (GIsEditor)
+	if( !GEnableVREditorHacks )
 	{
-		Settings->PositionOffset = FVector::ZeroVector;
-		Settings->BaseOrientation = FQuat::Identity;
-		Settings->BaseOffset = FVector::ZeroVector;
-		Settings->WorldToMetersScale = InWorldContext.World()->GetWorldSettings()->WorldToMeters;
-		//Settings->Flags.bWorldToMetersOverride = false;
-		InitDevice();
+#if WITH_EDITOR
+		// @TODO: add more values here.
+		// This call make sense when 'Play' is used from the Editor;
+		if( GIsEditor )
+		{
+			Settings->PositionOffset = FVector::ZeroVector;
+			Settings->BaseOrientation = FQuat::Identity;
+			Settings->BaseOffset = FVector::ZeroVector;
+			Settings->WorldToMetersScale = InWorldContext.World()->GetWorldSettings()->WorldToMeters;
+			//Settings->Flags.bWorldToMetersOverride = false;
+			InitDevice();
 
-		FApp::SetUseVRFocus(true);
-		FApp::SetHasVRFocus(true);
-		OnStartGameFrame(InWorldContext);
-	}
-	else
+			FApp::SetUseVRFocus( true );
+			FApp::SetHasVRFocus( true );
+			OnStartGameFrame( InWorldContext );
+		}
+		else
 #endif
-	{
-		MakeSureValidFrameExists(InWorldContext.World()->GetWorldSettings());
+		{
+			MakeSureValidFrameExists( InWorldContext.World()->GetWorldSettings() );
+		}
 	}
 }
 
 void FOculusRiftHMD::OnEndPlay(FWorldContext& InWorldContext)
 {
-	if (GIsEditor)
+	if (GIsEditor && !GEnableVREditorHacks)
 	{
 		// @todo vreditor: If we add support for starting PIE while in VR Editor, we don't want to kill stereo mode when exiting PIE
 		EnableStereo(false);
diff --git a/Engine/Plugins/Runtime/OculusRift/Source/OculusRift/Private/OculusRiftRender.cpp b/Engine/Plugins/Runtime/OculusRift/Source/OculusRift/Private/OculusRiftRender.cpp
--- a/Engine/Plugins/Runtime/OculusRift/Source/OculusRift/Private/OculusRiftRender.cpp
+++ b/Engine/Plugins/Runtime/OculusRift/Source/OculusRift/Private/OculusRiftRender.cpp
@@ -157,11 +157,6 @@ void FViewExtension::PreRenderView_RenderThread(FRHICommandListImmediate& RHICmd
 
 		CurrentFrame->PoseToOrientationAndPosition(CurrentFrame->EyeRenderPose[eyeIdx], GameEyeOrient, GameEyePosition);
 		const FQuat DeltaControlOrientation =  ViewOrientation * GameEyeOrient.Inverse();
-		// make sure we use the same viewrotation as we had on a game thread
-		if( !GEnableVREditorHacks )	// @todo vreditor: This assert goes off sometimes when dragging properties in the details pane (UE-27540)
-		{
-			check(View.ViewRotation == CurrentFrame->CachedViewRotation[eyeIdx]);
-		}
 
 		if (CurrentFrame->Flags.bOrientationChanged)
 		{
@@ -284,7 +279,15 @@ void FCustomPresent::CopyTexture_RenderThread(FRHICommandListImmediate& RHICmdLi
 	static FGlobalBoundShaderState BoundShaderState;
 	SetGlobalBoundShaderState(RHICmdList, FeatureLevel, BoundShaderState, RendererModule->GetFilterVertexDeclaration().VertexDeclarationRHI, *VertexShader, *PixelShader);
 
-	PixelShader->SetParameters(RHICmdList, TStaticSamplerState<SF_Bilinear>::GetRHI(), SrcTextureRHI);
+	const bool bSameSize = DstRect.Size() == SrcRect.Size();
+	if( bSameSize )
+	{
+		PixelShader->SetParameters(RHICmdList, TStaticSamplerState<SF_Point>::GetRHI(), SrcTextureRHI);
+	}
+	else
+	{
+		PixelShader->SetParameters(RHICmdList, TStaticSamplerState<SF_Bilinear>::GetRHI(), SrcTextureRHI);
+	}
 
 	RendererModule->DrawRectangle(
 		RHICmdList,
@@ -346,7 +349,7 @@ void FOculusRiftHMD::RenderTexture_RenderThread(class FRHICommandListImmediate&
 				const float SrcRectAspect = (float)SrcViewRect.Width() / (float)SrcViewRect.Height();
 				const float DstRectAspect = (float)DstViewRect.Width() / (float)DstViewRect.Height();
 
-				if( SrcRectAspect < 1.0f )
+				if( SrcRectAspect < DstRectAspect )
 				{
 					// Source is taller than destination
 					if( bCropToFill )
diff --git a/Engine/Plugins/Runtime/Steam/SteamVR/Source/SteamVR/Private/SteamVRHMD.cpp b/Engine/Plugins/Runtime/Steam/SteamVR/Source/SteamVR/Private/SteamVRHMD.cpp
--- a/Engine/Plugins/Runtime/Steam/SteamVR/Source/SteamVR/Private/SteamVRHMD.cpp
+++ b/Engine/Plugins/Runtime/Steam/SteamVR/Source/SteamVR/Private/SteamVRHMD.cpp
@@ -833,7 +833,10 @@ bool FSteamVRHMD::IsInLowPersistenceMode() const
 
 void FSteamVRHMD::OnEndPlay(FWorldContext& InWorldContext)
 {
-	EnableStereo(false);
+	if (!GEnableVREditorHacks)
+	{
+		EnableStereo(false);
+	}
 }
 
 bool FSteamVRHMD::OnStartGameFrame(FWorldContext& WorldContext)
@@ -1031,6 +1034,11 @@ bool FSteamVRHMD::IsStereoEnabled() const
 
 bool FSteamVRHMD::EnableStereo(bool bStereo)
 {
+	if( bStereoEnabled == bStereo )
+	{
+		return false;
+	}
+
 	if (bStereo && bIsQuitting)
 	{
 		// Cancel shutting down the vr subsystem if re-enabling stereo before we're done counting down
diff --git a/Engine/Shaders/MaterialTemplate.usf b/Engine/Shaders/MaterialTemplate.usf
--- a/Engine/Shaders/MaterialTemplate.usf
+++ b/Engine/Shaders/MaterialTemplate.usf
@@ -207,6 +207,13 @@ struct FMaterialPixelParameters
 	/** Used in mobile custom pp material to preserve original SceneColor Alpha */
 	half BackupSceneColorAlpha;
 #endif
+
+#if COMPILER_HLSL
+	// Workaround for "error X3067: 'GetObjectWorldPosition': ambiguous function call"
+	// Which happens when FMaterialPixelParameters and FMaterialVertexParameters have the same number of floats with the HLSL compiler ver 9.29.952.3111
+	// Function overload resolution appears to identify types based on how many floats / ints / etc they contain
+	uint Dummy;
+#endif
 };
 
 // @todo compat hack
diff --git a/Engine/Source/Editor/ContentBrowser/Private/SAssetView.cpp b/Engine/Source/Editor/ContentBrowser/Private/SAssetView.cpp
--- a/Engine/Source/Editor/ContentBrowser/Private/SAssetView.cpp
+++ b/Engine/Source/Editor/ContentBrowser/Private/SAssetView.cpp
@@ -3748,8 +3748,12 @@ FReply SAssetView::OnDraggingAssetItem( const FGeometry& MyGeometry, const FPoin
 			if( InAssetData.Num() > 0 )
 			{
 				UActorFactory* FactoryToUse = nullptr;
-				FEditorDelegates::OnAssetDragStarted.Broadcast( InAssetData, FactoryToUse );
-				if( MouseEvent.IsMouseButtonDown( EKeys::LeftMouseButton ) )
+				if( FEditorDelegates::OnAssetDragStarted.IsBound() )
+				{
+					FEditorDelegates::OnAssetDragStarted.Broadcast( InAssetData, FactoryToUse );
+					return FReply::Handled();
+				}
+				else if( MouseEvent.IsMouseButtonDown( EKeys::LeftMouseButton ) )
 				{
 					return FReply::Handled().BeginDragDrop( FAssetDragDropOp::New( InAssetData ) );
 				}
diff --git a/Engine/Source/Editor/DetailCustomizations/Private/FbxImportUIDetails.cpp b/Engine/Source/Editor/DetailCustomizations/Private/FbxImportUIDetails.cpp
--- a/Engine/Source/Editor/DetailCustomizations/Private/FbxImportUIDetails.cpp
+++ b/Engine/Source/Editor/DetailCustomizations/Private/FbxImportUIDetails.cpp
@@ -107,11 +107,6 @@ void FFbxImportUIDetails::CustomizeDetails( IDetailLayoutBuilder& DetailBuilder
 			Prop->SetOnPropertyValueChanged(FSimpleDelegate::CreateSP(this, &FFbxImportUIDetails::MeshImportModeChanged));
 			MeshCategory.AddProperty(Prop);
 		}
-		{
-			TSharedRef<IPropertyHandle> Prop = DetailBuilder.GetProperty(GET_MEMBER_NAME_CHECKED(UFbxImportUI, bImportAsSubDSurface));
-			Prop->SetOnPropertyValueChanged(FSimpleDelegate::CreateSP(this, &FFbxImportUIDetails::MeshImportModeChanged));
-			MeshCategory.AddProperty(Prop);
-		}
 	}
 
 	TSharedRef<IPropertyHandle> ImportMeshProp = DetailBuilder.GetProperty(GET_MEMBER_NAME_CHECKED(UFbxImportUI, bImportMesh));
diff --git a/Engine/Source/Editor/DetailCustomizations/Private/FbxImportUIDetails.h b/Engine/Source/Editor/DetailCustomizations/Private/FbxImportUIDetails.h
--- a/Engine/Source/Editor/DetailCustomizations/Private/FbxImportUIDetails.h
+++ b/Engine/Source/Editor/DetailCustomizations/Private/FbxImportUIDetails.h
@@ -33,7 +33,7 @@ public:
 	/** Called if the bImportMaterials changes */
 	void ImportMaterialsChanged();
 
-	/** Called if the mesh mode (static / skeletal / SubDSurface) changes */
+	/** Called if the mesh mode (static / skeletal) changes */
 	void MeshImportModeChanged();
 
 	/** Called if the import mesh option for skeletal meshes is changed */
diff --git a/Engine/Source/Editor/EditorStyle/Private/SlateEditorStyle.cpp b/Engine/Source/Editor/EditorStyle/Private/SlateEditorStyle.cpp
--- a/Engine/Source/Editor/EditorStyle/Private/SlateEditorStyle.cpp
+++ b/Engine/Source/Editor/EditorStyle/Private/SlateEditorStyle.cpp
@@ -4624,10 +4624,10 @@ void FSlateEditorStyle::FStyle::SetupLevelEditorStyle()
 		Set( "LevelEditor.FoliageMode.Selected", new IMAGE_BRUSH( "Icons/icon_Mode_Foliage_selected_40x", Icon40x40 ) );
 		Set( "LevelEditor.FoliageMode.Selected.Small", new IMAGE_BRUSH( "Icons/icon_Mode_Foliage_selected_40x", Icon20x20 ) );
 
-		Set( "LevelEditor.BspMode", new IMAGE_BRUSH( "Icons/icon_Mode_GeoEdit_40px", Icon40x40 ) );
-		Set( "LevelEditor.BspMode.Small", new IMAGE_BRUSH( "Icons/icon_Mode_GeoEdit_40px", Icon20x20 ) );
-		Set( "LevelEditor.BspMode.Selected", new IMAGE_BRUSH( "Icons/icon_Mode_GeoEdit-a_40px", Icon40x40 ) );
-		Set( "LevelEditor.BspMode.Selected.Small", new IMAGE_BRUSH( "Icons/icon_Mode_GeoEdit-a_40px", Icon20x20 ) );
+		Set( "LevelEditor.BspMode", new IMAGE_BRUSH( "Icons/icon_Mode_MeshEditor_40px", Icon40x40 ) );
+		Set( "LevelEditor.BspMode.Small", new IMAGE_BRUSH( "Icons/icon_Mode_MeshEditor_40px", Icon20x20 ) );
+		Set( "LevelEditor.BspMode.Selected", new IMAGE_BRUSH( "Icons/icon_Mode_MeshEditor_40px-a_40px", Icon40x40 ) );
+		Set( "LevelEditor.BspMode.Selected.Small", new IMAGE_BRUSH( "Icons/icon_Mode_MeshEditor_40px-a_40px", Icon20x20 ) );
 
 		Set( "LevelEditor.WorldProperties", new IMAGE_BRUSH( "Icons/icon_worldscript_40x", Icon40x40 ) );
 		Set( "LevelEditor.WorldProperties.Small", new IMAGE_BRUSH( "Icons/icon_worldscript_40x", Icon20x20 ) );
@@ -4644,8 +4644,8 @@ void FSlateEditorStyle::FStyle::SetupLevelEditorStyle()
 		Set( "LevelEditor.EditMatinee", new IMAGE_BRUSH( "Icons/icon_matinee_40x", Icon40x40 ) );
 		Set( "LevelEditor.EditMatinee.Small", new IMAGE_BRUSH( "Icons/icon_matinee_40x", Icon20x20 ) );
 
-		Set( "LevelEditor.ToggleVR", new IMAGE_BRUSH( "Icons/icon_ToggleVREditor_40x", Icon40x40 ) );
-		Set( "LevelEditor.ToggleVR.Small", new IMAGE_BRUSH( "Icons/icon_ToggleVREditor_40x", Icon20x20 ) );
+		Set( "LevelEditor.ToggleVR", new IMAGE_BRUSH( "Icons/VREditor/VR_Editor_Toolbar_Icon", Icon40x40 ) );
+		Set( "LevelEditor.ToggleVR.Small", new IMAGE_BRUSH( "Icons/VREditor/VR_Editor_Toolbar_Icon_Small", Icon20x20 ) );
 
 		Set( "MergeActors.MeshMergingTool", new IMAGE_BRUSH( "Icons/icon_MergeActors_MeshMerging_40x", Icon40x40 ) );
 		Set( "MergeActors.MeshProxyTool", new IMAGE_BRUSH( "Icons/icon_MergeActors_MeshProxy_40x", Icon40x40 ) );
diff --git a/Engine/Source/Editor/FoliageEdit/Private/FoliageEdMode.cpp b/Engine/Source/Editor/FoliageEdit/Private/FoliageEdMode.cpp
--- a/Engine/Source/Editor/FoliageEdit/Private/FoliageEdMode.cpp
+++ b/Engine/Source/Editor/FoliageEdit/Private/FoliageEdMode.cpp
@@ -14,6 +14,7 @@
 #include "Framework/Application/SlateApplication.h"
 #include "Framework/Commands/UICommandList.h"
 #include "Materials/Material.h"
+#include "Materials/MaterialInstanceDynamic.h"
 #include "Engine/CollisionProfile.h"
 #include "Engine/StaticMeshActor.h"
 #include "Components/ModelComponent.h"
@@ -53,6 +54,11 @@
 
 DEFINE_LOG_CATEGORY_STATIC(LogFoliage, Log, Warning);
 
+namespace VREd
+{
+	static FAutoConsoleVariable FoliageOpacity(TEXT("VREd.FoliageOpacity"), 0.02f, TEXT("The foliage brush opacity."));
+}
+
 //
 // FFoliageMeshUIInfo
 //
@@ -189,11 +195,13 @@ FEdModeFoliage::FEdModeFoliage()
 	, FoliageInteractor(nullptr)
 {
 	// Load resources and construct brush component
-	UMaterial* BrushMaterial = nullptr;
 	UStaticMesh* StaticMesh = nullptr;
 	if (!IsRunningCommandlet())
 	{
-		BrushMaterial = LoadObject<UMaterial>(nullptr, TEXT("/Engine/EditorLandscapeResources/FoliageBrushSphereMaterial.FoliageBrushSphereMaterial"), nullptr, LOAD_None, nullptr);
+		UMaterial* BrushMaterial = LoadObject<UMaterial>(nullptr, TEXT("/Engine/EditorLandscapeResources/FoliageBrushSphereMaterial.FoliageBrushSphereMaterial"), nullptr, LOAD_None, nullptr);
+		BrushMID = UMaterialInstanceDynamic::Create(BrushMaterial, GetTransientPackage());
+		check(BrushMID != nullptr);
+
 		StaticMesh = LoadObject<UStaticMesh>(nullptr, TEXT("/Engine/EngineMeshes/Sphere.Sphere"), nullptr, LOAD_None, nullptr);
 	}
 
@@ -201,12 +209,16 @@ FEdModeFoliage::FEdModeFoliage()
 	SphereBrushComponent->SetCollisionProfileName(UCollisionProfile::NoCollision_ProfileName);
 	SphereBrushComponent->SetCollisionObjectType(ECC_WorldDynamic);
 	SphereBrushComponent->SetStaticMesh(StaticMesh);
-	SphereBrushComponent->OverrideMaterials.Add(BrushMaterial);
+	SphereBrushComponent->SetMaterial(0, BrushMID);
 	SphereBrushComponent->SetAbsolute(true, true, true);
 	SphereBrushComponent->CastShadow = false;
 
 	bBrushTraceValid = false;
 	BrushLocation = FVector::ZeroVector;
+	
+	// Get the default opacity from the material.
+	FName OpacityParamName("OpacityAmount");
+	BrushMID->GetScalarParameterValue(OpacityParamName, DefaultBrushOpacity);
 
 	FFoliageEditCommands::Register();
 	UICommandList = MakeShareable(new FUICommandList);
@@ -364,6 +376,11 @@ void FEdModeFoliage::Enter()
 
 		ViewportWorldInteraction->OnViewportInteractionHoverUpdate().RemoveAll(this);
 		ViewportWorldInteraction->OnViewportInteractionHoverUpdate().AddRaw(this, &FEdModeFoliage::OnVRHoverUpdate);
+
+		// Hide the VR transform gizmo while we're in foliage mode.  It sort of gets in the way of painting.
+		ViewportWorldInteraction->SetTransformGizmoVisible( false );
+
+		SetBrushOpacity(VREd::FoliageOpacity->GetFloat());
 	}
 }
 
@@ -375,9 +392,15 @@ void FEdModeFoliage::Exit()
 		UViewportWorldInteraction* ViewportWorldInteraction = Cast<UViewportWorldInteraction>( GEditor->GetEditorWorldExtensionsManager()->GetEditorWorldExtensions( GetWorld() )->FindExtension( UViewportWorldInteraction::StaticClass() ) );
 		if (ViewportWorldInteraction != nullptr)
 		{
+			// Restore the transform gizmo visibility
+			ViewportWorldInteraction->SetTransformGizmoVisible( true );
+
 			ViewportWorldInteraction->OnViewportInteractionInputAction().RemoveAll(this);
 			ViewportWorldInteraction->OnViewportInteractionHoverUpdate().RemoveAll(this);
 			FoliageInteractor = nullptr;
+
+			// Reset the brush opacity to default.
+			SetBrushOpacity(DefaultBrushOpacity);
 		}
 	}
 
@@ -439,58 +462,63 @@ void FEdModeFoliage::Exit()
 
 void FEdModeFoliage::OnVRHoverUpdate(UViewportInteractor* Interactor, FVector& HoverImpactPoint, bool& bWasHandled)
 {
+	UVREditorMode* VREditorMode = Cast<UVREditorMode>( GEditor->GetEditorWorldExtensionsManager()->GetEditorWorldExtensions( GetWorld() )->FindExtension( UVREditorMode::StaticClass() ) );
+	if( VREditorMode != nullptr && VREditorMode->IsFullyInitialized() )
 	{
-		UVREditorMode* VREditorMode = Cast<UVREditorMode>( GEditor->GetEditorWorldExtensionsManager()->GetEditorWorldExtensions( GetWorld() )->FindExtension( UVREditorMode::StaticClass() ) );
-		if( VREditorMode != nullptr && VREditorMode->IsFullyInitialized() )
+		// Check if we're hovering over UI. If so, stop painting so we don't display the preview brush sphere
+		if (FoliageInteractor && (FoliageInteractor->IsHoveringOverPriorityType() || FoliageInteractor->GetDraggingMode() != EViewportInteractionDraggingMode::Nothing))
 		{
-			// Check if we're hovering over UI. If so, stop painting so we don't display the preview brush sphere
-			if (FoliageInteractor && FoliageInteractor->IsHoveringOverPriorityType())
-			{
-				EndFoliageBrushTrace();
-				FoliageInteractor = nullptr;
-			}
-			// If there isn't currently a foliage interactor and we are hovering over something valid
-			else if (FoliageInteractor == nullptr && Interactor->GetHitResultFromLaserPointer().GetActor() != nullptr)
-			{
-				FoliageInteractor = Interactor;
-			}
-			// If we aren't hovering over something valid and the tool isn't active
-			else if (Interactor->GetHitResultFromLaserPointer().GetActor() == nullptr && !bToolActive)
-			{
-				FoliageInteractor = nullptr;
-			}
+			EndFoliageBrushTrace();
+			FoliageInteractor = nullptr;
+		}
+		// If there isn't currently a foliage interactor and we are hovering over something valid
+		else if (FoliageInteractor == nullptr && !Interactor->IsHoveringOverPriorityType() && Interactor->GetHitResultFromLaserPointer().GetActor() != nullptr)
+		{
+			FoliageInteractor = Interactor;
+		}
+		// If we aren't hovering over something valid and the tool isn't active
+		else if (Interactor->GetHitResultFromLaserPointer().GetActor() == nullptr && !bToolActive)
+		{
+			FoliageInteractor = nullptr;
+		}
 
-			// Skip other interactors if we are painting with one
-			if (FoliageInteractor && Interactor == FoliageInteractor)
+		// Skip other interactors if we are painting with one
+		if (FoliageInteractor && Interactor == FoliageInteractor)
+		{
+			// Go ahead and paint immediately
+			FVector LaserPointerStart, LaserPointerEnd;
+			if (FoliageInteractor->GetLaserPointer( /* Out */ LaserPointerStart, /* Out */ LaserPointerEnd))
 			{
-				// Go ahead and paint immediately
-				FVector LaserPointerStart, LaserPointerEnd;
-				if (FoliageInteractor->GetLaserPointer( /* Out */ LaserPointerStart, /* Out */ LaserPointerEnd))
-				{
-					const FVector LaserPointerDirection = (LaserPointerEnd - LaserPointerStart).GetSafeNormal();
+				const FVector LaserPointerDirection = (LaserPointerEnd - LaserPointerStart).GetSafeNormal();
+
+				FoliageBrushTrace(nullptr, LaserPointerStart, LaserPointerDirection);
 
-					FoliageBrushTrace(nullptr, LaserPointerStart, LaserPointerDirection);
+				if( bBrushTraceValid )
+				{
+					HoverImpactPoint = BrushLocation;
+					bWasHandled = true;
 				}
 			}
-			
 		}
+			
+		const bool bBrushMeshVisible = !(FoliageInteractor == nullptr || Interactor->GetDraggingMode() != EViewportInteractionDraggingMode::Nothing);
+		SphereBrushComponent->SetVisibility(bBrushMeshVisible);
 	}
 }
 
 void FEdModeFoliage::OnVRAction(class FEditorViewportClient& ViewportClient, UViewportInteractor* Interactor, const FViewportActionKeyInput& Action, bool& bOutIsInputCaptured, bool& bWasHandled)
 {
 	UVREditorMode* VREditorMode = Cast<UVREditorMode>( GEditor->GetEditorWorldExtensionsManager()->GetEditorWorldExtensions( GetWorld() )->FindExtension( UVREditorMode::StaticClass() ) );
-	if (VREditorMode != nullptr && Interactor != nullptr)
+	if (VREditorMode != nullptr && Interactor != nullptr && Interactor->GetDraggingMode() == EViewportInteractionDraggingMode::Nothing)
 	{
-		const UVREditorInteractor* VRInteractor = Cast<UVREditorInteractor>(Interactor);
-
-		// Only allow light press
-		Interactor->SetAllowTriggerFullPress(false);
-		// Consume light press
-		if ((Action.ActionType == ViewportWorldActionTypes::SelectAndMove_LightlyPressed) && (!VREditorMode->IsShowingRadialMenu(VRInteractor)))
+		const UVREditorInteractor* VREditorInteractor = Cast<UVREditorInteractor>(Interactor);
+		if (Action.ActionType == ViewportWorldActionTypes::SelectAndMove && ( VREditorInteractor == nullptr || !VREditorMode->IsShowingRadialMenu( VREditorInteractor )))
 		{
-			if ((Action.Event == IE_Pressed) && (!VRInteractor->IsHoveringOverPriorityType()))
+			if (Action.Event == IE_Pressed && !Interactor->IsHoveringOverPriorityType())
 			{
+				bWasHandled = true;
+				bOutIsInputCaptured = true;
+
 				// Go ahead and paint immediately
 				FVector LaserPointerStart, LaserPointerEnd;
 				if (Interactor->GetLaserPointer( /* Out */ LaserPointerStart, /* Out */ LaserPointerEnd))
@@ -561,6 +589,9 @@ void FEdModeFoliage::OnVRAction(class FEditorViewportClient& ViewportClient, UVi
 			{
 				EndFoliageBrushTrace();
 				FoliageInteractor = nullptr;
+				
+				bWasHandled = true;
+				bOutIsInputCaptured = false;
 			}
 		}
 	}
@@ -1114,6 +1145,12 @@ bool FEdModeFoliage::VertexMaskCheck(const FHitResult& Hit, const UFoliageType*
 	return true;
 }
 
+void FEdModeFoliage::SetBrushOpacity(const float InOpacity)
+{
+	static FName OpacityParamName("OpacityAmount");
+	BrushMID->SetScalarParameterValue(OpacityParamName, InOpacity);
+}
+
 bool LandscapeLayerCheck(const FHitResult& Hit, const UFoliageType* Settings, FEdModeFoliage::LandscapeLayerCacheData& LandscapeLayersCache, float& OutHitWeight)
 {
 	OutHitWeight = 1.f;
diff --git a/Engine/Source/Editor/FoliageEdit/Private/FoliageEdMode.h b/Engine/Source/Editor/FoliageEdit/Private/FoliageEdMode.h
--- a/Engine/Source/Editor/FoliageEdit/Private/FoliageEdMode.h
+++ b/Engine/Source/Editor/FoliageEdit/Private/FoliageEdMode.h
@@ -116,7 +116,7 @@ struct FFoliageUISettings
 		, bShowPaletteItemTooltips(true)
 		, ActivePaletteViewMode(EFoliagePaletteViewMode::Thumbnail)
 		, PaletteThumbnailScale(0.3f)
-		, Radius(512.f)
+		, Radius(250.f)
 		, PaintDensity(0.5f)
 		, UnpaintDensity(0.f)
 		, IsInSingleInstantiationMode(false)
@@ -559,12 +559,21 @@ private:
 	/** Does a filter based on the vertex color of a static mesh */
 	static bool VertexMaskCheck(const FHitResult& Hit, const UFoliageType* Settings);
 
+	/** Set the brush mesh opacity */
+	void SetBrushOpacity(const float InOpacity);
+
 	bool bBrushTraceValid;
 	FVector BrushLocation;
 	FVector BrushNormal;
 	FVector BrushTraceDirection;
 	UStaticMeshComponent* SphereBrushComponent;
 
+	/** The dynamic material of the sphere brush. */
+	class UMaterialInstanceDynamic* BrushMID;
+
+	/** Default opacity received from the brush material to reset it when closing. */
+	float DefaultBrushOpacity;
+
 	// Landscape layer cache data
 	LandscapeLayerCacheData LandscapeLayerCaches;
 
diff --git a/Engine/Source/Editor/LandscapeEditor/Private/LandscapeEdMode.cpp b/Engine/Source/Editor/LandscapeEditor/Private/LandscapeEdMode.cpp
--- a/Engine/Source/Editor/LandscapeEditor/Private/LandscapeEdMode.cpp
+++ b/Engine/Source/Editor/LandscapeEditor/Private/LandscapeEdMode.cpp
@@ -571,7 +571,7 @@ void FEdModeLandscape::OnVRHoverUpdate(UViewportInteractor* Interactor, FVector&
 	if( InteractorPainting != nullptr && InteractorPainting == Interactor )
 	{
 		UVREditorMode* VREditorMode = Cast<UVREditorMode>( GEditor->GetEditorWorldExtensionsManager()->GetEditorWorldExtensions( GetWorld() )->FindExtension( UVREditorMode::StaticClass() ) );
-		if( VREditorMode != nullptr && VREditorMode->IsActive() && Interactor != nullptr )
+		if( VREditorMode != nullptr && VREditorMode->IsActive() && Interactor != nullptr && Interactor->GetDraggingMode() == EViewportInteractionDraggingMode::Nothing )
 		{
 			const UVREditorInteractor* VRInteractor = Cast<UVREditorInteractor>( Interactor );
 
@@ -601,9 +601,9 @@ void FEdModeLandscape::OnVRAction(FEditorViewportClient& ViewportClient, UViewpo
 	// Never show the traditional Unreal transform widget.  It doesn't work in VR because we don't have hit proxies.
 	ViewportClient.EngineShowFlags.SetModeWidgets(false);
 
-	if (VREditorMode != nullptr && VREditorMode->IsActive() && Interactor != nullptr)
+	if (VREditorMode != nullptr && VREditorMode->IsActive() && Interactor != nullptr && Interactor->GetDraggingMode() == EViewportInteractionDraggingMode::Nothing)
 	{
-		if (Action.ActionType == ViewportWorldActionTypes::SelectAndMove_LightlyPressed || Action.ActionType == ViewportWorldActionTypes::SelectAndMove)
+		if (Action.ActionType == ViewportWorldActionTypes::SelectAndMove)
 		{
 			const UVREditorInteractor* VRInteractor = Cast<UVREditorInteractor>(Interactor);
 
diff --git a/Engine/Source/Editor/LevelEditor/Private/LevelEditorActions.cpp b/Engine/Source/Editor/LevelEditor/Private/LevelEditorActions.cpp
--- a/Engine/Source/Editor/LevelEditor/Private/LevelEditorActions.cpp
+++ b/Engine/Source/Editor/LevelEditor/Private/LevelEditorActions.cpp
@@ -2083,7 +2083,7 @@ void FLevelEditorActionCallbacks::ToggleVR()
 bool FLevelEditorActionCallbacks::ToggleVR_CanExecute()
 {
 	IVREditorModule& VREditorModule = IVREditorModule::Get();
-	return VREditorModule.IsVREditorAvailable() && !GEditor->bIsSimulatingInEditor;
+	return VREditorModule.IsVREditorAvailable();
 }
 
 
@@ -2093,13 +2093,6 @@ bool FLevelEditorActionCallbacks::ToggleVR_IsChecked()
 	return VREditorModule.IsVREditorEnabled();
 }
 
-
-bool FLevelEditorActionCallbacks::ToggleVR_IsVisible()
-{
-	return GetDefault<UEditorExperimentalSettings>()->bEnableVREditing;
-}
-
-
 bool FLevelEditorActionCallbacks::CanSelectGameModeBlueprint()
 {
 	bool bCheckOutNeeded = false;
diff --git a/Engine/Source/Editor/LevelEditor/Private/LevelEditorToolBar.cpp b/Engine/Source/Editor/LevelEditor/Private/LevelEditorToolBar.cpp
--- a/Engine/Source/Editor/LevelEditor/Private/LevelEditorToolBar.cpp
+++ b/Engine/Source/Editor/LevelEditor/Private/LevelEditorToolBar.cpp
@@ -1269,7 +1269,7 @@ TSharedRef< SWidget > FLevelEditorToolBar::MakeLevelEditorToolBar( const TShared
 			FSlateIcon(FEditorStyle::GetStyleSetName(), "LevelEditor.EditMatinee") 
 			);
 
-		ToolbarBuilder.AddToolBarButton( FLevelEditorCommands::Get().ToggleVR, NAME_None, LOCTEXT("ToggleVR", "VR") );
+		ToolbarBuilder.AddToolBarButton( FLevelEditorCommands::Get().ToggleVR, NAME_None, LOCTEXT("ToggleVR", "VR Mode") );
 	}
 	ToolbarBuilder.EndSection();
 	
diff --git a/Engine/Source/Editor/LevelEditor/Private/SLevelEditor.cpp b/Engine/Source/Editor/LevelEditor/Private/SLevelEditor.cpp
--- a/Engine/Source/Editor/LevelEditor/Private/SLevelEditor.cpp
+++ b/Engine/Source/Editor/LevelEditor/Private/SLevelEditor.cpp
@@ -129,8 +129,7 @@ void SLevelEditor::BindCommands()
 		Actions.ToggleVR,
 		FExecuteAction::CreateStatic( &FLevelEditorActionCallbacks::ToggleVR ),
 		FCanExecuteAction::CreateStatic( &FLevelEditorActionCallbacks::ToggleVR_CanExecute ),
-		FIsActionChecked::CreateStatic( &FLevelEditorActionCallbacks::ToggleVR_IsChecked ),
-		FIsActionButtonVisible::CreateStatic( &FLevelEditorActionCallbacks::ToggleVR_IsVisible ) );
+		FIsActionChecked::CreateStatic( &FLevelEditorActionCallbacks::ToggleVR_IsChecked ) );
 
 	LevelEditorCommands->MapAction(
 		Actions.WorldProperties,
diff --git a/Engine/Source/Editor/LevelEditor/Private/SLevelViewport.cpp b/Engine/Source/Editor/LevelEditor/Private/SLevelViewport.cpp
--- a/Engine/Source/Editor/LevelEditor/Private/SLevelViewport.cpp
+++ b/Engine/Source/Editor/LevelEditor/Private/SLevelViewport.cpp
@@ -1541,13 +1541,14 @@ EVisibility SLevelViewport::OnGetViewportContentVisibility() const
 	{
 		return BaseVisibility;
 	}
-	return IsPlayInEditorViewportActive() && IsImmersive() ? EVisibility::Collapsed : EVisibility::Visible;
+
+	return ( ( IsPlayInEditorViewportActive() && IsImmersive() ) || GEngine->IsStereoscopic3D( ActiveViewport.Get() ) ) ? EVisibility::Collapsed : EVisibility::Visible;
 }
 
 EVisibility SLevelViewport::GetToolBarVisibility() const
 {
-	// Do not show the toolbar if this viewport has a play in editor session
-	return IsPlayInEditorViewportActive() ? EVisibility::Collapsed : OnGetViewportContentVisibility();
+	// Do not show the toolbar if this viewport has a play in editor session, or we're in the VR Editor
+	return ( IsPlayInEditorViewportActive() || GEngine->IsStereoscopic3D( ActiveViewport.Get() ) ) ? EVisibility::Collapsed : OnGetViewportContentVisibility();
 }
 
 EVisibility SLevelViewport::GetMaximizeToggleVisibility() const
@@ -3490,7 +3491,7 @@ void SLevelViewport::StartPlayInEditorSession(UGameViewportClient* PlayClient, c
 	ULevelEditorPlaySettings const* EditorPlayInSettings = GetDefault<ULevelEditorPlaySettings>();
 	check(EditorPlayInSettings);
 
-	if(EditorPlayInSettings->ShowMouseControlLabel)
+	if(EditorPlayInSettings->ShowMouseControlLabel && !GEngine->IsStereoscopic3D( ActiveViewport.Get() ) )
 	{
 		ELabelAnchorMode AnchorMode = EditorPlayInSettings->MouseControlLabelPosition.GetValue();
 
@@ -3646,31 +3647,41 @@ void SLevelViewport::EndPlayInEditorSession()
 
 	if( IsPlayInEditorViewportActive() )
 	{
-		ActiveViewport->OnPlayWorldViewportSwapped(*InactiveViewport);
+		{
+			TSharedPtr<FSceneViewport> GameViewport = ActiveViewport;
+			ActiveViewport = InactiveViewport;
+			ActiveViewport->OnPlayWorldViewportSwapped( *GameViewport );
 
-		// Play in editor viewport was active, swap back to our level editor viewport
-		ActiveViewport->SetViewportClient( nullptr );
+			// Play in editor viewport was active, swap back to our level editor viewport
+			GameViewport->SetViewportClient( nullptr );
 
-		// We should be the only thing holding on to viewports
-		check( ActiveViewport.IsUnique() );
+			// We should be the only thing holding on to viewports
+			check( GameViewport.IsUnique() );
+		}
 
-		ActiveViewport = InactiveViewport;
-	
 		// Ensure our active viewport is for level editing
 		check( ActiveViewport->GetClient() == LevelViewportClient.Get() );
 
-		// Restore camera settings that may be adversely affected by PIE
-		LevelViewportClient->RestoreCameraFromPIE();
-		RedrawViewport(true);
+		// If we're going back to VR Editor, refresh the level viewport's render target so the HMD will present frames here
+		if( GEngine->IsStereoscopic3D( ActiveViewport.Get() ) )
+		{
+			ActiveViewport->UpdateViewportRHI( false, ActiveViewport->GetSizeXY().X, ActiveViewport->GetSizeXY().Y, ActiveViewport->GetWindowMode(), PF_Unknown );
+		}
+		else
+		{
+			// Restore camera settings that may be adversely affected by PIE
+			LevelViewportClient->RestoreCameraFromPIE();
+			RedrawViewport(true);
+
+			// Remove camera roll from any PIE camera applied in this viewport. A rolled camera is hard to use for editing
+			LevelViewportClient->RemoveCameraRoll();
+		}
 	}
 	else
 	{
 		InactiveViewport->SetViewportClient( nullptr );
 	}
 
-	// Remove camera roll from any PIE camera applied in this viewport. A rolled camera is hard to use for editing
-	LevelViewportClient->RemoveCameraRoll();
-
 	// Reset the inactive viewport
 	InactiveViewport.Reset();
 
@@ -3739,7 +3750,7 @@ void SLevelViewport::SwapViewportsForPlayInEditor()
 	ULevelEditorPlaySettings const* EditorPlayInSettings = GetDefault<ULevelEditorPlaySettings>();
 	check(EditorPlayInSettings);
 	
-	if(EditorPlayInSettings->ShowMouseControlLabel)
+	if(EditorPlayInSettings->ShowMouseControlLabel && !GEngine->IsStereoscopic3D( ActiveViewport.Get() ) )
 	{
 		ELabelAnchorMode AnchorMode = EditorPlayInSettings->MouseControlLabelPosition.GetValue();
 		
diff --git a/Engine/Source/Editor/LevelEditor/Public/LevelEditorActions.h b/Engine/Source/Editor/LevelEditor/Public/LevelEditorActions.h
--- a/Engine/Source/Editor/LevelEditor/Public/LevelEditorActions.h
+++ b/Engine/Source/Editor/LevelEditor/Public/LevelEditorActions.h
@@ -654,7 +654,6 @@ public:
 	static void ToggleVR();
 	static bool ToggleVR_CanExecute();
 	static bool ToggleVR_IsChecked();
-	static bool ToggleVR_IsVisible();
 
 	/** Opens delta transform */
 	static void DeltaTransform();
diff --git a/Engine/Source/Editor/MeshPaint/Private/MeshPaintEdMode.cpp b/Engine/Source/Editor/MeshPaint/Private/MeshPaintEdMode.cpp
--- a/Engine/Source/Editor/MeshPaint/Private/MeshPaintEdMode.cpp
+++ b/Engine/Source/Editor/MeshPaint/Private/MeshPaintEdMode.cpp
@@ -255,11 +255,11 @@ void FEdModeMeshPaint::Enter()
 		ApplyOrRemoveForceBestLOD(/*bApply=*/ true, PaintingMeshLODIndex);
 	}
 
-	TSharedPtr<FEditorWorldExtensionCollection> EditorWorld = GEditor->GetEditorWorldExtensionsManager()->GetEditorWorldExtensions( GetWorld() );
-	if(EditorWorld.IsValid())
+	UEditorWorldExtensionCollection* ExtensionCollection = GEditor->GetEditorWorldExtensionsManager()->GetEditorWorldExtensions(GetWorld());
+	if(ExtensionCollection != nullptr)
 	{
 		// Register to find out about VR input events
-		UViewportWorldInteraction* ViewportWorldInteraction = Cast<UViewportWorldInteraction>(GEditor->GetEditorWorldExtensionsManager()->GetEditorWorldExtensions(GetWorld())->FindExtension(UViewportWorldInteraction::StaticClass()));
+		UViewportWorldInteraction* ViewportWorldInteraction = Cast<UViewportWorldInteraction>(ExtensionCollection->FindExtension(UViewportWorldInteraction::StaticClass()));
 		if(ViewportWorldInteraction != nullptr)
 		{
 			ViewportWorldInteraction->OnViewportInteractionInputAction().RemoveAll(this);
@@ -275,14 +275,18 @@ void FEdModeMeshPaint::Enter()
 void FEdModeMeshPaint::Exit()
 {
 	{
-		UViewportWorldInteraction* ViewportWorldInteraction = Cast<UViewportWorldInteraction>(GEditor->GetEditorWorldExtensionsManager()->GetEditorWorldExtensions(GetWorld())->FindExtension(UViewportWorldInteraction::StaticClass()));
-		if(ViewportWorldInteraction != nullptr )
+		UEditorWorldExtensionCollection* ExtensionCollection = GEditor->GetEditorWorldExtensionsManager()->GetEditorWorldExtensions(GetWorld());
+		if (ExtensionCollection != nullptr)
 		{
-			// Restore the transform gizmo visibility
-			ViewportWorldInteraction->SetTransformGizmoVisible( true );
+			UViewportWorldInteraction* ViewportWorldInteraction = Cast<UViewportWorldInteraction>(ExtensionCollection->FindExtension(UViewportWorldInteraction::StaticClass()));
+			if(ViewportWorldInteraction != nullptr )
+			{
+				// Restore the transform gizmo visibility
+				ViewportWorldInteraction->SetTransformGizmoVisible( true );
 
-			// Unregister from event handlers
-			ViewportWorldInteraction->OnViewportInteractionInputAction().RemoveAll( this );
+				// Unregister from event handlers
+				ViewportWorldInteraction->OnViewportInteractionInputAction().RemoveAll( this );
+			}
 		}
 	}
 
@@ -5425,9 +5429,9 @@ void FEdModeMeshPaint::OnVRAction( class FEditorViewportClient& ViewportClient,
 {
 	UVREditorMode* VREditorMode = Cast<UVREditorMode>( GEditor->GetEditorWorldExtensionsManager()->GetEditorWorldExtensions( GetWorld() )->FindExtension( UVREditorMode::StaticClass() ) );
 	UVREditorInteractor* VRInteractor = Cast<UVREditorInteractor>( Interactor );
-	if( VREditorMode != nullptr && VREditorMode->IsActive() && VRInteractor != nullptr )
+	if( VREditorMode != nullptr && VREditorMode->IsActive() && VRInteractor != nullptr && VRInteractor->GetDraggingMode() == EViewportInteractionDraggingMode::Nothing)
 	{
-		if( Action.ActionType == ViewportWorldActionTypes::SelectAndMove_LightlyPressed )
+		if( Action.ActionType == ViewportWorldActionTypes::SelectAndMove )
 		{
 			if( !bIsPainting && Action.Event == IE_Pressed && !VRInteractor->IsHoveringOverPriorityType() )
 			{
@@ -5465,10 +5469,6 @@ void FEdModeMeshPaint::OnVRAction( class FEditorViewportClient& ViewportClient,
 					bWasHandled = true;
 					bOutIsInputCaptured = true;
 
-					// Don't allow a "full press" to happen at all with this trigger pull.  We don't need full press for anything, and it
-					// would interrupt our light press.
-					VRInteractor->SetAllowTriggerFullPress( false );
-
 					StartPainting();
 					PaintingWithInteractorInVR = Interactor;
 
@@ -5498,6 +5498,12 @@ void FEdModeMeshPaint::OnVRAction( class FEditorViewportClient& ViewportClient,
 				bWasHandled = true;
 				bOutIsInputCaptured = false;
 			}
+
+			else if( bIsPainting )
+			{
+				// A different hand might be painting, so absorb the input
+				bOutIsInputCaptured = bWasHandled = ( Action.Event == IE_Released ? false : true );
+			}
 		}
 	}
 }
diff --git a/Engine/Source/Editor/MovieSceneTools/Private/Sections/ThumbnailSection.cpp b/Engine/Source/Editor/MovieSceneTools/Private/Sections/ThumbnailSection.cpp
--- a/Engine/Source/Editor/MovieSceneTools/Private/Sections/ThumbnailSection.cpp
+++ b/Engine/Source/Editor/MovieSceneTools/Private/Sections/ThumbnailSection.cpp
@@ -15,6 +15,7 @@
 #include "Widgets/Text/SInlineEditableTextBlock.h"
 #include "PropertyEditorModule.h"
 #include "IDetailsView.h"
+#include "IVREditorModule.h"
 
 
 #define LOCTEXT_NAMESPACE "FThumbnailSection"
@@ -228,7 +229,7 @@ int32 FThumbnailSection::OnPaintSection( FSequencerSectionPainter& InPainter ) c
 
 	static const float SectionThumbnailPadding = 4.f;
 
-	const ESlateDrawEffect::Type DrawEffects = InPainter.bParentEnabled ? ESlateDrawEffect::None : ESlateDrawEffect::DisabledEffect;
+	ESlateDrawEffect::Type DrawEffects = InPainter.bParentEnabled ? ESlateDrawEffect::None : ESlateDrawEffect::DisabledEffect;
 
 	int32 LayerId = InPainter.LayerId;
 
@@ -267,13 +268,22 @@ int32 FThumbnailSection::OnPaintSection( FSequencerSectionPainter& InPainter ) c
 
 		if (Fade <= 1.f)
 		{
+			DrawEffects |= ESlateDrawEffect::NoGamma;
+
+			if (IVREditorModule::Get().IsVREditorModeActive())
+			{
+				// In VR editor every widget is in the world and gamma corrected by the scene renderer.  Thumbnails will have already been gamma
+				// corrected and so they need to be reversed
+				DrawEffects |= ESlateDrawEffect::ReverseGamma;
+			}
+
 			FSlateDrawElement::MakeViewport(
 				InPainter.DrawElements,
 				LayerId,
 				PaintGeometry,
 				Thumbnail,
 				ThumbnailClipRect,
-				DrawEffects | ESlateDrawEffect::NoGamma,
+				DrawEffects,
 				FLinearColor(1.f, 1.f, 1.f, 1.f - Fade)
 				);
 		}
diff --git a/Engine/Source/Editor/Sequencer/Private/MovieSceneTrackEditor.cpp b/Engine/Source/Editor/Sequencer/Private/MovieSceneTrackEditor.cpp
--- a/Engine/Source/Editor/Sequencer/Private/MovieSceneTrackEditor.cpp
+++ b/Engine/Source/Editor/Sequencer/Private/MovieSceneTrackEditor.cpp
@@ -30,14 +30,12 @@ void FMovieSceneTrackEditor::AnimatablePropertyChanged(FOnKeyProperty OnKeyPrope
 			}
 
 			UpdatePlaybackRange();
-	
+
+
 			TSharedPtr<FSequencer> SequencerToUpdate = StaticCastSharedPtr<FSequencer>(Sequencer.Pin());
 			if (SequencerToUpdate.IsValid())
 			{
-				SequencerToUpdate->NotifyMovieSceneDataChanged(EMovieSceneDataChangeType::Unknown);
 				SequencerToUpdate->SynchronizeSequencerSelectionWithExternalSelection();
-				SequencerToUpdate->SynchronizeExternalSelectionWithSequencerSelection();
-				
 			}
 		}
 	}
diff --git a/Engine/Source/Editor/Sequencer/Private/Sequencer.cpp b/Engine/Source/Editor/Sequencer/Private/Sequencer.cpp
--- a/Engine/Source/Editor/Sequencer/Private/Sequencer.cpp
+++ b/Engine/Source/Editor/Sequencer/Private/Sequencer.cpp
@@ -134,8 +134,9 @@ struct FSequencerTemplateStore : FMovieSceneSequenceTemplateStore
 
 		if (TUniquePtr<FCachedMovieSceneEvaluationTemplate>* ExistingTemplate = Templates.Find(SequenceKey))
 		{
-			(*ExistingTemplate)->Regenerate(TemplateParameters);
-			return **ExistingTemplate;
+			FCachedMovieSceneEvaluationTemplate* Template = ExistingTemplate->Get();
+			Template->Regenerate(TemplateParameters);
+			return *Template;
 		}
 		else
 		{
@@ -502,6 +503,8 @@ void FSequencer::ResetToNewRootSequence(UMovieSceneSequence& NewSequence)
 	RootSequence = &NewSequence;
 	RestorePreAnimatedState();
 
+	RootTemplateInstance.Finish(*this);
+
 	TemplateStore->Reset();
 
 	ActiveTemplateIDs.Reset();
@@ -988,10 +991,14 @@ void FSequencer::NotifyMovieSceneDataChanged( EMovieSceneDataChangeType DataChan
 
 	if (DataChangeType == EMovieSceneDataChangeType::TrackValueChanged || 
 		DataChangeType == EMovieSceneDataChangeType::TrackValueChangedRefreshImmediately || 
-		DataChangeType == EMovieSceneDataChangeType::Unknown)
+		DataChangeType == EMovieSceneDataChangeType::Unknown ||
+		DataChangeType == EMovieSceneDataChangeType::MovieSceneStructureItemRemoved)
 	{
 		FSequencerEdMode* SequencerEdMode = (FSequencerEdMode*)(GLevelEditorModeTools().GetActiveMode(FSequencerEdMode::EM_SequencerMode));
-		SequencerEdMode->CleanUpMeshTrails();
+		if (SequencerEdMode != nullptr)
+		{
+			SequencerEdMode->CleanUpMeshTrails();
+		}
 	}
 
 	bNeedsEvaluate = true;
diff --git a/Engine/Source/Editor/Sequencer/Private/Sequencer.h b/Engine/Source/Editor/Sequencer/Private/Sequencer.h
--- a/Engine/Source/Editor/Sequencer/Private/Sequencer.h
+++ b/Engine/Source/Editor/Sequencer/Private/Sequencer.h
@@ -391,6 +391,7 @@ public:
 	/** Delegate handler for recording finishing */
 	void HandleRecordingFinished(UMovieSceneSequence* Sequence);
 
+	/** Whether or not playback is looping */
 	bool IsLooping() const;
 
 	/** Set the new global time, accounting for looping options */
@@ -618,6 +619,7 @@ public:
 	virtual FGuid MakeNewSpawnable(UObject& SourceObject) override;
 	virtual bool IsReadOnly() const override;
 	virtual void ExternalSelectionHasChanged() override { SynchronizeSequencerSelectionWithExternalSelection(); }
+	virtual USequencerSettings* GetSequencerSettings() override { return Settings; }
 
 public:
 
diff --git a/Engine/Source/Editor/Sequencer/Private/SequencerEdMode.cpp b/Engine/Source/Editor/Sequencer/Private/SequencerEdMode.cpp
--- a/Engine/Source/Editor/Sequencer/Private/SequencerEdMode.cpp
+++ b/Engine/Source/Editor/Sequencer/Private/SequencerEdMode.cpp
@@ -88,12 +88,12 @@ void FSequencerEdMode::Render(const FSceneView* View, FViewport* Viewport, FPrim
 
 #if WITH_EDITORONLY_DATA
 	// Draw spline trails using the PDI
-	if (View->Family->EngineShowFlags.Splines && !bDrawMeshTrails)
+	if (View->Family->EngineShowFlags.Splines)
 	{
 		DrawTracks3D(PDI);
 	}
 	// Draw mesh trails (doesn't use the PDI)
-	else 	
+	else if (bDrawMeshTrails)
 	{
 		PDI = nullptr;
 		DrawTracks3D(PDI);
@@ -121,7 +121,7 @@ void FSequencerEdMode::DrawHUD(FEditorViewportClient* ViewportClient,FViewport*
 
 void FSequencerEdMode::AddReferencedObjects(FReferenceCollector& Collector)
 {
-	for (FMeshTrailData MeshTrail : MeshTrails)
+	for (FMeshTrailData& MeshTrail : MeshTrails)
 	{
 		Collector.AddReferencedObject(MeshTrail.Track);
 		Collector.AddReferencedObject(MeshTrail.Trail);
@@ -148,20 +148,23 @@ void FSequencerEdMode::OnKeySelected(FViewport* Viewport, HMovieSceneKeyProxy* K
 void FSequencerEdMode::DrawMeshTransformTrailFromKey(const class ASequencerKeyActor* KeyActor)
 {
 	ASequencerMeshTrail* Trail = Cast<ASequencerMeshTrail>(KeyActor->GetOwner());
-	FMeshTrailData* TrailPtr = MeshTrails.FindByPredicate([Trail](const FMeshTrailData InTrail)
-	{
-		return Trail == InTrail.Trail;
-	});
-	if(TrailPtr)
+	if(Trail != nullptr)
 	{
-		// From the key, get the mesh trail, and then the track associated with that mesh trail
-		UMovieScene3DTransformTrack* Track = TrailPtr->Track;
-		// Draw a mesh trail for the key's associated actor
-		TArray<TWeakObjectPtr<UObject>> KeyObjects;
-		AActor* TrailActor = KeyActor->GetAssociatedActor();
-		KeyObjects.Add(TrailActor);
-		FPrimitiveDrawInterface* PDI = nullptr;
-		DrawTransformTrack(PDI, Track, KeyObjects, true);
+		FMeshTrailData* TrailPtr = MeshTrails.FindByPredicate([Trail](const FMeshTrailData InTrail)
+		{
+			return Trail == InTrail.Trail;
+		});
+		if(TrailPtr != nullptr)
+		{
+			// From the key, get the mesh trail, and then the track associated with that mesh trail
+			UMovieScene3DTransformTrack* Track = TrailPtr->Track;
+			// Draw a mesh trail for the key's associated actor
+			TArray<TWeakObjectPtr<UObject>> KeyObjects;
+			AActor* TrailActor = KeyActor->GetAssociatedActor();
+			KeyObjects.Add(TrailActor);
+			FPrimitiveDrawInterface* PDI = nullptr;
+			DrawTransformTrack(PDI, Track, KeyObjects, true);
+		}
 	}
 }
 
@@ -170,7 +173,10 @@ void FSequencerEdMode::CleanUpMeshTrails()
 	// Clean up any existing trails
 	for (FMeshTrailData& MeshTrail : MeshTrails)
 	{
-		MeshTrail.Trail->Cleanup();
+		if (MeshTrail.Trail)
+		{
+			MeshTrail.Trail->Cleanup();
+		}
 	}
 	MeshTrails.Empty();
 }
@@ -311,7 +317,6 @@ void FSequencerEdMode::DrawTransformTrack(FPrimitiveDrawInterface* PDI, UMovieSc
 
 					FVector OldPos_G = RefTM.TransformPosition(OldPos);
 					FVector NewKeyPos_G = RefTM.TransformPosition(NewKeyPos);
-					
 					// For constant interpolation - don't draw ticks - just draw dotted line.
 					if (bIsConstantKey)
 					{
@@ -332,7 +337,6 @@ void FSequencerEdMode::DrawTransformTrack(FPrimitiveDrawInterface* PDI, UMovieSc
 							GetLocationAtTime(TransformSection, NewTime, NewPos, NewRot);
 
 							FVector NewPos_G = RefTM.TransformPosition(NewPos);
-
 							if (PDI != nullptr)
 							{
 								PDI->DrawLine(OldPos_G, NewPos_G, TrackColor, SDPG_Foreground);
@@ -347,7 +351,7 @@ void FSequencerEdMode::DrawTransformTrack(FPrimitiveDrawInterface* PDI, UMovieSc
 								}
 								else if (TrailActor != nullptr)
 								{
-									TrailActor->AddFrameMeshComponent(NewTime, FTransform::FTransform(NewRot, NewPos_G, FVector::FVector(0.25f)));
+									TrailActor->AddFrameMeshComponent(NewTime, FTransform::FTransform(NewRot, NewPos, FVector::FVector(3.0f)));
 								}
 							}
 							OldTime = NewTime;
@@ -404,7 +408,7 @@ void FSequencerEdMode::DrawTransformTrack(FPrimitiveDrawInterface* PDI, UMovieSc
 				}
 				else if (TrailActor != nullptr)
 				{
-					TrailActor->AddKeyMeshActor(NewKeyTime, FTransform::FTransform(NewKeyRot, NewKeyPos_G, FVector::FVector(0.5f)), TransformSection);
+					TrailActor->AddKeyMeshActor(NewKeyTime, FTransform::FTransform(NewKeyRot, NewKeyPos, FVector::FVector(3.0f)), TransformSection);
 				}
 
 				//@todo
@@ -547,9 +551,13 @@ void FSequencerEdMode::DrawTracks3D(FPrimitiveDrawInterface* PDI)
 							});
 							if (TrailPtr == nullptr)
 							{
-								ASequencerMeshTrail* TrailActor = UViewportWorldInteraction::SpawnTransientSceneActor<ASequencerMeshTrail>(GetWorld(), TEXT("SequencerMeshTrail"), true);
-								FMeshTrailData MeshTrail = FMeshTrailData(TransformTrack, TrailActor);
-								MeshTrails.Add(MeshTrail);
+								UViewportWorldInteraction* WorldInteraction = Cast<UViewportWorldInteraction>( GEditor->GetEditorWorldExtensionsManager()->GetEditorWorldExtensions( GetWorld() )->FindExtension( UViewportWorldInteraction::StaticClass() ) );
+								if( WorldInteraction != nullptr )
+								{
+									ASequencerMeshTrail* TrailActor = WorldInteraction->SpawnTransientSceneActor<ASequencerMeshTrail>(TEXT("SequencerMeshTrail"), true);
+									FMeshTrailData MeshTrail = FMeshTrailData(TransformTrack, TrailActor);
+									MeshTrails.Add(MeshTrail);
+								}
 							}
 						}
 						DrawTransformTrack(PDI, TransformTrack, BoundObjects, ObjectBindingNode.Value);
diff --git a/Engine/Source/Editor/Sequencer/Private/SequencerKeyActor.cpp b/Engine/Source/Editor/Sequencer/Private/SequencerKeyActor.cpp
--- a/Engine/Source/Editor/Sequencer/Private/SequencerKeyActor.cpp
+++ b/Engine/Source/Editor/Sequencer/Private/SequencerKeyActor.cpp
@@ -4,6 +4,7 @@
 #include "Sections/MovieScene3DTransformSection.h"
 #include "SequencerEdMode.h"
 #include "Materials/MaterialInstance.h" 
+#include "Materials/MaterialInstanceDynamic.h"
 #include "Engine/Selection.h"
 #include "EditorModeTools.h"
 #include "EditorModeManager.h"
@@ -12,12 +13,32 @@
 ASequencerKeyActor::ASequencerKeyActor()
 	: Super()
 {
-	UStaticMesh* KeyEditorMesh = LoadObject<UStaticMesh>(nullptr, TEXT("/Engine/BasicShapes/Sphere"));
+	UStaticMesh* KeyEditorMesh = LoadObject<UStaticMesh>(nullptr, TEXT("/Engine/VREditor/TransformGizmo/SM_Sequencer_Key"));
 	check(KeyEditorMesh != nullptr);
-	UMaterialInstance* KeyEditorMaterial = LoadObject<UMaterialInstance>(nullptr, TEXT("/Engine/VREditor/LaserPointer/LaserPointerMaterial_Inst"));
+	UMaterial* KeyEditorMaterial = LoadObject<UMaterial>(nullptr, TEXT("/Engine/VREditor/TransformGizmo/Main"));
 	check(KeyEditorMaterial != nullptr);
-	GetStaticMeshComponent()->SetStaticMesh(KeyEditorMesh);
-	GetStaticMeshComponent()->SetMaterial(0, KeyEditorMaterial);
+
+	const bool bTransient = true;
+	USceneComponent* SceneComponent = CreateDefaultSubobject<USceneComponent>(TEXT("SceneComponent"), bTransient);
+	check(SceneComponent != nullptr);
+	this->RootComponent = SceneComponent;
+
+	KeyMeshComponent = CreateDefaultSubobject<UStaticMeshComponent>(TEXT("KeyMesh"));
+	KeyMeshComponent->SetMobility(EComponentMobility::Movable);
+	KeyMeshComponent->SetupAttachment(RootComponent);
+	KeyMeshComponent->SetStaticMesh(KeyEditorMesh);
+	KeyMeshComponent->CreateAndSetMaterialInstanceDynamicFromMaterial(0, KeyEditorMaterial);
+
+	KeyMeshComponent->SetCollisionEnabled(ECollisionEnabled::QueryOnly);
+	KeyMeshComponent->SetCollisionResponseToAllChannels(ECR_Ignore);
+	KeyMeshComponent->SetCollisionResponseToChannel(ECC_Visibility, ECR_Block);
+	
+	KeyMeshComponent->bGenerateOverlapEvents = false;
+	KeyMeshComponent->SetCanEverAffectNavigation(false);
+	KeyMeshComponent->bCastDynamicShadow = false;
+	KeyMeshComponent->bCastStaticShadow = false;
+	KeyMeshComponent->bAffectDistanceFieldLighting = false;
+	KeyMeshComponent->bAffectDynamicIndirectLighting = false;
 }
 
 
@@ -34,6 +55,13 @@ void ASequencerKeyActor::SetKeyData(class UMovieScene3DTransformSection* NewTrac
 	KeyTime = NewKeyTime;
 	// Associate the currently selected actor with this key
 	AssociatedActor = GEditor->GetSelectedActors()->GetTop<AActor>();
+	// Draw a single transform track based on the data from this key
+	FEditorViewportClient* ViewportClient = StaticCast<FEditorViewportClient*>(GEditor->GetActiveViewport()->GetClient());
+	if (ViewportClient != nullptr)
+	{
+		FSequencerEdMode* SequencerEdMode = (FSequencerEdMode*)(ViewportClient->GetModeTools()->GetActiveMode(FSequencerEdMode::EM_SequencerMode));
+		SequencerEdMode->DrawMeshTransformTrailFromKey(this);
+	}
 }
 
 
@@ -52,14 +80,6 @@ void ASequencerKeyActor::PropagateKeyChange()
 		TransYCurve.UpdateOrAddKey(KeyTime, GetActorTransform().GetLocation().Y);
 		TransZCurve.UpdateOrAddKey(KeyTime, GetActorTransform().GetLocation().Z);
 
-		// Update the rotation keys
-		FRichCurve& RotXCurve = TrackSection->GetRotationCurve(EAxis::X);
-		FRichCurve& RotYCurve = TrackSection->GetRotationCurve(EAxis::Y);
-		FRichCurve& RotZCurve = TrackSection->GetRotationCurve(EAxis::Z);
-		RotXCurve.UpdateOrAddKey(KeyTime, GetActorTransform().GetRotation().X);
-		RotYCurve.UpdateOrAddKey(KeyTime, GetActorTransform().GetRotation().Y);
-		RotZCurve.UpdateOrAddKey(KeyTime, GetActorTransform().GetRotation().Z);
-
 		// Draw a single transform track based on the data from this key
 		FEditorViewportClient* ViewportClient = StaticCast<FEditorViewportClient*>(GEditor->GetActiveViewport()->GetClient());
 		if (ViewportClient != nullptr)
diff --git a/Engine/Source/Editor/Sequencer/Private/SequencerKeyActor.h b/Engine/Source/Editor/Sequencer/Private/SequencerKeyActor.h
new file mode 100644
--- /dev/null
+++ b/Engine/Source/Editor/Sequencer/Private/SequencerKeyActor.h
@@ -0,0 +1,62 @@
+// Copyright 1998-2016 Epic Games, Inc. All Rights Reserved.
+
+#pragma once
+
+#include "Engine/StaticMeshActor.h"
+#include "SequencerKeyActor.generated.h"
+
+UCLASS()
+class SEQUENCER_API ASequencerKeyActor : public AActor
+{
+	GENERATED_BODY()
+
+public:
+
+	ASequencerKeyActor();
+
+	/** AActor overrides */
+	virtual void PostEditMove(bool bFinished) override;
+	virtual bool IsEditorOnly() const final
+	{
+		return true;
+	}
+	/** AActor */
+
+	/** Set the track section and time associated with this key */
+	void SetKeyData(class UMovieScene3DTransformSection* NewTrackSection, float NewKeyTime);
+
+	/** Return the actor associated with this key */
+	AActor* GetAssociatedActor() const
+	{ 
+		return AssociatedActor; 
+	}
+
+	/** Returns the mesh component for this key, or nullptr if not spawned right now */
+	class UStaticMeshComponent* GetMeshComponent()
+	{
+		return KeyMeshComponent;
+	}
+
+protected:
+
+	/** Push the data from a key actor change back to Sequencer */
+	void PropagateKeyChange();
+
+private:
+
+	/** The key mesh */
+	UPROPERTY()
+	class UStaticMeshComponent* KeyMeshComponent;
+
+	/** The actor whose transform was used to build this key */
+	UPROPERTY()
+	AActor*	AssociatedActor;
+
+	/** The track section this key resides on */
+	UPROPERTY()
+	class UMovieScene3DTransformSection* TrackSection;
+
+	/** The time this key is associated with in Sequencer */
+	UPROPERTY()
+	float KeyTime;
+};
\ No newline at end of file
diff --git a/Engine/Source/Editor/Sequencer/Private/SequencerMeshTrail.cpp b/Engine/Source/Editor/Sequencer/Private/SequencerMeshTrail.cpp
--- a/Engine/Source/Editor/Sequencer/Private/SequencerMeshTrail.cpp
+++ b/Engine/Source/Editor/Sequencer/Private/SequencerMeshTrail.cpp
@@ -5,15 +5,25 @@
 #include "ViewportWorldInteraction.h"
 #include "EditorWorldExtension.h"
 #include "Materials/MaterialInstance.h" 
+#include "Materials/MaterialInstanceDynamic.h"
 #include "Components/BillboardComponent.h"
+#include "MovieScene3DTransformSection.h"
+#include "TimerManager.h"
 
 ASequencerMeshTrail::ASequencerMeshTrail()
 	: Super()
 {
+	TrailTime = 0.0f;
+	MaxTrailTime = 1.0f;
 	UBillboardComponent* PrimitiveRootComponent = CreateDefaultSubobject<UBillboardComponent>(TEXT("RootComponent"));
 	PrimitiveRootComponent->bSelectable = false;
 	PrimitiveRootComponent->SetVisibility(false, false);
 	SetRootComponent(PrimitiveRootComponent);
+
+	if (!IsTemplate() && GEditor)
+	{
+		GEditor->GetTimerManager()->SetTimer(TrailUpdate, FTimerDelegate::CreateUObject(this, &ASequencerMeshTrail::UpdateTrailAppearance, 0.0005f), 0.0005f, true);
+	}
 }
 
 void ASequencerMeshTrail::Cleanup()
@@ -31,7 +41,7 @@ void ASequencerMeshTrail::AddKeyMeshActor(float KeyTime, const FTransform KeyTra
 	ASequencerKeyActor* KeyMeshActor = nullptr;
 	UStaticMesh* KeyEditorMesh = nullptr;
 	UMaterialInstance* KeyEditorMaterial = nullptr;
-
+	MaxTrailTime = FMath::Max(KeyTime, MaxTrailTime);
 	FKeyActorData* KeyPtr = KeyMeshActors.FindByPredicate([KeyTime](const FKeyActorData InKey)
 	{
 		return FMath::IsNearlyEqual(KeyTime, InKey.Time);
@@ -40,13 +50,22 @@ void ASequencerMeshTrail::AddKeyMeshActor(float KeyTime, const FTransform KeyTra
 	// If we don't currently have an actor for this time, create one
 	if (KeyPtr == nullptr)
 	{
-		KeyMeshActor = UViewportWorldInteraction::SpawnTransientSceneActor<ASequencerKeyActor>(GetWorld(), TEXT("KeyMesh"), false);
-		KeyMeshActor->SetMobility(EComponentMobility::Movable);
-		KeyMeshActor->SetActorTransform(KeyTransform);
-		KeyMeshActor->SetKeyData(TrackSection, KeyTime);
-		KeyMeshActor->SetOwner(this);
-		FKeyActorData NewKey = FKeyActorData(KeyTime, KeyMeshActor);
-		KeyMeshActors.Add(NewKey);
+		UViewportWorldInteraction* WorldInteraction = Cast<UViewportWorldInteraction>( GEditor->GetEditorWorldExtensionsManager()->GetEditorWorldExtensions( GetWorld() )->FindExtension( UViewportWorldInteraction::StaticClass() ) );
+		if( WorldInteraction != nullptr )
+		{
+			KeyMeshActor = WorldInteraction->SpawnTransientSceneActor<ASequencerKeyActor>( TEXT( "KeyMesh" ), false );
+			KeyMeshActor->SetActorTransform(KeyTransform);
+			// Destroy all the frame components this trail created
+			for (FFrameComponentData FrameMesh : FrameMeshComponents)
+			{
+				FrameMesh.FrameComponent->DestroyComponent();
+			}
+			FrameMeshComponents.Reset();
+			KeyMeshActor->SetKeyData(TrackSection, KeyTime);
+			KeyMeshActor->SetOwner(this);
+			FKeyActorData NewKey = FKeyActorData(KeyTime, KeyMeshActor);
+			KeyMeshActors.Add(NewKey);
+		}
 	}
 	else
 	{
@@ -58,8 +77,8 @@ void ASequencerMeshTrail::AddKeyMeshActor(float KeyTime, const FTransform KeyTra
 void ASequencerMeshTrail::AddFrameMeshComponent(const float FrameTime, const FTransform FrameTransform)
 {
 	UStaticMeshComponent* FrameMeshComponent = nullptr;
-	UStaticMesh* KeyEditorMesh = nullptr;
-	UMaterialInterface* KeyEditorMaterial = nullptr;
+	UStaticMesh* FrameEditorMesh = nullptr;
+	UMaterial* FrameEditorMaterial = nullptr;
 
 	FFrameComponentData* FramePtr = FrameMeshComponents.FindByPredicate([FrameTime](const FFrameComponentData InFrame)
 	{
@@ -72,15 +91,16 @@ void ASequencerMeshTrail::AddFrameMeshComponent(const float FrameTime, const FTr
 		FrameMeshComponent = NewObject<UStaticMeshComponent>(this);
 		this->AddOwnedComponent(FrameMeshComponent);
 		FrameMeshComponent->RegisterComponent();
-
-		KeyEditorMesh = LoadObject<UStaticMesh>(nullptr, TEXT("/Engine/BasicShapes/Sphere"));
-		check(KeyEditorMesh != nullptr);
-		KeyEditorMaterial = LoadObject<UMaterialInstance>(nullptr, TEXT("/Engine/VREditor/LaserPointer/LaserPointerMaterial_Inst"));
-		check(KeyEditorMaterial != nullptr);
-		FrameMeshComponent->SetStaticMesh(KeyEditorMesh);
-		FrameMeshComponent->SetMaterial(0, KeyEditorMaterial);
-		FrameMeshComponent->SetWorldTransform(FrameTransform);
+		FrameEditorMesh = LoadObject<UStaticMesh>(nullptr, TEXT("/Engine/VREditor/TransformGizmo/SM_Sequencer_Node"));
+		check(FrameEditorMesh != nullptr);
+		FrameEditorMaterial = LoadObject<UMaterial>(nullptr, TEXT("/Engine/VREditor/TransformGizmo/Main"));
+		check(FrameEditorMaterial != nullptr);
+		FrameMeshComponent->SetStaticMesh(FrameEditorMesh);
+		FrameMeshComponent->CreateAndSetMaterialInstanceDynamicFromMaterial(0, FrameEditorMaterial);
 		FrameMeshComponent->SetMobility(EComponentMobility::Movable);
+		FrameMeshComponent->SetWorldTransform(FrameTransform);
+		FrameMeshComponent->SetCastShadow(false);
+		FrameMeshComponent->SetCollisionEnabled(ECollisionEnabled::NoCollision);
 		FrameMeshComponent->bSelectable = false;
 		FFrameComponentData NewFrame = FFrameComponentData(FrameTime, FrameMeshComponent);
 		FrameMeshComponents.Add(NewFrame);
@@ -91,3 +111,47 @@ void ASequencerMeshTrail::AddFrameMeshComponent(const float FrameTime, const FTr
 		FramePtr->FrameComponent->SetWorldTransform(FrameTransform);
 	}
 }
+
+void ASequencerMeshTrail::UpdateTrailAppearance(float UpdateTime)
+{
+	if (FMath::IsNearlyEqual(MaxTrailTime, 0.0f))
+	{
+		MaxTrailTime = 1.0f;
+	}
+	// TODO: No hardcoded glow values
+	TrailTime = FMath::Fmod(TrailTime + UpdateTime, MaxTrailTime);
+	for (FFrameComponentData Frame : FrameMeshComponents)
+	{
+		UMaterialInstanceDynamic* FrameMaterial = Cast<UMaterialInstanceDynamic>(Frame.FrameComponent->GetMaterial(0));
+		if (FrameMaterial != nullptr)
+		{
+			float GlowValue = 0.0f;
+			if (FMath::IsNearlyEqual(TrailTime, Frame.Time, 0.1f))
+			{
+				GlowValue = 12.0f;
+			}
+			else
+			{
+				GlowValue = 3.0f;
+			}
+			FrameMaterial->SetScalarParameterValue("GlowAmount", GlowValue);
+		}
+	}
+	for (FKeyActorData Key : KeyMeshActors)
+	{
+		UMaterialInstanceDynamic* KeyMaterial = Cast<UMaterialInstanceDynamic>(Key.KeyActor->GetMeshComponent()->GetMaterial(0));
+		if (KeyMaterial != nullptr)
+		{
+			float GlowValue = 0.0f;
+			if (FMath::IsNearlyEqual(TrailTime, Key.Time, 0.1f))
+			{
+				GlowValue = 20.0f;
+			}
+			else
+			{
+				GlowValue = 5.0f;
+			}
+			KeyMaterial->SetScalarParameterValue("GlowAmount", GlowValue);
+		}
+	}
+}
diff --git a/Engine/Source/Editor/Sequencer/Private/SequencerMeshTrail.h b/Engine/Source/Editor/Sequencer/Private/SequencerMeshTrail.h
--- a/Engine/Source/Editor/Sequencer/Private/SequencerMeshTrail.h
+++ b/Engine/Source/Editor/Sequencer/Private/SequencerMeshTrail.h
@@ -60,4 +60,9 @@ private:
 	TArray<FFrameComponentData> FrameMeshComponents;
 
 	TSharedPtr<class FSequencer> Sequencer;
+
+	void UpdateTrailAppearance(float UpdateTime);
+	float TrailTime;
+	float MaxTrailTime;
+	FTimerHandle TrailUpdate; 
 };
\ No newline at end of file
diff --git a/Engine/Source/Editor/Sequencer/Public/ISequencer.h b/Engine/Source/Editor/Sequencer/Public/ISequencer.h
--- a/Engine/Source/Editor/Sequencer/Public/ISequencer.h
+++ b/Engine/Source/Editor/Sequencer/Public/ISequencer.h
@@ -371,6 +371,9 @@ public:
 	/** Pause playback */
 	virtual void Pause() = 0;
 
+	/** Getter for sequencer settings */
+	virtual USequencerSettings* GetSequencerSettings() = 0;
+
 	/** Attempt to find a spawned object in the currently focused movie scene, or the template object for the specified binding ID, if possible */
 	virtual UObject* FindSpawnedObjectOrTemplate(const FGuid& BindingId) = 0;
 
diff --git a/Engine/Source/Editor/UnrealEd/Classes/Editor/EditorEngine.h b/Engine/Source/Editor/UnrealEd/Classes/Editor/EditorEngine.h
--- a/Engine/Source/Editor/UnrealEd/Classes/Editor/EditorEngine.h
+++ b/Engine/Source/Editor/UnrealEd/Classes/Editor/EditorEngine.h
@@ -55,6 +55,7 @@ class USoundBase;
 class USoundNode;
 class UTextureRenderTarget2D;
 struct FAnalyticsEventAttribute;
+class UEditorWorldExtensionManager;
 
 //
 // Things to set in mapSetBrush.
@@ -614,6 +615,12 @@ public:
 	/** The feature level we should use when loading or creating a new world */
 	ERHIFeatureLevel::Type DefaultWorldFeatureLevel;
 
+private:
+
+	/** Manager that holds all extensions paired with a world */
+	UPROPERTY()
+	UEditorWorldExtensionManager* EditorWorldExtensionsManager;
+
 protected:
 
 	/** Count of how many PIE instances are waiting to log in */
@@ -1819,7 +1826,7 @@ public:
 	/**
 	*  Returns the Editors world manager instance.
 	*/
-	TSharedRef<class FEditorWorldExtensionManager> GetEditorWorldExtensionsManager() { return EditorWorldExtensionsManager.ToSharedRef(); }
+	UEditorWorldExtensionManager* GetEditorWorldExtensionsManager() { return EditorWorldExtensionsManager; }
 
 	// Editor specific
 
@@ -2823,9 +2830,6 @@ private:
 	/** The Timer manager for all timer delegates */
 	TSharedPtr<class FTimerManager> TimerManager;
 
-	/** Manager that holds all the worlds contexts with viewport interaction */
-	TSharedPtr<class FEditorWorldExtensionManager> EditorWorldExtensionsManager;
-
 	/** The output log -> message log redirector for use during PIE */
 	TSharedPtr<class FOutputLogErrorsToMessageLogProxy> OutputLogErrorsToMessageLogProxyPtr;
 
diff --git a/Engine/Source/Editor/UnrealEd/Classes/Factories/FbxImportUI.h b/Engine/Source/Editor/UnrealEd/Classes/Factories/FbxImportUI.h
--- a/Engine/Source/Editor/UnrealEd/Classes/Factories/FbxImportUI.h
+++ b/Engine/Source/Editor/UnrealEd/Classes/Factories/FbxImportUI.h
@@ -22,8 +22,6 @@ enum EFBXImportType
 	FBXIT_SkeletalMesh UMETA(DisplayName="Skeletal Mesh"),
 	/** Select Animation if you'd like to import only animation. */
 	FBXIT_Animation UMETA(DisplayName="Animation"),
-	/** Subdivision Surface (Experimental, Early work in progress) */
-	FBXIT_SubDSurface UMETA(DisplayName="SubDSurface"),
 
 	FBXIT_MAX,
 };
@@ -55,9 +53,6 @@ public:
 	bool bImportAsSkeletal;
 	
 	/** Whether to import the incoming FBX as a Subdivision Surface (could be made a combo box together with bImportAsSkeletal) (Experimental, Early work in progress) */
-	UPROPERTY(EditAnywhere, AdvancedDisplay, Category = Mesh, meta = (ImportType = "StaticMesh|SkeletalMesh", DisplayName="Subdivison Surface"))
-	bool bImportAsSubDSurface;
-
 	/** Whether to import the mesh. Allows animation only import when importing a skeletal mesh. */
 	UPROPERTY(EditAnywhere, Category=Mesh, meta=(ImportType="SkeletalMesh"))
 	bool bImportMesh;
@@ -124,11 +119,6 @@ public:
 	void SetMeshTypeToImport()
 	{
 		MeshTypeToImport = bImportAsSkeletal ? FBXIT_SkeletalMesh : FBXIT_StaticMesh;
-
-		if(bImportAsSubDSurface)
-		{
-			MeshTypeToImport = FBXIT_SubDSurface;
-		}
 	}
 };
 
diff --git a/Engine/Source/Editor/UnrealEd/Classes/Settings/EditorExperimentalSettings.h b/Engine/Source/Editor/UnrealEd/Classes/Settings/EditorExperimentalSettings.h
--- a/Engine/Source/Editor/UnrealEd/Classes/Settings/EditorExperimentalSettings.h
+++ b/Engine/Source/Editor/UnrealEd/Classes/Settings/EditorExperimentalSettings.h
@@ -125,18 +125,6 @@ public:
 	UPROPERTY(EditAnywhere, config, Category = Tools, meta = (DisplayName = "Use OpenCL for convex hull decomposition"))
 	bool bUseOpenCLForConvexHullDecomp;
 
-	/** Enables a preview of the Unreal Editor in VR.  This adds a new tool bar button that allows you to toggle into "VR Mode" instantly.  This feature is still in development, but your feedback is appreciated! */
-	UPROPERTY(EditAnywhere, config, Category=VR, meta = (DisplayName="Enable VR Editing"))
-	bool bEnableVREditing;
-
-	/**If true, wearing a Vive or Oculus Rift headset will automatically enter VR Editing mode if Enable VR Editing is true. */
-	UPROPERTY(EditAnywhere, config, Category = VR, meta = (DisplayName = "Enable VR Mode Auto-Entry"))
-	bool bEnableAutoVREditMode;
-
-	/**If true, the tutorial window should always appear when the VR Editor is launched. */
-	UPROPERTY(EditAnywhere, config, Category = VR, meta = (DisplayName = "Always Show VR Tutorial at Start"))
-	bool bAlwaysShowVRTutorial;
-
 	/** Allows editing of potentially unsafe properties during PIE. Advanced use only - use with caution. */
 	UPROPERTY(EditAnywhere, config, Category = Tools, meta = (DisplayName = "Allow editing of potentially unsafe properties."))
 	bool bAllowPotentiallyUnsafePropertyEditing;
diff --git a/Engine/Source/Editor/UnrealEd/Private/AssetThumbnail.cpp b/Engine/Source/Editor/UnrealEd/Private/AssetThumbnail.cpp
--- a/Engine/Source/Editor/UnrealEd/Private/AssetThumbnail.cpp
+++ b/Engine/Source/Editor/UnrealEd/Private/AssetThumbnail.cpp
@@ -34,6 +34,7 @@
 #include "AssetToolsModule.h"
 #include "Styling/SlateIconFinder.h"
 #include "ClassIconFinder.h"
+#include "IVREditorModule.h"
 
 class SAssetThumbnail : public SCompoundWidget
 {
@@ -150,6 +151,9 @@ public:
 			TSharedPtr<SViewport> Viewport = 
 				SNew( SViewport )
 				.EnableGammaCorrection(false)
+				// In VR editor every widget is in the world and gamma corrected by the scene renderer.  Thumbnails will have already been gamma
+				// corrected and so they need to be reversed
+				.ReverseGammaCorrection(IVREditorModule::Get().IsVREditorModeActive())
 				.EnableBlending(true);
 
 			Viewport->SetViewportInterface( AssetThumbnail.ToSharedRef() );
diff --git a/Engine/Source/Editor/UnrealEd/Private/Editor.cpp b/Engine/Source/Editor/UnrealEd/Private/Editor.cpp
--- a/Engine/Source/Editor/UnrealEd/Private/Editor.cpp
+++ b/Engine/Source/Editor/UnrealEd/Private/Editor.cpp
@@ -88,6 +88,8 @@ FEditorDelegates::FOnEditorModeTransitioned				FEditorDelegates::EditorModeEnter
 FEditorDelegates::FOnEditorModeTransitioned				FEditorDelegates::EditorModeExit;
 FEditorDelegates::FOnPIEEvent							FEditorDelegates::PreBeginPIE;
 FEditorDelegates::FOnPIEEvent							FEditorDelegates::BeginPIE;
+FEditorDelegates::FOnPIEEvent							FEditorDelegates::PrePIEEnded;
+FEditorDelegates::FOnPIEEvent							FEditorDelegates::PostPIEStarted;
 FEditorDelegates::FOnPIEEvent							FEditorDelegates::EndPIE;
 FEditorDelegates::FOnPIEEvent							FEditorDelegates::PausePIE;
 FEditorDelegates::FOnPIEEvent							FEditorDelegates::ResumePIE;
diff --git a/Engine/Source/Editor/UnrealEd/Private/EditorEngine.cpp b/Engine/Source/Editor/UnrealEd/Private/EditorEngine.cpp
--- a/Engine/Source/Editor/UnrealEd/Private/EditorEngine.cpp
+++ b/Engine/Source/Editor/UnrealEd/Private/EditorEngine.cpp
@@ -317,6 +317,8 @@ UEditorEngine::UEditorEngine(const FObjectInitializer& ObjectInitializer)
 	NumOnlinePIEInstances = 0;
 	DefaultWorldFeatureLevel = GMaxRHIFeatureLevel;
 
+	EditorWorldExtensionsManager = nullptr;
+
 #if !UE_BUILD_SHIPPING
 	if (!AutomationCommon::OnEditorAutomationMapLoadDelegate().IsBound())
 	{
@@ -613,7 +615,7 @@ void UEditorEngine::InitEditor(IEngineLoop* InEngineLoop)
 	TimerManager = MakeShareable(new FTimerManager());
 
 	// create the editor world manager
-	EditorWorldExtensionsManager = MakeShareable(new FEditorWorldExtensionManager());
+	EditorWorldExtensionsManager = NewObject<UEditorWorldExtensionManager>();
 
 	// Settings.
 	FBSPOps::GFastRebuild = 0;
@@ -896,6 +898,7 @@ void UEditorEngine::Init(IEngineLoop* InEngineLoop)
 			TEXT("UndoHistory"),
 			TEXT("DeviceProfileEditor"),
 			TEXT("SourceCodeAccess"),
+			TEXT("MeshEditor"),
 			TEXT("BehaviorTreeEditor"),
 			TEXT("HardwareTargeting"),
 			TEXT("LocalizationDashboard"),
@@ -1699,8 +1702,8 @@ void UEditorEngine::Tick( float DeltaSeconds, bool bIdleMode )
 		}
 	}
 
-	// Updates the ViewportWorldInteraction
-	EditorWorldExtensionsManager->Tick( DeltaSeconds );
+	// Updates all the extensions for all the editor worlds
+	EditorWorldExtensionsManager->Tick(DeltaSeconds);
 
 	bool bIsMouseOverAnyLevelViewport = false;
 
diff --git a/Engine/Source/Editor/UnrealEd/Private/EditorModeManager.cpp b/Engine/Source/Editor/UnrealEd/Private/EditorModeManager.cpp
--- a/Engine/Source/Editor/UnrealEd/Private/EditorModeManager.cpp
+++ b/Engine/Source/Editor/UnrealEd/Private/EditorModeManager.cpp
@@ -18,6 +18,9 @@
 #include "Framework/Notifications/NotificationManager.h"
 #include "Widgets/Notifications/SNotificationList.h"
 #include "Engine/LevelStreaming.h"
+#include "ViewportWorldInteraction.h"
+#include "Public/EditorWorldExtension.h"
+#include "Editor/EditorEngine.h"
 
 /*------------------------------------------------------------------------------
 	FEditorModeTools.
@@ -217,7 +220,18 @@ void FEditorModeTools::SetPivotLocation( const FVector& Location, const bool bIn
 
 ECoordSystem FEditorModeTools::GetCoordSystem(bool bGetRawValue)
 {
-	if (!bGetRawValue && (GetWidgetMode() == FWidget::WM_Scale))
+	bool bAligningToActors = false;
+	if (GEditor->GetEditorWorldExtensionsManager() != nullptr)
+	{
+
+		UViewportWorldInteraction* ViewportWorldInteraction = Cast<UViewportWorldInteraction>(GEditor->GetEditorWorldExtensionsManager()->GetEditorWorldExtensions(GetWorld())->FindExtension(UViewportWorldInteraction::StaticClass()));
+		if (ViewportWorldInteraction != nullptr && ViewportWorldInteraction->AreAligningToActors() == true)
+		{
+			bAligningToActors = true;
+		}
+	}
+	if (!bGetRawValue && 
+		((GetWidgetMode() == FWidget::WM_Scale) || bAligningToActors))
 	{
 		return COORD_Local;
 	}
@@ -229,6 +243,20 @@ ECoordSystem FEditorModeTools::GetCoordSystem(bool bGetRawValue)
 
 void FEditorModeTools::SetCoordSystem(ECoordSystem NewCoordSystem)
 {
+	// If we are trying to enter world space but are aligning to actors, turn off aligning to actors
+	if (GEditor->GetEditorWorldExtensionsManager() != nullptr &&
+		NewCoordSystem == COORD_World)
+	{
+		UViewportWorldInteraction* ViewportWorldInteraction = Cast<UViewportWorldInteraction>(GEditor->GetEditorWorldExtensionsManager()->GetEditorWorldExtensions(GetWorld())->FindExtension(UViewportWorldInteraction::StaticClass()));
+		if (ViewportWorldInteraction != nullptr && ViewportWorldInteraction->AreAligningToActors() == true)
+		{
+			if (ViewportWorldInteraction->HasCandidatesSelected())
+			{
+				ViewportWorldInteraction->SetSelectionAsCandidates();
+			}
+			GUnrealEd->Exec(GetWorld(), TEXT("VI.EnableGuides 0"));
+		}
+	}
 	CoordSystem = NewCoordSystem;
 }
 
diff --git a/Engine/Source/Editor/UnrealEd/Private/EditorViewportClient.cpp b/Engine/Source/Editor/UnrealEd/Private/EditorViewportClient.cpp
--- a/Engine/Source/Editor/UnrealEd/Private/EditorViewportClient.cpp
+++ b/Engine/Source/Editor/UnrealEd/Private/EditorViewportClient.cpp
@@ -44,6 +44,7 @@
 #include "EditorBuildUtils.h"
 #include "AudioDevice.h"
 #include "EditorWorldExtension.h"
+#include "ViewportWorldInteraction.h"
 
 #define LOCTEXT_NAMESPACE "EditorViewportClient"
 
@@ -1418,8 +1419,16 @@ EMouseCursor::Type FEditorViewportClient::GetCursor(FViewport* InViewport,int32
 			}
 		}
 	}
-
-
+	
+	// Allow the viewport interaction to override any previously set mouse cursor
+	UViewportWorldInteraction* WorldInteraction = Cast<UViewportWorldInteraction>(GEditor->GetEditorWorldExtensionsManager()->GetEditorWorldExtensions(GetWorld())->FindExtension(UViewportWorldInteraction::StaticClass()));
+	if (WorldInteraction != nullptr && WorldInteraction->ShouldSuppressExistingCursor())
+	{
+			MouseCursor = EMouseCursor::None;
+			RequiredCursorVisibiltyAndAppearance.bHardwareCursorVisible = false;
+			RequiredCursorVisibiltyAndAppearance.bSoftwareCursorVisible = false;
+			UpdateRequiredCursorVisibility();
+	}
 
 	CachedMouseX = X;
 	CachedMouseY = Y;
@@ -2271,7 +2280,7 @@ bool FEditorViewportClient::InputKey(FViewport* InViewport, int32 ControllerId,
 		return true;
 	}
 
-	FEditorWorldExtensionCollection& EditorWorldExtensionCollection = *GEditor->GetEditorWorldExtensionsManager()->GetEditorWorldExtensions( GetWorld() );
+	UEditorWorldExtensionCollection& EditorWorldExtensionCollection = *GEditor->GetEditorWorldExtensionsManager()->GetEditorWorldExtensions( GetWorld() );
 	if( EditorWorldExtensionCollection.InputKey(this, Viewport, Key, Event))
 	{
 		return true;
diff --git a/Engine/Source/Editor/UnrealEd/Private/EditorWorldExtension.cpp b/Engine/Source/Editor/UnrealEd/Private/EditorWorldExtension.cpp
--- a/Engine/Source/Editor/UnrealEd/Private/EditorWorldExtension.cpp
+++ b/Engine/Source/Editor/UnrealEd/Private/EditorWorldExtension.cpp
@@ -3,20 +3,22 @@
 #include "EditorWorldExtension.h"
 #include "EngineGlobals.h"
 #include "Editor.h"
+#include "ILevelEditor.h"
+#include "LevelEditor.h"
 
 /************************************************************************/
 /* UEditorExtension		                                                */
 /************************************************************************/
 UEditorWorldExtension::UEditorWorldExtension() :
 	Super(),
-	OwningWorld( nullptr )
+	OwningExtensionsCollection( nullptr )
 {
 
 }
 
 UEditorWorldExtension::~UEditorWorldExtension()
 {
-	OwningWorld = nullptr;
+	OwningExtensionsCollection = nullptr;
 }
 
 bool UEditorWorldExtension::InputKey( FEditorViewportClient* InViewportClient, FViewport* Viewport, FKey Key, EInputEvent Event )
@@ -31,46 +33,183 @@ bool UEditorWorldExtension::InputAxis( FEditorViewportClient* InViewportClient,
 
 UWorld* UEditorWorldExtension::GetWorld() const
 {
-	return OwningWorld;
+	return OwningExtensionsCollection->GetWorld();
 }
 
-void UEditorWorldExtension::SetWorld( UWorld* InitWorld )
+AActor* UEditorWorldExtension::SpawnTransientSceneActor(TSubclassOf<AActor> ActorClass, const FString& ActorName, const bool bWithSceneComponent /*= false*/, const EObjectFlags InObjectFlags /*= EObjectFlags::RF_DuplicateTransient*/)
 {
-	OwningWorld = InitWorld;
+	UWorld* World = GetWorld();
+	check(World != nullptr);
+	const bool bWasWorldPackageDirty = World->GetOutermost()->IsDirty();
+
+	FActorSpawnParameters ActorSpawnParameters;
+	ActorSpawnParameters.Name = MakeUniqueObjectName(World, ActorClass, *ActorName);
+	ActorSpawnParameters.SpawnCollisionHandlingOverride = ESpawnActorCollisionHandlingMethod::AlwaysSpawn;
+	ActorSpawnParameters.ObjectFlags = InObjectFlags;
+
+	check(ActorClass != nullptr);
+	AActor* NewActor = World->SpawnActor< AActor >(ActorClass, ActorSpawnParameters);
+	NewActor->SetActorLabel(ActorName);
+
+	// Keep track of this actor so that we can migrate it between worlds if needed
+	ExtensionActors.Add( NewActor );
+
+	if (bWithSceneComponent)
+	{
+		// Give the new actor a root scene component, so we can attach multiple sibling components to it
+		USceneComponent* SceneComponent = NewObject<USceneComponent>(NewActor);
+		NewActor->AddOwnedComponent(SceneComponent);
+		NewActor->SetRootComponent(SceneComponent);
+		SceneComponent->RegisterComponent();
+	}
+
+	// Don't dirty the level file after spawning a transient actor
+	if (!bWasWorldPackageDirty)
+	{
+		World->GetOutermost()->SetDirtyFlag(false);
+	}
+
+	return NewActor;
+}
+
+void UEditorWorldExtension::DestroyTransientActor(AActor* Actor)
+{
+	if (Actor != nullptr)
+	{
+		ExtensionActors.RemoveSingleSwap(Actor);
+	}
+
+	UWorld* World = GetWorld();
+	check(World != nullptr);
+	if (Actor != nullptr)
+	{
+		const bool bWasWorldPackageDirty = World->GetOutermost()->IsDirty();
+
+		const bool bNetForce = false;
+		const bool bShouldModifyLevel = false;	// Don't modify level for transient actor destruction
+		World->DestroyActor(Actor, bNetForce, bShouldModifyLevel);
+
+		// Don't dirty the level file after destroying a transient actor
+		if (!bWasWorldPackageDirty)
+		{
+			World->GetOutermost()->SetDirtyFlag(false);
+		}
+	}
+}
+
+UEditorWorldExtensionCollection* UEditorWorldExtension::GetOwningCollection()
+{
+	return OwningExtensionsCollection;
+}
+
+bool UEditorWorldExtension::ExecCommand(const FString& InCommand)
+{
+	bool bResult = false;
+
+	UWorld* World = GetWorld();
+
+	// @todo vreditor: The following should not be needed.  It's a workaround because the input preprocessor
+	// in VREditor fires input events without setting GWorld to the PlayWorld during event processing.  This
+	// is inconsistent with the normal editor.  We're working around it by only using this logic when
+	// GIsPlayInEditorWorld is not set. (Which would be the case in Force VR Mode)
+	if (!GIsPlayInEditorWorld && GEditor->bIsSimulatingInEditor && GEditor->PlayWorld != nullptr && GEditor->PlayWorld == World)
+	{
+		UWorld* OldWorld = World;
+
+		// The play world needs to be selected if it exists
+		OldWorld = SetPlayInEditorWorld(OldWorld);
+
+		bResult = GUnrealEd->Exec(World, *InCommand);
+
+		// Restore the old world if there was one
+		if (OldWorld)
+		{
+			RestoreEditorWorld(OldWorld);
+		}
+	}
+	else
+	{
+		bResult = GUnrealEd->Exec(World, *InCommand);
+	}
+
+	return bResult;
+}
+
+void UEditorWorldExtension::TransitionWorld(UWorld* NewWorld)
+{
+	for (int32 ActorIndex = 0; ActorIndex < ExtensionActors.Num(); ++ActorIndex)
+	{
+		AActor* Actor = ExtensionActors[ ActorIndex ];
+		if( Actor != nullptr )
+		{
+			ReparentActor( Actor, NewWorld );
+		}
+		else
+		{
+			ExtensionActors.RemoveAtSwap( ActorIndex-- );
+		}
+	}
+}
+
+void UEditorWorldExtension::ReparentActor(AActor* Actor, UWorld* NewWorld)
+{
+	// Do not try to reparent the actor if it is in the same world as the requested new world
+	if(Actor->GetWorld() == NewWorld)
+	{
+		return;
+	}
+
+	ULevel* Level = NewWorld->PersistentLevel;
+	Actor->Rename(nullptr, Level);
+	Actor->DispatchBeginPlay();
+}
+
+void UEditorWorldExtension::InitInternal(UEditorWorldExtensionCollection* InOwningExtensionsCollection)
+{
+	OwningExtensionsCollection = InOwningExtensionsCollection;
 }
 
 /************************************************************************/
-/* FEditorWorldExtensionCollection                                      */
+/* UEditorWorldExtensionCollection                                      */
 /************************************************************************/
-FEditorWorldExtensionCollection::FEditorWorldExtensionCollection( FWorldContext& InWorldContext ) :
-	WorldContext ( InWorldContext )
+UEditorWorldExtensionCollection::UEditorWorldExtensionCollection() :
+	Super(),
+	EditorWorldOnSimulate(nullptr)
 {
-
+	if( !IsTemplate() )
+	{
+		FEditorDelegates::PostPIEStarted.AddUObject( this, &UEditorWorldExtensionCollection::PostPIEStarted );
+		FEditorDelegates::EndPIE.AddUObject( this, &UEditorWorldExtensionCollection::OnEndPIE );
+	}
 }
 
-FEditorWorldExtensionCollection::~FEditorWorldExtensionCollection()
+UEditorWorldExtensionCollection::~UEditorWorldExtensionCollection()
 {
+	FEditorDelegates::PostPIEStarted.RemoveAll( this );
+	FEditorDelegates::EndPIE.RemoveAll( this );
+
 	EditorExtensions.Empty();
+	EditorWorldOnSimulate = nullptr;
 }
 
-UWorld* FEditorWorldExtensionCollection::GetWorld() const
+UWorld* UEditorWorldExtensionCollection::GetWorld() const
 {
-	return WorldContext.World();
+	return WorldContext->World();
 }
 
-FWorldContext& FEditorWorldExtensionCollection::GetWorldContext() const
+FWorldContext* UEditorWorldExtensionCollection::GetWorldContext() const
 {
 	return WorldContext;
 }
 
-UEditorWorldExtension* FEditorWorldExtensionCollection::AddExtension( TSubclassOf<UEditorWorldExtension> EditorExtensionClass )
+UEditorWorldExtension* UEditorWorldExtensionCollection::AddExtension( TSubclassOf<UEditorWorldExtension> EditorExtensionClass )
 {
 	UEditorWorldExtension* CreatedExtension = EditorExtensionClass.GetDefaultObject();
 	AddExtension( CreatedExtension );
 	return CreatedExtension;
 }
 
-void FEditorWorldExtensionCollection::AddExtension( UEditorWorldExtension* EditorExtension )
+void UEditorWorldExtensionCollection::AddExtension( UEditorWorldExtension* EditorExtension )
 {
 	check( EditorExtension != nullptr );
 
@@ -92,12 +231,12 @@ void FEditorWorldExtensionCollection::AddExtension( UEditorWorldExtension* Edito
 		const int32 InitialRefCount = 1;
 		EditorExtensions.Add( FEditorExtensionTuple( EditorExtension, InitialRefCount ) );
 
-		EditorExtension->SetWorld( GetWorld() );
+		EditorExtension->InitInternal(this);
 		EditorExtension->Init();
 	}
 }
 
-void FEditorWorldExtensionCollection::RemoveExtension( UEditorWorldExtension* EditorExtension )
+void UEditorWorldExtensionCollection::RemoveExtension( UEditorWorldExtension* EditorExtension )
 {
 	check( EditorExtension != nullptr );
 
@@ -121,7 +260,7 @@ void FEditorWorldExtensionCollection::RemoveExtension( UEditorWorldExtension* Ed
 	}
 }
 
-UEditorWorldExtension* FEditorWorldExtensionCollection::FindExtension( TSubclassOf<UEditorWorldExtension> EditorExtensionClass )
+UEditorWorldExtension* UEditorWorldExtensionCollection::FindExtension( TSubclassOf<UEditorWorldExtension> EditorExtensionClass )
 {
 	UEditorWorldExtension* ResultExtension = nullptr;
 	for( FEditorExtensionTuple& EditorExtensionTuple : EditorExtensions )
@@ -137,19 +276,18 @@ UEditorWorldExtension* FEditorWorldExtensionCollection::FindExtension( TSubclass
 	return ResultExtension;
 }
 
-void FEditorWorldExtensionCollection::Tick( float DeltaSeconds )
+void UEditorWorldExtensionCollection::Tick( float DeltaSeconds )
 {
-	for( FEditorExtensionTuple& EditorExtensionTuple : EditorExtensions )
+	for (FEditorExtensionTuple& EditorExtensionTuple : EditorExtensions)
 	{
 		UEditorWorldExtension* EditorExtension = EditorExtensionTuple.Get<0>();
-		EditorExtension->Tick( DeltaSeconds );
+		EditorExtension->Tick(DeltaSeconds);
 	}
 }
 
-bool FEditorWorldExtensionCollection::InputKey( FEditorViewportClient* InViewportClient, FViewport* Viewport, FKey Key, EInputEvent Event )
+bool UEditorWorldExtensionCollection::InputKey( FEditorViewportClient* InViewportClient, FViewport* Viewport, FKey Key, EInputEvent Event )
 {
 	bool bHandled = false;
-
 	for( FEditorExtensionTuple& EditorExtensionTuple : EditorExtensions )
 	{
 		UEditorWorldExtension* EditorExtension = EditorExtensionTuple.Get<0>();
@@ -159,7 +297,7 @@ bool FEditorWorldExtensionCollection::InputKey( FEditorViewportClient* InViewpor
 	return bHandled;
 }
 
-bool FEditorWorldExtensionCollection::InputAxis( FEditorViewportClient* InViewportClient, FViewport* Viewport, int32 ControllerId, FKey Key, float Delta, float DeltaTime )
+bool UEditorWorldExtensionCollection::InputAxis( FEditorViewportClient* InViewportClient, FViewport* Viewport, int32 ControllerId, FKey Key, float Delta, float DeltaTime )
 {
 	bool bHandled = false;
 
@@ -172,27 +310,76 @@ bool FEditorWorldExtensionCollection::InputAxis( FEditorViewportClient* InViewpo
 	return bHandled;
 }
 
-void FEditorWorldExtensionCollection::AddReferencedObjects( FReferenceCollector& Collector )
+
+void UEditorWorldExtensionCollection::PostPIEStarted( bool bIsSimulatingInEditor )
 {
-	for( FEditorExtensionTuple& EditorExtensionTuple : EditorExtensions )
+	if( bIsSimulatingInEditor && GEditor->EditorWorld != nullptr && GEditor->EditorWorld == WorldContext->World() )
 	{
-		UEditorWorldExtension* EditorExtension = EditorExtensionTuple.Get<0>();
-		Collector.AddReferencedObject( EditorExtension );
+		SetWorldContext( GEditor->GetPIEWorldContext() );
+		EditorWorldOnSimulate = &GEditor->GetEditorWorldContext();
+
+		for (FEditorExtensionTuple& EditorExtensionTuple : EditorExtensions)
+		{
+			UEditorWorldExtension* EditorExtension = EditorExtensionTuple.Get<0>();
+			EditorExtension->EnteredSimulateInEditor();
+		}
+	}
+}
+
+
+void UEditorWorldExtensionCollection::OnEndPIE( bool bWasSimulatingInEditor )
+{
+	if( bWasSimulatingInEditor && EditorWorldOnSimulate != nullptr && EditorWorldOnSimulate->World() == GEditor->EditorWorld )
+	{
+		if( !GIsRequestingExit )
+		{
+			// Revert back to the editor world before closing the play world, otherwise actors and objects will be destroyed.
+			SetWorldContext( EditorWorldOnSimulate );
+			EditorWorldOnSimulate = nullptr;
+
+			for( FEditorExtensionTuple& EditorExtensionTuple : EditorExtensions )
+			{
+				UEditorWorldExtension* EditorExtension = EditorExtensionTuple.Get<0>();
+				EditorExtension->LeftSimulateInEditor();
+			}
+		}
 	}
 }
 
+
+void UEditorWorldExtensionCollection::SetWorldContext(FWorldContext* InWorldContext)
+{
+	check( InWorldContext != nullptr );
+
+	if( this->WorldContext != nullptr )
+	{
+		UWorld* CurrentWorld = WorldContext->World();
+		if (CurrentWorld != nullptr)
+		{
+			for (FEditorExtensionTuple& EditorExtensionTuple : EditorExtensions)
+			{
+				UEditorWorldExtension* EditorExtension = EditorExtensionTuple.Get<0>();
+				EditorExtension->TransitionWorld(InWorldContext->World());
+			}
+		}
+	}
+
+	WorldContext = InWorldContext;
+}
+
 /************************************************************************/
-/* FEditorWorldExtensionManager                                         */
+/* UEditorWorldExtensionManager                                         */
 /************************************************************************/
-FEditorWorldExtensionManager::FEditorWorldExtensionManager()
+UEditorWorldExtensionManager::UEditorWorldExtensionManager() :
+	Super()
 {
 	if( GEngine )
 	{
-		GEngine->OnWorldContextDestroyed().AddRaw( this, &FEditorWorldExtensionManager::OnWorldContextRemove );
+		GEngine->OnWorldContextDestroyed().AddUObject( this, &UEditorWorldExtensionManager::OnWorldContextRemove );
 	}
 }
 
-FEditorWorldExtensionManager::~FEditorWorldExtensionManager()
+UEditorWorldExtensionManager::~UEditorWorldExtensionManager()
 {
 	if( GEngine )
 	{
@@ -202,55 +389,74 @@ FEditorWorldExtensionManager::~FEditorWorldExtensionManager()
 	EditorWorldExtensionCollection.Empty();
 }
 
-TSharedPtr<FEditorWorldExtensionCollection> FEditorWorldExtensionManager::GetEditorWorldExtensions( const UWorld* InWorld )
+UEditorWorldExtensionCollection* UEditorWorldExtensionManager::GetEditorWorldExtensions(UWorld* InWorld)
 {
 	// Try to find this world in the map and return it or create and add one if nothing found
-	TSharedPtr<FEditorWorldExtensionCollection> Result;
-	if( InWorld )
+	UEditorWorldExtensionCollection* Result = nullptr;
+	if (InWorld)
 	{
-		TSharedPtr<FEditorWorldExtensionCollection>* FoundWorld = EditorWorldExtensionCollection.Find( InWorld->GetUniqueID() );
-		if( FoundWorld != nullptr )
+		UEditorWorldExtensionCollection* FoundExtensionCollection = FindExtensionCollection(InWorld);
+		if (FoundExtensionCollection != nullptr)
 		{
-			Result = *FoundWorld;
+			Result = FoundExtensionCollection;
 		}
 		else
 		{
-			FWorldContext* WorldContext = GEditor->GetWorldContextFromWorld( InWorld );
-			Result = OnWorldContextAdd( *WorldContext );
+			FWorldContext* WorldContext = GEditor->GetWorldContextFromWorld(InWorld);
+			Result = OnWorldContextAdd(WorldContext);
 		}
 	}
 	return Result;
 }
 
-TSharedPtr<FEditorWorldExtensionCollection> FEditorWorldExtensionManager::OnWorldContextAdd( FWorldContext& InWorldContext )
+UEditorWorldExtensionCollection* UEditorWorldExtensionManager::OnWorldContextAdd(FWorldContext* InWorldContext)
 {
-	//Only add editor type world to the map
-	UWorld* World = InWorldContext.World();
-	TSharedPtr<FEditorWorldExtensionCollection> Result;
-	if( World != nullptr )
+	const UWorld* World = InWorldContext->World();
+	UEditorWorldExtensionCollection* Result = nullptr;
+	if (World != nullptr)
 	{
-		TSharedPtr<FEditorWorldExtensionCollection> ExtensionCollection( new FEditorWorldExtensionCollection( InWorldContext ) );
+		UEditorWorldExtensionCollection* ExtensionCollection = NewObject<UEditorWorldExtensionCollection>();
+		ExtensionCollection->SetWorldContext(InWorldContext);
 		Result = ExtensionCollection;
-		EditorWorldExtensionCollection.Add( World->GetUniqueID(), Result );
+		EditorWorldExtensionCollection.Add(Result);
 	}
 	return Result;
 }
 
-void FEditorWorldExtensionManager::OnWorldContextRemove( FWorldContext& InWorldContext )
+void UEditorWorldExtensionManager::OnWorldContextRemove(FWorldContext& InWorldContext)
 {
 	UWorld* World = InWorldContext.World();
 	if(World)
 	{
-		EditorWorldExtensionCollection.Remove( World->GetUniqueID() );
+		UEditorWorldExtensionCollection* ExtensionCollection = FindExtensionCollection(World);
+		if (ExtensionCollection)
+		{
+			EditorWorldExtensionCollection.Remove(ExtensionCollection);
+		}
 	}
 }
 
-void FEditorWorldExtensionManager::Tick( float DeltaSeconds )
+UEditorWorldExtensionCollection* UEditorWorldExtensionManager::FindExtensionCollection(const UWorld* InWorld)
+{
+	UEditorWorldExtensionCollection* ResultCollection = nullptr;
+	for (UEditorWorldExtensionCollection* ExtensionCollection : EditorWorldExtensionCollection)
+	{
+		if (ExtensionCollection != nullptr && ExtensionCollection->GetWorld() == InWorld)
+		{
+			ResultCollection = ExtensionCollection;
+			break;
+		}
+	}
+
+	return ResultCollection;
+}
+
+void UEditorWorldExtensionManager::Tick( float DeltaSeconds )
 {
 	// Tick all the collections
-	for( auto& Collection : EditorWorldExtensionCollection )
+	for(UEditorWorldExtensionCollection* ExtensionCollection : EditorWorldExtensionCollection)
 	{
-		FEditorWorldExtensionCollection& ExtensionCollection = *Collection.Value;
-		ExtensionCollection.Tick( DeltaSeconds );
+		check(ExtensionCollection != nullptr && ExtensionCollection->IsValidLowLevel());
+		ExtensionCollection->Tick(DeltaSeconds);
 	}
-}
\ No newline at end of file
+}
