Misc. module cleanup

diff --git a/Engine/Source/Editor/UnrealEd/Private/EditorEngine.cpp b/Engine/Source/Editor/UnrealEd/Private/EditorEngine.cpp
--- a/Engine/Source/Editor/UnrealEd/Private/EditorEngine.cpp
+++ b/Engine/Source/Editor/UnrealEd/Private/EditorEngine.cpp
@@ -898,7 +898,6 @@ void UEditorEngine::Init(IEngineLoop* InEngineLoop)
 			TEXT("UndoHistory"),
 			TEXT("DeviceProfileEditor"),
 			TEXT("SourceCodeAccess"),
-			TEXT("MeshEditor"),
 			TEXT("BehaviorTreeEditor"),
 			TEXT("HardwareTargeting"),
 			TEXT("LocalizationDashboard"),
