Forward declaration of UEditorWorldExtensionCollection

diff --git a/Engine/Source/Editor/UnrealEd/Public/EditorWorldExtension.h b/Engine/Source/Editor/UnrealEd/Public/EditorWorldExtension.h
--- a/Engine/Source/Editor/UnrealEd/Public/EditorWorldExtension.h
+++ b/Engine/Source/Editor/UnrealEd/Public/EditorWorldExtension.h
@@ -7,6 +7,7 @@
 #include "EditorWorldExtension.generated.h"
 
 struct FWorldContext;
+class UEditorWorldExtensionCollection;
 
 UCLASS()
 class UNREALED_API UEditorWorldExtension : public UObject
