Manual copy of 3341154
- Added CanPlaySound to UEditorEngine for external use.
- Use bEnableEditorSounds setting for VREditor and ViewportWorldInteraction.

diff --git a/Engine/Source/Editor/UnrealEd/Classes/Editor/EditorEngine.h b/Engine/Source/Editor/UnrealEd/Classes/Editor/EditorEngine.h
--- a/Engine/Source/Editor/UnrealEd/Classes/Editor/EditorEngine.h
+++ b/Engine/Source/Editor/UnrealEd/Classes/Editor/EditorEngine.h
@@ -1160,6 +1160,12 @@ public:
 	/** Plays an editor sound (if the user has sounds enabled.) */
 	void PlayEditorSound( USoundBase* InSound );
 
+	/**
+	 * Returns true if currently able to play a sound and if the user has sounds enabled.
+	 */
+	bool CanPlayEditorSound() const;
+
+
 	/**
 	 * Returns the preview audio component
 	 */
diff --git a/Engine/Source/Editor/UnrealEd/Private/EditorEngine.cpp b/Engine/Source/Editor/UnrealEd/Private/EditorEngine.cpp
--- a/Engine/Source/Editor/UnrealEd/Private/EditorEngine.cpp
+++ b/Engine/Source/Editor/UnrealEd/Private/EditorEngine.cpp
@@ -2263,7 +2263,7 @@ void UEditorEngine::PlayEditorSound( const FString& SoundAssetName )
 void UEditorEngine::PlayEditorSound( USoundBase* InSound )
 {
 	// Only play sounds if the user has that feature enabled
-	if (!GIsSavingPackage && IsInGameThread() && GetDefault<ULevelEditorMiscSettings>()->bEnableEditorSounds)
+	if (!GIsSavingPackage && CanPlayEditorSound())
 	{
 		if (InSound != nullptr)
 		{
@@ -2272,6 +2272,11 @@ void UEditorEngine::PlayEditorSound( USoundBase* InSound )
 	}
 }
 
+bool UEditorEngine::CanPlayEditorSound() const
+{
+	return IsInGameThread() && GetDefault<ULevelEditorMiscSettings>()->bEnableEditorSounds;
+}
+
 void UEditorEngine::ClearPreviewComponents()
 {
 	if( PreviewAudioComponent )
diff --git a/Engine/Source/Editor/VREditor/VREditorMode.cpp b/Engine/Source/Editor/VREditor/VREditorMode.cpp
--- a/Engine/Source/Editor/VREditor/VREditorMode.cpp
+++ b/Engine/Source/Editor/VREditor/VREditorMode.cpp
@@ -1123,7 +1123,7 @@ const UVREditorAssetContainer& UVREditorMode::GetAssetContainer() const
 
 void UVREditorMode::PlaySound(USoundBase* SoundBase, const FVector& InWorldLocation, const float InVolume /*= 1.0f*/)
 {
-	if (bIsActive && bIsFullyInitialized)
+	if (bIsActive && bIsFullyInitialized && GEditor != nullptr && GEditor->CanPlayEditorSound())
 	{
 		const float Volume = InVolume*VREd::SFXMultiplier->GetFloat();
 		UGameplayStatics::PlaySoundAtLocation(GetWorld(), SoundBase, InWorldLocation, Volume);
diff --git a/Engine/Source/Editor/ViewportInteraction/ViewportWorldInteraction.cpp b/Engine/Source/Editor/ViewportInteraction/ViewportWorldInteraction.cpp
--- a/Engine/Source/Editor/ViewportInteraction/ViewportWorldInteraction.cpp
+++ b/Engine/Source/Editor/ViewportInteraction/ViewportWorldInteraction.cpp
@@ -3084,7 +3084,7 @@ const class UViewportInteractionAssetContainer& UViewportWorldInteraction::LoadA
 
 void UViewportWorldInteraction::PlaySound(USoundBase* SoundBase, const FVector& InWorldLocation, const float InVolume /*= 1.0f*/)
 {
-	if (bActive)
+	if (bActive && GEditor != nullptr && GEditor->CanPlayEditorSound())
 	{
 		const float Volume = InVolume*VI::SFXMultiplier->GetFloat();
 		UGameplayStatics::PlaySoundAtLocation(GetWorld(), SoundBase, InWorldLocation, FRotator::ZeroRotator, Volume);
