Disabled code caching data for all platforms until we can figure out why it's filling up DDC cache

diff --git a/Engine/Source/Editor/UnrealEd/Private/Commandlets/DerivedDataCacheCommandlet.cpp b/Engine/Source/Editor/UnrealEd/Private/Commandlets/DerivedDataCacheCommandlet.cpp
--- a/Engine/Source/Editor/UnrealEd/Private/Commandlets/DerivedDataCacheCommandlet.cpp
+++ b/Engine/Source/Editor/UnrealEd/Private/Commandlets/DerivedDataCacheCommandlet.cpp
@@ -216,7 +216,7 @@ int32 UDerivedDataCacheCommandlet::Main( const FString& Params )
 			GRedirectCollector.ResolveStringAssetReference();
 
 			// cache all the resources for this platform
-			for (TObjectIterator<UObject> It; It; ++It)
+			/*for (TObjectIterator<UObject> It; It; ++It)
 			{
 				if ((PackageFilter&NORMALIZE_ExcludeEnginePackages) == 0 || !It->GetOutermost()->GetName().StartsWith(TEXT("/Engine")))
 				{
@@ -229,7 +229,7 @@ int32 UDerivedDataCacheCommandlet::Main( const FString& Params )
 						}
 					}
 				}
-			}
+			}*/
 
 
 			// Keep track of which packages have already been processed along with the map.
