Copying //UE4/Dev-Editor to //UE4/Dev-Main (Source: //UE4/Dev-Editor @ 3279756)

diff --git a/Engine/Source/Editor/UMGEditor/Private/WidgetBlueprintEditor.cpp b/Engine/Source/Editor/UMGEditor/Private/WidgetBlueprintEditor.cpp
--- a/Engine/Source/Editor/UMGEditor/Private/WidgetBlueprintEditor.cpp
+++ b/Engine/Source/Editor/UMGEditor/Private/WidgetBlueprintEditor.cpp
@@ -235,6 +235,10 @@ void FWidgetBlueprintEditor::SetSelectedNamedSlot(TOptional<FNamedSlotSelection>
 	SelectedNamedSlot.Reset();
 
 	SelectedNamedSlot = InSelectedNamedSlot;
+	if (InSelectedNamedSlot.IsSet())
+	{
+		SelectedWidgets.Add(InSelectedNamedSlot->NamedSlotHostWidget);
+	}
 
 	OnSelectedWidgetsChanged.Broadcast();
 }
@@ -920,7 +924,15 @@ EWidgetDesignFlags::Type FWidgetBlueprintEditor::GetCurrentDesignerFlags() const
 	
 	if ( bShowDashedOutlines )
 	{
-		Flags = ( EWidgetDesignFlags::Type )(Flags | EWidgetDesignFlags::ShowOutline);
+		Flags = ( EWidgetDesignFlags::Type )( Flags | EWidgetDesignFlags::ShowOutline );
+	}
+
+	if ( const UWidgetDesignerSettings* Designer = GetDefault<UWidgetDesignerSettings>() )
+	{
+		if ( Designer->bExecutePreConstructEvent )
+		{
+			Flags = ( EWidgetDesignFlags::Type )( Flags | EWidgetDesignFlags::ExecutePreConstruct );
+		}
 	}
 
 	return Flags;
diff --git a/Engine/Source/Editor/UMGEditor/Private/WidgetBlueprintEditorUtils.cpp b/Engine/Source/Editor/UMGEditor/Private/WidgetBlueprintEditorUtils.cpp
--- a/Engine/Source/Editor/UMGEditor/Private/WidgetBlueprintEditorUtils.cpp
+++ b/Engine/Source/Editor/UMGEditor/Private/WidgetBlueprintEditorUtils.cpp
@@ -430,6 +430,42 @@ UWidget* FWidgetBlueprintEditorUtils::FindNamedSlotHostWidgetForContent(UWidget*
 	return HostWidget;
 }
 
+void FWidgetBlueprintEditorUtils::FindAllAncestorNamedSlotHostWidgetsForContent(TArray<FWidgetReference>& OutSlotHostWidgets, UWidget* WidgetTemplate, TSharedRef<FWidgetBlueprintEditor> BlueprintEditor)
+{
+	OutSlotHostWidgets.Empty();
+	UUserWidget* Preview = BlueprintEditor->GetPreview();
+	UWidgetBlueprint* WidgetBP = BlueprintEditor->GetWidgetBlueprintObj();
+	UWidgetTree* WidgetTree = (WidgetBP != nullptr) ? WidgetBP->WidgetTree : nullptr;
+
+	if (Preview != nullptr && WidgetTree != nullptr)
+	{
+		UWidget* SlotHostWidget = FindNamedSlotHostWidgetForContent(WidgetTemplate, WidgetTree);
+		while (SlotHostWidget != nullptr)
+		{
+			UWidget* SlotWidget = Preview->GetWidgetFromName(SlotHostWidget->GetFName());
+			FWidgetReference WidgetRef;
+
+			if (SlotWidget != nullptr)
+			{
+				WidgetRef = BlueprintEditor->GetReferenceFromPreview(SlotWidget);
+
+				if (WidgetRef.IsValid())
+				{
+					OutSlotHostWidgets.Add(WidgetRef);
+				}
+			}
+
+			WidgetTemplate = WidgetRef.GetTemplate();
+
+			SlotHostWidget = nullptr;
+			if (WidgetTemplate != nullptr)
+			{
+				SlotHostWidget = FindNamedSlotHostWidgetForContent(WidgetRef.GetTemplate(), WidgetTree);
+			}
+		}
+	}
+}
+
 bool FWidgetBlueprintEditorUtils::RemoveNamedSlotHostContent(UWidget* WidgetTemplate, INamedSlotInterface* NamedSlotHost)
 {
 	TArray<FName> SlotNames;
@@ -1117,7 +1153,11 @@ void FWidgetBlueprintEditorUtils::ImportWidgetsFromText(UWidgetBlueprint* BP, co
 			}
 
 			Widget->Rename(nullptr, BP->WidgetTree);
-			Widget->SetDisplayLabel(Widget->GetName());
+
+			if (Widget->GetDisplayLabel().Equals(WidgetOldName))
+			{
+				Widget->SetDisplayLabel(Widget->GetName());
+			}
 
 			if ( SlotData )
 			{
diff --git a/Engine/Source/Editor/UMGEditor/Public/Settings/WidgetDesignerSettings.h b/Engine/Source/Editor/UMGEditor/Public/Settings/WidgetDesignerSettings.h
--- a/Engine/Source/Editor/UMGEditor/Public/Settings/WidgetDesignerSettings.h
+++ b/Engine/Source/Editor/UMGEditor/Public/Settings/WidgetDesignerSettings.h
@@ -15,16 +15,34 @@ class UMGEDITOR_API UWidgetDesignerSettings : public UObject
 {
 	GENERATED_UCLASS_BODY()
 
+public:
+
 	/** If enabled, actor positions will snap to the grid. */
 	UPROPERTY(EditAnywhere, config, Category = GridSnapping, meta = (DisplayName = "Enable Grid Snapping"))
 	uint32 GridSnapEnabled:1;
 
+	/**
+	 * 
+	 */
 	UPROPERTY(config)
 	int32 GridSnapSize;
 
+	/**
+	 * 
+	 */
 	UPROPERTY(EditAnywhere, config, Category = Dragging)
 	bool bLockToPanelOnDragByDefault;
 
+	/**
+	 * Should the designer show outlines by default?
+	 */
 	UPROPERTY(EditAnywhere, config, Category = Visuals, meta = ( DisplayName = "Show Dashed Outlines By Default" ))
 	bool bShowOutlines;
+
+	/**
+	 * Should the designer run the design event?  Disable this if you're seeing crashes in the designer,
+	 * you may have some unsafe code running in the designer.
+	 */
+	UPROPERTY(EditAnywhere, config, Category = Visuals)
+	bool bExecutePreConstructEvent;
 };
diff --git a/Engine/Source/Editor/UnrealEd/Classes/Commandlets/GatherTextFromSourceCommandlet.h b/Engine/Source/Editor/UnrealEd/Classes/Commandlets/GatherTextFromSourceCommandlet.h
--- a/Engine/Source/Editor/UnrealEd/Classes/Commandlets/GatherTextFromSourceCommandlet.h
+++ b/Engine/Source/Editor/UnrealEd/Classes/Commandlets/GatherTextFromSourceCommandlet.h
@@ -6,6 +6,7 @@
 #include "UObject/ObjectMacros.h"
 #include "Framework/Commands/Commands.h"
 #include "Commandlets/GatherTextCommandletBase.h"
+#include "Internationalization/StringTableCore.h"
 #include "GatherTextFromSourceCommandlet.generated.h"
 
 class Error;
@@ -27,6 +28,53 @@ private:
 		EditorOnly,
 	};
 
+	struct FSourceLocation
+	{
+		FSourceLocation()
+			: File()
+			, Line(INDEX_NONE)
+		{
+		}
+
+		FSourceLocation(FString InFile, const int32 InLine)
+			: File(MoveTemp(InFile))
+			, Line(InLine)
+		{
+		}
+
+		FString ToString() const
+		{
+			return (Line == INDEX_NONE) ? File : FString::Printf(TEXT("%s - line %d"), *File, Line);
+		}
+
+		FString File;
+		int32 Line;
+	};
+
+	struct FParsedStringTableEntry
+	{
+		FString SourceString;
+		FSourceLocation SourceLocation;
+		bool bIsEditorOnly;
+	};
+
+	struct FParsedStringTableEntryMetaData
+	{
+		FString MetaData;
+		FSourceLocation SourceLocation;
+		bool bIsEditorOnly;
+	};
+
+	typedef TMap<FName, FParsedStringTableEntryMetaData> FParsedStringTableEntryMetaDataMap;
+
+	struct FParsedStringTable
+	{
+		FString TableNamespace;
+		FSourceLocation SourceLocation;
+		TMap<FString, FParsedStringTableEntry, FDefaultSetAllocator, FLocKeyMapFuncs<FParsedStringTableEntry>> TableEntries;
+		TMap<FString, FParsedStringTableEntryMetaDataMap, FDefaultSetAllocator, FLocKeyMapFuncs<FParsedStringTableEntryMetaDataMap>> MetaDataEntries;
+	};
+
 	struct FSourceFileParseContext
 	{
 		bool AddManifestText( const FString& Token, const FString& Namespace, const FString& SourceText, const FManifestContext& Context );
@@ -43,6 +91,14 @@ private:
 
 		void RemoveDefine( const FString& InDefineCtx );
 
+		void AddStringTable( const FName InTableId, const FString& InTableNamespace );
+
+		void AddStringTableFromFile( const FName InTableId, const FString& InTableNamespace, const FString& InTableFilename, const FString& InRootPath );
+
+		void AddStringTableEntry( const FName InTableId, const FString& InKey, const FString& InSourceString );
+
+		void AddStringTableEntryMetaData( const FName InTableId, const FString& InKey, const FName InMetaDataId, const FString& InMetaData );
+
 		//Working data
 		FString Filename;
 		int32 LineNumber;
@@ -62,6 +118,9 @@ private:
 		//Destination location of the parsed FLocTextEntrys
 		TSharedPtr< FLocTextHelper > GatherManifestHelper;
 
+		//Discovered string table data from all files
+		TMap<FName, FParsedStringTable> ParsedStringTables;
+
 		FSourceFileParseContext()
 			: Filename()
 			, LineNumber(0)
@@ -83,6 +142,10 @@ private:
 		}
 
 	private:
+		bool AddStringTableImpl( const FName InTableId, const FString& InTableNamespace );
+		bool AddStringTableEntryImpl( const FName InTableId, const FString& InKey, const FString& InSourceString, const FSourceLocation& InSourceLocation );
+		bool AddStringTableEntryMetaDataImpl( const FName InTableId, const FString& InKey, const FName InMetaDataId, const FString& InMetaData, const FSourceLocation& InSourceLocation );
+
 		//Working data
 		TArray<FString> MacroBlockStack;
 		mutable TOptional<EMacroBlockState> CachedMacroBlockState;
@@ -170,24 +233,9 @@ private:
 	class FMacroDescriptor : public FParsableDescriptor
 	{
 	public:
-		enum EMacroArgSemantic
-		{
-			MAS_Namespace,
-			MAS_Identifier,
-			MAS_SourceText,
-		};
-
-		struct FMacroArg
-		{
-			EMacroArgSemantic Semantic;
-			bool IsAutoText;
-
-			FMacroArg(EMacroArgSemantic InSema, bool InIsAutoText) : Semantic(InSema), IsAutoText(InIsAutoText) {}
-		};
-
 		static const FString TextMacroString;
 
-		FMacroDescriptor(FString InName) : Name(InName) {}
+		FMacroDescriptor(FString InName) : Name(MoveTemp(InName)) {}
 
 		virtual const FString& GetToken() const override { return Name; }
 
@@ -211,6 +259,21 @@ private:
 	class FStringMacroDescriptor : public FMacroDescriptor
 	{
 	public:
+		enum EMacroArgSemantic
+		{
+			MAS_Namespace,
+			MAS_Identifier,
+			MAS_SourceText,
+		};
+
+		struct FMacroArg
+		{
+			EMacroArgSemantic Semantic;
+			bool IsAutoText;
+
+			FMacroArg(EMacroArgSemantic InSema, bool InIsAutoText) : Semantic(InSema), IsAutoText(InIsAutoText) {}
+		};
+
 		FStringMacroDescriptor(FString InName, FMacroArg Arg0, FMacroArg Arg1, FMacroArg Arg2) : FMacroDescriptor(InName)
 		{
 			Arguments.Add(Arg0);
@@ -235,6 +298,41 @@ private:
 		TArray<FMacroArg> Arguments;
 	};
 
+	class FStringTableMacroDescriptor : public FMacroDescriptor
+	{
+	public:
+		FStringTableMacroDescriptor() : FMacroDescriptor(TEXT("LOCTABLE_NEW")) {}
+
+		virtual void TryParse(const FString& Text, FSourceFileParseContext& Context) const override;
+	};
+
+	class FStringTableFromFileMacroDescriptor : public FMacroDescriptor
+	{
+	public:
+		FStringTableFromFileMacroDescriptor(FString InName, FString InRootPath) : FMacroDescriptor(MoveTemp(InName)), RootPath(MoveTemp(InRootPath)) {}
+
+		virtual void TryParse(const FString& Text, FSourceFileParseContext& Context) const override;
+
+	private:
+		FString RootPath;
+	};
+
+	class FStringTableEntryMacroDescriptor : public FMacroDescriptor
+	{
+	public:
+		FStringTableEntryMacroDescriptor() : FMacroDescriptor(TEXT("LOCTABLE_SETSTRING")) {}
+
+		virtual void TryParse(const FString& Text, FSourceFileParseContext& Context) const override;
+	};
+
+	class FStringTableEntryMetaDataMacroDescriptor : public FMacroDescriptor
+	{
+	public:
+		FStringTableEntryMetaDataMacroDescriptor() : FMacroDescriptor(TEXT("LOCTABLE_SETMETA")) {}
+
+		virtual void TryParse(const FString& Text, FSourceFileParseContext& Context) const override;
+	};
+
 	class FIniNamespaceDescriptor : public FPreProcessorDescriptor
 	{
 	public:
diff --git a/Engine/Source/Editor/UnrealEd/Classes/Editor/EditorEngine.h b/Engine/Source/Editor/UnrealEd/Classes/Editor/EditorEngine.h
--- a/Engine/Source/Editor/UnrealEd/Classes/Editor/EditorEngine.h
+++ b/Engine/Source/Editor/UnrealEd/Classes/Editor/EditorEngine.h
@@ -218,7 +218,7 @@ struct FCachedActorLabels
 	FCachedActorLabels();
 
 	/** Constructor that populates the set of actor names */
-	explicit FCachedActorLabels(UWorld* World, const TSet<AActor*>& IgnoredActors = TSet<AActor*>());
+	explicit UNREALED_API FCachedActorLabels(UWorld* World, const TSet<AActor*>& IgnoredActors = TSet<AActor*>());
 
 	/** Populate the set of actor names */
 	void Populate(UWorld* World, const TSet<AActor*>& IgnoredActors = TSet<AActor*>());
@@ -2976,6 +2976,9 @@ public:
 
 private:
 	FTimerHandle CleanupPIEOnlineSessionsTimerHandle;
+
+	/** Delegate handle for game viewport close requests in PIE sessions. */
+	FDelegateHandle ViewportCloseRequestedDelegateHandle;
 };
 
 //////////////////////////////////////////////////////////////////////////
diff --git a/Engine/Source/Editor/UnrealEd/Classes/Factories/FbxSceneImportFactory.h b/Engine/Source/Editor/UnrealEd/Classes/Factories/FbxSceneImportFactory.h
--- a/Engine/Source/Editor/UnrealEd/Classes/Factories/FbxSceneImportFactory.h
+++ b/Engine/Source/Editor/UnrealEd/Classes/Factories/FbxSceneImportFactory.h
@@ -4,7 +4,7 @@
 
 #include "CoreMinimal.h"
 #include "UObject/ObjectMacros.h"
-#include "Factories/Factory.h"
+#include "SceneImportFactory.h"
 #include "FbxSceneImportFactory.generated.h"
 
 class AActor;
@@ -338,14 +338,17 @@ typedef TMap<FString, UnFbx::FBXImportOptions*> ImportOptionsNameMap;
 typedef ImportOptionsNameMap* ImportOptionsNameMapPtr;
 
 UCLASS(hidecategories=Object)
-class UNREALED_API UFbxSceneImportFactory : public UFactory
+class UNREALED_API UFbxSceneImportFactory : public USceneImportFactory
 {
 	GENERATED_UCLASS_BODY()
 
-	//~ Begin UFactory Interface
+	/** UFactory Interface */
 	virtual UObject* FactoryCreateBinary(UClass* InClass, UObject* InParent, FName InName, EObjectFlags Flags, UObject* Context, const TCHAR* Type, const uint8*& Buffer, const uint8* BufferEnd, FFeedbackContext* Warn) override;
 	virtual UObject* FactoryCreateBinary(UClass* InClass, UObject* InParent, FName InName, EObjectFlags Flags, UObject* Context, const TCHAR* Type, const uint8*& Buffer, const uint8* BufferEnd, FFeedbackContext* Warn, bool& bOutOperationCanceled) override;
-	//~ Begin UFactory Interface
+	virtual bool FactoryCanImport(const FString& Filename) override;
+
+	/** USceneImportFactory Interface */
+	virtual bool ImportsAssets() const override { return true; }
 
 	/** Import options UI detail when importing fbx scene */
 	UPROPERTY(Transient)
diff --git a/Engine/Source/Editor/UnrealEd/Classes/Factories/StringTableFactory.h b/Engine/Source/Editor/UnrealEd/Classes/Factories/StringTableFactory.h
new file mode 100644
--- /dev/null
+++ b/Engine/Source/Editor/UnrealEd/Classes/Factories/StringTableFactory.h
@@ -0,0 +1,18 @@
+// Copyright 1998-2017 Epic Games, Inc. All Rights Reserved.
+
+#pragma once
+
+#include "CoreMinimal.h"
+#include "UObject/ObjectMacros.h"
+#include "Factories/Factory.h"
+#include "StringTableFactory.generated.h"
+
+UCLASS(hidecategories=Object)
+class UNREALED_API UStringTableFactory : public UFactory
+{
+	GENERATED_UCLASS_BODY()
+
+	//~ Begin UFactory Interface
+	virtual UObject* FactoryCreateNew(UClass* Class, UObject* InParent, FName Name, EObjectFlags Flags, UObject* Context, FFeedbackContext* Warn) override;
+	//~ Begin UFactory Interface
+};
diff --git a/Engine/Source/Editor/UnrealEd/Classes/Settings/LevelEditorPlaySettings.h b/Engine/Source/Editor/UnrealEd/Classes/Settings/LevelEditorPlaySettings.h
--- a/Engine/Source/Editor/UnrealEd/Classes/Settings/LevelEditorPlaySettings.h
+++ b/Engine/Source/Editor/UnrealEd/Classes/Settings/LevelEditorPlaySettings.h
@@ -180,7 +180,11 @@ public:
 
 	/** Whether to play sounds during PIE */
 	UPROPERTY(config, EditAnywhere, Category=PlayInEditor, meta=(ToolTip="Whether to play sounds when in a Play In Editor session"))
-	bool EnableSound;
+	bool EnableGameSound;
+
+	/** Whether to play a sound when entering and exiting PIE */
+	UPROPERTY(config, EditAnywhere, AdvancedDisplay, Category = PlayInEditor, meta = (DisplayName = "Enable PIE Enter and Exit Sounds"))
+	bool EnablePIEEnterAndExitSounds;
 
 	/** Which quality level to use when playing in editor */
 	UPROPERTY(config, EditAnywhere, Category=PlayInEditor)
diff --git a/Engine/Source/Editor/UnrealEd/Private/Analytics/EditorAnalytics.cpp b/Engine/Source/Editor/UnrealEd/Private/Analytics/EditorAnalytics.cpp
--- a/Engine/Source/Editor/UnrealEd/Private/Analytics/EditorAnalytics.cpp
+++ b/Engine/Source/Editor/UnrealEd/Private/Analytics/EditorAnalytics.cpp
@@ -177,7 +177,7 @@ FString FEditorAnalytics::TranslateErrorCode(int32 ErrorCode)
 	case EAnalyticsErrorCodes::FilesInstallFailed:
 		return TEXT("Failed to deploy files to device.  Check to make sure your device is connected.");
 	case EAnalyticsErrorCodes::DeviceNotSetupForDevelopment:
-		return TEXT("Failed to launch on device.  Make sure your device has been enabled for development from within the Xcode Devices window.");
+		return TEXT("Failed to launch on device.  Make sure your device is currently unlocked and has been enabled for development by using a mobile provision including your device id.");
 	case EAnalyticsErrorCodes::DeviceOSNewerThanSDK:
 		return TEXT("Failed to launch on device.  Make sure your install of Xcode matches or is newer than the OS on your device.");
 	case EAnalyticsErrorCodes::RemoteCertificatesNotFound:
diff --git a/Engine/Source/Editor/UnrealEd/Private/Commandlets/ContentCommandlets.cpp b/Engine/Source/Editor/UnrealEd/Private/Commandlets/ContentCommandlets.cpp
--- a/Engine/Source/Editor/UnrealEd/Private/Commandlets/ContentCommandlets.cpp
+++ b/Engine/Source/Editor/UnrealEd/Private/Commandlets/ContentCommandlets.cpp
@@ -542,11 +542,15 @@ void UResavePackagesCommandlet::LoadAndSaveOnePackage(const FString& Filename)
 
 							UE_LOG(LogContentCommandlet, Display, TEXT("Deleting '%s' from source control..."), *Filename);
 							SourceControlProvider.Execute(ISourceControlOperation::Create<FDelete>(), PackageFilename);
+
+							FilesToSubmit.Add(Filename);
 						}
 						else if (SourceControlState.IsValid() && SourceControlState->CanCheckout())
 						{
 							UE_LOG(LogContentCommandlet, Display, TEXT("Deleting '%s' from source control..."), *Filename);
 							SourceControlProvider.Execute(ISourceControlOperation::Create<FDelete>(), PackageFilename);
+
+							FilesToSubmit.Add(Filename);
 						}
 						else if (SourceControlState.IsValid() && SourceControlState->IsCheckedOutOther())
 						{
diff --git a/Engine/Source/Editor/UnrealEd/Private/Commandlets/DiffAssetsCommandlet.cpp b/Engine/Source/Editor/UnrealEd/Private/Commandlets/DiffAssetsCommandlet.cpp
--- a/Engine/Source/Editor/UnrealEd/Private/Commandlets/DiffAssetsCommandlet.cpp
+++ b/Engine/Source/Editor/UnrealEd/Private/Commandlets/DiffAssetsCommandlet.cpp
@@ -78,7 +78,7 @@ bool UDiffAssetsCommandlet::CopyFileToTempLocation(FString& InOutFilename)
 
 bool UDiffAssetsCommandlet::LoadFile(const FString& Filename, TArray<UObject *>& LoadedObjects)
 {
-	UPackage* Package = Cast<UPackage>(LoadPackage( NULL, *Filename, LOAD_ForDiff));
+	UPackage* Package = LoadPackage( NULL, *Filename, LOAD_ForDiff);
 	if (!Package)
 	{
 		UE_LOG(LogDiffAssetsCommandlet, Warning, TEXT("Could not load %s"), *Filename);
diff --git a/Engine/Source/Editor/UnrealEd/Private/Commandlets/FixupRedirectsCommandlet.cpp b/Engine/Source/Editor/UnrealEd/Private/Commandlets/FixupRedirectsCommandlet.cpp
--- a/Engine/Source/Editor/UnrealEd/Private/Commandlets/FixupRedirectsCommandlet.cpp
+++ b/Engine/Source/Editor/UnrealEd/Private/Commandlets/FixupRedirectsCommandlet.cpp
@@ -168,7 +168,7 @@ int32 UFixupRedirectsCommandlet::Main( const FString& Params )
 			UE_LOG(LogFixupRedirectsCommandlet, Display, TEXT("Looking for redirects in %s..."), *Filename);
 
 			// Assert if package couldn't be opened so we have no chance of messing up saving later packages.
-			UPackage* Package = Cast<UPackage>(LoadPackage(NULL, *Filename, 0));
+			UPackage* Package = LoadPackage(NULL, *Filename, 0);
 
 			// load all string asset reference targets, and add fake redirectors for them
 			GRedirectCollector.ResolveStringAssetReference();
@@ -775,8 +775,8 @@ int32 UFixupRedirectsCommandlet::Main( const FString& Params )
 		{
 			const FString& Filename = PackageList[PackageIndex];
 
-			FString PackageName(FPackageName::FilenameToLongPackageName(Filename));
-			FSourceControlStatePtr SourceControlState = SourceControlProvider.GetState(SourceControlHelpers::PackageFilename(Filename), EStateCacheUsage::ForceUpdate);
+			FString PackageName(SourceControlHelpers::PackageFilename(Filename));
+			FSourceControlStatePtr SourceControlState = SourceControlProvider.GetState(PackageName, EStateCacheUsage::ForceUpdate);
 
 			if( SourceControlState.IsValid() && (SourceControlState->IsCheckedOut() || SourceControlState->IsAdded() || SourceControlState->IsDeleted()) )
 			{
@@ -788,11 +788,14 @@ int32 UFixupRedirectsCommandlet::Main( const FString& Params )
 			}
 		}
 
-		// Check in all changed files
-		const FText Description = NSLOCTEXT("FixupRedirectsCmdlet", "ChangelistDescription", "Fixed up Redirects");
-		TSharedRef<FCheckIn, ESPMode::ThreadSafe> CheckInOperation = ISourceControlOperation::Create<FCheckIn>();
-		CheckInOperation->SetDescription( Description );
-		SourceControlProvider.Execute(CheckInOperation, SourceControlHelpers::PackageFilenames(FilesToSubmit));
+		// Check in all changed files, but don't actually create a changelist if there's nothing to submit.
+		if (FilesToSubmit.Num() > 0)
+		{
+			const FText Description = NSLOCTEXT("FixupRedirectsCmdlet", "ChangelistDescription", "Fixed up Redirects");
+			TSharedRef<FCheckIn, ESPMode::ThreadSafe> CheckInOperation = ISourceControlOperation::Create<FCheckIn>();
+			CheckInOperation->SetDescription(Description);
+			SourceControlProvider.Execute(CheckInOperation, FilesToSubmit);
+		}
 
 		// toss the SCC manager
 		ISourceControlModule::Get().GetProvider().Close();
diff --git a/Engine/Source/Editor/UnrealEd/Private/Commandlets/GatherTextCommandlet.cpp b/Engine/Source/Editor/UnrealEd/Private/Commandlets/GatherTextCommandlet.cpp
--- a/Engine/Source/Editor/UnrealEd/Private/Commandlets/GatherTextCommandlet.cpp
+++ b/Engine/Source/Editor/UnrealEd/Private/Commandlets/GatherTextCommandlet.cpp
@@ -63,9 +63,7 @@ int32 UGatherTextCommandlet::Main( const FString& Params )
 
 	GConfig->LoadFile(*GatherTextConfigPath);
 
-	FConfigFile* ConfigFile = GConfig->FindConfigFile(*GatherTextConfigPath);
-
-	if( NULL == ConfigFile )
+	if (!GConfig->FindConfigFile(*GatherTextConfigPath))
 	{
 		UE_LOG(LogGatherTextCommandlet, Error, TEXT("Loading Config File \"%s\" failed."), *GatherTextConfigPath);
 		return -1; 
@@ -94,25 +92,37 @@ int32 UGatherTextCommandlet::Main( const FString& Params )
 	TSharedRef<FLocTextHelper> CommandletGatherManifestHelper = MakeShareable(new FLocTextHelper(MakeShareable(new FLocFileSCCNotifies(CommandletSourceControlInfo))));
 	CommandletGatherManifestHelper->LoadManifest(ELocTextHelperLoadFlags::Create);
 
-	int32 NumSteps = (ConfigFile->Find("CommonSettings") != NULL) ? ConfigFile->Num() - 1 :  ConfigFile->Num();
+	const FString GatherTextStepPrefix = TEXT("GatherTextStep");
 
-	//Execute each step defined in the config file.
-	for( int32 i=0; i<NumSteps ; ++i )
+	// Read the list of steps from the config file (they all have the format GatherTextStep{N})
+	TArray<FString> StepNames;
+	GConfig->GetSectionNames(GatherTextConfigPath, StepNames);
+	StepNames.RemoveAllSwap([&GatherTextStepPrefix](const FString& InStepName)
 	{
-		FString SectionName = FString::Printf(TEXT("GatherTextStep%d"),i);
-		FConfigSection* CurrCommandletSection = ConfigFile->Find(SectionName);
-		if( NULL == CurrCommandletSection )
-		{
-			UE_LOG(LogGatherTextCommandlet, Error, TEXT("Could not find %s"),*SectionName);
-			continue;
-		}
+		return !InStepName.StartsWith(GatherTextStepPrefix);
+	});
+
+	// Make sure the steps are sorted in ascending order (by numerical suffix)
+	StepNames.Sort([&GatherTextStepPrefix](const FString& InStepNameOne, const FString& InStepNameTwo)
+	{
+		const FString NumericalSuffixOneStr = InStepNameOne.RightChop(GatherTextStepPrefix.Len());
+		const int32 NumericalSuffixOne = FCString::Atoi(*NumericalSuffixOneStr);
+
+		const FString NumericalSuffixTwoStr = InStepNameTwo.RightChop(GatherTextStepPrefix.Len());
+		const int32 NumericalSuffixTwo = FCString::Atoi(*NumericalSuffixTwoStr);
+
+		return NumericalSuffixOne < NumericalSuffixTwo;
+	});
 
-		FString CommandletClassName = GConfig->GetStr( *SectionName, TEXT("CommandletClass"), GatherTextConfigPath ) + TEXT("Commandlet");
+	//Execute each step defined in the config file.
+	for (const FString& StepName : StepNames)
+	{
+		FString CommandletClassName = GConfig->GetStr( *StepName, TEXT("CommandletClass"), GatherTextConfigPath ) + TEXT("Commandlet");
 
 		UClass* CommandletClass = FindObject<UClass>(ANY_PACKAGE,*CommandletClassName,false);
 		if (!CommandletClass)
 		{
-			UE_LOG(LogGatherTextCommandlet, Error,TEXT("The commandlet name %s in section %s is invalid."), *CommandletClassName, *SectionName);
+			UE_LOG(LogGatherTextCommandlet, Error,TEXT("The commandlet name %s in section %s is invalid."), *CommandletClassName, *StepName);
 			continue;
 		}
 
@@ -124,10 +134,9 @@ int32 UGatherTextCommandlet::Main( const FString& Params )
 		// Execute the commandlet.
 		double CommandletExecutionStartTime = FPlatformTime::Seconds();
 
-		UE_LOG(LogGatherTextCommandlet, Log,TEXT("Executing %s: %s"), *SectionName, *CommandletClassName);
-		
+		UE_LOG(LogGatherTextCommandlet, Log,TEXT("Executing %s: %s"), *StepName, *CommandletClassName);
 		
-		FString GeneratedCmdLine = FString::Printf(TEXT("-Config=\"%s\" -Section=%s"), *GatherTextConfigPath , *SectionName);
+		FString GeneratedCmdLine = FString::Printf(TEXT("-Config=\"%s\" -Section=%s"), *GatherTextConfigPath , *StepName);
 
 		// Add all the command params with the exception of config
 		for(auto ParamIter = ParamVals.CreateConstIterator(); ParamIter; ++ParamIter)
@@ -149,7 +158,7 @@ int32 UGatherTextCommandlet::Main( const FString& Params )
 
 		if( 0 != Commandlet->Main( GeneratedCmdLine ) )
 		{
-			UE_LOG(LogGatherTextCommandlet, Error,TEXT("%s-%s reported an error."),*SectionName, *CommandletClassName);
+			UE_LOG(LogGatherTextCommandlet, Error,TEXT("%s-%s reported an error."), *StepName, *CommandletClassName);
 			if( CommandletSourceControlInfo.IsValid() )
 			{
 				FText SCCErrorStr;
@@ -161,7 +170,7 @@ int32 UGatherTextCommandlet::Main( const FString& Params )
 			return -1;
 		}
 
-		UE_LOG(LogGatherTextCommandlet, Log,TEXT("Completed %s: %s"), *SectionName, *CommandletClassName);
+		UE_LOG(LogGatherTextCommandlet, Log,TEXT("Completed %s: %s"), *StepName, *CommandletClassName);
 	}
 
 	if( CommandletSourceControlInfo.IsValid() && !bDisableSubmit )
diff --git a/Engine/Source/Editor/UnrealEd/Private/Commandlets/GatherTextFromSourceCommandlet.cpp b/Engine/Source/Editor/UnrealEd/Private/Commandlets/GatherTextFromSourceCommandlet.cpp
--- a/Engine/Source/Editor/UnrealEd/Private/Commandlets/GatherTextFromSourceCommandlet.cpp
+++ b/Engine/Source/Editor/UnrealEd/Private/Commandlets/GatherTextFromSourceCommandlet.cpp
@@ -5,6 +5,7 @@
 #include "Misc/FileHelper.h"
 #include "Misc/Paths.h"
 #include "Misc/ExpressionParserTypes.h"
+#include "Internationalization/InternationalizationMetadata.h"
 
 DEFINE_LOG_CATEGORY_STATIC(LogGatherTextFromSourceCommandlet, Log, All);
 
@@ -237,14 +238,24 @@ int32 UGatherTextFromSourceCommandlet::Main( const FString& Params )
 
 	// New Localization System with Namespace as literal argument.
 	Parsables.Add(new FStringMacroDescriptor( FString(TEXT("NSLOCTEXT")),
-		FMacroDescriptor::FMacroArg(FMacroDescriptor::MAS_Namespace, true),
-		FMacroDescriptor::FMacroArg(FMacroDescriptor::MAS_Identifier, true),
-		FMacroDescriptor::FMacroArg(FMacroDescriptor::MAS_SourceText, true)));
+		FStringMacroDescriptor::FMacroArg(FStringMacroDescriptor::MAS_Namespace, true),
+		FStringMacroDescriptor::FMacroArg(FStringMacroDescriptor::MAS_Identifier, true),
+		FStringMacroDescriptor::FMacroArg(FStringMacroDescriptor::MAS_SourceText, true)));
 	
 	// New Localization System with Namespace as preprocessor define.
 	Parsables.Add(new FStringMacroDescriptor( FString(TEXT("LOCTEXT")),
-		FMacroDescriptor::FMacroArg(FMacroDescriptor::MAS_Identifier, true),
-		FMacroDescriptor::FMacroArg(FMacroDescriptor::MAS_SourceText, true)));
+		FStringMacroDescriptor::FMacroArg(FStringMacroDescriptor::MAS_Identifier, true),
+		FStringMacroDescriptor::FMacroArg(FStringMacroDescriptor::MAS_SourceText, true)));
+
+	Parsables.Add(new FStringTableMacroDescriptor());
+
+	Parsables.Add(new FStringTableFromFileMacroDescriptor(TEXT("LOCTABLE_FROMFILE_ENGINE"), FPaths::EngineContentDir()));
+
+	Parsables.Add(new FStringTableFromFileMacroDescriptor(TEXT("LOCTABLE_FROMFILE_GAME"), FPaths::GameContentDir()));
+
+	Parsables.Add(new FStringTableEntryMacroDescriptor());
+
+	Parsables.Add(new FStringTableEntryMetaDataMacroDescriptor());
 
 	Parsables.Add(new FIniNamespaceDescriptor());
 
@@ -305,6 +316,47 @@ int32 UGatherTextFromSourceCommandlet::Main( const FString& Params )
 		}
 	}
 	
+	// Process any parsed string tables
+	for (const auto& ParsedStringTablePair : ParseCtxt.ParsedStringTables)
+	{
+		if (ParsedStringTablePair.Value.SourceLocation.Line == INDEX_NONE)
+		{
+			UE_LOG(LogGatherTextFromSourceCommandlet, Warning, TEXT("String table with ID '%s' had %s entries parsed for it, but the table was never registered. Skipping for gather."), *ParsedStringTablePair.Key.ToString(), ParsedStringTablePair.Value.TableEntries.Num());
+		}
+		else
+		{
+			for (const auto& ParsedStringTableEntryPair : ParsedStringTablePair.Value.TableEntries)
+			{
+				if (!ParsedStringTableEntryPair.Value.bIsEditorOnly || ParseCtxt.ShouldGatherFromEditorOnlyData)
+				{
+					FManifestContext SourceContext;
+					SourceContext.Key = ParsedStringTableEntryPair.Key;
+					SourceContext.SourceLocation = ParsedStringTableEntryPair.Value.SourceLocation.ToString();
+
+					const FParsedStringTableEntryMetaDataMap* ParsedMetaDataMap = ParsedStringTablePair.Value.MetaDataEntries.Find(ParsedStringTableEntryPair.Key);
+					if (ParsedMetaDataMap && ParsedMetaDataMap->Num() > 0)
+					{
+						SourceContext.InfoMetadataObj = MakeShareable(new FLocMetadataObject());
+						for (const auto& ParsedMetaDataPair : *ParsedMetaDataMap)
+						{
+							if (!ParsedMetaDataPair.Value.bIsEditorOnly || ParseCtxt.ShouldGatherFromEditorOnlyData)
+							{
+								SourceContext.InfoMetadataObj->SetStringField(ParsedMetaDataPair.Key.ToString(), ParsedMetaDataPair.Value.MetaData);
+							}
+						}
+					}
+
+					FString SourceDescription = FString::Printf(TEXT("In string table '%s' at '%s'"),
+						*ParsedStringTablePair.Key.ToString(),
+						*ParsedStringTableEntryPair.Value.SourceLocation.ToString()
+						);
+
+					ParseCtxt.GatherManifestHelper->AddSourceText(ParsedStringTablePair.Value.TableNamespace, FLocItem(ParsedStringTableEntryPair.Value.SourceString), SourceContext, &SourceDescription);
+				}
+			}
+		}
+	}
+
 	// Clear parsables list safely
 	for (int32 i=0; i<Parsables.Num(); i++)
 	{
@@ -497,7 +549,7 @@ FString UGatherTextFromSourceCommandlet::UnescapeLiteralCharacterEscapeSequences
 		}
 	}
 
-	return RetString;
+	return RetString.ReplaceEscapedCharWithChar();
 }
 
 FString UGatherTextFromSourceCommandlet::RemoveStringFromTextMacro(const FString& TextMacro, const FString& IdentForLogging, bool& Error)
@@ -877,8 +929,7 @@ bool UGatherTextFromSourceCommandlet::FSourceFileParseContext::AddManifestText(
 			*Filename, 
 			LineNumber, 
 			*LineText);
-		FLocItem Source( SourceText.ReplaceEscapedCharWithChar() );
-		return GatherManifestHelper->AddSourceText(InNamespace, Source, Context, &EntryDescription);
+		return GatherManifestHelper->AddSourceText(InNamespace, FLocItem(SourceText), Context, &EntryDescription);
 	}
 
 	return false;
@@ -1010,6 +1061,129 @@ void UGatherTextFromSourceCommandlet::FSourceFileParseContext::RemoveDefine(cons
 	}
 }
 
+bool UGatherTextFromSourceCommandlet::FSourceFileParseContext::AddStringTableImpl(const FName InTableId, const FString& InTableNamespace)
+{
+	// String table entries may be parsed before the string table itself (due to code ordering), so only warn about duplication here if we've already got a source location for the string table (as adding entries doesn't set that)
+	FParsedStringTable& ParsedStringTable = ParsedStringTables.FindOrAdd(InTableId);
+	if (ParsedStringTable.SourceLocation.Line != INDEX_NONE)
+	{
+		UE_LOG(LogGatherTextFromSourceCommandlet, Warning, TEXT("String table with ID '%s' at '%s' was already parsed at '%s'. Ignoring additional definition."), *InTableId.ToString(), *FSourceLocation(Filename, LineNumber).ToString(), *ParsedStringTable.SourceLocation.ToString());
+		return false;
+	}
+
+	ParsedStringTable.TableNamespace = InTableNamespace;
+	ParsedStringTable.SourceLocation = FSourceLocation(Filename, LineNumber);
+	return true;
+}
+
+bool UGatherTextFromSourceCommandlet::FSourceFileParseContext::AddStringTableEntryImpl(const FName InTableId, const FString& InKey, const FString& InSourceString, const FSourceLocation& InSourceLocation)
+{
+	const bool bIsEditorOnly = EvaluateMacroStack() == EMacroBlockState::EditorOnly;
+
+	// String table entries may be parsed before the string table itself (due to code ordering), so we may need to add our string table below
+	FParsedStringTable& ParsedStringTable = ParsedStringTables.FindOrAdd(InTableId);
+
+	FParsedStringTableEntry* ExistingEntry = ParsedStringTable.TableEntries.Find(InKey);
+	if (ExistingEntry)
+	{
+		if (ExistingEntry->SourceString.Equals(InSourceString, ESearchCase::CaseSensitive))
+		{
+			ExistingEntry->bIsEditorOnly &= bIsEditorOnly;
+			return true;
+		}
+		else
+		{
+			UE_LOG(LogGatherTextFromSourceCommandlet, Warning, TEXT("String table entry with ID '%s' and key '%s' at '%s' was already parsed at '%s'. Ignoring additional definition."), *InTableId.ToString(), *InKey, *FSourceLocation(Filename, LineNumber).ToString(), *ExistingEntry->SourceLocation.ToString());
+			return false;
+		}
+	}
+	else
+	{
+		FParsedStringTableEntry& ParsedStringTableEntry = ParsedStringTable.TableEntries.Add(InKey);
+		ParsedStringTableEntry.SourceString = InSourceString;
+		ParsedStringTableEntry.SourceLocation = InSourceLocation;
+		ParsedStringTableEntry.bIsEditorOnly = bIsEditorOnly;
+		return true;
+	}
+}
+
+bool UGatherTextFromSourceCommandlet::FSourceFileParseContext::AddStringTableEntryMetaDataImpl(const FName InTableId, const FString& InKey, const FName InMetaDataId, const FString& InMetaData, const FSourceLocation& InSourceLocation)
+{
+	const bool bIsEditorOnly = EvaluateMacroStack() == EMacroBlockState::EditorOnly;
+
+	// String table meta-data may be parsed before the string table itself (due to code ordering), so we may need to add our string table below
+	FParsedStringTable& ParsedStringTable = ParsedStringTables.FindOrAdd(InTableId);
+	FParsedStringTableEntryMetaDataMap& MetaDataMap = ParsedStringTable.MetaDataEntries.FindOrAdd(InKey);
+
+	FParsedStringTableEntryMetaData* ExistingMetaData = MetaDataMap.Find(InMetaDataId);
+	if (ExistingMetaData)
+	{
+		if (ExistingMetaData->MetaData.Equals(InMetaData, ESearchCase::CaseSensitive))
+		{
+			ExistingMetaData->bIsEditorOnly &= bIsEditorOnly;
+			return true;
+		}
+		else
+		{
+			UE_LOG(LogGatherTextFromSourceCommandlet, Warning, TEXT("String table entry meta-data with ID '%s' and key '%s' at '%s' was already parsed at '%s'. Ignoring additional definition."), *InTableId.ToString(), *InKey, *FSourceLocation(Filename, LineNumber).ToString(), *ExistingMetaData->SourceLocation.ToString());
+			return false;
+		}
+	}
+	else
+	{
+		FParsedStringTableEntryMetaData& ParsedMetaData = MetaDataMap.Add(InMetaDataId);
+		ParsedMetaData.MetaData = InMetaData;
+		ParsedMetaData.SourceLocation = InSourceLocation;
+		ParsedMetaData.bIsEditorOnly = bIsEditorOnly;
+		return true;
+	}
+}
+
+void UGatherTextFromSourceCommandlet::FSourceFileParseContext::AddStringTable(const FName InTableId, const FString& InTableNamespace)
+{
+	AddStringTableImpl(InTableId, InTableNamespace);
+}
+
+void UGatherTextFromSourceCommandlet::FSourceFileParseContext::AddStringTableFromFile(const FName InTableId, const FString& InTableNamespace, const FString& InTableFilename, const FString& InRootPath)
+{
+	if (AddStringTableImpl(InTableId, InTableNamespace))
+	{
+		const FString FullImportPath = InRootPath / InTableFilename;
+
+		FStringTableRef TmpStringTable = FStringTable::NewStringTable();
+		if (TmpStringTable->ImportStrings(FullImportPath))
+		{
+			TmpStringTable->EnumerateSourceStrings([&](const FString& InKey, const FString& InSourceString)
+			{
+				const FSourceLocation SourceLocation = FSourceLocation(InTableFilename, INDEX_NONE);
+				AddStringTableEntryImpl(InTableId, InKey, InSourceString, SourceLocation);
+
+				TmpStringTable->EnumerateMetaData(InKey, [&](const FName InMetaDataId, const FString& InMetaData)
+				{
+					AddStringTableEntryMetaDataImpl(InTableId, InKey, InMetaDataId, InMetaData, SourceLocation);
+					return true; // continue enumeration
+				});
+
+				return true; // continue enumeration
+			});
+		}
+		else
+		{
+			UE_LOG(LogGatherTextFromSourceCommandlet, Warning, TEXT("String table with ID '%s' at '%s' failed to import strings from '%s'."), *InTableId.ToString(), *FSourceLocation(Filename, LineNumber).ToString(), *FullImportPath);
+		}
+	}
+}
+
+void UGatherTextFromSourceCommandlet::FSourceFileParseContext::AddStringTableEntry(const FName InTableId, const FString& InKey, const FString& InSourceString)
+{
+	AddStringTableEntryImpl(InTableId, InKey, InSourceString, FSourceLocation(Filename, LineNumber));
+}
+
+void UGatherTextFromSourceCommandlet::FSourceFileParseContext::AddStringTableEntryMetaData(const FName InTableId, const FString& InKey, const FName InMetaDataId, const FString& InMetaData)
+{
+	AddStringTableEntryMetaDataImpl(InTableId, InKey, InMetaDataId, InMetaData, FSourceLocation(Filename, LineNumber));
+}
+
 void UGatherTextFromSourceCommandlet::FDefineDescriptor::TryParse(const FString& Text, FSourceFileParseContext& Context) const
 {
 	// Attempt to parse something of the format
@@ -1255,13 +1429,13 @@ void UGatherTextFromSourceCommandlet::FCommandMacroDescriptor::TryParse(const FS
 				FString Identifier = Arguments[0].Trim();
 				const FString UICommandRootNamespace = TEXT("UICommands");
 				FString Namespace = Context.WithinNamespaceDefine && !Context.Namespace.IsEmpty() ? FString::Printf( TEXT("%s.%s"), *UICommandRootNamespace, *Context.Namespace) : UICommandRootNamespace;
-				FString SourceLocation = FString( Context.Filename + TEXT(" - line ") + FString::FromInt(Context.LineNumber) );
+				FString SourceLocation = FSourceLocation(Context.Filename, Context.LineNumber).ToString();
 				FString SourceText = Arguments[1].Trim();
 
 				if ( Identifier.IsEmpty() )
 				{
 					//The command doesn't have an identifier so we can't gather it
-					UE_LOG(LogGatherTextFromSourceCommandlet, Warning, TEXT("UICOMMAND macro doesn't have unique identifier. %s"), *SourceLocation );
+					UE_LOG(LogGatherTextFromSourceCommandlet, Warning, TEXT("UICOMMAND macro has an empty identifier and cannot be gathered. %s"), *SourceLocation );
 					return;
 				}
 
@@ -1321,7 +1495,7 @@ void UGatherTextFromSourceCommandlet::FStringMacroDescriptor::TryParse(const FSt
 			{
 				FString Identifier;
 				FString Namespace = Context.Namespace;
-				FString SourceLocation = FString( Context.Filename + TEXT(" - line ") + FString::FromInt(Context.LineNumber) );
+				FString SourceLocation = FSourceLocation(Context.Filename, Context.LineNumber).ToString();
 				FString SourceText;
 
 				bool ArgParseError = false;
@@ -1378,6 +1552,179 @@ void UGatherTextFromSourceCommandlet::FStringMacroDescriptor::TryParse(const FSt
 	}
 }
 
+void UGatherTextFromSourceCommandlet::FStringTableMacroDescriptor::TryParse(const FString& Text, FSourceFileParseContext& Context) const
+{
+	// Attempt to parse something of the format
+	// LOCTABLE_NEW(Id, Namespace)
+
+	if (!Context.ExcludedRegion && !Context.WithinBlockComment && !Context.WithinLineComment && !Context.WithinStringLiteral)
+	{
+		TArray<FString> Arguments;
+		if (ParseArgsFromMacro(StripCommentsFromToken(Text, Context), Arguments, Context))
+		{
+			if (Arguments.Num() != 2)
+			{
+				UE_LOG(LogGatherTextFromSourceCommandlet, Warning, TEXT("Unexpected number of arguments for %s macro in %s(%d):%s"), *GetToken(), *Context.Filename, Context.LineNumber, *FLocTextHelper::SanitizeLogOutput(Context.LineText));
+			}
+			else
+			{
+				FString TableId = Arguments[0].Trim();
+				FString TableNamespace = Arguments[1].Trim();
+
+				const FString TableIdMacroDesc = FString::Printf(TEXT("\"Id\" argument in %s macro %s(%d):%s"), *GetToken(), *Context.Filename, Context.LineNumber, *Context.LineText);
+				const FString TableNamespaceMacroDesc = FString::Printf(TEXT("\"Namespace\" argument in %s macro %s(%d):%s"), *GetToken(), *Context.Filename, Context.LineNumber, *Context.LineText);
+
+				bool HasQuotes;
+				if (PrepareArgument(TableId, true, TableIdMacroDesc, HasQuotes) && PrepareArgument(TableNamespace, true, TableNamespaceMacroDesc, HasQuotes))
+				{
+					const FName TableIdName = *TableId;
+
+					if (TableIdName.IsNone())
+					{
+						UE_LOG(LogGatherTextFromSourceCommandlet, Warning, TEXT("%s macro has an empty identifier and cannot be gathered. %s"), *GetToken(), *FSourceLocation(Context.Filename, Context.LineNumber).ToString());
+					}
+					else
+					{
+						Context.AddStringTable(TableIdName, TableNamespace);
+					}
+				}
+			}
+		}
+	}
+}
+
+void UGatherTextFromSourceCommandlet::FStringTableFromFileMacroDescriptor::TryParse(const FString& Text, FSourceFileParseContext& Context) const
+{
+	// Attempt to parse something of the format
+	// LOCTABLE_FROMFILE_X(Id, Namespace, FilePath)
+
+	if (!Context.ExcludedRegion && !Context.WithinBlockComment && !Context.WithinLineComment && !Context.WithinStringLiteral)
+	{
+		TArray<FString> Arguments;
+		if (ParseArgsFromMacro(StripCommentsFromToken(Text, Context), Arguments, Context))
+		{
+			if (Arguments.Num() != 3)
+			{
+				UE_LOG(LogGatherTextFromSourceCommandlet, Warning, TEXT("Unexpected number of arguments for %s macro in %s(%d):%s"), *GetToken(), *Context.Filename, Context.LineNumber, *FLocTextHelper::SanitizeLogOutput(Context.LineText));
+			}
+			else
+			{
+				FString TableId = Arguments[0].Trim();
+				FString TableNamespace = Arguments[1].Trim();
+				FString TableFilename = Arguments[2].Trim();
+
+				const FString TableIdMacroDesc = FString::Printf(TEXT("\"Id\" argument in %s macro %s(%d):%s"), *GetToken(), *Context.Filename, Context.LineNumber, *Context.LineText);
+				const FString TableNamespaceMacroDesc = FString::Printf(TEXT("\"Namespace\" argument in %s macro %s(%d):%s"), *GetToken(), *Context.Filename, Context.LineNumber, *Context.LineText);
+				const FString TableFilenameMacroDesc = FString::Printf(TEXT("\"FilePath\" argument in %s macro %s(%d):%s"), *GetToken(), *Context.Filename, Context.LineNumber, *Context.LineText);
+
+				bool HasQuotes;
+				if (PrepareArgument(TableId, true, TableIdMacroDesc, HasQuotes) && PrepareArgument(TableNamespace, true, TableNamespaceMacroDesc, HasQuotes) && PrepareArgument(TableFilename, true, TableFilenameMacroDesc, HasQuotes))
+				{
+					const FName TableIdName = *TableId;
+
+					if (TableIdName.IsNone())
+					{
+						UE_LOG(LogGatherTextFromSourceCommandlet, Warning, TEXT("%s macro has an empty identifier and cannot be gathered. %s"), *GetToken(), *FSourceLocation(Context.Filename, Context.LineNumber).ToString());
+					}
+					else
+					{
+						Context.AddStringTableFromFile(TableIdName, TableNamespace, TableFilename, RootPath);
+					}
+				}
+			}
+		}
+	}
+}
+
+void UGatherTextFromSourceCommandlet::FStringTableEntryMacroDescriptor::TryParse(const FString& Text, FSourceFileParseContext& Context) const
+{
+	// Attempt to parse something of the format
+	// LOCTABLE_SETSTRING(Id, Key, SourceString)
+
+	if (!Context.ExcludedRegion && !Context.WithinBlockComment && !Context.WithinLineComment && !Context.WithinStringLiteral)
+	{
+		TArray<FString> Arguments;
+		if (ParseArgsFromMacro(StripCommentsFromToken(Text, Context), Arguments, Context))
+		{
+			if (Arguments.Num() != 3)
+			{
+				UE_LOG(LogGatherTextFromSourceCommandlet, Warning, TEXT("Unexpected number of arguments for %s macro in %s(%d):%s"), *GetToken(), *Context.Filename, Context.LineNumber, *FLocTextHelper::SanitizeLogOutput(Context.LineText));
+			}
+			else
+			{
+				FString TableId = Arguments[0].Trim();
+				FString Key = Arguments[1].Trim();
+				FString SourceString = Arguments[2].Trim();
+
+				const FString TableIdMacroDesc = FString::Printf(TEXT("\"Id\" argument in %s macro %s(%d):%s"), *GetToken(), *Context.Filename, Context.LineNumber, *Context.LineText);
+				const FString KeyMacroDesc = FString::Printf(TEXT("\"Key\" argument in %s macro %s(%d):%s"), *GetToken(), *Context.Filename, Context.LineNumber, *Context.LineText);
+				const FString SourceStringMacroDesc = FString::Printf(TEXT("\"SourceString\" argument in %s macro %s(%d):%s"), *GetToken(), *Context.Filename, Context.LineNumber, *Context.LineText);
+
+				bool HasQuotes;
+				if (PrepareArgument(TableId, true, TableIdMacroDesc, HasQuotes) && PrepareArgument(Key, true, KeyMacroDesc, HasQuotes) && PrepareArgument(SourceString, true, SourceStringMacroDesc, HasQuotes))
+				{
+					const FName TableIdName = *TableId;
+
+					if (TableIdName.IsNone() || Key.IsEmpty())
+					{
+						UE_LOG(LogGatherTextFromSourceCommandlet, Warning, TEXT("%s macro has an empty identifier and cannot be gathered. %s"), *GetToken(), *FSourceLocation(Context.Filename, Context.LineNumber).ToString());
+					}
+					else if (!SourceString.IsEmpty())
+					{
+						Context.AddStringTableEntry(TableIdName, Key, SourceString);
+					}
+				}
+			}
+		}
+	}
+}
+
+void UGatherTextFromSourceCommandlet::FStringTableEntryMetaDataMacroDescriptor::TryParse(const FString& Text, FSourceFileParseContext& Context) const
+{
+	// Attempt to parse something of the format
+	// LOCTABLE_SETMETA(Id, Key, SourceString)
+
+	if (!Context.ExcludedRegion && !Context.WithinBlockComment && !Context.WithinLineComment && !Context.WithinStringLiteral)
+	{
+		TArray<FString> Arguments;
+		if (ParseArgsFromMacro(StripCommentsFromToken(Text, Context), Arguments, Context))
+		{
+			if (Arguments.Num() != 4)
+			{
+				UE_LOG(LogGatherTextFromSourceCommandlet, Warning, TEXT("Unexpected number of arguments for %s macro in %s(%d):%s"), *GetToken(), *Context.Filename, Context.LineNumber, *FLocTextHelper::SanitizeLogOutput(Context.LineText));
+			}
+			else
+			{
+				FString TableId = Arguments[0].Trim();
+				FString Key = Arguments[1].Trim();
+				FString MetaDataId = Arguments[2].Trim();
+				FString MetaData = Arguments[3].Trim();
+
+				const FString TableIdMacroDesc = FString::Printf(TEXT("\"Id\" argument in %s macro %s(%d):%s"), *GetToken(), *Context.Filename, Context.LineNumber, *Context.LineText);
+				const FString KeyMacroDesc = FString::Printf(TEXT("\"Key\" argument in %s macro %s(%d):%s"), *GetToken(), *Context.Filename, Context.LineNumber, *Context.LineText);
+				const FString MetaDataIdMacroDesc = FString::Printf(TEXT("\"MetaDataId\" argument in %s macro %s(%d):%s"), *GetToken(), *Context.Filename, Context.LineNumber, *Context.LineText);
+				const FString MetaDataMacroDesc = FString::Printf(TEXT("\"MetaData\" argument in %s macro %s(%d):%s"), *GetToken(), *Context.Filename, Context.LineNumber, *Context.LineText);
+
+				bool HasQuotes;
+				if (PrepareArgument(TableId, true, TableIdMacroDesc, HasQuotes) && PrepareArgument(Key, true, KeyMacroDesc, HasQuotes) && PrepareArgument(MetaDataId, true, MetaDataIdMacroDesc, HasQuotes) && PrepareArgument(MetaData, true, MetaDataMacroDesc, HasQuotes))
+				{
+					const FName TableIdName = *TableId;
+					const FName MetaDataIdName = *MetaDataId;
+
+					if (TableIdName.IsNone() || Key.IsEmpty() || MetaDataIdName.IsNone())
+					{
+						UE_LOG(LogGatherTextFromSourceCommandlet, Warning, TEXT("%s macro has an empty identifier and cannot be gathered. %s"), *GetToken(), *FSourceLocation(Context.Filename, Context.LineNumber).ToString());
+					}
+					else if (!MetaData.IsEmpty())
+					{
+						Context.AddStringTableEntryMetaData(TableIdName, Key, MetaDataIdName, MetaData);
+					}
+				}
+			}
+		}
+	}
+}
+
 void UGatherTextFromSourceCommandlet::FIniNamespaceDescriptor::TryParse(const FString& Text, FSourceFileParseContext& Context) const
 {
 	// Attempt to parse something of the format
diff --git a/Engine/Source/Editor/UnrealEd/Private/ConvexDecompTool.h b/Engine/Source/Editor/UnrealEd/Private/ConvexDecompTool.h
--- a/Engine/Source/Editor/UnrealEd/Private/ConvexDecompTool.h
+++ b/Engine/Source/Editor/UnrealEd/Private/ConvexDecompTool.h
@@ -4,8 +4,7 @@
 	ConvexDecompTool.h: Utility for turning graphics mesh into convex hulls.
 =============================================================================*/
 
-#ifndef __CONVEXDECOMPTOOL_H__
-#define __CONVEXDECOMPTOOL_H__
+#pragma once
 
 #include "CoreMinimal.h"
 
@@ -20,5 +19,3 @@ class UBodySetup;
  *	@param		InMaxHullVerts		Number of verts allowed in a hull
  */
 UNREALED_API void DecomposeMeshToHulls(UBodySetup* InBodySetup, const TArray<FVector>& InVertices, const TArray<uint32>& InIndices, float InAccuracy, int32 InMaxHullVerts);
-
-#endif // __CONVEXDECOMPTOOL_H__
diff --git a/Engine/Source/Editor/UnrealEd/Private/DataTableEditorUtils.cpp b/Engine/Source/Editor/UnrealEd/Private/DataTableEditorUtils.cpp
--- a/Engine/Source/Editor/UnrealEd/Private/DataTableEditorUtils.cpp
+++ b/Engine/Source/Editor/UnrealEd/Private/DataTableEditorUtils.cpp
@@ -170,6 +170,48 @@ bool FDataTableEditorUtils::MoveRow(UDataTable* DataTable, FName RowName, ERowMo
 	return true;
 }
 
+bool FDataTableEditorUtils::DiffersFromDefault(UDataTable* DataTable, FName RowName)
+{
+	bool bDiffers = false;
+
+	if (DataTable && DataTable->RowMap.Contains(RowName))
+	{
+		uint8* RowData = DataTable->RowMap[RowName];
+
+		if (const UUserDefinedStruct* UDStruct = Cast<const UUserDefinedStruct>(DataTable->RowStruct))
+		{
+			bDiffers = UDStruct->DiffersFromDefaultValue(RowData);
+		}
+	}
+
+	return bDiffers;
+}
+
+bool FDataTableEditorUtils::ResetToDefault(UDataTable* DataTable, FName RowName)
+{
+	bool bResult = false;
+
+	if (DataTable && DataTable->RowMap.Contains(RowName))
+	{
+		const FScopedTransaction Transaction(LOCTEXT("ResetDataTableRowToDefault", "Reset Data Table Row to Default Values"));
+
+		BroadcastPreChange(DataTable, EDataTableChangeInfo::RowData);
+		DataTable->Modify();
+
+		uint8* RowData = DataTable->RowMap[RowName];
+
+		if (const UUserDefinedStruct* UDStruct = Cast<const UUserDefinedStruct>(DataTable->RowStruct))
+		{
+			UDStruct->InitializeDefaultValue(RowData);
+			bResult = true;
+		}
+
+		BroadcastPostChange(DataTable, EDataTableChangeInfo::RowData);
+	}
+
+	return bResult;
+}
+
 void FDataTableEditorUtils::BroadcastPreChange(UDataTable* DataTable, EDataTableChangeInfo Info)
 {
 	FDataTableEditorManager::Get().PreChange(DataTable, Info);
diff --git a/Engine/Source/Editor/UnrealEd/Private/Editor.cpp b/Engine/Source/Editor/UnrealEd/Private/Editor.cpp
--- a/Engine/Source/Editor/UnrealEd/Private/Editor.cpp
+++ b/Engine/Source/Editor/UnrealEd/Private/Editor.cpp
@@ -81,6 +81,7 @@ FSimpleMulticastDelegate								FEditorDelegates::ActorPropertiesChange;
 FSimpleMulticastDelegate								FEditorDelegates::RefreshEditor;
 FSimpleMulticastDelegate								FEditorDelegates::RefreshAllBrowsers;
 FSimpleMulticastDelegate								FEditorDelegates::RefreshLayerBrowser;
+FSimpleMulticastDelegate								FEditorDelegates::RefreshLevelBrowser;
 FSimpleMulticastDelegate								FEditorDelegates::RefreshPrimitiveStatsBrowser;
 FSimpleMulticastDelegate								FEditorDelegates::LoadSelectedAssetsIfNeeded;
 FSimpleMulticastDelegate								FEditorDelegates::DisplayLoadErrors;
@@ -189,7 +190,7 @@ void FReimportManager::UpdateReimportPaths( UObject* Obj, const TArray<FString>&
 	}
 }
 
-bool FReimportManager::Reimport( UObject* Obj, bool bAskForNewFileIfMissing, bool bShowNotification, FString PreferedReimportFile, FReimportHandler* SpecifiedReimportHandler)
+bool FReimportManager::Reimport( UObject* Obj, bool bAskForNewFileIfMissing, bool bShowNotification, FString PreferredReimportFile, FReimportHandler* SpecifiedReimportHandler)
 {
 	// Warn that were about to reimport, so prep for it
 	PreReimport.Broadcast( Obj );
@@ -242,12 +243,12 @@ bool FReimportManager::Reimport( UObject* Obj, bool bAskForNewFileIfMissing, boo
 			}
 
 			bValidSourceFilename = true;
-			if ((bAskForNewFileIfMissing || !PreferedReimportFile.IsEmpty()) && bMissingFiles )
+			if ((bAskForNewFileIfMissing || !PreferredReimportFile.IsEmpty()) && bMissingFiles )
 			{
-				if (!bAskForNewFileIfMissing && !PreferedReimportFile.IsEmpty())
+				if (!bAskForNewFileIfMissing && !PreferredReimportFile.IsEmpty())
 				{
 					SourceFilenames.Empty();
-					SourceFilenames.Add(PreferedReimportFile);
+					SourceFilenames.Add(PreferredReimportFile);
 				}
 				else
 				{
@@ -265,6 +266,13 @@ bool FReimportManager::Reimport( UObject* Obj, bool bAskForNewFileIfMissing, boo
 					CanReimportHandler->SetReimportPaths(Obj, SourceFilenames);
 				}
 			}
+			else if (!PreferredReimportFile.IsEmpty() && !SourceFilenames.Contains(PreferredReimportFile))
+			{
+				// Reimporting the asset from a new file
+				SourceFilenames.Empty();
+				SourceFilenames.Add(PreferredReimportFile);
+				CanReimportHandler->SetReimportPaths(Obj, SourceFilenames);
+			}
 
 			if ( bValidSourceFilename )
 			{
@@ -460,7 +468,7 @@ void FReimportManager::AddReferencedObjects( FReferenceCollector& Collector )
 	}
 }
 
-bool FReimportManager::ReimportMultiple(TArrayView<UObject*> Objects, bool bAskForNewFileIfMissing /*= false*/, bool bShowNotification /*= true*/, FString PreferedReimportFile /*= TEXT("")*/, FReimportHandler* SpecifiedReimportHandler /*= nullptr */)
+bool FReimportManager::ReimportMultiple(TArrayView<UObject*> Objects, bool bAskForNewFileIfMissing /*= false*/, bool bShowNotification /*= true*/, FString PreferredReimportFile /*= TEXT("")*/, FReimportHandler* SpecifiedReimportHandler /*= nullptr */)
 {
 	bool bBulkSuccess = true;
 
@@ -474,7 +482,7 @@ bool FReimportManager::ReimportMultiple(TArrayView<UObject*> Objects, bool bAskF
 			FScopedSlowTask SingleObjectTask(1.0f, SingleTaskTest);
 			SingleObjectTask.EnterProgressFrame(1.0f);
 
-			bBulkSuccess = bBulkSuccess && Reimport(CurrentObject, bAskForNewFileIfMissing, bShowNotification, PreferedReimportFile, SpecifiedReimportHandler);
+			bBulkSuccess = bBulkSuccess && Reimport(CurrentObject, bAskForNewFileIfMissing, bShowNotification, PreferredReimportFile, SpecifiedReimportHandler);
 		}
 
 		BulkReimportTask.EnterProgressFrame(1.0f);
diff --git a/Engine/Source/Editor/UnrealEd/Private/EditorEngine.cpp b/Engine/Source/Editor/UnrealEd/Private/EditorEngine.cpp
--- a/Engine/Source/Editor/UnrealEd/Private/EditorEngine.cpp
+++ b/Engine/Source/Editor/UnrealEd/Private/EditorEngine.cpp
@@ -2723,15 +2723,11 @@ void UEditorEngine::ApplyDeltaToComponent(USceneComponent* InComponent,
 		{
 			if ( bDelta )
 			{
-				const FScaleMatrix ScaleMatrix( FVector( InDeltaScale.X , InDeltaScale.Y, InDeltaScale.Z ) );
-
-				FVector DeltaScale3D = ScaleMatrix.TransformPosition( InComponent->RelativeScale3D );
-				InComponent->SetRelativeScale3D(InComponent->RelativeScale3D + DeltaScale3D);
-
+				InComponent->SetRelativeScale3D(InComponent->RelativeScale3D + InDeltaScale);
 
 				FVector NewCompLocation = InComponent->RelativeLocation;
 				NewCompLocation -= PivotLocation;
-				NewCompLocation += ScaleMatrix.TransformPosition( NewCompLocation );
+				NewCompLocation += FScaleMatrix( InDeltaScale ).TransformPosition( NewCompLocation );
 				NewCompLocation += PivotLocation;
 				InComponent->SetRelativeLocation(NewCompLocation);
 			}
diff --git a/Engine/Source/Editor/UnrealEd/Private/EditorExporters.cpp b/Engine/Source/Editor/UnrealEd/Private/EditorExporters.cpp
--- a/Engine/Source/Editor/UnrealEd/Private/EditorExporters.cpp
+++ b/Engine/Source/Editor/UnrealEd/Private/EditorExporters.cpp
@@ -426,7 +426,7 @@ bool ULevelExporterT3D::ExportText( const FExportObjectInnerContext* Context, UO
 	if ((PortFlags & PPF_Copy) == 0)
 	{
 		// If we are not copying to clipboard, then export objects contained in the map package itself...
-		MapPackage = Cast<UPackage>(Object->GetOutermost());
+		MapPackage = Object->GetOutermost();
 	}
 
 	// this is the top level in the .t3d file
diff --git a/Engine/Source/Editor/UnrealEd/Private/EditorLevelUtils.cpp b/Engine/Source/Editor/UnrealEd/Private/EditorLevelUtils.cpp
--- a/Engine/Source/Editor/UnrealEd/Private/EditorLevelUtils.cpp
+++ b/Engine/Source/Editor/UnrealEd/Private/EditorLevelUtils.cpp
@@ -545,7 +545,7 @@ namespace EditorLevelUtils
 		InLevel->MarkPendingKill();
 		InLevel->GetOuter()->ClearFlags(RF_Public | RF_Standalone);
 
-		UPackage* Package = Cast<UPackage>(InLevel->GetOutermost());
+		UPackage* Package = InLevel->GetOutermost();
 		// We want to unconditionally destroy the level, so clear the dirty flag here so it can be unloaded successfully
 		Package->SetDirtyFlag(false);
 
diff --git a/Engine/Source/Editor/UnrealEd/Private/EditorViewportClient.cpp b/Engine/Source/Editor/UnrealEd/Private/EditorViewportClient.cpp
--- a/Engine/Source/Editor/UnrealEd/Private/EditorViewportClient.cpp
+++ b/Engine/Source/Editor/UnrealEd/Private/EditorViewportClient.cpp
@@ -4985,12 +4985,20 @@ void FEditorViewportClient::ProcessScreenShots(FViewport* InViewport)
 		// properties, such as it being aspect ratio constrained
 		if (GIsHighResScreenshot && !bCaptureAreaValid)
 		{
+			// Screen Percentage is an optimization and should not affect the editor by default, unless we're rendering in stereo
+			static TConsoleVariableData<int32>* ScreenPercentageEditorCVar = IConsoleManager::Get().FindTConsoleVariableDataInt(TEXT("r.ScreenPercentage.Editor"));
+			bool bUseScreenPercentage = (GEngine && GEngine->IsStereoscopic3D(InViewport)) 
+										|| (ScreenPercentageEditorCVar && ScreenPercentageEditorCVar->GetValueOnAnyThread() != 0);
+
 			FSceneViewFamilyContext ViewFamily(FSceneViewFamily::ConstructionValues(
 				InViewport,
 				GetScene(),
 				EngineShowFlags)
 				.SetRealtimeUpdate(IsRealtime())
 				.SetViewModeParam(ViewModeParam, ViewModeParamName));
+
+			ViewFamily.EngineShowFlags.SetScreenPercentage(bUseScreenPercentage);
+
 			auto* ViewportBak = Viewport;
 			Viewport = InViewport;
 			FSceneView* View = CalcSceneView(&ViewFamily);
diff --git a/Engine/Source/Editor/UnrealEd/Private/Factories/ActorFactory.cpp b/Engine/Source/Editor/UnrealEd/Private/Factories/ActorFactory.cpp
--- a/Engine/Source/Editor/UnrealEd/Private/Factories/ActorFactory.cpp
+++ b/Engine/Source/Editor/UnrealEd/Private/Factories/ActorFactory.cpp
@@ -267,6 +267,12 @@ AActor* UActorFactory::CreateActor( UObject* Asset, ULevel* InLevel, FTransform
 		{
 			PostSpawnActor(Asset, NewActor);
 		}
+
+		// Only do this if the actor wasn't already given a name
+		if (Name == NAME_None && Asset)
+		{
+			FActorLabelUtilities::SetActorLabelUnique(NewActor, Asset->GetName());
+		}
 	}
 
 	return NewActor;
@@ -306,11 +312,6 @@ AActor* UActorFactory::SpawnActor( UObject* Asset, ULevel* InLevel, const FTrans
 
 void UActorFactory::PostSpawnActor( UObject* Asset, AActor* NewActor)
 {
-	// Subclasses may implement this to modify the actor after it has been spawned
-	if (Asset)
-	{
-		FActorLabelUtilities::SetActorLabelUnique(NewActor, Asset->GetName());
-	}
 }
 
 void UActorFactory::PostCreateBlueprint( UObject* Asset, AActor* CDO )
diff --git a/Engine/Source/Editor/UnrealEd/Private/Factories/CSVImportFactory.cpp b/Engine/Source/Editor/UnrealEd/Private/Factories/CSVImportFactory.cpp
--- a/Engine/Source/Editor/UnrealEd/Private/Factories/CSVImportFactory.cpp
+++ b/Engine/Source/Editor/UnrealEd/Private/Factories/CSVImportFactory.cpp
@@ -157,10 +157,18 @@ UObject* UCSVImportFactory::FactoryCreateText(UClass* InClass, UObject* InParent
 			.Title( LOCTEXT("DataTableOptionsWindowTitle", "DataTable Options" ))
 			.SizingRule( ESizingRule::Autosized );
 		
+		FString ParentFullPath;
+
+		if (InParent)
+		{
+			ParentFullPath = InParent->GetPathName();
+		}
+
 		Window->SetContent
 		(
 			SAssignNew(ImportOptionsWindow, SCSVImportOptions)
-			.WidgetWindow(Window)
+				.WidgetWindow(Window)
+				.FullPath(FText::FromString(ParentFullPath))
 		);
 
 		FSlateApplication::Get().AddModalWindow(Window, ParentWindow, false);
diff --git a/Engine/Source/Editor/UnrealEd/Private/Factories/EditorFactories.cpp b/Engine/Source/Editor/UnrealEd/Private/Factories/EditorFactories.cpp
--- a/Engine/Source/Editor/UnrealEd/Private/Factories/EditorFactories.cpp
+++ b/Engine/Source/Editor/UnrealEd/Private/Factories/EditorFactories.cpp
@@ -105,6 +105,7 @@
 #include "Factories/SoundMixFactory.h"
 #include "Factories/ReimportSoundSurroundFactory.h"
 #include "Factories/StructureFactory.h"
+#include "Factories/StringTableFactory.h"
 #include "Factories/SubsurfaceProfileFactory.h"
 #include "Factories/SubDSurfaceFactory.h"
 #include "Factories/Texture2dFactoryNew.h"
@@ -165,6 +166,7 @@
 #include "GameFramework/TouchInterface.h"
 #include "Engine/UserDefinedEnum.h"
 #include "Engine/UserDefinedStruct.h"
+#include "Internationalization/StringTable.h"
 #include "Editor.h"
 #include "Matinee/InterpData.h"
 #include "Matinee/InterpGroupCamera.h"
@@ -6811,6 +6813,22 @@ UObject* UDataTableFactory::FactoryCreateNew(UClass* Class, UObject* InParent, F
 	return DataTable;
 }
 
+/*------------------------------------------------------------------------------
+UStringTableFactory implementation.
+------------------------------------------------------------------------------*/
+UStringTableFactory::UStringTableFactory(const FObjectInitializer& ObjectInitializer)
+	: Super(ObjectInitializer)
+{
+	SupportedClass = UStringTable::StaticClass();
+	bCreateNew = true;
+	bEditAfterNew = true;
+}
+
+UObject* UStringTableFactory::FactoryCreateNew(UClass* Class, UObject* InParent, FName Name, EObjectFlags Flags, UObject* Context, FFeedbackContext* Warn)
+{
+	return NewObject<UStringTable>(InParent, Name, Flags);
+}
+
 /*------------------------------------------------------------------------------
  UPreviewMeshCollectionFactory implementation.
 ------------------------------------------------------------------------------*/
diff --git a/Engine/Source/Editor/UnrealEd/Private/Fbx/FbxFactory.cpp b/Engine/Source/Editor/UnrealEd/Private/Fbx/FbxFactory.cpp
--- a/Engine/Source/Editor/UnrealEd/Private/Fbx/FbxFactory.cpp
+++ b/Engine/Source/Editor/UnrealEd/Private/Fbx/FbxFactory.cpp
@@ -306,7 +306,8 @@ UObject* UFbxFactory::FactoryCreateBinary
 							NewStaticMesh = FbxImporter->ImportStaticMeshAsSingle(InParent, FbxMeshArray, Name, Flags, ImportUI->StaticMeshImportData, NULL, 0);
 							if (NewStaticMesh != nullptr)
 							{
-								FbxImporter->ReorderMaterialToFbxOrder(NewStaticMesh, FbxMeshArray);
+								//Build the staticmesh
+								FbxImporter->PostImportStaticMesh(NewStaticMesh, FbxMeshArray);
 							}
 						}
 
@@ -358,6 +359,13 @@ UObject* UFbxFactory::FactoryCreateBinary
 							TArray<FbxNode*> &LODMeshesArray = FbxMeshesLod[LODIndex];
 							FbxImporter->ImportStaticMeshAsSingle(InParent, LODMeshesArray, Name, Flags, ImportUI->StaticMeshImportData, NewStaticMesh, LODIndex);
 						}
+						
+						//Build the staticmesh
+						if (NewStaticMesh)
+						{
+							TArray<FbxNode*> &LODMeshesArray = FbxMeshesLod[0];
+							FbxImporter->PostImportStaticMesh(NewStaticMesh, LODMeshesArray);
+						}
 					}
 					else
 					{
@@ -668,12 +676,16 @@ UObject* UFbxFactory::RecursiveImportNode(void* VoidFbxImporter, void* VoidNode,
 		
 		if (NewObject)
 		{
-			//Reorder the material
-			TArray<FbxNode*> Nodes;
-			FbxImporter->FindAllLODGroupNode(Nodes, Node, 0);
-			if (Nodes.Num() > 0)
+			UStaticMesh *NewStaticMesh = Cast<UStaticMesh>(NewObject);
+			if (NewStaticMesh != nullptr)
 			{
-				FbxImporter->ReorderMaterialToFbxOrder(Cast<UStaticMesh>(NewObject), Nodes);
+				//Reorder the material
+				TArray<FbxNode*> Nodes;
+				FbxImporter->FindAllLODGroupNode(Nodes, Node, 0);
+				if (Nodes.Num() > 0)
+				{
+					FbxImporter->PostImportStaticMesh(NewStaticMesh, Nodes);
+				}
 			}
 		}
 	}
@@ -687,10 +699,14 @@ UObject* UFbxFactory::RecursiveImportNode(void* VoidFbxImporter, void* VoidNode,
 
 			if ( NewObject )
 			{
-				//Reorder the material
-				TArray<FbxNode*> Nodes;
-				Nodes.Add(Node);
-				FbxImporter->ReorderMaterialToFbxOrder(Cast<UStaticMesh>(NewObject), Nodes);
+				UStaticMesh *NewStaticMesh = Cast<UStaticMesh>(NewObject);
+				if (NewStaticMesh != nullptr)
+				{
+					//Reorder the material
+					TArray<FbxNode*> Nodes;
+					Nodes.Add(Node);
+					FbxImporter->PostImportStaticMesh(NewStaticMesh, Nodes);
+				}
 
 				OutNewAssets.Add(NewObject);
 			}
diff --git a/Engine/Source/Editor/UnrealEd/Private/Fbx/FbxMainExport.cpp b/Engine/Source/Editor/UnrealEd/Private/Fbx/FbxMainExport.cpp
--- a/Engine/Source/Editor/UnrealEd/Private/Fbx/FbxMainExport.cpp
+++ b/Engine/Source/Editor/UnrealEd/Private/Fbx/FbxMainExport.cpp
@@ -1529,16 +1529,17 @@ FbxNode* FFbxExporter::ExportActor(AActor* Actor, bool bExportComponents, INodeN
 			if (ParentActor != NULL)
 			{
 				//In case the parent was not export, get the absolute transform
-				const FTransform AbsoluteTransform = Actor->GetTransform();
+				const FTransform AbsoluteTransform = RotationDirectionConvert * Actor->GetTransform();
 				ActorLocation = AbsoluteTransform.GetTranslation();
 				ActorRotation = AbsoluteTransform.GetRotation().Euler();
 				ActorScale = AbsoluteTransform.GetScale3D();
 			}
 			else
 			{
-				ActorLocation = Actor->GetActorLocation();
-				ActorRotation = Actor->GetActorRotation().Euler();
-				ActorScale = Actor->GetRootComponent() ? Actor->GetRootComponent()->RelativeScale3D : FVector(1.f, 1.f, 1.f);
+				const FTransform ConvertedTransform = RotationDirectionConvert * Actor->GetTransform();
+				ActorLocation = ConvertedTransform.GetTranslation();
+				ActorRotation = ConvertedTransform.GetRotation().Euler();
+				ActorScale = ConvertedTransform.GetScale3D();
 			}
 		}
 
diff --git a/Engine/Source/Editor/UnrealEd/Private/Fbx/FbxMainImport.cpp b/Engine/Source/Editor/UnrealEd/Private/Fbx/FbxMainImport.cpp
--- a/Engine/Source/Editor/UnrealEd/Private/Fbx/FbxMainImport.cpp
+++ b/Engine/Source/Editor/UnrealEd/Private/Fbx/FbxMainImport.cpp
@@ -25,6 +25,8 @@
 #include "EngineAnalytics.h"
 #include "AnalyticsEventAttribute.h"
 #include "Interfaces/IAnalyticsProvider.h"
+#include "UObject/UObjectGlobals.h"
+#include "UObject/Package.h"
 
 DEFINE_LOG_CATEGORY(LogFbx);
 
@@ -198,6 +200,7 @@ void ApplyImportUIToImportOptions(UFbxImportUI* ImportUI, FBXImportOptions& InOu
 	InOutImportOptions.bUsedAsFullName = ImportUI->bOverrideFullName;
 	InOutImportOptions.bImportAnimations = ImportUI->bImportAnimations;
 	InOutImportOptions.SkeletonForAnimation = ImportUI->Skeleton;
+	InOutImportOptions.ImportType = ImportUI->MeshTypeToImport;
 
 	if ( ImportUI->MeshTypeToImport == FBXIT_StaticMesh )
 	{
@@ -940,6 +943,7 @@ bool FFbxImporter::ImportFile(FString Filename, bool bPreventMaterialNameClash /
 
 	// Get the version number of the FBX file format.
 	Importer->GetFileVersion(FileMajor, FileMinor, FileRevision);
+	FbxFileVersion = FString::Printf(TEXT("%d.%d.%d"), FileMajor, FileMinor, FileRevision);
 
 	// output result
 	if(bStatus)
@@ -1047,6 +1051,7 @@ bool FFbxImporter::ImportFromFile(const FString& Filename, const FString& Type,
 				{
 					if( FEngineAnalytics::IsAvailable() )
 					{
+						const static UEnum* FBXImportTypeEnum = FindObject<UEnum>(ANY_PACKAGE, TEXT("EFBXImportType"));
 						TArray<FAnalyticsEventAttribute> Attribs;
 
 						FString OriginalVendor(UTF8_TO_TCHAR(DocInfo->Original_ApplicationVendor.Get().Buffer()));
@@ -1057,6 +1062,8 @@ bool FFbxImporter::ImportFromFile(const FString& Filename, const FString& Type,
 						FString LastSavedAppName(UTF8_TO_TCHAR(DocInfo->LastSaved_ApplicationName.Get().Buffer()));
 						FString LastSavedAppVersion(UTF8_TO_TCHAR(DocInfo->LastSaved_ApplicationVersion.Get().Buffer()));
 
+						FString FilenameHash = FMD5::HashAnsiString(*Filename);
+
 						Attribs.Add(FAnalyticsEventAttribute(TEXT("Original Application Vendor"), OriginalVendor));
 						Attribs.Add(FAnalyticsEventAttribute(TEXT("Original Application Name"), OriginalAppName));
 						Attribs.Add(FAnalyticsEventAttribute(TEXT("Original Application Version"), OriginalAppVersion));
@@ -1065,6 +1072,11 @@ bool FFbxImporter::ImportFromFile(const FString& Filename, const FString& Type,
 						Attribs.Add(FAnalyticsEventAttribute(TEXT("LastSaved Application Name"), LastSavedAppName));
 						Attribs.Add(FAnalyticsEventAttribute(TEXT("LastSaved Application Version"), LastSavedAppVersion));
 
+						Attribs.Add(FAnalyticsEventAttribute(TEXT("FBX Version"), FbxFileVersion));
+						Attribs.Add(FAnalyticsEventAttribute(TEXT("Filename Hash"), FilenameHash));
+
+						Attribs.Add(FAnalyticsEventAttribute(TEXT("Import Type"), FBXImportTypeEnum->GetEnumName(ImportOptions->ImportType)));
+
 						FString EventString = FString::Printf(TEXT("Editor.Usage.FBX.Import"));
 						FEngineAnalytics::GetProvider().RecordEvent(EventString, Attribs);
 					}
diff --git a/Engine/Source/Editor/UnrealEd/Private/Fbx/FbxSceneImportFactory.cpp b/Engine/Source/Editor/UnrealEd/Private/Fbx/FbxSceneImportFactory.cpp
--- a/Engine/Source/Editor/UnrealEd/Private/Fbx/FbxSceneImportFactory.cpp
+++ b/Engine/Source/Editor/UnrealEd/Private/Fbx/FbxSceneImportFactory.cpp
@@ -835,6 +835,17 @@ UObject* UFbxSceneImportFactory::FactoryCreateBinary
 	return ReturnObject;
 }
 
+bool UFbxSceneImportFactory::FactoryCanImport(const FString& Filename)
+{
+	const FString Extension = FPaths::GetExtension(Filename);
+
+	if (Extension == TEXT("fbx"))
+	{
+		return true;
+	}
+	return false;
+}
+
 TSharedPtr<FFbxNodeInfo> GetNodeInfoPtrById(TArray<TSharedPtr<FFbxNodeInfo>> &HierarchyInfo, uint64 SearchId)
 {
 	for (TSharedPtr<FFbxNodeInfo> NodeInfoPtr : HierarchyInfo)
@@ -2043,6 +2054,13 @@ UObject* UFbxSceneImportFactory::RecursiveImportNode(void* VoidFbxImporter, void
 					ImportANode(VoidFbxImporter, TmpVoidArray, Flags, NodeIndex, SceneInfo, OutNodeInfo, PackagePath, Total, NewObject, LODIndex);
 				}
 			}
+			UStaticMesh *NewStaticMesh = Cast<UStaticMesh>(NewObject);
+			if (NewStaticMesh != nullptr)
+			{
+				//Build the staticmesh
+				FFbxImporter->FindAllLODGroupNode(AllNodeInLod, Node, 0);
+				FFbxImporter->PostImportStaticMesh(NewStaticMesh, AllNodeInLod);
+			}
 		}
 	}
 	else
@@ -2058,6 +2076,15 @@ UObject* UFbxSceneImportFactory::RecursiveImportNode(void* VoidFbxImporter, void
 				//We should always have a valid attribute if we just create a new asset
 				check(OutNodeInfo.IsValid() && OutNodeInfo->AttributeInfo.IsValid());
 
+				UStaticMesh *NewStaticMesh = Cast<UStaticMesh>(FirstBaseObject);
+				if (NewStaticMesh != nullptr)
+				{
+					//Build the staticmesh
+					TArray<FbxNode*> AllNodeInLod;
+					AllNodeInLod.Add(Node);
+					FFbxImporter->PostImportStaticMesh(NewStaticMesh, AllNodeInLod);
+				}
+
 				AllNewAssets.Add(OutNodeInfo->AttributeInfo, FirstBaseObject);
 			}
 		}
diff --git a/Engine/Source/Editor/UnrealEd/Private/Fbx/FbxSkeletalMeshImport.cpp b/Engine/Source/Editor/UnrealEd/Private/Fbx/FbxSkeletalMeshImport.cpp
--- a/Engine/Source/Editor/UnrealEd/Private/Fbx/FbxSkeletalMeshImport.cpp
+++ b/Engine/Source/Editor/UnrealEd/Private/Fbx/FbxSkeletalMeshImport.cpp
@@ -870,7 +870,10 @@ bool UnFbx::FFbxImporter::ImportBone(TArray<FbxNode*>& NodeArray, FSkeletalMeshI
 		}
 		else
 		{
-			AddTokenizedErrorMessage(FTokenizedMessage::Create(EMessageSeverity::Warning, LOCTEXT("FbxSkeletaLMeshimport_MissingBindPose", "Could not find the bind pose.  It will use time 0 as bind pose.")), FFbxErrors::SkeletalMesh_InvalidBindPose);
+			if (!GIsAutomationTesting)
+			{
+				AddTokenizedErrorMessage(FTokenizedMessage::Create(EMessageSeverity::Warning, LOCTEXT("FbxSkeletaLMeshimport_MissingBindPose", "Could not find the bind pose.  It will use time 0 as bind pose.")), FFbxErrors::SkeletalMesh_InvalidBindPose);
+			}
 			bUseTime0AsRefPose = true;
 		}
 	}
diff --git a/Engine/Source/Editor/UnrealEd/Private/Fbx/FbxStaticMeshImport.cpp b/Engine/Source/Editor/UnrealEd/Private/Fbx/FbxStaticMeshImport.cpp
--- a/Engine/Source/Editor/UnrealEd/Private/Fbx/FbxStaticMeshImport.cpp
+++ b/Engine/Source/Editor/UnrealEd/Private/Fbx/FbxStaticMeshImport.cpp
@@ -838,10 +838,21 @@ UStaticMesh* UnFbx::FFbxImporter::ReimportSceneStaticMesh(uint64 FbxNodeUniqueId
 					ImportStaticMeshAsSingle(Mesh->GetOutermost(), AllNodeInLod, *Mesh->GetName(), RF_Public | RF_Standalone, TemplateImportData, FirstBaseMesh, LODIndex, nullptr);
 				}
 			}
+			if (FirstBaseMesh != nullptr)
+			{
+				FindAllLODGroupNode(AllNodeInLod, NodeParent, 0);
+				PostImportStaticMesh(FirstBaseMesh, AllNodeInLod);
+			}
 		}
 		else
 		{
 			FirstBaseMesh = ImportStaticMesh(Mesh->GetOutermost(), Node, *Mesh->GetName(), RF_Public | RF_Standalone, TemplateImportData, Mesh, 0, ExistMeshDataPtr);
+			if (FirstBaseMesh != nullptr)
+			{
+				TArray<FbxNode*> AllNodeInLod;
+				AllNodeInLod.Add(Node);
+				PostImportStaticMesh(FirstBaseMesh, AllNodeInLod);
+			}
 		}
 	}
 	else
@@ -850,6 +861,10 @@ UStaticMesh* UnFbx::FFbxImporter::ReimportSceneStaticMesh(uint64 FbxNodeUniqueId
 		if (FbxMeshArray.Num() > 0)
 		{
 			FirstBaseMesh = ImportStaticMeshAsSingle(Mesh->GetOutermost(), FbxMeshArray, *Mesh->GetName(), RF_Public | RF_Standalone, TemplateImportData, Mesh, 0, ExistMeshDataPtr);
+			if (FirstBaseMesh != nullptr)
+			{
+				PostImportStaticMesh(FirstBaseMesh, FbxMeshArray);
+			}
 		}
 		else // no mesh found in the FBX file
 		{
@@ -977,7 +992,12 @@ UStaticMesh* UnFbx::FFbxImporter::ReimportStaticMesh(UStaticMesh* Mesh, UFbxStat
 			LodZeroNodes.Add(Node);
 			NewMesh = ImportStaticMesh(Mesh->GetOuter(), Node, *Mesh->GetName(), RF_Public|RF_Standalone, TemplateImportData, Mesh, 0, ExistMeshDataPtr);
 		}
-		ReorderMaterialToFbxOrder(NewMesh, LodZeroNodes);
+
+		if (NewMesh != nullptr)
+		{
+			PostImportStaticMesh(NewMesh, LodZeroNodes);
+		}
+		
 	}
 	else
 	{
@@ -985,7 +1005,10 @@ UStaticMesh* UnFbx::FFbxImporter::ReimportStaticMesh(UStaticMesh* Mesh, UFbxStat
 		if (FbxMeshArray.Num() > 0)
 		{
 			NewMesh = ImportStaticMeshAsSingle(Mesh->GetOuter(), FbxMeshArray, *Mesh->GetName(), RF_Public|RF_Standalone, TemplateImportData, Mesh, 0, ExistMeshDataPtr);
-			ReorderMaterialToFbxOrder(NewMesh, FbxMeshArray);
+			if (NewMesh != nullptr)
+			{
+				PostImportStaticMesh(NewMesh, FbxMeshArray);
+			}
 		}
 		else // no mesh found in the FBX file
 		{
@@ -1492,14 +1515,13 @@ UStaticMesh* UnFbx::FFbxImporter::ImportStaticMeshAsSingle(UObject* InParent, TA
 		{
 			SrcModel.BuildSettings.bGenerateLightmapUVs = true;
 			SrcModel.BuildSettings.DstLightmapIndex = FirstOpenUVChannel;
-StaticMesh->LightMapCoordinateIndex = FirstOpenUVChannel;
+			StaticMesh->LightMapCoordinateIndex = FirstOpenUVChannel;
 		}
 		else
 		{
 			SrcModel.BuildSettings.bGenerateLightmapUVs = false;
 		}
 
-		TArray<FText> BuildErrors;
 		StaticMesh->LODGroup = ImportOptions->StaticMeshLODGroup;
 
 		//Set the Imported version before calling the build
@@ -1511,31 +1533,7 @@ StaticMesh->LightMapCoordinateIndex = FirstOpenUVChannel;
 			RestoreExistingMeshSettings(ExistMeshData, InStaticMesh, StaticMesh->LODGroup != NAME_None ? INDEX_NONE : LODIndex);
 		}
 
-		StaticMesh->Build(false, &BuildErrors);
-
-		for (FText& Error : BuildErrors)
-		{
-			AddTokenizedErrorMessage(FTokenizedMessage::Create(EMessageSeverity::Warning, Error), FFbxErrors::StaticMesh_BuildError);
-		}
-
-		// this is damage control. After build, we'd like to absolutely sure that 
-		// all index is pointing correctly and they're all used. Otherwise we remove them
-		FMeshSectionInfoMap OldSectionInfoMap = StaticMesh->SectionInfoMap;
-		StaticMesh->SectionInfoMap.Clear();
-		// fix up section data
-		for (int32 LODResoureceIndex = 0; LODResoureceIndex < StaticMesh->RenderData->LODResources.Num(); ++LODResoureceIndex)
-		{
-			FStaticMeshLODResources& LOD = StaticMesh->RenderData->LODResources[LODResoureceIndex];
-			int32 NumSections = LOD.Sections.Num();
-			for (int32 SectionIndex = 0; SectionIndex < NumSections; ++SectionIndex)
-			{
-				FMeshSectionInfo Info = OldSectionInfoMap.Get(LODResoureceIndex, SectionIndex);
-				if (StaticMesh->StaticMaterials.IsValidIndex(Info.MaterialIndex))
-				{
-					StaticMesh->SectionInfoMap.Set(LODResoureceIndex, SectionIndex, Info);
-				}
-			}
-		}
+		
 
 		// The code to check for bad lightmap UVs doesn't scale well with number of triangles.
 		// Skip it here because Lightmass will warn about it during a light build anyway.
@@ -1568,20 +1566,6 @@ StaticMesh->LightMapCoordinateIndex = FirstOpenUVChannel;
 
 	if (StaticMesh)
 	{
-
-		//collision generation
-		if (StaticMesh->bCustomizedCollision == false && ImportOptions->bAutoGenerateCollision)
-		{
-			FKAggregateGeom & AggGeom = StaticMesh->BodySetup->AggGeom;
-			AggGeom.ConvexElems.Empty(1);	//if no custom collision is setup we just regenerate collision when reimport
-
-			const int32 NumDirs = 18;
-			TArray<FVector> Dirs;
-			Dirs.AddUninitialized(NumDirs);
-			for (int32 DirIdx = 0; DirIdx < NumDirs; ++DirIdx) { Dirs[DirIdx] = KDopDir18[DirIdx]; }
-			GenerateKDopAsSimpleCollision(StaticMesh, Dirs);
-		}
-
 		//warnings based on geometry
 		VerifyGeometry(StaticMesh);
 
@@ -1591,10 +1575,57 @@ StaticMesh->LightMapCoordinateIndex = FirstOpenUVChannel;
 	return StaticMesh;
 }
 
-void UnFbx::FFbxImporter::ReorderMaterialToFbxOrder(UStaticMesh* StaticMesh, TArray<FbxNode*>& MeshNodeArray)
+void UnFbx::FFbxImporter::PostImportStaticMesh(UStaticMesh* StaticMesh, TArray<FbxNode*>& MeshNodeArray)
 {
+	if (StaticMesh == nullptr)
+	{
+		return;
+	}
+
+	// Build the staticmesh, we move the build here because we want to avoid building the staticmesh for every LOD
+	// when we import the mesh.
+	TArray<FText> BuildErrors;
+	StaticMesh->Build(false, &BuildErrors);
+
+	for (FText& Error : BuildErrors)
+	{
+		AddTokenizedErrorMessage(FTokenizedMessage::Create(EMessageSeverity::Warning, Error), FFbxErrors::StaticMesh_BuildError);
+	}
+
+	// this is damage control. After build, we'd like to absolutely sure that 
+	// all index is pointing correctly and they're all used. Otherwise we remove them
+	FMeshSectionInfoMap TempOldSectionInfoMap = StaticMesh->SectionInfoMap;
+	StaticMesh->SectionInfoMap.Clear();
+	// fix up section data
+	for (int32 LODResoureceIndex = 0; LODResoureceIndex < StaticMesh->RenderData->LODResources.Num(); ++LODResoureceIndex)
+	{
+		FStaticMeshLODResources& LOD = StaticMesh->RenderData->LODResources[LODResoureceIndex];
+		int32 NumSections = LOD.Sections.Num();
+		for (int32 SectionIndex = 0; SectionIndex < NumSections; ++SectionIndex)
+		{
+			FMeshSectionInfo Info = TempOldSectionInfoMap.Get(LODResoureceIndex, SectionIndex);
+			if (StaticMesh->StaticMaterials.IsValidIndex(Info.MaterialIndex))
+			{
+				StaticMesh->SectionInfoMap.Set(LODResoureceIndex, SectionIndex, Info);
+			}
+		}
+	}
+
+	//collision generation must be done after the build, this will ensure a valid BodySetup
+	if (StaticMesh->bCustomizedCollision == false && ImportOptions->bAutoGenerateCollision && StaticMesh->BodySetup)
+	{
+		FKAggregateGeom & AggGeom = StaticMesh->BodySetup->AggGeom;
+		AggGeom.ConvexElems.Empty(1);	//if no custom collision is setup we just regenerate collision when reimport
+
+		const int32 NumDirs = 18;
+		TArray<FVector> Dirs;
+		Dirs.AddUninitialized(NumDirs);
+		for (int32 DirIdx = 0; DirIdx < NumDirs; ++DirIdx) { Dirs[DirIdx] = KDopDir18[DirIdx]; }
+		GenerateKDopAsSimpleCollision(StaticMesh, Dirs);
+	}
+
 	//If there is less the 2 materials in the fbx file there is no need to reorder them
-	if (StaticMesh == nullptr || StaticMesh->StaticMaterials.Num() < 2)
+	if (StaticMesh->StaticMaterials.Num() < 2)
 	{
 		return;
 	}
@@ -1626,6 +1657,61 @@ void UnFbx::FFbxImporter::ReorderMaterialToFbxOrder(UStaticMesh* StaticMesh, TAr
 		return;
 	}
 
+	//If there is some skinxx material name we will reorder the material to follow the skinxx workflow instead of the fbx order
+	bool IsUsingSkinxxWorkflow = true;
+	TArray<FString> MeshMaterialsSkinXX;
+	MeshMaterialsSkinXX.AddZeroed(MeshMaterials.Num());
+	for (int32 FbxMaterialIndex = 0; FbxMaterialIndex < MeshMaterials.Num(); ++FbxMaterialIndex)
+	{
+		const FString &FbxMaterialName = MeshMaterials[FbxMaterialIndex];
+		//If we have all skinxx material name we have to re-order to skinxx workflow
+		int32 Offset = FbxMaterialName.Find(TEXT("_SKIN"), ESearchCase::IgnoreCase, ESearchDir::FromEnd);
+		if (Offset == INDEX_NONE)
+		{
+			IsUsingSkinxxWorkflow = false;
+			MeshMaterialsSkinXX.Empty();
+			break;
+		}
+		int32 SkinIndex = INDEX_NONE;
+		// Chop off the material name so we are left with the number in _SKINXX
+		FString SkinXXNumber = FbxMaterialName.Right(FbxMaterialName.Len() - (Offset + 1)).RightChop(4);
+		if (SkinXXNumber.IsNumeric())
+		{
+			SkinIndex = FPlatformString::Atoi(*SkinXXNumber);
+		}
+
+		if (SkinIndex >= MeshMaterialsSkinXX.Num())
+		{
+			MeshMaterialsSkinXX.AddZeroed((SkinIndex + 1) - MeshMaterialsSkinXX.Num());
+		}
+		if (MeshMaterialsSkinXX.IsValidIndex(SkinIndex))
+		{
+			MeshMaterialsSkinXX[SkinIndex] = FbxMaterialName;
+		}
+		else
+		{
+			//Cannot reorder this item
+			IsUsingSkinxxWorkflow = false;
+			MeshMaterialsSkinXX.Empty();
+			break;
+		}
+	}
+
+	if (IsUsingSkinxxWorkflow)
+	{
+		//Shrink the array to valid entry, in case the skinxx has some hole like _skin[01, 02, 04, 05...]
+		for (int32 FbxMaterialIndex = MeshMaterialsSkinXX.Num() - 1; FbxMaterialIndex >= 0; --FbxMaterialIndex)
+		{
+			const FString &FbxMaterial = MeshMaterialsSkinXX[FbxMaterialIndex];
+			if (FbxMaterial.IsEmpty())
+			{
+				MeshMaterialsSkinXX.RemoveAt(FbxMaterialIndex);
+			}
+		}
+		//Replace the fbx ordered materials by the skinxx ordered material
+		MeshMaterials = MeshMaterialsSkinXX;
+	}
+
 	//Reorder the StaticMaterials array to reflect the order in the fbx file
 	//So we make sure the order reflect the material ID in the DCCs
 	FMeshSectionInfoMap OldSectionInfoMap = StaticMesh->SectionInfoMap;
diff --git a/Engine/Source/Editor/UnrealEd/Private/Fbx/ReimportFbxSceneFactory.cpp b/Engine/Source/Editor/UnrealEd/Private/Fbx/ReimportFbxSceneFactory.cpp
--- a/Engine/Source/Editor/UnrealEd/Private/Fbx/ReimportFbxSceneFactory.cpp
+++ b/Engine/Source/Editor/UnrealEd/Private/Fbx/ReimportFbxSceneFactory.cpp
@@ -976,7 +976,7 @@ UBlueprint *UReimportFbxSceneFactory::UpdateOriginalBluePrint(FString &BluePrint
 		CollectGarbage(GARBAGE_COLLECTION_KEEPFLAGS);
 
 		//Create the new nodes from the hierarchy actor
-		FKismetEditorUtilities::AddComponentsToBlueprint(BluePrint, HierarchyActor->GetInstanceComponents());
+		FKismetEditorUtilities::AddComponentsToBlueprint(BluePrint, HierarchyActor->GetInstanceComponents(), false, nullptr, true);
 		
 		UWorld* World = HierarchyActor->GetWorld();
 		World->DestroyActor(HierarchyActor);
@@ -1150,6 +1150,12 @@ EReimportResult::Type UReimportFbxSceneFactory::ImportStaticMesh(void* VoidFbxIm
 				FbxImporter->FindAllLODGroupNode(AllNodeInLod, NodeParent, LODIndex);
 				FbxImporter->ImportStaticMeshAsSingle(Pkg, AllNodeInLod, StaticMeshFName, RF_Public | RF_Standalone, StaticMeshImportData, NewObject, LODIndex);
 			}
+			UStaticMesh *NewMesh = Cast<UStaticMesh>(NewObject);
+			if (NewMesh != nullptr)
+			{
+				FbxImporter->FindAllLODGroupNode(AllNodeInLod, NodeParent, 0);
+				FbxImporter->PostImportStaticMesh(NewMesh, AllNodeInLod);
+			}
 		}
 	}
 	else
diff --git a/Engine/Source/Editor/UnrealEd/Private/FbxImporter.h b/Engine/Source/Editor/UnrealEd/Private/FbxImporter.h
--- a/Engine/Source/Editor/UnrealEd/Private/FbxImporter.h
+++ b/Engine/Source/Editor/UnrealEd/Private/FbxImporter.h
@@ -125,6 +125,7 @@ struct FBXImportOptions
 	EFBXNormalGenerationMethod::Type NormalGenerationMethod;
 	bool bTransformVertexToAbsolute;
 	bool bBakePivotInVertex;
+	EFBXImportType ImportType;
 	// Static Mesh options
 	bool bCombineToSingle;
 	EVertexColorImportOption::Type VertexColorImportOption;
@@ -607,11 +608,11 @@ public:
 	UNREALED_API UStaticMesh* ImportStaticMeshAsSingle(UObject* InParent, TArray<FbxNode*>& MeshNodeArray, const FName InName, EObjectFlags Flags, UFbxStaticMeshImportData* TemplateImportData, UStaticMesh* InStaticMesh, int LODIndex = 0, void *ExistMeshDataPtr = nullptr);
 
 	/**
-	 * Helper function to reorder the material array after we build the staticmesh.
-	 * It order the material like it is in the fbx file and reassign section material index properly.
-	 * This must be call once all LOD are imported.
-	 */
-	void ReorderMaterialToFbxOrder(UStaticMesh* StaticMesh, TArray<FbxNode*>& MeshNodeArray);
+	* Finish the import of the staticmesh after all LOD have been process (cannot be call before all LOD are imported). There is two main operation done by this function
+	* 1. Build the staticmesh render data
+	* 2. Reorder the material array to follow the fbx file material order
+	*/
+	UNREALED_API void PostImportStaticMesh(UStaticMesh* StaticMesh, TArray<FbxNode*>& MeshNodeArray);
 
 	/**
 	* Creates a SubDSurface mesh from all the meshes in FBX scene with the given name and flags.
@@ -1014,6 +1015,8 @@ protected:
 	// base path of fbx file
 	FString FileBasePath;
 	TWeakObjectPtr<UObject> Parent;
+	FString FbxFileVersion;
+
 	// Flag that the mesh is the first mesh to import in current FBX scene
 	// FBX scene may contain multiple meshes, importer can import them at one time.
 	// Initialized as true when start to import a FBX scene
diff --git a/Engine/Source/Editor/UnrealEd/Private/FbxMeshUtils.cpp b/Engine/Source/Editor/UnrealEd/Private/FbxMeshUtils.cpp
--- a/Engine/Source/Editor/UnrealEd/Private/FbxMeshUtils.cpp
+++ b/Engine/Source/Editor/UnrealEd/Private/FbxMeshUtils.cpp
@@ -160,6 +160,9 @@ namespace FbxMeshUtils
 				// Add imported mesh to existing model
 				if( TempStaticMesh )
 				{
+					//Build the staticmesh
+					FFbxImporter->PostImportStaticMesh(TempStaticMesh, *(LODNodeList[bUseLODs ? LODLevel : 0]));
+
 					// Update mesh component
 					BaseStaticMesh->MarkPackageDirty();
 
diff --git a/Engine/Source/Editor/UnrealEd/Private/FileHelpers.cpp b/Engine/Source/Editor/UnrealEd/Private/FileHelpers.cpp
--- a/Engine/Source/Editor/UnrealEd/Private/FileHelpers.cpp
+++ b/Engine/Source/Editor/UnrealEd/Private/FileHelpers.cpp
@@ -326,7 +326,32 @@ FString FEditorFileUtils::GetFilterString(EFileInteraction Interaction)
 		break;
 
 	case FI_ImportScene:
-		Result = TEXT("FBX (*.fbx) OBJ (*.obj)|*.fbx;*.obj|FBX (*.fbx)|*.fbx|OBJ (*.obj)|*.obj");
+		{
+			TArray<UFactory*> Factories;
+			for (UClass* Class : TObjectRange<UClass>())
+			{
+				if (Class->IsChildOf<USceneImportFactory>())
+				{
+					Factories.Add(Class->GetDefaultObject<UFactory>());
+				}
+
+			}
+
+			if (Factories.Num() > 0)
+			{
+				FString FileTypes;
+				FString AllExtensions;
+				TMultiMap<uint32, UFactory*> FilterIndexToFactory;
+
+				ObjectTools::GenerateFactoryFileExtensions(Factories, FileTypes, AllExtensions, FilterIndexToFactory);
+
+				FileTypes = FString::Printf(TEXT("All Files (%s)|%s|%s"), *AllExtensions, *AllExtensions, *FileTypes);
+
+
+				Result = FileTypes;
+			}
+
+		}
 		break;
 
 	case FI_ExportScene:
@@ -1060,39 +1085,52 @@ void FEditorFileUtils::Import(const FString& InFilename)
 {
 	const FScopedBusyCursor BusyCursor;
 
-	FFormatNamedArguments Args;
-	//Import scene support only fbx for now
-	//Check the extension because the import map don't support fbx file import
-	if (FPaths::GetExtension(InFilename).Compare(TEXT("fbx"), ESearchCase::IgnoreCase) == 0)
+	USceneImportFactory *SceneFactory = nullptr;
+	for (UClass* Class : TObjectRange<UClass>() )
 	{
-		//Ask a root content path to the user
-		TSharedRef<SDlgPickPath> PickContentPathDlg =
-			SNew(SDlgPickPath)
-			.Title(LOCTEXT("FbxChooseImportRootContentPath", "Choose Location for importing the fbx scene content"));
-
-		FString Path = "";
-		if (PickContentPathDlg->ShowModal() == EAppReturnType::Cancel)
+		if (Class->IsChildOf<USceneImportFactory>())
 		{
-			return;
+			USceneImportFactory* TestFactory = Class->GetDefaultObject<USceneImportFactory>();
+			if (TestFactory->FactoryCanImport(InFilename))
+			{
+				/// Pick the first one for now 
+				SceneFactory = TestFactory;
+				break;
+			}
 		}
-		Path = PickContentPathDlg->GetPath().ToString();
-		UFactory *FbxSceneFactory = NULL;
-		//Search the UFbxSceneImportFactory instance
-		for (TObjectIterator<UClass> It; It; ++It)
+
+	}
+	
+	if (SceneFactory)
+	{
+		FString Path = "";
+
+		//Ask the user for the root path where they want to any content to be placed
+		if(SceneFactory->ImportsAssets())
 		{
-			if (It->IsChildOf(UFbxSceneImportFactory::StaticClass()))
+			TSharedRef<SDlgPickPath> PickContentPathDlg =
+				SNew(SDlgPickPath)
+				.Title(LOCTEXT("ChooseImportRootContentPath", "Choose Location for importing the scene content"));
+		
+			if (PickContentPathDlg->ShowModal() == EAppReturnType::Cancel)
 			{
-				FbxSceneFactory = It->GetDefaultObject<UFactory>();
-				break;
+				return;
 			}
+
+			Path = PickContentPathDlg->GetPath().ToString();
 		}
+
+
 		FAssetToolsModule& AssetToolsModule = FModuleManager::Get().LoadModuleChecked<FAssetToolsModule>("AssetTools");
 		TArray<FString> Files;
 		Files.Add(InFilename);
-		AssetToolsModule.Get().ImportAssets(Files, Path, FbxSceneFactory);
+		AssetToolsModule.Get().ImportAssets(Files, Path, SceneFactory);
 	}
 	else
 	{
+
+		FFormatNamedArguments Args;
+
 		Args.Add(TEXT("MapFilename"), FText::FromString(FPaths::GetCleanFilename(InFilename)));
 		GWarn->BeginSlowTask(FText::Format(NSLOCTEXT("UnrealEd", "ImportingMap_F", "Importing map: {MapFilename}..."), Args), true);
 		GUnrealEd->Exec(GWorld, *FString::Printf(TEXT("MAP IMPORTADD FILE=\"%s\""), *InFilename));
@@ -3584,7 +3622,7 @@ void FEditorFileUtils::FindAllSubmittablePackageFiles(TMap<FString, FSourceContr
 	FSourceControlStatePtr ProjectFileSourceControlState = SourceControlProvider.GetState(FPaths::ConvertRelativePathToFull(FPaths::GetProjectFilePath()), EStateCacheUsage::Use);
 
 	if (ProjectFileSourceControlState.IsValid() && ProjectFileSourceControlState->IsCurrent() &&
-		(ProjectFileSourceControlState->IsCheckedOut() || ProjectFileSourceControlState->IsAdded() || (!ProjectFileSourceControlState->IsSourceControlled() && ProjectFileSourceControlState->CanAdd())))
+		(ProjectFileSourceControlState->CanCheckIn() || (!ProjectFileSourceControlState->IsSourceControlled() && ProjectFileSourceControlState->CanAdd())))
 	{
 		OutPackages.Add(FPaths::GetProjectFilePath(), MoveTemp(ProjectFileSourceControlState));
 	}
@@ -3604,8 +3642,8 @@ void FEditorFileUtils::FindAllSubmittablePackageFiles(TMap<FString, FSourceContr
 		FSourceControlStatePtr SourceControlState = SourceControlProvider.GetState(FPaths::ConvertRelativePathToFull(Filename), EStateCacheUsage::Use);
 
 		// Only include non-map packages that are currently checked out or packages not under source control
-		if (SourceControlState.IsValid() && SourceControlState->IsCurrent() && 
-			(SourceControlState->IsCheckedOut() || SourceControlState->IsAdded() || (!SourceControlState->IsSourceControlled() && SourceControlState->CanAdd())) &&
+		if (SourceControlState.IsValid() && SourceControlState->IsCurrent() &&
+			(SourceControlState->CanCheckIn() || (!SourceControlState->IsSourceControlled() && SourceControlState->CanAdd())) &&
 			(bIncludeMaps || !IsMapPackageAsset(*Filename)))
 		{
 			OutPackages.Add(MoveTemp(PackageName), MoveTemp(SourceControlState));
@@ -3650,8 +3688,8 @@ void FEditorFileUtils::FindAllSubmittableConfigFiles(TMap<FString, TSharedPtr<cl
 			FSourceControlStatePtr SourceControlState = SourceControlProvider.GetState(ConfigFilename, EStateCacheUsage::Use);
 
 			// Only include config files that are currently checked out or packages not under source control
-			if (SourceControlState.IsValid() &&
-				(SourceControlState->IsCheckedOut() || SourceControlState->IsAdded() || (!SourceControlState->IsSourceControlled() && SourceControlState->CanAdd())))
+			if (SourceControlState.IsValid() && SourceControlState->IsCurrent() &&
+				(SourceControlState->CanCheckIn() || (!SourceControlState->IsSourceControlled() && SourceControlState->CanAdd())))
 			{
 				OutConfigFiles.Add(ConfigFilename, SourceControlState);
 			}
diff --git a/Engine/Source/Editor/UnrealEd/Private/Kismet2/BlueprintEditorUtils.cpp b/Engine/Source/Editor/UnrealEd/Private/Kismet2/BlueprintEditorUtils.cpp
--- a/Engine/Source/Editor/UnrealEd/Private/Kismet2/BlueprintEditorUtils.cpp
+++ b/Engine/Source/Editor/UnrealEd/Private/Kismet2/BlueprintEditorUtils.cpp
@@ -1300,7 +1300,7 @@ UClass* FBlueprintEditorUtils::RegenerateBlueprintClass(UBlueprint* Blueprint, U
 	bool bRegenerated = false;
 
 	// Cache off the dirty flag for the package, so we can restore it later
-	UPackage* Package = Cast<UPackage>(Blueprint->GetOutermost());
+	UPackage* Package = Blueprint->GetOutermost();
 	bool bIsPackageDirty = Package ? Package->IsDirty() : false;
 
 	// Preload the blueprint and all its parts before refreshing nodes. 
diff --git a/Engine/Source/Editor/UnrealEd/Private/Kismet2/Kismet2.cpp b/Engine/Source/Editor/UnrealEd/Private/Kismet2/Kismet2.cpp
--- a/Engine/Source/Editor/UnrealEd/Private/Kismet2/Kismet2.cpp
+++ b/Engine/Source/Editor/UnrealEd/Private/Kismet2/Kismet2.cpp
@@ -267,7 +267,7 @@ void FBlueprintUnloader::UnloadBlueprint(const bool bResetPackage)
 	{
 		UBlueprint* UnloadingBp = OldBlueprint.Get();
 
-		UPackage* const OldPackage = Cast<UPackage>(UnloadingBp->GetOutermost());
+		UPackage* const OldPackage = UnloadingBp->GetOutermost();
 		bool const bIsDirty = OldPackage->IsDirty();
 
 		UPackage* const TransientPackage = GetTransientPackage();
@@ -773,7 +773,7 @@ void FKismetEditorUtilities::CompileBlueprint(UBlueprint* BlueprintObj, bool bIs
 	// Reset the flag, so if the user tries to use PIE it will warn them if the BP did not compile
 	BlueprintObj->bDisplayCompilePIEWarning = true;
 
-	UPackage* const BlueprintPackage = Cast<UPackage>(BlueprintObj->GetOutermost());
+	UPackage* const BlueprintPackage = BlueprintObj->GetOutermost();
 	// compiling the blueprint will inherently dirty the package, but if there 
 	// weren't any changes to save before, there shouldn't be after
 	bool const bStartedWithUnsavedChanges = (BlueprintPackage != NULL) ? BlueprintPackage->IsDirty() : true;
@@ -963,7 +963,7 @@ bool FKismetEditorUtilities::GenerateBlueprintSkeleton(UBlueprint* BlueprintObj,
 
 	if( BlueprintObj->SkeletonGeneratedClass == NULL || bForceRegeneration )
 	{
-		UPackage* Package = Cast<UPackage>(BlueprintObj->GetOutermost());
+		UPackage* Package = BlueprintObj->GetOutermost();
 		bool bIsPackageDirty = Package ? Package->IsDirty() : false;
 					
 		IKismetCompilerInterface& Compiler = FModuleManager::LoadModuleChecked<IKismetCompilerInterface>(KISMET_COMPILER_MODULENAME);
@@ -993,7 +993,7 @@ void FKismetEditorUtilities::RecompileBlueprintBytecode(UBlueprint* BlueprintObj
 	check(BlueprintObj);
 	checkf(BlueprintObj->GeneratedClass, TEXT("Invalid generated class for %s"), *BlueprintObj->GetName());
 
-	UPackage* const BlueprintPackage = Cast<UPackage>(BlueprintObj->GetOutermost());
+	UPackage* const BlueprintPackage = BlueprintObj->GetOutermost();
 	bool const bStartedWithUnsavedChanges = (BlueprintPackage != NULL) ? BlueprintPackage->IsDirty() : true;
 
 	IKismetCompilerInterface& Compiler = FModuleManager::LoadModuleChecked<IKismetCompilerInterface>(KISMET_COMPILER_MODULENAME);
diff --git a/Engine/Source/Editor/UnrealEd/Private/Kismet2/KismetReinstanceUtilities.cpp b/Engine/Source/Editor/UnrealEd/Private/Kismet2/KismetReinstanceUtilities.cpp
--- a/Engine/Source/Editor/UnrealEd/Private/Kismet2/KismetReinstanceUtilities.cpp
+++ b/Engine/Source/Editor/UnrealEd/Private/Kismet2/KismetReinstanceUtilities.cpp
@@ -427,7 +427,7 @@ void FBlueprintCompileReinstancer::OptionallyRefreshNodes(UBlueprint* CurrentBP)
 {
 	if (HotReloadedNewClass)
 	{
-		UPackage* const Package = Cast<UPackage>(CurrentBP->GetOutermost());
+		UPackage* const Package = CurrentBP->GetOutermost();
 		const bool bStartedWithUnsavedChanges = Package != nullptr ? Package->IsDirty() : true;
 
 		FBlueprintEditorUtils::RefreshExternalBlueprintDependencyNodes(CurrentBP, HotReloadedNewClass);
diff --git a/Engine/Source/Editor/UnrealEd/Private/Kismet2/StructureEditorUtils.cpp b/Engine/Source/Editor/UnrealEd/Private/Kismet2/StructureEditorUtils.cpp
--- a/Engine/Source/Editor/UnrealEd/Private/Kismet2/StructureEditorUtils.cpp
+++ b/Engine/Source/Editor/UnrealEd/Private/Kismet2/StructureEditorUtils.cpp
@@ -516,6 +516,22 @@ bool FStructureEditorUtils::Fill_MakeStructureDefaultValue(const UUserDefinedStr
 	return bResult;
 }
 
+bool FStructureEditorUtils::DiffersFromDefaultValue(const UUserDefinedStruct* Struct, uint8* StructData)
+{
+	bool bDiffers = false;
+	if (Struct && StructData)
+	{
+		UUserDefinedStructEditorData* StructEditorData = CastChecked<UUserDefinedStructEditorData>(Struct->EditorData);
+		const uint8* DefaultInstance = StructEditorData->GetDefaultInstance();
+		if (DefaultInstance)
+		{
+			const int32 PortFlags = PPF_None;
+			bDiffers = !Struct->CompareScriptStruct(StructData, DefaultInstance, PortFlags);
+		}
+	}
+	return bDiffers;
+}
+
 bool FStructureEditorUtils::Fill_MakeStructureDefaultValue(const UProperty* Property, uint8* PropertyData)
 {
 	bool bResult = true;
diff --git a/Engine/Source/Editor/UnrealEd/Private/Lightmass/Lightmass.h b/Engine/Source/Editor/UnrealEd/Private/Lightmass/Lightmass.h
--- a/Engine/Source/Editor/UnrealEd/Private/Lightmass/Lightmass.h
+++ b/Engine/Source/Editor/UnrealEd/Private/Lightmass/Lightmass.h
@@ -4,8 +4,7 @@
 	Lightmass.h: lightmass import/export definitions.
 =============================================================================*/
 
-#ifndef __LIGHTMASS_H__
-#define __LIGHTMASS_H__
+#pragma once
 
 #include "CoreMinimal.h"
 #include "Misc/Guid.h"
@@ -711,5 +710,3 @@ protected:
 	 **/
 	static void SwarmCallback( NSwarm::FMessage* CallbackMessage, void* CallbackData );
 };
-
-#endif //__LIGHTMASS_H__
