- Fixed issue with cooker always cooking AllMaps even when a single map was specified on commandline

diff --git a/Engine/Source/Editor/UnrealEd/Private/CookOnTheFlyServer.cpp b/Engine/Source/Editor/UnrealEd/Private/CookOnTheFlyServer.cpp
--- a/Engine/Source/Editor/UnrealEd/Private/CookOnTheFlyServer.cpp
+++ b/Engine/Source/Editor/UnrealEd/Private/CookOnTheFlyServer.cpp
@@ -5112,8 +5112,8 @@ void UCookOnTheFlyServer::CollectFilesToCook(TArray<FName>& FilesInPath, const T
 
 
 
-		// if we didn't find maps to cook then cook the allmaps section
-		if (bFoundMapsToCook == false)
+		// if we didn't find maps to cook, and we don't have any commandline maps (CookMaps), then cook the allmaps section
+		if (bFoundMapsToCook == false && CookMaps.Num() == 0)
 		{
 			UE_LOG(LogCook, Verbose, TEXT("Loading default map ini section AllMaps "));
 			TArray<FString> AllMapsSection;
