Fixed fatal error when failing to import a cube texture.
Fixed null access crash when exiting.

diff --git a/Engine/Source/Editor/UnrealEd/Private/Factories/EditorFactories.cpp b/Engine/Source/Editor/UnrealEd/Private/Factories/EditorFactories.cpp
--- a/Engine/Source/Editor/UnrealEd/Private/Factories/EditorFactories.cpp
+++ b/Engine/Source/Editor/UnrealEd/Private/Factories/EditorFactories.cpp
@@ -2650,8 +2650,9 @@ UTexture2D* UTextureFactory::CreateTexture2D( UObject* InParent, FName Name, EOb
 
 UTextureCube* UTextureFactory::CreateTextureCube( UObject* InParent, FName Name, EObjectFlags Flags )
 {
-	UTextureCube* NewTextureCube = CastChecked<UTextureCube>( CreateOrOverwriteAsset(UTextureCube::StaticClass(), InParent, Name, Flags) );
-	return NewTextureCube;
+	// CreateOrOverwriteAsset could fail if this cubemap replaces an asset that still has references.
+	UObject* NewObject = CreateOrOverwriteAsset(UTextureCube::StaticClass(), InParent, Name, Flags);
+	return NewObject ? CastChecked<UTextureCube>(NewObject) : nullptr;
 }
 
 void UTextureFactory::SuppressImportOverwriteDialog()
diff --git a/Engine/Source/Runtime/Engine/Private/Streaming/StreamingTexture.cpp b/Engine/Source/Runtime/Engine/Private/Streaming/StreamingTexture.cpp
--- a/Engine/Source/Runtime/Engine/Private/Streaming/StreamingTexture.cpp
+++ b/Engine/Source/Runtime/Engine/Private/Streaming/StreamingTexture.cpp
@@ -351,7 +351,7 @@ void FStreamingTexture::StreamWantedMips(FStreamingManagerTexture& Manager)
 	{
 		ensure(ResidentMips == RequestedMips);
 
-		if (Texture->PendingMipChangeRequestStatus.GetValue() == TexState_ReadyFor_Requests)
+		if (Texture->Resource && Texture->PendingMipChangeRequestStatus.GetValue() == TexState_ReadyFor_Requests)
 		{
 			Texture->RequestedMips = WantedMips;
 
