Sequencer: Fix fbx skeletal animation export so that it ticks animation to update the bones per frame.

diff --git a/Engine/Source/Editor/UnrealEd/Private/Fbx/FbxAnimationExport.cpp b/Engine/Source/Editor/UnrealEd/Private/Fbx/FbxAnimationExport.cpp
--- a/Engine/Source/Editor/UnrealEd/Private/Fbx/FbxAnimationExport.cpp
+++ b/Engine/Source/Editor/UnrealEd/Private/Fbx/FbxAnimationExport.cpp
@@ -392,6 +392,17 @@ void FFbxExporter::ExportAnimTrack(IAnimTrackAdapter& AnimTrackAdapter, AActor*
 		// This will call UpdateSkelPose on the skeletal mesh component to move bones based on animations in the matinee group
 		AnimTrackAdapter.UpdateAnimation( SampleTime );
 
+		// Update space bases so new animation position has an effect.
+		// @todo - hack - this will be removed at some point
+		SkeletalMeshComponent->TickAnimation(0.03f, false);
+
+		SkeletalMeshComponent->RefreshBoneTransforms();
+		SkeletalMeshComponent->RefreshSlaveComponents();
+		SkeletalMeshComponent->UpdateComponentToWorld();
+		SkeletalMeshComponent->FinalizeBoneTransform();
+		SkeletalMeshComponent->MarkRenderTransformDirty();
+		SkeletalMeshComponent->MarkRenderDynamicDataDirty();
+
 		FbxTime ExportTime;
 		ExportTime.SetSecondDouble(SampleTime);
 
