Copying //UE4/Orion-Staging to //UE4/Main (Source: //Orion/Dev-General @ 3316439)

diff --git a/Engine/Source/Editor/UnrealEd/Private/CookOnTheFlyServer.cpp b/Engine/Source/Editor/UnrealEd/Private/CookOnTheFlyServer.cpp
--- a/Engine/Source/Editor/UnrealEd/Private/CookOnTheFlyServer.cpp
+++ b/Engine/Source/Editor/UnrealEd/Private/CookOnTheFlyServer.cpp
@@ -820,6 +820,7 @@ bool UCookOnTheFlyServer::StartNetworkFileServer( const bool BindAnyPort )
 	check( IsCookOnTheFlyMode() );
 	//GetDerivedDataCacheRef().WaitForQuiescence(false);
 
+	ValidateCookOnTheFlySettings();
 
 	ITargetPlatformManagerModule& TPM = GetTargetPlatformManagerRef();
 	const TArray<ITargetPlatform*>& Platforms = TPM.GetCookingTargetPlatforms();
@@ -1501,10 +1502,9 @@ uint32 UCookOnTheFlyServer::TickCookOnTheSide( const float TimeSlice, uint32 &Co
 					//  mark the original request as processed (if this isn't actually the file they were requesting then it will fail)
 					//	and then also save our new request as processed so we don't do it again
 					PackagesToSave.AddUnique( Package );
-#if DEBUG_COOKONTHEFLY
-					UE_LOG( LogCook, Display, TEXT("Request for %s received going to save %s"), *BuildFilename, *PackageFilename );
-#endif
-					CookedPackages.Add( FFilePlatformCookedPackage(ToBuild,false) );
+					UE_LOG( LogCook, Verbose, TEXT("Request for %s received going to save %s"), *BuildFilename, *PackageFilename );
+
+					CookedPackages.Add( FFilePlatformCookedPackage(ToBuild.GetFilename(), TargetPlatformNames) );
 					
 					ToBuild.SetFilename( PackageFilename );
 				}
@@ -1524,9 +1524,8 @@ uint32 UCookOnTheFlyServer::TickCookOnTheSide( const float TimeSlice, uint32 &Co
 			// this could also happen if the source file doesn't exist which is often as we request files with different extensions when we are searching for files
 			// just return that we processed the cook request
 			// the network file manager will then handle the missing file and search somewhere else
-#if DEBUG_COOKONTHEFLY
-			UE_LOG(LogCook, Display, TEXT("Not cooking package %s"), *ToBuild.GetFilename().ToString());
-#endif
+			UE_LOG(LogCook, Verbose, TEXT("Not cooking package %s"), *ToBuild.GetFilename().ToString());
+
 			// did not cook this package 
 #if DO_CHECK
 			// make sure this package doesn't exist
@@ -1543,7 +1542,7 @@ uint32 UCookOnTheFlyServer::TickCookOnTheSide( const float TimeSlice, uint32 &Co
 				}
 			}
 #endif
-			CookedPackages.Add( FFilePlatformCookedPackage( ToBuild.GetFilename(), TargetPlatformNames, false) );
+			CookedPackages.Add( FFilePlatformCookedPackage( ToBuild.GetFilename(), TargetPlatformNames) );
 			continue;
 		}
 
@@ -1977,8 +1976,8 @@ uint32 UCookOnTheFlyServer::TickCookOnTheSide( const float TimeSlice, uint32 &Co
 					bShouldSaveAsync = false;
 				}
 
-				ESavePackageResult SavePackageResult = ESavePackageResult::Error;
-				bool bSucceededSavePackage = false;
+				TArray<bool> bSucceededSavePackage;
+				TArray<ESavePackageResult> SavePackageResults;
 				{
 					COOK_STAT(FScopedDurationTimer TickCookOnTheSideSaveCookedPackageTimer(DetailedCookStats::TickCookOnTheSideSaveCookedPackageTimeSec));
 					SCOPE_TIMER(SaveCookedPackage);
@@ -1992,24 +1991,32 @@ uint32 UCookOnTheFlyServer::TickCookOnTheSide( const float TimeSlice, uint32 &Co
 					KeepEditorOnlyPackages |= IsChildCooker() || (CookByTheBookOptions && CookByTheBookOptions->ChildCookers.Num() > 0);
 					SaveFlags |= KeepEditorOnlyPackages ? SAVE_KeepEditorOnlyCookedPackages : SAVE_None;
 
-					GOutputCookingWarnings = IsCookFlagSet(ECookInitializationFlags::OutputVerboseCookerWarnings);
-					SavePackageResult = SaveCookedPackage(Package, SaveFlags, bWasUpToDate, AllTargetPlatformNames);
+					GOutputCookingWarnings = IsCookFlagSet(ECookInitializationFlags::OutputVerboseCookerWarnings);					
+					SaveCookedPackage(Package, SaveFlags, bWasUpToDate, AllTargetPlatformNames, SavePackageResults);
 					GOutputCookingWarnings = false;
-
-					if (SavePackageResult == ESavePackageResult::Success)
+					check(AllTargetPlatformNames.Num() == SavePackageResults.Num());
+					for ( const auto& SavePackageResult : SavePackageResults )
 					{
-						bSucceededSavePackage = true;
-						// Update flags used to determine garbage collection.
-						if (Package->ContainsMap())
+						if (SavePackageResult == ESavePackageResult::Success || SavePackageResult == ESavePackageResult::GenerateStub || SavePackageResult == ESavePackageResult::ReplaceCompletely)
 						{
-							Result |= COSR_CookedMap;
+							bSucceededSavePackage.Add(true);
+							// Update flags used to determine garbage collection.
+							if (Package->ContainsMap())
+							{
+								Result |= COSR_CookedMap;
+							}
+							else
+							{
+								++CookedPackageCount;
+								Result |= COSR_CookedPackage;
+							}
 						}
 						else
 						{
-							++CookedPackageCount;
-							Result |= COSR_CookedPackage;
+							bSucceededSavePackage.Add(false);
 						}
 					}
+					check( SavePackageResults.Num() == bSucceededSavePackage.Num() );
 					Timer.SavedPackage();
 				}
 
@@ -2039,7 +2046,16 @@ uint32 UCookOnTheFlyServer::TickCookOnTheSide( const float TimeSlice, uint32 &Co
 				{
 					// mark the package as cooked
 					FFilePlatformCookedPackage FileRequest(StandardFilename, AllTargetPlatformNames, bSucceededSavePackage);
-					if (SavePackageResult != ESavePackageResult::ReferencedOnlyByEditorOnlyData)
+					bool bWasReferencedOnlyByEditorOnlyData = false;
+					for (const auto& SavePackageResult : SavePackageResults)
+					{
+						if ( SavePackageResult == ESavePackageResult::ReferencedOnlyByEditorOnlyData )
+						{
+							bWasReferencedOnlyByEditorOnlyData = true;
+							// if this is the case all of the packages should be referenced only by editor only data
+						}
+					}
+					if (!bWasReferencedOnlyByEditorOnlyData)
 					{
 						CookedPackages.Add(FileRequest);
 						if ((CurrentCookMode == ECookMode::CookOnTheFly) && (I >= FirstUnsolicitedPackage))
@@ -2062,7 +2078,10 @@ uint32 UCookOnTheFlyServer::TickCookOnTheSide( const float TimeSlice, uint32 &Co
 				}
 				else
 				{
-					check( bSucceededSavePackage == false );
+					for ( const auto& succeededSavePackage : bSucceededSavePackage )
+					{
+						check( succeededSavePackage == false );
+					}
 				}
 			}
 		}
@@ -2706,10 +2725,10 @@ void UCookOnTheFlyServer::TickRecompileShaderRequests()
 }
 
 
-ESavePackageResult UCookOnTheFlyServer::SaveCookedPackage(UPackage* Package, uint32 SaveFlags, bool& bOutWasUpToDate)
+void UCookOnTheFlyServer::SaveCookedPackage(UPackage* Package, uint32 SaveFlags, bool& bOutWasUpToDate, TArray<ESavePackageResult>& SavePackageResults)
 {
 	TArray<FName> TargetPlatformNames; 
-	return SaveCookedPackage( Package, SaveFlags, bOutWasUpToDate, TargetPlatformNames );
+	return SaveCookedPackage( Package, SaveFlags, bOutWasUpToDate, TargetPlatformNames, SavePackageResults);
 }
 
 bool UCookOnTheFlyServer::ShouldCook(const FString& InFileName, const FName &InPlatformName)
@@ -2753,9 +2772,9 @@ bool UCookOnTheFlyServer::MakePackageFullyLoaded(UPackage* Package)
 	return bPackageFullyLoaded;
 }
 
-ESavePackageResult UCookOnTheFlyServer::SaveCookedPackage(UPackage* Package, uint32 SaveFlags, bool& bOutWasUpToDate, TArray<FName> &TargetPlatformNames)
+void UCookOnTheFlyServer::SaveCookedPackage(UPackage* Package, uint32 SaveFlags, bool& bOutWasUpToDate, TArray<FName> &TargetPlatformNames, TArray<ESavePackageResult>& SavePackageResults)
 {
-	ESavePackageResult Result = ESavePackageResult::Success;
+	check( SavePackageResults.Num() == 0);
 	check( bIsSavingPackage == false );
 	bIsSavingPackage = true;
 	FString Filename(GetCachedPackageFilename(Package));
@@ -2840,10 +2859,16 @@ ESavePackageResult UCookOnTheFlyServer::SaveCookedPackage(UPackage* Package, uin
 			Package->SetPackageFlags(PKG_StoreCompressed);
 		}
 
-		for (ITargetPlatform* Target : Platforms)
+		
+		for (int32 PlatformIndex = 0; PlatformIndex < Platforms.Num(); ++PlatformIndex )
 		{
+			SavePackageResults.Add(ESavePackageResult::Error);
+			ITargetPlatform* Target = Platforms[PlatformIndex];
 			FString PlatFilename = Filename.Replace(TEXT("[Platform]"), *Target->PlatformName());
 
+			ESavePackageResult& Result = SavePackageResults[PlatformIndex];
+			Result = ESavePackageResult::Success;
+
 			bool bCookPackage = true;
 
 			// don't save Editor resources from the Engine if the target doesn't have editoronly data
@@ -2889,7 +2914,7 @@ ESavePackageResult UCookOnTheFlyServer::SaveCookedPackage(UPackage* Package, uin
 				}
 
 				if (World)
-					{
+				{
 					// Fixup legacy lightmaps before saving
 					// This should be done after loading, but FRedirectCollector::ResolveStringAssetReference in Core loads UWorlds with LoadObject so there's no opportunity to handle this fixup on load
 					World->PersistentLevel->HandleLegacyMapBuildData();
@@ -2941,8 +2966,6 @@ ESavePackageResult UCookOnTheFlyServer::SaveCookedPackage(UPackage* Package, uin
 	check( bIsSavingPackage == true );
 	bIsSavingPackage = false;
 
-	// return success
-	return Result;
 }
 
 
@@ -3768,7 +3791,8 @@ FString UCookOnTheFlyServer::ConvertCookedPathToUncookedPath(const FString& Cook
 	FString CookedFilenameNoSandbox = CookedFilename;
 	CookedFilenameNoSandbox.RemoveFromStart(SandboxDirectory);
 	int32 EndOfPlatformIndex = 0;
-	
+
+	// assume at this point the cook platform is the next thing on the path
 	FString CookedFilenameNoPlatform = CookedFilename;
 
 	if (CookedFilenameNoSandbox.FindChar(TEXT('/'), EndOfPlatformIndex))
@@ -3777,9 +3801,10 @@ FString UCookOnTheFlyServer::ConvertCookedPathToUncookedPath(const FString& Cook
 		CookedFilenameNoPlatform.ReplaceInline(TEXT("//"), TEXT("/"));
 	}
 	
-
+	// after removing the cooked platform we can use the sandbox file to convert back to a uncooked path
 	FString FullUncookedPath = SandboxFile->ConvertFromSandboxPath(*CookedFilenameNoPlatform);
 
+	// make the result a standard filename (relative)
 	FPaths::MakeStandardFilename(FullUncookedPath);
 	return FullUncookedPath;
 
@@ -4022,25 +4047,28 @@ bool UCookOnTheFlyServer::SaveCookedAssetRegistry(const FString& InCookedAssetRe
 			Json->WriteValue(TEXT("SucceededSave"), TEXT("false"));
 		}
 		else
-			{
+		{
 			Json->WriteValue(TEXT("SucceededSave"), TEXT("true"));
 		}
-		
 
-				{
+
+		{
+			SCOPE_TIMER(DeterminePackageDependentHash);
 			FMD5Hash Hash;
 			PDInfoModule.DeterminePackageDependentHash(*RelativePath, Hash);
 			FString DependentPackageHashString = Lex::ToString(Hash);
 			Json->WriteValue(TEXT("DependentPackageHash"), *DependentPackageHashString);
-				}
+		}
 		{
+			SCOPE_TIMER(DetermineFullPackageHash);
 			FMD5Hash Hash;
 			PDInfoModule.DetermineFullPackageHash(*RelativePath, Hash);
 			FString FullPackageHashString = Lex::ToString(Hash);
 			Json->WriteValue(TEXT("FullPackageHash"), *FullPackageHashString);
-			}
+		}
 
 		{
+			SCOPE_TIMER(GetPackageDependencies);
 			TArray<FString> DependencyNames;
 			if (PDInfoModule.GetPackageDependencies(*RelativePath, DependencyNames))
 			{
@@ -4049,10 +4077,10 @@ bool UCookOnTheFlyServer::SaveCookedAssetRegistry(const FString& InCookedAssetRe
 				for (const FString& DependencyName : DependencyNames)
 				{
 					Json->WriteValue(DependencyName);
-		}
+				}
 				Json->WriteArrayEnd();
-	}
-}
+			}
+		}
 
 		Json->WriteValue(TEXT("FileSize"), FString::Printf(TEXT("%lld"), FileSize));
 		Json->WriteArrayStart("AssetData");
@@ -4241,15 +4269,15 @@ bool UCookOnTheFlyServer::SaveCookedAssetRegistry(const FString& InCookedAssetRe
 			{
 				UE_LOG(LogCook, Warning, TEXT("Unable to read or json is invalid format %s"), *CookedAssetRegistryFilename);
 			}
-				}
-			}
+		}
+	}
+
 
-			
 	Json->WriteArrayEnd();
 	Json->WriteObjectEnd();
 
 	if (Json->Close())
-			{
+	{
 		FArchive* ItemTemplatesFile = IFileManager::Get().CreateFileWriter(*CookedAssetRegistryFilename);
 		if (ItemTemplatesFile)
 		{
@@ -4277,9 +4305,26 @@ bool UCookOnTheFlyServer::SaveCookedAssetRegistry(const FString& InCookedAssetRe
 		UE_LOG(LogCook, Error, TEXT("Error closing Json Writer"));
 	}
 	return bSuccess;
-			}
+}
+
+
+void UCookOnTheFlyServer::GetAllCookedFiles(TMap<FName, FName>& UncookedPathToCookedPath, const FString& SandboxPath)
+{
+	IPlatformFile& PlatformFile = FPlatformFileManager::Get().GetPlatformFile();
+	TArray<FString> CookedFiles;
+	FPackageSearchVisitor PackageSearch(CookedFiles);
+	PlatformFile.IterateDirectoryRecursively(*SandboxPath, PackageSearch);
+	for (const FString& CookedFile : CookedFiles)
+	{
+		const FName CookedFName(*CookedFile);
 
+		const FString CookedFullPath = FPaths::ConvertRelativePathToFull(CookedFile);
+		const FString UncookedFilename = ConvertCookedPathToUncookedPath(CookedFullPath);
 
+		const FName UncookedFName(*UncookedFilename);
+		UncookedPathToCookedPath.Add(UncookedFName, CookedFName);
+	}
+}
 
 void UCookOnTheFlyServer::PopulateCookedPackagesFromDisk(const TArray<ITargetPlatform*>& Platforms)
 {
@@ -4289,31 +4334,31 @@ void UCookOnTheFlyServer::PopulateCookedPackagesFromDisk(const TArray<ITargetPla
 	FPackageDependencyInfoModule& PDInfoModule = FModuleManager::LoadModuleChecked<FPackageDependencyInfoModule>("PackageDependencyInfo");
 
 
-			// See what files are out of date in the sandbox folder
-			for (int32 Index = 0; Index < Platforms.Num(); Index++)
-			{
+	// See what files are out of date in the sandbox folder
+	for (int32 Index = 0; Index < Platforms.Num(); Index++)
+	{
 		TArray<FString> CookedPackagesToDelete;
 
-				ITargetPlatform* Target = Platforms[Index];
+		ITargetPlatform* Target = Platforms[Index];
 		FString SandboxPath = GetSandboxDirectory(Target->PlatformName());
 		FName PlatformFName(*Target->PlatformName());
-				
-				FString EngineSandboxPath = SandboxFile->ConvertToSandboxPath(*FPaths::EngineDir()) + TEXT("/");
-				EngineSandboxPath.ReplaceInline(TEXT("[Platform]"), *Target->PlatformName());
 
-				FString GameSandboxPath = SandboxFile->ConvertToSandboxPath(*(FPaths::GameDir() + TEXT("a.txt")));
-				GameSandboxPath.ReplaceInline(TEXT("a.txt"), TEXT(""));
-				GameSandboxPath.ReplaceInline(TEXT("[Platform]"), *Target->PlatformName());
+		FString EngineSandboxPath = SandboxFile->ConvertToSandboxPath(*FPaths::EngineDir()) + TEXT("/");
+		EngineSandboxPath.ReplaceInline(TEXT("[Platform]"), *Target->PlatformName());
 
-				FString LocalGamePath = FPaths::GameDir();
-				if (FPaths::IsProjectFilePathSet())
-				{
-					LocalGamePath = FPaths::GetPath(FPaths::GetProjectFilePath()) + TEXT("/");
-				}
+		FString GameSandboxPath = SandboxFile->ConvertToSandboxPath(*(FPaths::GameDir() + TEXT("a.txt")));
+		GameSandboxPath.ReplaceInline(TEXT("a.txt"), TEXT(""));
+		GameSandboxPath.ReplaceInline(TEXT("[Platform]"), *Target->PlatformName());
 
-				FString LocalEnginePath = FPaths::EngineDir();
+		FString LocalGamePath = FPaths::GameDir();
+		if (FPaths::IsProjectFilePathSet())
+		{
+			LocalGamePath = FPaths::GetPath(FPaths::GetProjectFilePath()) + TEXT("/");
+		}
 
-		static bool bFindCulprit = false;
+		FString LocalEnginePath = FPaths::EngineDir();
+
+		static bool bFindCulprit = false; //debugging setting if we want to find culprit for iterative cook issues
 
 		if (IsCookFlagSet(ECookInitializationFlags::IterateOnHash))
 		{
@@ -4321,41 +4366,52 @@ void UCookOnTheFlyServer::PopulateCookedPackagesFromDisk(const TArray<ITargetPla
 			const FString CookedAssetRegistry = FPaths::GameDir() / TEXT("CookedAssetRegistry.json");
 			const FString SandboxCookedAssetRegistryFilename = ConvertToFullSandboxPath(*CookedAssetRegistry, true, Target->PlatformName());
 
+			bool bIsIterateSharedBuild = IsCookFlagSet(ECookInitializationFlags::IterateSharedBuild);
 			TArray<FCookedAssetPackageInfo> CookedPackagesInfo;
-			ReadCookedAssetRegistry(SandboxCookedAssetRegistryFilename, CookedPackagesInfo);
+			if ( ReadCookedAssetRegistry(SandboxCookedAssetRegistryFilename, CookedPackagesInfo) == false )
+			{
+				if ( IsCookFlagSet(ECookInitializationFlags::IterateSharedBuild) )
+				{
+					// iterate on the shared build if the option is set
+					FString SharedCookedAssetRegistry = FPaths::Combine( *FPaths::GameSavedDir(), TEXT("SharedIterativeBuild"), *Target->PlatformName(), TEXT("CookedAssetRegistry.json"));
+					ReadCookedAssetRegistry( SharedCookedAssetRegistry, CookedPackagesInfo );
+				}
+			}
 			TMap<FName, FCookedAssetPackageInfo*> CookedPackageInfoLookup;
 			for (auto& PackageInfo : CookedPackagesInfo)
 			{
 				CookedPackageInfoLookup.Add(FName(*PackageInfo.RelativeFilename), &PackageInfo);
 			}
 
-			// check the files on disk 
-
 
+			// check the files on disk 
+			TMap<FName, FName> UncookedPathToCookedPath;
+			// get all the on disk cooked files
+			GetAllCookedFiles(UncookedPathToCookedPath, SandboxPath);
+			
+			const static FName NAME_DummyCookedFilename(TEXT("DummyCookedFilename")); // pls never name a package dummycookedfilename otherwise shit might go wonky
+			if ( bIsIterateSharedBuild )
+			{
+				// if we are iterating of a shared build the cooked files might not exist in the cooked directory because we assume they are packaged in the pak file (which we don't want to extract)
+				for (const auto& CookedPackageInfo : CookedPackagesInfo)
+				{
+					UncookedPathToCookedPath.Add(FName(*CookedPackageInfo.RelativeFilename), NAME_DummyCookedFilename);
+				}
+			}		
 
-			TMap<FName, FName> CookedPathToUncookedPath;
 			TArray<FName> AllUncookedFNames;
 			{
-				IPlatformFile& PlatformFile = FPlatformFileManager::Get().GetPlatformFile();
-				TArray<FString> CookedFiles;
-				FPackageSearchVisitor PackageSearch(CookedFiles);
-				PlatformFile.IterateDirectoryRecursively(*SandboxPath, PackageSearch);
 				TArray<FString> AllUncookedFilenames;
-				for (const FString& CookedFile : CookedFiles)
+				for (const auto& PathPair : UncookedPathToCookedPath)
 				{
-					const FName CookedFName(*CookedFile);
-
-					const FString CookedFullPath = FPaths::ConvertRelativePathToFull(CookedFile);
-					const FString UncookedFilename = ConvertCookedPathToUncookedPath(CookedFullPath);
-
-					const FName UncookedFName(*UncookedFilename);
-					CookedPathToUncookedPath.Add(CookedFName, UncookedFName);
-					if (CookedPackageInfoLookup.Contains(UncookedFName))
+					if (CookedPackageInfoLookup.Contains(PathPair.Key))
 					{
-						AllUncookedFilenames.Add(UncookedFilename);
-						AllUncookedFNames.Add(UncookedFName);
+						AllUncookedFNames.Add(PathPair.Key);
+						AllUncookedFilenames.Add(PathPair.Key.ToString());
 					}
 				}
+
+				if (AllUncookedFilenames.Num() > 0)
 				{
 #if OUTPUT_TIMING
 					SCOPE_TIMER(DetermineAllDependentPackageInfoTime);
@@ -4366,18 +4422,20 @@ void UCookOnTheFlyServer::PopulateCookedPackagesFromDisk(const TArray<ITargetPla
 					PDInfoModule.AsyncDetermineAllDependentPackageInfo(AllUncookedFilenames, NumThreadsToDeterminePackageInfo);
 				}
 			}
-
-			uint32 NumPackagesConsidered = CookedPathToUncookedPath.Num();
+#if OUTPUT_TIMING
+			SCOPE_TIMER(ProcessingKeptPackageHashes);
+#endif
+			uint32 NumPackagesConsidered = UncookedPathToCookedPath.Num();
 			uint32 NumPackagesUnableToDetermineCurrentDependencyInfo = 0;
 			uint32 NumPackagesUnableToFindCookedPackageInfo = 0;
 			uint32 NumPackagesFileHashMissmatch = 0;
 			uint32 NumPackagesDependentHashMissmatch = 0;
 			uint32 NumPackagesKept = 0;
 
-			for (const auto& CookedPaths : CookedPathToUncookedPath)
+			for (const auto& CookedPaths : UncookedPathToCookedPath)
 			{
-				const FName CookedFile = CookedPaths.Key;
-				const FName UncookedFilename = CookedPaths.Value;
+				const FName CookedFile = CookedPaths.Value;
+				const FName UncookedFilename = CookedPaths.Key;
 				const FString UncookedFilenameString = UncookedFilename.ToString();
 				bool bShouldKeep = true;
 
@@ -4397,10 +4455,15 @@ void UCookOnTheFlyServer::PopulateCookedPackagesFromDisk(const TArray<ITargetPla
 				}
 				
 				FMD5Hash FullPackageHash;
-				if (bShouldKeep && PDInfoModule.DetermineFullPackageHash(*UncookedFilenameString, FullPackageHash) == false)
 				{
-					++NumPackagesUnableToDetermineCurrentDependencyInfo;
-					bShouldKeep = false;
+#if OUTPUT_TIMING
+					SCOPE_TIMER(DetermineFullPackageHash);
+#endif
+					if (bShouldKeep && PDInfoModule.DetermineFullPackageHash(*UncookedFilenameString, FullPackageHash) == false)
+					{
+						++NumPackagesUnableToDetermineCurrentDependencyInfo;
+						bShouldKeep = false;
+					}
 				}
 
 				FMD5Hash PackageDependentHash;
@@ -4461,23 +4524,40 @@ void UCookOnTheFlyServer::PopulateCookedPackagesFromDisk(const TArray<ITargetPla
 					bShouldKeep = false;
 				}
 
+				if (CookedFile == NAME_DummyCookedFilename)
+				{
+					check(IFileManager::Get().FileExists(*CookedFile.ToString()) == false);
+				}
+
 				if (bShouldKeep)
 				{
 					// Add this package to the CookedPackages list so that we don't try cook it again
-					check( IFileManager::Get().FileExists(*CookedFile.ToString()));
+					if ( CookedFile != NAME_DummyCookedFilename )
+					{
+						check( IFileManager::Get().FileExists(*CookedFile.ToString()));
+					}
 					TArray<FName> PlatformNames;
 					PlatformNames.Add(PlatformFName);
-					CookedPackages.Add(FFilePlatformCookedPackage(UncookedFilename, MoveTemp(PlatformNames), true));
+					TArray<bool> Succeeded;
+					Succeeded.Add(true);
+					CookedPackages.Add(FFilePlatformCookedPackage(UncookedFilename, MoveTemp(PlatformNames), MoveTemp(Succeeded)));
 					++NumPackagesKept;
 				}
 				else
 				{
-					// delete the old package 
-					const FString CookedFullPath = FPaths::ConvertRelativePathToFull(CookedFile.ToString());
-#if DEBUG_COOKONTHEFLY
-					UE_LOG(LogCook, Display, TEXT("Deleting cooked package %s failed filehash test"), *CookedFullPath);
-#endif
-					CookedPackagesToDelete.Add(CookedFullPath); 
+					if ( CookedFile != NAME_DummyCookedFilename )
+					{
+						// delete the old package 
+						const FString CookedFullPath = FPaths::ConvertRelativePathToFull(CookedFile.ToString());
+						UE_LOG(LogCook, Verbose, TEXT("Deleting cooked package %s failed filehash test"), *CookedFullPath);
+						CookedPackagesToDelete.Add(CookedFullPath); 
+					}
+					else
+					{
+						// the cooker should rebuild this package because it's not in the cooked package list
+						// the new package will have higher priority then the package in the original shared cooked build
+						UE_LOG(LogCook, Verbose, TEXT("Shared cooked build: Detected package is out of date %s"), *UncookedFilenameString);
+					}
 				}
 			}
 
@@ -4526,7 +4606,7 @@ void UCookOnTheFlyServer::PopulateCookedPackagesFromDisk(const TArray<ITargetPla
 					if ( bShouldKeep )
 					{
 						++NumMarkedFailedSaveKept;
-						CookedPackages.Add(MoveTemp(FFilePlatformCookedPackage(StandardName, MoveTemp(PlatformNames), false)));
+						CookedPackages.Add(MoveTemp(FFilePlatformCookedPackage(StandardName, MoveTemp(PlatformNames))));
 					}
 					
 				}
@@ -4600,7 +4680,9 @@ void UCookOnTheFlyServer::PopulateCookedPackagesFromDisk(const TArray<ITargetPla
 					{
 						TArray<FName> PlatformNames;
 						PlatformNames.Add(PlatformFName);
-						CookedPackages.Add(FFilePlatformCookedPackage(FName(*UncookedFilename), MoveTemp(PlatformNames), true));
+						TArray<bool> Succeeded;
+						Succeeded.Add(true);
+						CookedPackages.Add(FFilePlatformCookedPackage(FName(*UncookedFilename), MoveTemp(PlatformNames), MoveTemp(Succeeded)));
 					}
 				}
 			}
@@ -4638,6 +4720,13 @@ void UCookOnTheFlyServer::VerifyCookedPackagesAreUptodate(const TArray<ITargetPl
 	// we remove cooked packages which are updated using the "MarkObjectDirty" callbacks from the editor
 	// here we process the dependencies by removing any packages which awere dependent on a package which was removed
 
+	if (IsCookFlagSet(ECookInitializationFlags::IterateSharedBuild))
+	{
+		// can't verify cooked packages in shared builds
+		check( IsCookOnTheFlyMode() == false );
+		return;
+	}
+
 	uint32 CookedPackagesMissingOnDisk = 0;
 
 	// loop through all packages on disk and remove any packages from the cooked packages array which don't exist
@@ -4764,10 +4853,10 @@ void UCookOnTheFlyServer::VerifyCookedPackagesAreUptodate(const TArray<ITargetPl
 						continue;
 					}
 				}
-				
+
 				const FName StandardDependencyName = GetCachedStandardPackageFileFName(DependencyPackageName);
-				if ( StandardDependencyName == NAME_None )
-					{
+				if (StandardDependencyName == NAME_None)
+				{
 					UE_LOG(LogCook, Warning, TEXT("Unable to resolve package name %s not able to iterative cook"), *DependencyPackageName.ToString());
 					bShouldRemovePackage = true;
 					break;
@@ -4775,18 +4864,18 @@ void UCookOnTheFlyServer::VerifyCookedPackagesAreUptodate(const TArray<ITargetPl
 
 				// check this package is cooked for the platform
 				if (!CookedPackages.Exists(StandardDependencyName, PlatformNamesArray))
-						{
-					if (IsCookFlagSet(ECookInitializationFlags::LogDebugInfo))
-							{
-						UE_LOG(LogCook, Display, TEXT("Invalidating cooked package %s because %s was considered out of date"), *CookedFile.ToString(), *StandardDependencyName.ToString());
-							}
+				{
+					//if (IsCookFlagSet(ECookInitializationFlags::LogDebugInfo))
+					{
+						UE_LOG(LogCook, Verbose, TEXT("Invalidating cooked package %s because %s was considered out of date"), *CookedFile.ToString(), *StandardDependencyName.ToString());
+					}
 					++CookedPackagesFailedDependencyCheck;
 					bShouldRemovePackage = true;
 					break;
-						}
+				}
 
-				if (CookedPackages.Exists(StandardDependencyName, PlatformNamesArray,false))
-						{
+				if (CookedPackages.Exists(StandardDependencyName, PlatformNamesArray, false))
+				{
 					TArray<FName> NewDependencies;
 					AssetRegistry.GetDependencies(DependencyPackageName, NewDependencies, EAssetRegistryDependencyType::All);
 					for (const auto& NewDependency : NewDependencies)
@@ -4797,14 +4886,15 @@ void UCookOnTheFlyServer::VerifyCookedPackagesAreUptodate(const TArray<ITargetPl
 			}
 
 			if (bShouldRemovePackage)
-							{
+			{
+				UE_LOG(LogCook, Verbose, TEXT("Removing package from cooked package list %s"), *CookedFile.ToString());
 				CookedPackages.RemoveFileForPlatform(CookedFile, PlatformFName);
-							}
+			}
 			else
 			{
 				++CookedPackagesKept;
-						}
-					}
+			}
+		}
 	}
 
 
@@ -4869,7 +4959,7 @@ void UCookOnTheFlyServer::CleanSandbox(const bool bIterative)
 #if OUTPUT_TIMING
 	double SandboxCleanTime = 0.0;
 #endif
-						{
+	{
 #if OUTPUT_TIMING
 		SCOPE_SECONDS_COUNTER(SandboxCleanTime);
 		SCOPE_TIMER(CleanSandboxTime);
@@ -4889,9 +4979,9 @@ void UCookOnTheFlyServer::CleanSandbox(const bool bIterative)
 
 				IniSettingsOutOfDate(Target);
 				SaveCurrentIniSettings(Target);
-						}
+			}
 
-					}
+		}
 		else
 		{
 			for (const auto& Target : Platforms)
@@ -5076,13 +5166,13 @@ void UCookOnTheFlyServer::CollectFilesToCook(TArray<FName>& FilesInPath, const T
 	TArray<FString> CookDirectories = InCookDirectories;
 	
 	if (!IsCookingDLC() && 
-		!(FilesToCookFlags & ECookByTheBookOptions::NoAlwaysCookMaps) )
+		!(FilesToCookFlags & ECookByTheBookOptions::NoAlwaysCookMaps))
 	{
 
 		{
-		TArray<FString> MapList;
-		// Add the default map section
-		GEditor->LoadMapListFromIni(TEXT("AlwaysCookMaps"), MapList);
+			TArray<FString> MapList;
+			// Add the default map section
+			GEditor->LoadMapListFromIni(TEXT("AlwaysCookMaps"), MapList);
 
 			for (int32 MapIdx = 0; MapIdx < MapList.Num(); MapIdx++)
 			{
@@ -5092,19 +5182,19 @@ void UCookOnTheFlyServer::CollectFilesToCook(TArray<FName>& FilesInPath, const T
 		}
 
 
-		bool bFoundMapsToCook = false;
+		bool bFoundMapsToCook = CookMaps.Num() > 0;
 
 		{
 			TArray<FString> MapList;
-		for (const auto &IniMapSection : IniMapSections)
-		{
+			for (const auto &IniMapSection : IniMapSections)
+			{
 				UE_LOG(LogCook, Verbose, TEXT("Loading map ini section %s "), *IniMapSection);
-			GEditor->LoadMapListFromIni(*IniMapSection, MapList);
-		}
-		for (int32 MapIdx = 0; MapIdx < MapList.Num(); MapIdx++)
-		{
+				GEditor->LoadMapListFromIni(*IniMapSection, MapList);
+			}
+			for (int32 MapIdx = 0; MapIdx < MapList.Num(); MapIdx++)
+			{
 				UE_LOG(LogCook, Verbose, TEXT("Maplist contains has %s "), *MapList[MapIdx]);
-			AddFileToCook(FilesInPath, MapList[MapIdx]);
+				AddFileToCook(FilesInPath, MapList[MapIdx]);
 				bFoundMapsToCook = true;
 			}
 		}
@@ -5125,7 +5215,7 @@ void UCookOnTheFlyServer::CollectFilesToCook(TArray<FName>& FilesInPath, const T
 			UE_LOG(LogCook, Verbose, TEXT("Loading default map ini section AllMaps "));
 			TArray<FString> AllMapsSection;
 			GEditor->LoadMapListFromIni(TEXT("AllMaps"), AllMapsSection);
-			for ( const auto& MapName : AllMapsSection )
+			for (const auto& MapName : AllMapsSection)
 			{
 				AddFileToCook(FilesInPath, MapName);
 			}
@@ -5659,15 +5749,20 @@ void UCookOnTheFlyServer::CookByTheBookFinished()
 				{
 					SCOPE_TIMER(SaveCookedPackageAssetRegistry);
 
+					
 					TArray<FName> AllCookedPackages;
-					AllCookedPackages.Append(CookedPackagesForPlatform);
 					TArray<FName> UncookedEditorOnlyPackageNamesArray;
-					for (FName UncookedEditorOnlyPackage : UncookedEditorOnlyPackageNames)
 					{
-						UncookedEditorOnlyPackageNamesArray.Add(UncookedEditorOnlyPackage);
+						SCOPE_TIMER(SortingCookedPackages);
+						AllCookedPackages.Append(CookedPackagesForPlatform);
+					
+						for (FName UncookedEditorOnlyPackage : UncookedEditorOnlyPackageNames)
+						{
+							UncookedEditorOnlyPackageNamesArray.Add(UncookedEditorOnlyPackage);
+						}
+						AllCookedPackages.Sort();
+						UncookedEditorOnlyPackageNamesArray.Sort();
 					}
-					AllCookedPackages.Sort();
-					UncookedEditorOnlyPackageNamesArray.Sort();
 					SaveCookedAssetRegistry(SandboxCookedAssetRegistryFilename, AllCookedPackages, UncookedEditorOnlyPackageNamesArray, IgnorePackageNames, PlatformName.ToString(), false);
 				}
 				if (IsCreatingReleaseVersion())
@@ -5896,6 +5991,11 @@ void UCookOnTheFlyServer::TermSandbox()
 	SandboxFile = nullptr;
 }
 
+void UCookOnTheFlyServer::ValidateCookOnTheFlySettings() const
+{
+	check( IsCookFlagSet(ECookInitializationFlags::IterateSharedBuild) == false );
+}
+
 void UCookOnTheFlyServer::ValidateCookByTheBookSettings() const
 {
 	if (IsChildCooker())
@@ -6135,9 +6235,12 @@ void UCookOnTheFlyServer::StartCookByTheBook( const FCookByTheBookStartupOptions
 			{
 				TArray<FName> PlatformNames;
 				PlatformNames.Add(PlatformName);
+				TArray<bool> Succeeded;
+				Succeeded.Add(true);
+				PlatformNames.Add(PlatformName);
 				for (const auto& PackageFilename : PackageList)
 				{
-					CookedPackages.Add( MoveTemp( FFilePlatformCookedPackage( PackageFilename, PlatformNames, true ) ) );
+					CookedPackages.Add( MoveTemp( FFilePlatformCookedPackage( PackageFilename, PlatformNames, Succeeded) ) );
 				}
 			}
 			CookByTheBookOptions->BasedOnReleaseCookedPackages.Add(PlatformName, MoveTemp(PackageList));
@@ -6476,7 +6579,12 @@ void UCookOnTheFlyServer::StartChildCookers(int32 NumCookersToSpawn, const TArra
 			FName PackageFName = FName(*PackageName);
 			check(PackageFName== StandardPackageName);
 #endif	
-			CookedPackages.Add(FFilePlatformCookedPackage(StandardPackageName, TargetPlatformNames, true));
+			TArray<bool> Succeeded;
+			for ( int32 T = 0; T < TargetPlatformNames.Num(); ++T )
+			{
+				Succeeded.Add(true);
+			}
+			CookedPackages.Add(FFilePlatformCookedPackage(StandardPackageName, TargetPlatformNames, MoveTemp(Succeeded)));
 		}
 		DistributeStandardFilenames.RemoveAt(0, NumFilesForCooker);
 
diff --git a/Engine/Source/Editor/UnrealEd/Private/Fbx/FbxSkeletalMeshImport.cpp b/Engine/Source/Editor/UnrealEd/Private/Fbx/FbxSkeletalMeshImport.cpp
--- a/Engine/Source/Editor/UnrealEd/Private/Fbx/FbxSkeletalMeshImport.cpp
+++ b/Engine/Source/Editor/UnrealEd/Private/Fbx/FbxSkeletalMeshImport.cpp
@@ -1401,6 +1401,7 @@ USkeletalMesh* UnFbx::FFbxImporter::ImportSkeletalMesh(UObject* InParent, TArray
 		ExistSkelMeshDataPtr = SaveExistingSkelMeshData(ExistingSkelMesh, !ImportOptions->bImportMaterials);
 	}
 	
+
 	if (SkeletalMesh == nullptr)
 	{
 		// Create the new mesh after saving the old data, since it will replace the old skeletalmesh
@@ -2097,7 +2098,7 @@ void UnFbx::FFbxImporter::SetMaterialOrderByName(FSkeletalMeshImportData& Import
 	TMap<int32, int32> NameIndexGreaterThenMaterialArraySize;
 	{
 		int32 MaterialCount = ImportData.Materials.Num();
-
+		int32 MaxMaterialOrderedCount = 0;
 		bool bNeedsReorder = false;
 		for (int32 MaterialIndex = 0; MaterialIndex < MaterialCount; ++MaterialIndex)
 		{
@@ -2114,6 +2115,7 @@ void UnFbx::FFbxImporter::SetMaterialOrderByName(FSkeletalMeshImportData& Import
 				{
 					if (OrderedIndex < MaterialCount)
 					{
+						MaxMaterialOrderedCount = FMath::Max(MaxMaterialOrderedCount, OrderedIndex+1);
 						NameIndexToMaterialIndex.Add(OrderedIndex, MaterialIndex);
 					}
 					else
@@ -2128,6 +2130,7 @@ void UnFbx::FFbxImporter::SetMaterialOrderByName(FSkeletalMeshImportData& Import
 			if (!bFoundValidName)
 			{
 				MissingNameSuffixMaterial.Add(MaterialIndex);
+				MaxMaterialOrderedCount = FMath::Max(MaxMaterialOrderedCount, MaterialIndex+1);
 			}
 		}
 
@@ -2135,18 +2138,22 @@ void UnFbx::FFbxImporter::SetMaterialOrderByName(FSkeletalMeshImportData& Import
 		{
 			//Add the missing name material at the end to not disturb the existing order
 			TArray<int32> OrderedListMissing;
-			OrderedListMissing.AddZeroed(NameIndexToMaterialIndex.Num() + MissingNameSuffixMaterial.Num());
+			OrderedListMissing.AddZeroed(MaxMaterialOrderedCount);
 			for (auto Kvp : NameIndexToMaterialIndex)
 			{
 				OrderedListMissing[Kvp.Key] = -1;
 			}
 			for (int32 OrderedListMissingIndex = 0; OrderedListMissingIndex < OrderedListMissing.Num(); ++OrderedListMissingIndex)
 			{
+				if (MissingNameSuffixMaterial.Num() <= 0)
+				{
+					break;
+				}
+
 				if (OrderedListMissing[OrderedListMissingIndex] != 0)
 					continue;
 
 				NameIndexToMaterialIndex.Add(OrderedListMissingIndex, MissingNameSuffixMaterial.Pop());
-				
 			}
 		}
 
diff --git a/Engine/Source/Editor/UnrealEd/Private/PreviewMaterial.cpp b/Engine/Source/Editor/UnrealEd/Private/PreviewMaterial.cpp
--- a/Engine/Source/Editor/UnrealEd/Private/PreviewMaterial.cpp
+++ b/Engine/Source/Editor/UnrealEd/Private/PreviewMaterial.cpp
@@ -77,9 +77,8 @@ public:
 					return false;
 				}
 
-				extern ENGINE_API int32 GEnableGPUSkinCacheShaders;
-
-				bool bSkinCache = GEnableGPUSkinCacheShaders && (VertexFactoryType == FindVertexFactoryType(FName(TEXT("FGPUSkinPassthroughVertexFactory"), FNAME_Find)));
+				extern ENGINE_API bool IsGPUSkinCacheAvailable();
+				bool bSkinCache = IsGPUSkinCacheAvailable() && (VertexFactoryType == FindVertexFactoryType(FName(TEXT("FGPUSkinPassthroughVertexFactory"), FNAME_Find)));
 					
 				if (
 					VertexFactoryType != FindVertexFactoryType(FName(TEXT("TGPUSkinVertexFactoryfalse"), FNAME_Find)) &&
diff --git a/Engine/Source/Programs/AutomationTool/AutomationTool.csproj b/Engine/Source/Programs/AutomationTool/AutomationTool.csproj
--- a/Engine/Source/Programs/AutomationTool/AutomationTool.csproj
+++ b/Engine/Source/Programs/AutomationTool/AutomationTool.csproj
@@ -14,6 +14,7 @@
     <TargetFrameworkProfile>
     </TargetFrameworkProfile>
     <FileAlignment>512</FileAlignment>
+    <AutoGenerateBindingRedirects>true</AutoGenerateBindingRedirects>
   </PropertyGroup>
   <PropertyGroup Condition="'$(Configuration)|$(Platform)' == 'Debug|AnyCPU'">
     <DebugSymbols>true</DebugSymbols>
diff --git a/Engine/Source/Programs/AutomationTool/Scripts/CopyBuildToStagingDirectory.Automation.cs b/Engine/Source/Programs/AutomationTool/Scripts/CopyBuildToStagingDirectory.Automation.cs
--- a/Engine/Source/Programs/AutomationTool/Scripts/CopyBuildToStagingDirectory.Automation.cs
+++ b/Engine/Source/Programs/AutomationTool/Scripts/CopyBuildToStagingDirectory.Automation.cs
@@ -111,8 +111,9 @@ public partial class Project : CommandUtils
 		var UnrealPakExe = CombinePaths(CmdEnv.LocalRoot, "Engine/Binaries/Win64/UnrealPak.exe");
 		Log("Running UnrealPak *******");
 		string CmdLine = CommandUtils.MakePathSafeToUseWithCommandLine(OutputLocation) + " -create=" + CommandUtils.MakePathSafeToUseWithCommandLine(UnrealPakResponseFileName);
+		string LogFileName = CombinePaths(CmdEnv.LogFolder, "PakLog_" + PakName + ".log");
 
-		CmdLine += String.Format(" -encryptionini -enginedir=\"{0}\" -projectdir=\"{1}\" -platform={2}", EngineDir, ProjectDir, Platform);
+		CmdLine += String.Format(" -encryptionini -enginedir=\"{0}\" -projectdir=\"{1}\" -platform={2} -abslog=\"{3}\"", EngineDir, ProjectDir, Platform, LogFileName);
 
 		if (GlobalCommandLine.Installed)
 		{
@@ -253,6 +254,14 @@ public partial class Project : CommandUtils
 
 		Log("Creating Staging Manifest...");
 
+		if (Params.IterateSharedCookedBuild)
+		{
+			// can't do shared cooked builds with DLC that's madness!!
+			//check( Params.HasDLCName == false );
+
+			// stage all the previously staged files
+			SC.StageFiles(StagedFileType.NonUFS, CombinePaths(SC.ProjectRoot, "Saved", "SharedIterativeBuild", SC.CookPlatform, "StagedBuild"), "*", true, null, "", true); // remap to the root directory
+		}
         if (Params.HasDLCName)
         {
             string DLCName = Params.DLCName;
@@ -269,6 +278,7 @@ public partial class Project : CommandUtils
             string[] ExcludeWildCards = {"AssetRegistry.bin"};
 
             // Stage any loose files in the root folder
+			// TODO: not sure if we should stage the loose files if we have pak files enabled... 
 			SC.StageFiles(StagedFileType.UFS, PlatformCookDir, "*", false, ExcludeWildCards, SC.RelativeProjectRootForStage, true, !Params.UsePak(SC.StageTargetPlatform));
 
 			// Stage each sub directory separately so that we can skip Engine if need be
@@ -583,7 +593,7 @@ public partial class Project : CommandUtils
                 SC.StageRuntimeDependenciesFromReceipt(Receipt, true, Params.UsePak(SC.StageTargetPlatform));
 
 				// Add config files.
-				SC.StageFiles( StagedFileType.NonUFS, CombinePaths( SC.LocalRoot, "Engine/Programs/CrashReportClient/Config" ) );
+				SC.StageFiles( StagedFileType.NonUFS, CombinePaths( SC.LocalRoot, "Engine/Programs/CrashReportClient/Config" ), bAllowNotForLicenseesFiles:false );
 			}
 
 			// check if the game will be verifying ssl connections - if not, we can skip staging files that won't be needed
@@ -948,6 +958,12 @@ public partial class Project : CommandUtils
         {
             PostFix += "_P";
         }
+		if (Params.IterateSharedCookedBuild)
+		{
+			// shared cooked builds will produce a patch
+			// then be combined with the shared cooked build
+			PostFix += "_S_P";
+		}
 		var OutputRelativeLocation = CombinePaths(SC.RelativeProjectRootForStage, "Content/Paks/", PakName + "-" + SC.FinalCookPlatform + PostFix + ".pak");
 		if (SC.StageTargetPlatform.DeployLowerCaseFilenames(true))
 		{
diff --git a/Engine/Source/Programs/AutomationTool/Scripts/RunProjectCommand.Automation.cs b/Engine/Source/Programs/AutomationTool/Scripts/RunProjectCommand.Automation.cs
--- a/Engine/Source/Programs/AutomationTool/Scripts/RunProjectCommand.Automation.cs
+++ b/Engine/Source/Programs/AutomationTool/Scripts/RunProjectCommand.Automation.cs
@@ -981,9 +981,11 @@ public partial class Project : CommandUtils
 			{
 				Map += "?fake";
 			}
-
-			Args += String.Format("{0} -server -abslog={1}  -unattended -log -Messaging", Map, CommandUtils.MakePathSafeToUseWithCommandLine(ServerLogFile));
-
+			Args += String.Format("{0} -server -abslog={1}  -log -Messaging", Map, CommandUtils.MakePathSafeToUseWithCommandLine(ServerLogFile));
+			if (Params.Unattended)
+			{
+				Args += " -unattended";
+			}
 			// Do not blindly add -nomcp, only do so if the client is using it
 			if (Params.RunCommandline.Contains("-nomcp"))
 			{
diff --git a/Engine/Source/Programs/AutomationToolLauncher/AutomationToolLauncher.csproj b/Engine/Source/Programs/AutomationToolLauncher/AutomationToolLauncher.csproj
--- a/Engine/Source/Programs/AutomationToolLauncher/AutomationToolLauncher.csproj
+++ b/Engine/Source/Programs/AutomationToolLauncher/AutomationToolLauncher.csproj
@@ -11,6 +11,7 @@
     <AssemblyName>AutomationToolLauncher</AssemblyName>
     <TargetFrameworkVersion>v4.5</TargetFrameworkVersion>
     <FileAlignment>512</FileAlignment>
+    <AutoGenerateBindingRedirects>true</AutoGenerateBindingRedirects>
   </PropertyGroup>
   <PropertyGroup Condition=" '$(Configuration)|$(Platform)' == 'Debug|AnyCPU' ">
     <PlatformTarget>AnyCPU</PlatformTarget>
diff --git a/Engine/Source/Runtime/AIModule/Classes/AIController.h b/Engine/Source/Runtime/AIModule/Classes/AIController.h
--- a/Engine/Source/Runtime/AIModule/Classes/AIController.h
+++ b/Engine/Source/Runtime/AIModule/Classes/AIController.h
@@ -444,6 +444,8 @@ public:
 	const UAIPerceptionComponent* GetAIPerceptionComponent() const;
 
 	UBrainComponent* GetBrainComponent() { return BrainComponent; }
+	const UBlackboardComponent* GetBlackboardComponent() const { return Blackboard; }
+	UBlackboardComponent* GetBlackboardComponent() { return Blackboard; }
 };
 
 //----------------------------------------------------------------------//
diff --git a/Engine/Source/Runtime/AIModule/Classes/BehaviorTree/BlackboardData.h b/Engine/Source/Runtime/AIModule/Classes/BehaviorTree/BlackboardData.h
--- a/Engine/Source/Runtime/AIModule/Classes/BehaviorTree/BlackboardData.h
+++ b/Engine/Source/Runtime/AIModule/Classes/BehaviorTree/BlackboardData.h
@@ -86,6 +86,8 @@ public:
 	/** @return key data */
 	const FBlackboardEntry* GetKey(FBlackboard::FKey KeyID) const;
 
+	const TArray<FBlackboardEntry>& GetKeys() const { return Keys; }
+
 	virtual void PostLoad() override;
 #if WITH_EDITOR
 	virtual void PostEditChangeProperty(struct FPropertyChangedEvent& PropertyChangedEvent) override;
diff --git a/Engine/Source/Runtime/AIModule/Classes/BrainComponent.h b/Engine/Source/Runtime/AIModule/Classes/BrainComponent.h
--- a/Engine/Source/Runtime/AIModule/Classes/BrainComponent.h
+++ b/Engine/Source/Runtime/AIModule/Classes/BrainComponent.h
@@ -113,7 +113,7 @@ private:
 	FAIMessageObserver& operator=(const FAIMessageObserver&);
 };
 
-UCLASS(abstract, BlueprintType)
+UCLASS(BlueprintType)
 class AIMODULE_API UBrainComponent : public UActorComponent, public IAIResourceInterface
 {
 	GENERATED_UCLASS_BODY()
diff --git a/Engine/Source/Runtime/AIModule/Classes/EnvironmentQuery/EnvQueryTypes.h b/Engine/Source/Runtime/AIModule/Classes/EnvironmentQuery/EnvQueryTypes.h
--- a/Engine/Source/Runtime/AIModule/Classes/EnvironmentQuery/EnvQueryTypes.h
+++ b/Engine/Source/Runtime/AIModule/Classes/EnvironmentQuery/EnvQueryTypes.h
@@ -530,6 +530,7 @@ public:
 
 	FORCEINLINE bool IsFinished() const { return Status != EEnvQueryStatus::Processing; }
 	FORCEINLINE bool IsAborted() const { return Status == EEnvQueryStatus::Aborted; }
+	FORCEINLINE bool IsSuccsessful() const { return Status == EEnvQueryStatus::Success; }
 	FORCEINLINE void MarkAsMissingParam() { Status = EEnvQueryStatus::MissingParam; }
 	FORCEINLINE void MarkAsAborted() { Status = EEnvQueryStatus::Aborted; }
 	FORCEINLINE void MarkAsFailed() { Status = EEnvQueryStatus::Failed; }
@@ -921,7 +922,7 @@ public:
 #	define  UE_EQS_DBGMSG(Condition, Format, ...) \
 					if (Condition) \
 					{ \
-						Instance->ItemDetails[CurrentItem].FailedDescription = FString::Printf(Format, ##__VA_ARGS__); \
+						Instance.ItemDetails[CurrentItem].FailedDescription = FString::Printf(Format, ##__VA_ARGS__); \
 					}
 
 #	define UE_EQS_LOG(CategoryName, Verbosity, Format, ...) \
@@ -936,13 +937,63 @@ public:
 	FString GetExecutionTimeDescription() const;
 
 #if CPP || UE_BUILD_DOCS
-	struct AIMODULE_API ItemIterator
+	/** Note that this iterator is for read-only purposes. Please use FItemIterator for regular item iteration 
+	 *	while performing EQS testing and scoring */
+	struct AIMODULE_API FConstItemIterator
 	{
-		ItemIterator(const UEnvQueryTest* QueryTest, FEnvQueryInstance& QueryInstance, int32 StartingItemIndex = INDEX_NONE);
+		FConstItemIterator(FEnvQueryInstance& QueryInstance, int32 StartingItemIndex = INDEX_NONE)
+			: Instance(QueryInstance)
+			, CurrentItem(StartingItemIndex != INDEX_NONE ? StartingItemIndex : QueryInstance.CurrentTestStartingItem)
+		{
+			if (StartingItemIndex != INDEX_NONE)
+			{
+				CurrentItem = StartingItemIndex;
+			}
+			else
+			{
+				CurrentItem = QueryInstance.CurrentTestStartingItem;
+				if (Instance.Items.IsValidIndex(CurrentItem) == false || Instance.Items[CurrentItem].IsValid() == false)
+				{
+					++(*this);
+				}
+			}
+		}
+
+		uint8* GetItemData()
+		{
+			return Instance.RawData.GetData() + Instance.Items[CurrentItem].DataOffset;
+		}
 
-		~ItemIterator()
+		int32 GetIndex() const
+		{
+			return CurrentItem;
+		}
+
+		FORCEINLINE explicit operator bool() const
 		{
-			Instance->CurrentTestStartingItem = CurrentItem;
+			return CurrentItem < Instance.Items.Num();
+		}
+
+		void operator++()
+		{
+			++CurrentItem;
+			for (; CurrentItem < Instance.Items.Num() && !Instance.Items[CurrentItem].IsValid(); ++CurrentItem)
+				;
+		}
+
+	protected:
+
+		FEnvQueryInstance& Instance;
+		int32 CurrentItem;
+	};
+
+	struct AIMODULE_API FItemIterator : public FConstItemIterator
+	{
+		FItemIterator(const UEnvQueryTest* QueryTest, FEnvQueryInstance& QueryInstance, int32 StartingItemIndex = INDEX_NONE);
+
+		~FItemIterator()
+		{
+			Instance.CurrentTestStartingItem = CurrentItem;
 		}
 
 		/** Filter and score an item - used by tests working on float values
@@ -1050,11 +1101,6 @@ public:
 			NumTestsForItem++;
 		}
 
-		uint8* GetItemData()
-		{
-			return Instance->RawData.GetData() + Instance->Items[CurrentItem].DataOffset;
-		}
-
 		/** Force state and score of item
 		 *  Any following SetScore calls for current item will be ignored
 		 */
@@ -1066,37 +1112,31 @@ public:
 		}
 
 		/** Disables time slicing for this iterator, use with caution! */
-		void IgnoreTimeLimit()
+		FItemIterator& IgnoreTimeLimit()
 		{
 			Deadline = -1.0f;
-		}
-
-		int32 GetIndex() const
-		{
-			return CurrentItem;
+			return *this;
 		}
 
 		FORCEINLINE explicit operator bool() const
 		{
-			return CurrentItem < Instance->Items.Num() && !Instance->bFoundSingleResult && (Deadline < 0 || FPlatformTime::Seconds() < Deadline);
+			return CurrentItem < Instance.Items.Num() && !Instance.bFoundSingleResult && (Deadline < 0 || FPlatformTime::Seconds() < Deadline);
 		}
 
 		void operator++()
 		{
 			StoreTestResult();
-			if (!Instance->bFoundSingleResult)
+			if (!Instance.bFoundSingleResult)
 			{
 				InitItemScore();
-				FindNextValidIndex();
+				FConstItemIterator::operator++();
 			}
 		}
 
 	protected:
 
-		FEnvQueryInstance* Instance;
 		double Deadline;
 		float ItemScore;
-		int32 CurrentItem;
 		int16 NumPassedForItem;
 		int16 NumTestsForItem;
 		uint8 CachedFilterOp;
@@ -1117,12 +1157,6 @@ public:
 		void HandleFailedTestResult();
 		void StoreTestResult();
 
-		FORCEINLINE void FindNextValidIndex()
-		{
-			for (CurrentItem++; CurrentItem < Instance->Items.Num() && !Instance->Items[CurrentItem].IsValid(); CurrentItem++)
-				;
-		}
-
 		FORCEINLINE void SetScoreInternal(float Score)
 		{
 			switch (CachedScoreOp)
@@ -1170,6 +1204,7 @@ public:
 			}
 		}
 	};
+	typedef FItemIterator ItemIterator;
 #endif
 
 #undef UE_EQS_LOG
diff --git a/Engine/Source/Runtime/AIModule/Classes/Navigation/NavLocalGridData.h b/Engine/Source/Runtime/AIModule/Classes/Navigation/NavLocalGridData.h
--- a/Engine/Source/Runtime/AIModule/Classes/Navigation/NavLocalGridData.h
+++ b/Engine/Source/Runtime/AIModule/Classes/Navigation/NavLocalGridData.h
@@ -100,6 +100,7 @@ struct AIMODULE_API FNavLocalGridData : public TSimpleCellGrid<uint8, MAX_uint8>
 protected:
 	friend class UNavLocalGridManager;
 	TArray<float> CellZ;
+	float LastAccessTime;
 
 	/** convert PathIndices into pruned PathCoords */
 	void PostProcessPath(const FIntVector& StartCoords, const FIntVector& EndCoords, const TArray<int32>& PathIndices, TArray<FIntVector>& PathCoords) const;
diff --git a/Engine/Source/Runtime/AIModule/Classes/Navigation/NavLocalGridManager.h b/Engine/Source/Runtime/AIModule/Classes/Navigation/NavLocalGridManager.h
--- a/Engine/Source/Runtime/AIModule/Classes/Navigation/NavLocalGridManager.h
+++ b/Engine/Source/Runtime/AIModule/Classes/Navigation/NavLocalGridManager.h
@@ -24,7 +24,7 @@ struct FCombinedNavGridData : public FNavLocalGridData
  */
 
 UCLASS(Experimental)
-class UNavLocalGridManager : public UObject
+class AIMODULE_API UNavLocalGridManager : public UObject
 {
 	GENERATED_UCLASS_BODY()
 public:
@@ -47,6 +47,18 @@ public:
 	/** get shared size of grid cell, static but there's only one active manager */
 	static float GetCellSize() { return GridCellSize; }
 
+	/** set limit of source grids, 0 or negative means unlimited */
+	void SetMaxActiveSources(int32 NumActiveSources);
+
+	/** get limit of source grids */
+	int32 GetMaxActiveSources() const { return MaxActiveSourceGrids; }
+
+	/** check if source grid limit is set */
+	bool HasSourceGridLimit() const { return MaxActiveSourceGrids > 0; }
+
+	/** updates LastAccessTime in all source grids */
+	void UpdateAccessTime(int32 CombinedGridIdx);
+
 	/** get number of known source grids */
 	int32 GetNumSources() const { return SourceGrids.Num(); }
 
@@ -59,6 +71,9 @@ public:
 	/** get combined, non overlapping grid by index */
 	const FNavLocalGridData& GetGridData(int32 GridIdx) const { return CombinedGrids[GridIdx]; }
 
+	/** find combined grid value at world location */
+	uint8 GetGridValueAt(const FVector& WorldLocation) const;
+
 	/** find combined grid at location */
 	int32 GetGridIndex(const FVector& WorldLocation) const;
 
@@ -68,9 +83,14 @@ public:
 	UFUNCTION(BlueprintCallable, Category = "AI|Navigation", meta = (WorldContext = "WorldContext"))
 	static bool SetLocalNavigationGridDensity(UObject* WorldContext, float CellSize);
 
+	/** creates new grid data for single point */
 	UFUNCTION(BlueprintCallable, Category = "AI|Navigation", meta = (WorldContext = "WorldContext"))
 	static int32 AddLocalNavigationGridForPoint(UObject* WorldContext, const FVector& Location, const int32 Radius2D = 5, const float Height = 100.0f, bool bRebuildGrids = true);
 
+	/** creates single grid data for set of points */
+	UFUNCTION(BlueprintCallable, Category = "AI|Navigation", meta = (WorldContext = "WorldContext"))
+	static int32 AddLocalNavigationGridForPoints(UObject* WorldContext, const TArray<FVector>& Locations, const int32 Radius2D = 5, const float Height = 100.0f, bool bRebuildGrids = true);
+
 	UFUNCTION(BlueprintCallable, Category = "AI|Navigation", meta = (WorldContext = "WorldContext"))
 	static int32 AddLocalNavigationGridForBox(UObject* WorldContext, const FVector& Location, FVector Extent = FVector(1,1,1), FRotator Rotation = FRotator::ZeroRotator, const int32 Radius2D = 5, const float Height = 100.0f, bool bRebuildGrids = true);
 
@@ -99,8 +119,12 @@ protected:
 
 	int32 VersionNum;
 	int32 NextGridId;
+	int32 MaxActiveSourceGrids;
 	uint32 bNeedsRebuilds : 1;
 
 	/** projects combined grids to navigation data */
 	virtual void ProjectGrids(const TArray<int32>& GridIndices);
+
+	/** ensures limit of source grids, removing oldest entries (LastAccessTime) */
+	bool UpdateSourceGrids();
 };
diff --git a/Engine/Source/Runtime/AIModule/Classes/Tasks/AITask_RunEQS.h b/Engine/Source/Runtime/AIModule/Classes/Tasks/AITask_RunEQS.h
new file mode 100644
--- /dev/null
+++ b/Engine/Source/Runtime/AIModule/Classes/Tasks/AITask_RunEQS.h
@@ -0,0 +1,37 @@
+// Copyright 1998-2017 Epic Games, Inc. All Rights Reserved.
+#pragma once
+
+#include "CoreMinimal.h"
+#include "UObject/ObjectMacros.h"
+#include "Engine/EngineTypes.h"
+#include "Tasks/AITask.h"
+#include "EnvironmentQuery/EnvQueryTypes.h"
+#include "AITask_RunEQS.generated.h"
+
+class AAIController;
+class UEnvQuery;
+
+UCLASS()
+class AIMODULE_API UAITask_RunEQS : public UAITask
+{
+	GENERATED_BODY()
+		
+public:
+	UAITask_RunEQS(const FObjectInitializer& ObjectInitializer = FObjectInitializer::Get());
+
+	UFUNCTION(BlueprintCallable, Category = "AI|Tasks", meta = (DefaultToSelf = "Controller", BlueprintInternalUseOnly = "TRUE"))
+	static UAITask_RunEQS* RunEQS(AAIController* Controller, UEnvQuery* QueryTemplate);
+
+	void SetQueryTemplate(UEnvQuery& InQueryTemplate) { EQSRequest.QueryTemplate = &InQueryTemplate; }
+	void SetNotificationDelegate(const FQueryFinishedSignature& InNotificationDelegate) { NotificationDelegate = InNotificationDelegate; }
+		
+protected:
+	virtual void Activate() override;
+
+	void OnEQSRequestFinished(TSharedPtr<FEnvQueryResult> Result);
+
+	FEQSParametrizedQueryExecutionRequest EQSRequest;
+	FQueryFinishedSignature EQSFinishedDelegate;
+	FQueryFinishedSignature NotificationDelegate;
+	TSharedPtr<FEnvQueryResult> QueryResult;
+};
diff --git a/Engine/Source/Runtime/AIModule/Private/EnvironmentQuery/EnvQueryInstance.cpp b/Engine/Source/Runtime/AIModule/Private/EnvironmentQuery/EnvQueryInstance.cpp
--- a/Engine/Source/Runtime/AIModule/Private/EnvironmentQuery/EnvQueryInstance.cpp
+++ b/Engine/Source/Runtime/AIModule/Private/EnvironmentQuery/EnvQueryInstance.cpp
@@ -517,9 +517,8 @@ void FEnvQueryInstance::ReserveItemData(const int32 NumAdditionalItems)
 	INC_MEMORY_STAT_BY(STAT_AI_EQS_InstanceMemory, RawData.GetAllocatedSize());
 }
 
-FEnvQueryInstance::ItemIterator::ItemIterator(const UEnvQueryTest* QueryTest, FEnvQueryInstance& QueryInstance, int32 StartingItemIndex)
-	: Instance(&QueryInstance)
-	, CurrentItem(StartingItemIndex != INDEX_NONE ? StartingItemIndex : QueryInstance.CurrentTestStartingItem)
+FEnvQueryInstance::FItemIterator::FItemIterator(const UEnvQueryTest* QueryTest, FEnvQueryInstance& QueryInstance, int32 StartingItemIndex)
+	: FConstItemIterator(QueryInstance, StartingItemIndex)
 {
 	check(QueryTest);
 
@@ -528,32 +527,29 @@ FEnvQueryInstance::ItemIterator::ItemIterator(const UEnvQueryTest* QueryTest, FE
 	bIsFiltering = (QueryTest->TestPurpose == EEnvTestPurpose::Filter) || (QueryTest->TestPurpose == EEnvTestPurpose::FilterAndScore);
 
 	Deadline = QueryInstance.CurrentStepTimeLimit > 0.0 ? (FPlatformTime::Seconds() + QueryInstance.CurrentStepTimeLimit) : -1.0;
-	// it's possible item 'CurrentItem' has been already discarded. Find a valid starting index
-	--CurrentItem;
-	FindNextValidIndex();
 	InitItemScore();
 }
 
-void FEnvQueryInstance::ItemIterator::HandleFailedTestResult()
+void FEnvQueryInstance::FItemIterator::HandleFailedTestResult()
 {
 	ItemScore = -1.f;
-	Instance->Items[CurrentItem].Discard();
+	Instance.Items[CurrentItem].Discard();
 #if USE_EQS_DEBUGGER
-	Instance->ItemDetails[CurrentItem].FailedTestIndex = Instance->CurrentTest;
+	Instance.ItemDetails[CurrentItem].FailedTestIndex = Instance.CurrentTest;
 #endif
-	Instance->NumValidItems--;
+	Instance.NumValidItems--;
 }
 
-void FEnvQueryInstance::ItemIterator::StoreTestResult()
+void FEnvQueryInstance::FItemIterator::StoreTestResult()
 {
 	CheckItemPassed();
 	ensureAlways(FMath::IsNaN(ItemScore) == false);
 
 #if USE_EQS_DEBUGGER
-	Instance->NumProcessedItems++;
+	Instance.NumProcessedItems++;
 #endif // USE_EQS_DEBUGGER
 
-	if (Instance->IsInSingleItemFinalSearch())
+	if (Instance.IsInSingleItemFinalSearch())
 	{
 		// handle SingleResult mode
 		if (bPassed)
@@ -561,11 +557,11 @@ void FEnvQueryInstance::ItemIterator::StoreTestResult()
 			if (bForced)
 			{
 				// store item value in case it's using special "skipped" constant
-				Instance->ItemDetails[CurrentItem].TestResults[Instance->CurrentTest] = ItemScore;
+				Instance.ItemDetails[CurrentItem].TestResults[Instance.CurrentTest] = ItemScore;
 			}
 
-			Instance->PickSingleItem(CurrentItem);
-			Instance->bFoundSingleResult = true;
+			Instance.PickSingleItem(CurrentItem);
+			Instance.bFoundSingleResult = true;
 		}
 		else if (!bPassed)
 		{
@@ -584,7 +580,7 @@ void FEnvQueryInstance::ItemIterator::StoreTestResult()
 			ItemScore /= NumPassedForItem;
 		}
 
-		Instance->ItemDetails[CurrentItem].TestResults[Instance->CurrentTest] = ItemScore;
+		Instance.ItemDetails[CurrentItem].TestResults[Instance.CurrentTest] = ItemScore;
 	}
 }
 
diff --git a/Engine/Source/Runtime/AIModule/Private/EnvironmentQuery/EnvQueryTest.cpp b/Engine/Source/Runtime/AIModule/Private/EnvironmentQuery/EnvQueryTest.cpp
--- a/Engine/Source/Runtime/AIModule/Private/EnvironmentQuery/EnvQueryTest.cpp
+++ b/Engine/Source/Runtime/AIModule/Private/EnvironmentQuery/EnvQueryTest.cpp
@@ -109,12 +109,15 @@ void UEnvQueryTest::NormalizeItemScores(FEnvQueryInstance& QueryInstance)
 			if (TestValue != UEnvQueryTypes::SkippedItemValue)
 			{
 				const float ClampedScore = FMath::Clamp(TestValue, MinScore, MaxScore);
-				const float NormalizedScore = FMath::Abs(LocalReferenceValue - ClampedScore) / ValueSpan;
+				const float NormalizedScore = (ScoringFactorValue >= 0) 
+					? (FMath::Abs(LocalReferenceValue - ClampedScore) / ValueSpan)
+					: (1.f - (FMath::Abs(LocalReferenceValue - ClampedScore) / ValueSpan));
+				const float AbsoluteWeight = (ScoringFactorValue >= 0) ? ScoringFactorValue : -ScoringFactorValue;
 
 				switch (ScoringEquation)
 				{
 					case EEnvTestScoreEquation::Linear:
-						WeightedScore = ScoringFactorValue * NormalizedScore;
+						WeightedScore = AbsoluteWeight * NormalizedScore;
 						break;
 
 					case EEnvTestScoreEquation::InverseLinear:
@@ -124,21 +127,21 @@ void UEnvQueryTest::NormalizeItemScores(FEnvQueryInstance& QueryInstance)
 						// to add that flag later, we'll need to remove this option, since it should just be "mirror
 						// curve" plus "Linear".
 						float InverseNormalizedScore = (1.0f - NormalizedScore);
-						WeightedScore = ScoringFactorValue * InverseNormalizedScore;
+						WeightedScore = AbsoluteWeight * InverseNormalizedScore;
 						break;
 					}
 
 					case EEnvTestScoreEquation::Square:
-						WeightedScore = ScoringFactorValue * (NormalizedScore * NormalizedScore);
+						WeightedScore = AbsoluteWeight * (NormalizedScore * NormalizedScore);
 						break;
 
 					case EEnvTestScoreEquation::SquareRoot:
-						WeightedScore = ScoringFactorValue * FMath::Sqrt(NormalizedScore);
+						WeightedScore = AbsoluteWeight * FMath::Sqrt(NormalizedScore);
 						break;
 
 					case EEnvTestScoreEquation::Constant:
 						// I know, it's not "constant".  It's "Constant, or zero".  The tooltip should explain that.
-						WeightedScore = (NormalizedScore > 0) ? ScoringFactorValue : 0.0f;
+						WeightedScore = (NormalizedScore > 0) ? AbsoluteWeight : 0.0f;
 						break;
 						
 					default:
diff --git a/Engine/Source/Runtime/AIModule/Private/Navigation/GridPathFollowingComponent.cpp b/Engine/Source/Runtime/AIModule/Private/Navigation/GridPathFollowingComponent.cpp
--- a/Engine/Source/Runtime/AIModule/Private/Navigation/GridPathFollowingComponent.cpp
+++ b/Engine/Source/Runtime/AIModule/Private/Navigation/GridPathFollowingComponent.cpp
@@ -122,6 +122,7 @@ void UGridPathFollowingComponent::UpdateActiveGrid(const FVector& CurrentLocatio
 				ActiveGridId = GridData.GetGridId();
 				ActiveGridIdx = GridIdx;
 
+				GridManager->UpdateAccessTime(GridIdx);
 				GridData.FindPathForMovingAgent(*Path.Get(), CurrentLocation, MoveSegmentStartIndex, GridPathPoints, MoveSegmentStartIndexOffGrid);
 				
 				bHasGridPath = GridPathPoints.Num() > 1;
diff --git a/Engine/Source/Runtime/AIModule/Private/Navigation/NavLocalGridManager.cpp b/Engine/Source/Runtime/AIModule/Private/Navigation/NavLocalGridManager.cpp
--- a/Engine/Source/Runtime/AIModule/Private/Navigation/NavLocalGridManager.cpp
+++ b/Engine/Source/Runtime/AIModule/Private/Navigation/NavLocalGridManager.cpp
@@ -15,17 +15,24 @@ UNavLocalGridManager::UNavLocalGridManager(const FObjectInitializer& ObjectIniti
 
 int32 UNavLocalGridManager::AddGridData(const FNavLocalGridData& GridData, bool bUpdate)
 {
+	const float GameTime = HasSourceGridLimit() ? GetWorld()->GetTimeSeconds() : 0;
 	const int32 NewGridIdx = SourceGrids.Add(GridData);
+
 	SourceGrids[NewGridIdx].SetGridId(NextGridId);
+	SourceGrids[NewGridIdx].LastAccessTime = GameTime;
+
+	const int32 UsedId = SourceGrids[NewGridIdx].GetGridId();
 	NextGridId++;
-		
+	
+	UpdateSourceGrids();
+
 	bNeedsRebuilds = true;
 	if (bUpdate)
 	{
 		RebuildGrids();
 	}
 
-	return SourceGrids[NewGridIdx].GetGridId();
+	return UsedId;
 }
 
 void UNavLocalGridManager::RemoveGridData(int32 GridId, bool bUpdate)
@@ -204,6 +211,62 @@ void UNavLocalGridManager::ProjectGrids(const TArray<int32>& GridIndices)
 	}
 }
 
+void UNavLocalGridManager::UpdateAccessTime(int32 CombinedGridIdx)
+{
+	if (CombinedGrids.IsValidIndex(CombinedGridIdx))
+	{
+		const float GameTime = GetWorld()->GetTimeSeconds();
+
+		for (int32 Idx = 0; Idx < CombinedGrids[CombinedGridIdx].SourceIds.Num(); Idx++)
+		{
+			const int32 SourceGridIdx = CombinedGrids[CombinedGridIdx].SourceIds[Idx];
+			if (SourceGrids.IsValidIndex(SourceGridIdx))
+			{
+				SourceGrids[SourceGridIdx].LastAccessTime = GameTime;
+			}
+		}
+	}
+}
+
+void UNavLocalGridManager::SetMaxActiveSources(int32 NumActiveSources)
+{
+	MaxActiveSourceGrids = NumActiveSources;
+	
+	const bool bUpdated = UpdateSourceGrids();
+	if (bUpdated)
+	{
+		RebuildGrids();
+	}
+}
+
+bool UNavLocalGridManager::UpdateSourceGrids()
+{
+	if (SourceGrids.Num() < MaxActiveSourceGrids || !HasSourceGridLimit())
+	{
+		return false;
+	}
+
+	while (SourceGrids.Num() > MaxActiveSourceGrids)
+	{
+		float BestScore = FLT_MAX;
+		int32 BestIdx = 0;
+
+		for (int32 Idx = 0; Idx < SourceGrids.Num(); Idx++)
+		{
+			const FNavLocalGridData& GridData = SourceGrids[Idx];
+			if (BestScore > GridData.LastAccessTime)
+			{
+				BestScore = GridData.LastAccessTime;
+				BestIdx = Idx;
+			}
+		}
+
+		SourceGrids.RemoveAt(BestIdx, 1, false);
+	}
+
+	return true;
+}
+
 int32 UNavLocalGridManager::GetGridIndex(const FVector& WorldLocation) const
 {
 	DECLARE_SCOPE_CYCLE_COUNTER(TEXT("NavGrid: GetIndex"), STAT_GridFind, STATGROUP_AI);
@@ -219,6 +282,12 @@ int32 UNavLocalGridManager::GetGridIndex(const FVector& WorldLocation) const
 	return INDEX_NONE;
 }
 
+uint8 UNavLocalGridManager::GetGridValueAt(const FVector& WorldLocation) const
+{
+	const int32 GridIndex = GetGridIndex(WorldLocation);
+	return (GridIndex != INDEX_NONE) ? CombinedGrids[GridIndex].GetCellAtWorldLocation(WorldLocation) : 0;
+}
+
 bool UNavLocalGridManager::FindPath(const FVector& Start, const FVector& End, TArray<FVector>& PathPoints) const
 {
 	const int32 StartGridIdx = GetGridIndex(Start);
@@ -264,6 +333,30 @@ int32 UNavLocalGridManager::AddLocalNavigationGridForPoint(UObject* WorldContext
 	return GridId;
 }
 
+int32 UNavLocalGridManager::AddLocalNavigationGridForPoints(UObject* WorldContext, const TArray<FVector>& Locations, const int32 Radius2D, const float Height, bool bRebuildGrids)
+{
+	int32 GridId = 0;
+
+	UNavLocalGridManager* GridManager = UNavLocalGridManager::GetCurrent(WorldContext);
+	if (GridManager)
+	{
+		const FBox Bounds(Locations);
+		const float BoundsSize2D = FMath::Max(Bounds.Max.X - Bounds.Min.X, Bounds.Max.Y - Bounds.Min.Y);
+
+		FNavLocalGridData GridData(Bounds.GetCenter(), (UNavLocalGridManager::GridCellSize * Radius2D) + BoundsSize2D);
+		GridData.SetHeight(Height);
+
+		for (int32 Idx = 0; Idx < Locations.Num(); Idx++)
+		{
+			GridData.MarkPointObstacle(Locations[Idx]);
+		}
+
+		GridId = GridManager->AddGridData(GridData, bRebuildGrids);
+	}
+
+	return GridId;
+}
+
 int32 UNavLocalGridManager::AddLocalNavigationGridForBox(UObject* WorldContext, const FVector& Location, FVector Extent, FRotator Rotation, const int32 Radius2D, const float Height, bool bRebuildGrids)
 {
 	int32 GridId = 0;
diff --git a/Engine/Source/Runtime/AnimGraphRuntime/Private/AnimNodes/AnimNode_AimOffsetLookAt.cpp b/Engine/Source/Runtime/AnimGraphRuntime/Private/AnimNodes/AnimNode_AimOffsetLookAt.cpp
--- a/Engine/Source/Runtime/AnimGraphRuntime/Private/AnimNodes/AnimNode_AimOffsetLookAt.cpp
+++ b/Engine/Source/Runtime/AnimGraphRuntime/Private/AnimNodes/AnimNode_AimOffsetLookAt.cpp
@@ -4,6 +4,7 @@
 #include "Animation/AnimInstanceProxy.h"
 #include "AnimationRuntime.h"
 #include "Animation/BlendSpaceBase.h"
+#include "Animation/BlendSpace1D.h"
 #include "Engine/SkeletalMeshSocket.h"
 #include "DrawDebugHelpers.h"
 #include "EngineGlobals.h"
@@ -19,12 +20,46 @@ void FAnimNode_AimOffsetLookAt::Initialize(const FAnimationInitializeContext& Co
 {
 	FAnimNode_BlendSpacePlayer::Initialize(Context);
 	BasePose.Initialize(Context);
+
+	// SkeletalMesh possibly changed, update cached socket info.
+	bCachedSocketInfo = false;
 }
 
-void FAnimNode_AimOffsetLookAt::CacheBones(const FAnimationCacheBonesContext& Context)
+void FAnimNode_AimOffsetLookAt::PreUpdate(const UAnimInstance* InAnimInstance)
 {
-	FAnimNode_BlendSpacePlayer::CacheBones(Context);
-	BasePose.CacheBones(Context);
+	FAnimNode_BlendSpacePlayer::PreUpdate(InAnimInstance);
+
+	if (!bCachedSocketInfo)
+	{
+		bCachedSocketInfo = true;
+		SocketBoneReference.BoneName = NAME_None;
+		if (USkeletalMeshComponent* SkelMeshComp = InAnimInstance->GetSkelMeshComponent())
+		{
+			if (USkeletalMesh* SkelMesh = SkelMeshComp->SkeletalMesh)
+			{
+				if (const USkeletalMeshSocket* Socket = SkelMesh->FindSocket(SourceSocketName))
+				{
+					SocketLocalTransform = Socket->GetSocketLocalTransform();
+					SocketBoneReference.BoneName = Socket->BoneName;
+				}
+
+				if (const USkeletalMeshSocket* Socket = SkelMesh->FindSocket(PivotSocketName))
+				{
+					PivotSocketLocalTransform = Socket->GetSocketLocalTransform();
+					PivotSocketBoneReference.BoneName = Socket->BoneName;
+				}
+			}
+		}
+	}
+
+	if (USkeletalMeshComponent* SkelMeshComp = InAnimInstance->GetSkelMeshComponent())
+	{
+		ComponentToWorld = SkelMeshComp->ComponentToWorld;
+		if (const AActor* Owner = SkelMeshComp->GetOwner())
+		{
+			ActorTransform = Owner->GetTransform();
+		}
+	}
 }
 
 void FAnimNode_AimOffsetLookAt::UpdateAssetPlayer(const FAnimationUpdateContext& Context)
@@ -45,6 +80,18 @@ void FAnimNode_AimOffsetLookAt::UpdateAssetPlayer(const FAnimationUpdateContext&
 	BasePose.Update(Context);
 }
 
+void FAnimNode_AimOffsetLookAt::CacheBones(const FAnimationCacheBonesContext& Context)
+{
+	FAnimNode_BlendSpacePlayer::CacheBones(Context);
+	BasePose.CacheBones(Context);
+
+	if (bCachedSocketInfo)
+	{
+		SocketBoneReference.Initialize(Context.AnimInstanceProxy->GetRequiredBones());
+		PivotSocketBoneReference.Initialize(Context.AnimInstanceProxy->GetRequiredBones());
+	}
+}
+
 void FAnimNode_AimOffsetLookAt::Evaluate(FPoseContext& Context)
 {
 	// Evaluate base pose
@@ -68,76 +115,71 @@ void FAnimNode_AimOffsetLookAt::Evaluate(FPoseContext& Context)
 
 void FAnimNode_AimOffsetLookAt::UpdateFromLookAtTarget(FPoseContext& LocalPoseContext)
 {
+	FVector BlendInput(X, Y, Z);
+
 	const FBoneContainer& RequiredBones = LocalPoseContext.Pose.GetBoneContainer();
-	if (RequiredBones.GetSkeletalMeshAsset())
+	if (bCachedSocketInfo && BlendSpace && SocketBoneReference.IsValid(RequiredBones))
 	{
 		FCSPose<FCompactPose> GlobalPose;
 		GlobalPose.InitPose(LocalPoseContext.Pose);
 
-		USkeletalMeshComponent* Component = LocalPoseContext.AnimInstanceProxy->GetSkelMeshComponent();
-		AActor* Actor = Component ? Component->GetOwner() : nullptr;
+		const FCompactPoseBoneIndex SocketBoneIndex = SocketBoneReference.GetCompactPoseIndex(RequiredBones);
+		const FTransform BoneTransform = GlobalPose.GetComponentSpaceTransform(SocketBoneIndex);
 
-		FVector BlendInput(EForceInit::ForceInitToZero);
-		if (Component && Actor && BlendSpace)
+		FTransform SourceComponentTransform = SocketLocalTransform * BoneTransform;
+		if (PivotSocketBoneReference.IsValid(RequiredBones))
 		{
-			const USkeletalMeshSocket* Socket = RequiredBones.GetSkeletalMeshAsset()->FindSocket(SourceSocketName);
-			if (Socket)
-			{
-				const FTransform SocketLocalTransform = Socket->GetSocketLocalTransform();
-				FBoneReference SocketBoneReference;
-				SocketBoneReference.BoneName = Socket->BoneName;
-				SocketBoneReference.Initialize(RequiredBones);
-				// Only calculate blend space input if we have a valid source socket
-				if (SocketBoneReference.IsValid(RequiredBones))
-				{
-					const FCompactPoseBoneIndex SocketBoneIndex = SocketBoneReference.GetCompactPoseIndex(RequiredBones);
-					const FTransform ActorTransform = Actor->GetTransform();
+			const FCompactPoseBoneIndex PivotSocketBoneIndex = PivotSocketBoneReference.GetCompactPoseIndex(RequiredBones);
+			const FTransform PivotBoneComponentTransform = GlobalPose.GetComponentSpaceTransform(PivotSocketBoneIndex);
+
+			SourceComponentTransform.SetTranslation(PivotBoneComponentTransform.GetTranslation());
+		}
 
-					const FTransform BoneTransform = GlobalPose.GetComponentSpaceTransform(SocketBoneIndex);
-					const FTransform SocketWorldTransform = SocketLocalTransform * BoneTransform * Component->ComponentToWorld;
+		const FTransform SourceWorldTransform = SourceComponentTransform * ComponentToWorld;
 
-					// Convert Target to Actor Space
-					const FTransform TargetWorldTransform(LookAtLocation);
+		// Convert Target to Actor Space
+		const FTransform TargetWorldTransform(LookAtLocation);
 
-					const FVector DirectionToTarget = ActorTransform.InverseTransformVectorNoScale(TargetWorldTransform.GetLocation() - SocketWorldTransform.GetLocation()).GetSafeNormal();
-					const FVector CurrentDirection = ActorTransform.InverseTransformVectorNoScale(SocketWorldTransform.GetUnitAxis(EAxis::X));
+		const FVector DirectionToTarget = ActorTransform.InverseTransformVectorNoScale(TargetWorldTransform.GetLocation() - SourceWorldTransform.GetLocation()).GetSafeNormal();
+		const FVector CurrentDirection = ActorTransform.InverseTransformVectorNoScale(SourceWorldTransform.GetUnitAxis(EAxis::X));
 
-					const FVector AxisX = FVector::ForwardVector;
-					const FVector AxisY = FVector::RightVector;
-					const FVector AxisZ = FVector::UpVector;
+		const FVector AxisX = FVector::ForwardVector;
+		const FVector AxisY = FVector::RightVector;
+		const FVector AxisZ = FVector::UpVector;
 
-					const FVector2D CurrentCoords = FMath::GetAzimuthAndElevation(CurrentDirection, AxisX, AxisY, AxisZ);
-					const FVector2D TargetCoords = FMath::GetAzimuthAndElevation(DirectionToTarget, AxisX, AxisY, AxisZ);
-					BlendInput.X = FRotator::NormalizeAxis(FMath::RadiansToDegrees(TargetCoords.X - CurrentCoords.X));
-					BlendInput.Y = FRotator::NormalizeAxis(FMath::RadiansToDegrees(TargetCoords.Y - CurrentCoords.Y));
+		const FVector2D CurrentCoords = FMath::GetAzimuthAndElevation(CurrentDirection, AxisX, AxisY, AxisZ);
+		const FVector2D TargetCoords = FMath::GetAzimuthAndElevation(DirectionToTarget, AxisX, AxisY, AxisZ);
+		BlendInput.X = FRotator::NormalizeAxis(FMath::RadiansToDegrees(TargetCoords.X - CurrentCoords.X));
+		BlendInput.Y = FRotator::NormalizeAxis(FMath::RadiansToDegrees(TargetCoords.Y - CurrentCoords.Y));
 
 #if ENABLE_DRAW_DEBUG
-					if (CVarAimOffsetLookAtDebug.GetValueOnAnyThread() == 1)
-					{
-						DrawDebugLine(Component->GetWorld(), SocketWorldTransform.GetLocation(), TargetWorldTransform.GetLocation(), FColor::Green);
-						DrawDebugLine(Component->GetWorld(), SocketWorldTransform.GetLocation(), SocketWorldTransform.GetLocation() + SocketWorldTransform.GetUnitAxis(EAxis::X) * (TargetWorldTransform.GetLocation() - SocketWorldTransform.GetLocation()).Size(), FColor::Red);
-						DrawDebugCoordinateSystem(Component->GetWorld(), ActorTransform.GetLocation(), ActorTransform.GetRotation().Rotator(), 100.f);
-
-						FString DebugString = FString::Printf(TEXT("Socket (X:%f, Y:%f), Target (X:%f, Y:%f), Result (X:%f, Y:%f)")
-							, FMath::RadiansToDegrees(CurrentCoords.X)
-							, FMath::RadiansToDegrees(CurrentCoords.Y)
-							, FMath::RadiansToDegrees(TargetCoords.X)
-							, FMath::RadiansToDegrees(TargetCoords.Y)
-							, BlendInput.X
-							, BlendInput.Y);
-						GEngine->AddOnScreenDebugMessage(INDEX_NONE, 0.f, FColor::Red, DebugString, false);
-					}
+		if (CVarAimOffsetLookAtDebug.GetValueOnAnyThread() == 1)
+		{
+			FAnimInstanceProxy* AnimProxy = LocalPoseContext.AnimInstanceProxy;
+			AnimProxy->AnimDrawDebugLine(SourceWorldTransform.GetLocation(), TargetWorldTransform.GetLocation(), FColor::Green);
+			AnimProxy->AnimDrawDebugLine(SourceWorldTransform.GetLocation(), SourceWorldTransform.GetLocation() + SourceWorldTransform.GetUnitAxis(EAxis::X) * (TargetWorldTransform.GetLocation() - SourceWorldTransform.GetLocation()).Size(), FColor::Red);
+			AnimProxy->AnimDrawDebugCoordinateSystem(ActorTransform.GetLocation(), ActorTransform.GetRotation().Rotator(), 100.f);
+
+			FString DebugString = FString::Printf(TEXT("Socket (X:%f, Y:%f), Target (X:%f, Y:%f), Result (X:%f, Y:%f)")
+				, FMath::RadiansToDegrees(CurrentCoords.X)
+				, FMath::RadiansToDegrees(CurrentCoords.Y)
+				, FMath::RadiansToDegrees(TargetCoords.X)
+				, FMath::RadiansToDegrees(TargetCoords.Y)
+				, BlendInput.X
+				, BlendInput.Y);
+			AnimProxy->AnimDrawDebugOnScreenMessage(DebugString, FColor::Red);
+		}
 #endif // ENABLE_DRAW_DEBUG
-				}
-			}
+	}
 
-			// Set X and Y, so ticking next frame is based on correct weights.
-			X = BlendInput.X;
-			Y = BlendInput.Y;
+	// Set X and Y, so ticking next frame is based on correct weights.
+	X = BlendInput.X;
+	Y = BlendInput.Y;
 
-			// Generate BlendSampleDataCache from inputs.
-			BlendSpace->GetSamplesFromBlendInput(BlendInput, BlendSampleDataCache);
-		}
+	// Generate BlendSampleDataCache from inputs.
+	if (BlendSpace)
+	{
+		BlendSpace->GetSamplesFromBlendInput(BlendInput, BlendSampleDataCache);
 	}
 }
 
diff --git a/Engine/Source/Runtime/AnimGraphRuntime/Private/AnimNodes/AnimNode_LayeredBoneBlend.cpp b/Engine/Source/Runtime/AnimGraphRuntime/Private/AnimNodes/AnimNode_LayeredBoneBlend.cpp
--- a/Engine/Source/Runtime/AnimGraphRuntime/Private/AnimNodes/AnimNode_LayeredBoneBlend.cpp
+++ b/Engine/Source/Runtime/AnimGraphRuntime/Private/AnimNodes/AnimNode_LayeredBoneBlend.cpp
@@ -41,24 +41,27 @@ void FAnimNode_LayeredBoneBlend::PostCompile(const class USkeleton* InSkeleton)
 
 void FAnimNode_LayeredBoneBlend::RebuildCacheData(const USkeleton* InSkeleton)
 {
-		// this shouldn't happen outside of editor
-#if !WITH_EDITOR
-	ensureAlways(false);
-#endif 		
-	if (InSkeleton)
+	// make sure it's inished post load yet, this might cause more things to cache in initialize instead of during cooking
+	if (InSkeleton && !(InSkeleton->GetFlags() & RF_NeedPostLoad))
 	{
 		FAnimationRuntime::CreateMaskWeights(PerBoneBlendWeights, LayerSetup, InSkeleton);
 		SkeletonGuid = InSkeleton->GetGuid();
+		VirtualBoneGuid = InSkeleton->GetVirtualBoneGuid();
 	}
 }
 
+bool FAnimNode_LayeredBoneBlend::IsCacheInvalid(const USkeleton* InSkeleton) const
+{
+	return (InSkeleton->GetGuid() != SkeletonGuid || InSkeleton->GetVirtualBoneGuid() != VirtualBoneGuid);
+}
+
 void FAnimNode_LayeredBoneBlend::ReinitializeBoneBlendWeights(const FBoneContainer& RequiredBones, const USkeleton* Skeleton)
 {
-	if (Skeleton->GetGuid() != SkeletonGuid)
+	if (IsCacheInvalid(Skeleton))
 	{
 		RebuildCacheData(Skeleton);
 	}
-
+	
 	// build desired bone weights
 	const TArray<FBoneIndexType>& RequiredBoneIndices = RequiredBones.GetBoneIndicesArray();
 	DesiredBoneBlendWeights.SetNumZeroed(RequiredBoneIndices.Num());
@@ -67,7 +70,7 @@ void FAnimNode_LayeredBoneBlend::ReinitializeBoneBlendWeights(const FBoneContain
 		int32 SkeletonIndex = RequiredBones.GetSkeletonIndex(FCompactPoseBoneIndex(RequiredBoneIndex));
 		DesiredBoneBlendWeights[RequiredBoneIndex] = PerBoneBlendWeights[SkeletonIndex];
 	}
-
+	
 	CurrentBoneBlendWeights.Reset(DesiredBoneBlendWeights.Num());
 	CurrentBoneBlendWeights.AddZeroed(DesiredBoneBlendWeights.Num());
 
@@ -105,12 +108,16 @@ void FAnimNode_LayeredBoneBlend::ReinitializeBoneBlendWeights(const FBoneContain
 void FAnimNode_LayeredBoneBlend::CacheBones(const FAnimationCacheBonesContext& Context) 
 {
 	BasePose.CacheBones(Context);
-	for(int32 ChildIndex=0; ChildIndex<BlendPoses.Num(); ChildIndex++)
+	int32 NumPoses = BlendPoses.Num();
+	for(int32 ChildIndex=0; ChildIndex<NumPoses; ChildIndex++)
 	{
 		BlendPoses[ChildIndex].CacheBones(Context);
 	}
 
-	ReinitializeBoneBlendWeights(Context.AnimInstanceProxy->GetRequiredBones(), Context.AnimInstanceProxy->GetSkeleton());
+	if (NumPoses > 0)
+	{
+		ReinitializeBoneBlendWeights(Context.AnimInstanceProxy->GetRequiredBones(), Context.AnimInstanceProxy->GetSkeleton());
+	}
 }
 
 void FAnimNode_LayeredBoneBlend::Update(const FAnimationUpdateContext& Context)
diff --git a/Engine/Source/Runtime/AnimGraphRuntime/Public/AnimNodes/AnimNode_AimOffsetLookAt.h b/Engine/Source/Runtime/AnimGraphRuntime/Public/AnimNodes/AnimNode_AimOffsetLookAt.h
--- a/Engine/Source/Runtime/AnimGraphRuntime/Public/AnimNodes/AnimNode_AimOffsetLookAt.h
+++ b/Engine/Source/Runtime/AnimGraphRuntime/Public/AnimNodes/AnimNode_AimOffsetLookAt.h
@@ -35,9 +35,33 @@ struct ANIMGRAPHRUNTIME_API FAnimNode_AimOffsetLookAt : public FAnimNode_BlendSp
 	UPROPERTY(EditAnywhere, BlueprintReadWrite, Category = LookAt, meta = (PinHiddenByDefault))
 	FName SourceSocketName;
 
+	UPROPERTY(EditAnywhere, BlueprintReadWrite, Category = LookAt, meta = (PinHiddenByDefault))
+	FName PivotSocketName;
+
 	UPROPERTY(EditAnywhere, BlueprintReadWrite, Category = LookAt, meta = (PinShownByDefault))
 	float Alpha;
 
+	UPROPERTY(Transient)
+	bool bCachedSocketInfo;
+
+	UPROPERTY()
+	FBoneReference SocketBoneReference;
+
+	UPROPERTY()
+	FTransform SocketLocalTransform;
+
+	UPROPERTY()
+	FBoneReference PivotSocketBoneReference;
+
+	UPROPERTY()
+	FTransform PivotSocketLocalTransform;
+
+	UPROPERTY()
+	FTransform ActorTransform;
+
+	UPROPERTY()
+	FTransform ComponentToWorld;
+
 public:
 	FAnimNode_AimOffsetLookAt();
 
@@ -47,6 +71,10 @@ public:
 	virtual void UpdateAssetPlayer(const FAnimationUpdateContext& Context) override;
 	virtual void Evaluate(FPoseContext& Output) override;
 	virtual void GatherDebugData(FNodeDebugData& DebugData) override;
+
+	virtual bool HasPreUpdate() const override { return true; }
+	virtual void PreUpdate(const UAnimInstance* InAnimInstance) override;
+
 	// End of FAnimNode_Base interface
 
 	void UpdateFromLookAtTarget(FPoseContext& LocalPoseContext);
diff --git a/Engine/Source/Runtime/AnimGraphRuntime/Public/AnimNodes/AnimNode_LayeredBoneBlend.h b/Engine/Source/Runtime/AnimGraphRuntime/Public/AnimNodes/AnimNode_LayeredBoneBlend.h
--- a/Engine/Source/Runtime/AnimGraphRuntime/Public/AnimNodes/AnimNode_LayeredBoneBlend.h
+++ b/Engine/Source/Runtime/AnimGraphRuntime/Public/AnimNodes/AnimNode_LayeredBoneBlend.h
@@ -52,6 +52,10 @@ protected:
 	UPROPERTY()
 	FGuid						SkeletonGuid;
 
+	UPROPERTY()
+	FGuid						VirtualBoneGuid;
+
+
 	// transient data to handle weight and target weight
 	// this array changes based on required bones
 	TArray<FPerBoneBlendWeight> DesiredBoneBlendWeights;
@@ -72,7 +76,6 @@ public:
 	virtual void GatherDebugData(FNodeDebugData& DebugData) override;
 	// End of FAnimNode_Base interface
 
-#if WITH_EDITOR
 	void AddPose()
 	{
 		BlendWeights.Add(1.f);
@@ -87,6 +90,7 @@ public:
 		LayerSetup.RemoveAt(PoseIndex);
 	}
 
+#if WITH_EDITOR
 	// ideally you don't like to get to situation where it becomes inconsistent, but this happened, 
 	// and we don't know what caused this. Possibly copy/paste, but I tried copy/paste and that didn't work
 	// so here we add code to fix this up manually in editor, so that they can continue working on it. 
@@ -96,8 +100,11 @@ public:
 	// end FAnimNode_Base interface
 #endif
 
+	/** Reinitialize bone weights */
+	void ReinitializeBoneBlendWeights(const FBoneContainer& RequiredBones, const USkeleton* Skeleton);
+
 private:
 	// Rebuild cache data from the skeleton
 	void RebuildCacheData(const USkeleton* InSkeleton);
-	void ReinitializeBoneBlendWeights(const FBoneContainer& RequiredBones, const USkeleton* Skeleton);
+	bool IsCacheInvalid(const USkeleton* InSkeleton) const;
 };
diff --git a/Engine/Source/Runtime/ClothingSystemRuntime/Public/Assets/ClothingAsset.h b/Engine/Source/Runtime/ClothingSystemRuntime/Public/Assets/ClothingAsset.h
--- a/Engine/Source/Runtime/ClothingSystemRuntime/Public/Assets/ClothingAsset.h
+++ b/Engine/Source/Runtime/ClothingSystemRuntime/Public/Assets/ClothingAsset.h
@@ -110,7 +110,10 @@ struct FClothLODData
 		UScriptStruct* Struct = FClothLODData::StaticStruct();
 
 		// Serialize normal tagged data
-		Struct->SerializeTaggedProperties(Ar, (uint8*)this, Struct, nullptr);
+		if (!Ar.IsCountingMemory())
+		{
+			Struct->SerializeTaggedProperties(Ar, (uint8*)this, Struct, nullptr);
+		}
 
 		// Serialize the mesh to mesh data (not a USTRUCT)
 		Ar	<< TransitionUpSkinData
diff --git a/Engine/Source/Runtime/Core/Private/Windows/WindowsApplication.cpp b/Engine/Source/Runtime/Core/Private/Windows/WindowsApplication.cpp
--- a/Engine/Source/Runtime/Core/Private/Windows/WindowsApplication.cpp
+++ b/Engine/Source/Runtime/Core/Private/Windows/WindowsApplication.cpp
@@ -2442,7 +2442,7 @@ uint32 FWindowsApplication::GetFirstFreeTouchIndex()
 
 void FTaskbarList::Initialize()
 {
-	if (FWindowsPlatformMisc::CoInitialize())
+	if (LIKELY(FApp::CanEverRender()) && FWindowsPlatformMisc::CoInitialize())
 	{
 		if (CoCreateInstance(CLSID_TaskbarList, NULL, CLSCTX_INPROC_SERVER, IID_ITaskbarList3, (void **)&TaskBarList3) != S_OK)
 		{
@@ -2459,7 +2459,7 @@ FTaskbarList::FTaskbarList()
 
 FTaskbarList::~FTaskbarList()
 {
-	if (FWindowsPlatformMisc::CoInitialize() && TaskBarList3)
+	if (TaskBarList3 && FWindowsPlatformMisc::CoInitialize())
 	{
 		TaskBarList3->Release();
 	}
diff --git a/Engine/Source/Runtime/Core/Private/Windows/WindowsTextInputMethodSystem.cpp b/Engine/Source/Runtime/Core/Private/Windows/WindowsTextInputMethodSystem.cpp
--- a/Engine/Source/Runtime/Core/Private/Windows/WindowsTextInputMethodSystem.cpp
+++ b/Engine/Source/Runtime/Core/Private/Windows/WindowsTextInputMethodSystem.cpp
@@ -722,62 +722,68 @@ void FWindowsTextInputMethodSystem::ActivateContext(const TSharedRef<ITextInputM
 	const TSharedPtr<FGenericWindow> GenericWindow = Context->GetWindow();
 	InternalContext.WindowHandle = GenericWindow.IsValid() ? reinterpret_cast<HWND>(GenericWindow->GetOSWindowHandle()) : nullptr;
 
-	// IMM Implementation
-	InternalContext.IMMContext.IsComposing = false;
-	InternalContext.IMMContext.IsDeactivating = false;
-	::ImmAssociateContext(InternalContext.WindowHandle, IMMContextId);
-
-	// TSF Implementation
-	TComPtr<FTextStoreACP>& TextStore = InternalContext.TSFContext;
-	ITfDocumentMgr* Unused;
-	Result = TSFThreadManager->AssociateFocus(InternalContext.WindowHandle, TextStore->TSFDocumentManager, &Unused);
-	if(FAILED(Result))
+	if (InternalContext.WindowHandle)
 	{
-		TCHAR ErrorMsg[1024];
-		FPlatformMisc::GetSystemErrorMessage(ErrorMsg, 1024, Result);
-		UE_LOG(LogWindowsTextInputMethodSystem, Error, TEXT("Activating a context failed while setting focus on a TSF document manager. %s (0x%08x)"), ErrorMsg, Result);
-	}
+		// IMM Implementation
+		InternalContext.IMMContext.IsComposing = false;
+		InternalContext.IMMContext.IsDeactivating = false;
+		::ImmAssociateContext(InternalContext.WindowHandle, IMMContextId);
+
+		// TSF Implementation
+		TComPtr<FTextStoreACP>& TextStore = InternalContext.TSFContext;
+		ITfDocumentMgr* Unused;
+		Result = TSFThreadManager->AssociateFocus(InternalContext.WindowHandle, TextStore->TSFDocumentManager, &Unused);
+		if (FAILED(Result))
+		{
+			TCHAR ErrorMsg[1024];
+			FPlatformMisc::GetSystemErrorMessage(ErrorMsg, 1024, Result);
+			UE_LOG(LogWindowsTextInputMethodSystem, Error, TEXT("Activating a context failed while setting focus on a TSF document manager. %s (0x%08x)"), ErrorMsg, Result);
+		}
 
-	UE_LOG(LogWindowsTextInputMethodSystem, Verbose, TEXT("Activated context %p!"), &(Context.Get()));
+		UE_LOG(LogWindowsTextInputMethodSystem, Verbose, TEXT("Activated context %p!"), &(Context.Get()));
+	}
 }
 
 void FWindowsTextInputMethodSystem::DeactivateContext(const TSharedRef<ITextInputMethodContext>& Context)
 {
-	UE_LOG(LogWindowsTextInputMethodSystem, Verbose, TEXT("Deactivating context %p..."), &(Context.Get()));
-	
 	FInternalContext& InternalContext = ContextToInternalContextMap[Context];
 
-	// IMM Implementation
-	HIMC IMMContext = ::ImmGetContext(InternalContext.WindowHandle);
-	InternalContext.IMMContext.IsDeactivating = true;
-	// Request the composition is completed to ensure that the composition input UI is closed, and that a WM_IME_ENDCOMPOSITION message is sent
-	::ImmNotifyIME(IMMContext, NI_COMPOSITIONSTR, CPS_COMPLETE, 0);
-	::ImmReleaseContext(InternalContext.WindowHandle, IMMContext);
-
-	// Ensure all known windows are associated with a disabled IME
-	ClearStaleWindowHandles();
-	for(const TWeakPtr<FGenericWindow>& Window : KnownWindows)
+	if (InternalContext.WindowHandle)
 	{
-		TSharedPtr<FGenericWindow> WindowPtr = Window.Pin();
-		if(WindowPtr.IsValid())
+		UE_LOG(LogWindowsTextInputMethodSystem, Verbose, TEXT("Deactivating context %p..."), &(Context.Get()));
+
+		// IMM Implementation
+		HIMC IMMContext = ::ImmGetContext(InternalContext.WindowHandle);
+		InternalContext.IMMContext.IsDeactivating = true;
+		// Request the composition is completed to ensure that the composition input UI is closed, and that a WM_IME_ENDCOMPOSITION message is sent
+		::ImmNotifyIME(IMMContext, NI_COMPOSITIONSTR, CPS_COMPLETE, 0);
+		::ImmReleaseContext(InternalContext.WindowHandle, IMMContext);
+
+		// Ensure all known windows are associated with a disabled IME
+		ClearStaleWindowHandles();
+		for (const TWeakPtr<FGenericWindow>& Window : KnownWindows)
 		{
-			const HWND Hwnd = reinterpret_cast<HWND>(WindowPtr->GetOSWindowHandle());
-			if(Hwnd)
+			TSharedPtr<FGenericWindow> WindowPtr = Window.Pin();
+			if (WindowPtr.IsValid())
 			{
-				//  TSF Implementation
-				ITfDocumentMgr* Unused;
-				TSFThreadManager->AssociateFocus(Hwnd, TSFDisabledDocumentManager, &Unused);
+				const HWND Hwnd = reinterpret_cast<HWND>(WindowPtr->GetOSWindowHandle());
+				if (Hwnd)
+				{
+					//  TSF Implementation
+					ITfDocumentMgr* Unused;
+					TSFThreadManager->AssociateFocus(Hwnd, TSFDisabledDocumentManager, &Unused);
 
-				// IMM Implementation
-				::ImmAssociateContext(InternalContext.WindowHandle, nullptr);
+					// IMM Implementation
+					::ImmAssociateContext(InternalContext.WindowHandle, nullptr);
+				}
 			}
+
+			UE_LOG(LogWindowsTextInputMethodSystem, Verbose, TEXT("Deactivated context %p!"), &(Context.Get()));
 		}
 	}
 
 	// General Implementation
 	ActiveContext = nullptr;
-
-	UE_LOG(LogWindowsTextInputMethodSystem, Verbose, TEXT("Deactivated context %p!"), &(Context.Get()));
 }
 
 bool FWindowsTextInputMethodSystem::IsActiveContext(const TSharedRef<ITextInputMethodContext>& Context) const
diff --git a/Engine/Source/Runtime/Core/Public/HAL/FeedbackContextAnsi.h b/Engine/Source/Runtime/Core/Public/HAL/FeedbackContextAnsi.h
--- a/Engine/Source/Runtime/Core/Public/HAL/FeedbackContextAnsi.h
+++ b/Engine/Source/Runtime/Core/Public/HAL/FeedbackContextAnsi.h
@@ -49,7 +49,12 @@ public:
 	{}
 	void Serialize( const TCHAR* V, ELogVerbosity::Type Verbosity, const class FName& Category ) override
 	{
-		if (Verbosity == ELogVerbosity::Error || Verbosity == ELogVerbosity::Warning || Verbosity == ELogVerbosity::Display)
+		// When -stdout is specified then FOutputDeviceStdOutput will be installed and pipe logging to stdout. 
+		// If so don't use LocalPrint or else duplicate messages will be written to stdout
+		static bool bUsingStdOut = FParse::Param(FCommandLine::Get(), TEXT("stdout"));
+
+		if (bUsingStdOut == false &&
+			(Verbosity == ELogVerbosity::Error || Verbosity == ELogVerbosity::Warning || Verbosity == ELogVerbosity::Display))
 		{
 			if( TreatWarningsAsErrors && Verbosity==ELogVerbosity::Warning )
 			{
diff --git a/Engine/Source/Runtime/Core/Public/HAL/MallocStomp.h b/Engine/Source/Runtime/Core/Public/HAL/MallocStomp.h
--- a/Engine/Source/Runtime/Core/Public/HAL/MallocStomp.h
+++ b/Engine/Source/Runtime/Core/Public/HAL/MallocStomp.h
@@ -20,6 +20,11 @@
 #endif
 
 #if USE_MALLOC_STOMP
+
+#if PLATFORM_WINDOWS
+	#include "WindowsHWrapper.h"
+#endif
+
 /**
  * Stomp memory allocator. It helps find the following errors:
  * - Read or writes off the end of an allocation.
diff --git a/Engine/Source/Runtime/Core/Public/Misc/AutomationTest.h b/Engine/Source/Runtime/Core/Public/Misc/AutomationTest.h
--- a/Engine/Source/Runtime/Core/Public/Misc/AutomationTest.h
+++ b/Engine/Source/Runtime/Core/Public/Misc/AutomationTest.h
@@ -32,25 +32,40 @@ namespace EAutomationTestFlags
 	{
 		//~ Application context required for the test - not specifying means it will be valid for any context
 		// Test is suitable for running within the editor
-		EditorContext				= 0x00000001,
+		EditorContext = 0x00000001,
 		// Test is suitable for running within the client
-		ClientContext				= 0x00000002,
+		ClientContext = 0x00000002,
 		// Test is suitable for running within the server
-		ServerContext				= 0x00000004,
+		ServerContext = 0x00000004,
 		// Test is suitable for running within a commandlet
-		CommandletContext			= 0x00000008,
-		ApplicationContextMask		= EditorContext | ClientContext | ServerContext | CommandletContext,
+		CommandletContext = 0x00000008,
+		ApplicationContextMask = EditorContext | ClientContext | ServerContext | CommandletContext,
 
 		//~ Features required for the test - not specifying means it is valid for any feature combination
 		// Test requires a non-null RHI to run correctly
-		NonNullRHI					= 0x00000100,
+		NonNullRHI = 0x00000100,
 		// Test requires a user instigated session
-		RequiresUser				= 0x00000200,
-		FeatureMask					= NonNullRHI | RequiresUser,
+		RequiresUser = 0x00000200,
+		FeatureMask = NonNullRHI | RequiresUser,
 
 		//~ One-off flag to allow for fast disabling of tests without commenting code out
 		// Temp disabled and never returns for a filter
-		Disabled					= 0x00010000,
+		Disabled = 0x00010000,
+
+		//~ Priority of the test
+		// The highest priority possible. Showstopper/blocker.
+		CriticalPriority			= 0x00100000,
+		// High priority. Major feature functionality etc. 
+		HighPriority				= 0x00200000,
+		// Mask for High on SetMinimumPriority
+		HighPriorityAndAbove		= CriticalPriority | HighPriority,
+		// Medium Priority. Minor feature functionality, major generic content issues.
+		MediumPriority				= 0x00400000,
+		// Mask for Medium on SetMinimumPriority
+		MediumPriorityAndAbove		= CriticalPriority | HighPriority | MediumPriority,
+		// Low Priority. Minor content bugs. String errors. Etc.
+		LowPriority					= 0x00800000,
+		PriorityMask = CriticalPriority | HighPriority | MediumPriority | LowPriority,
 
 		//~ Speed of the test
 		//Super Fast Filter
diff --git a/Engine/Source/Runtime/CoreUObject/Private/Internationalization/PackageLocalizationCache.cpp b/Engine/Source/Runtime/CoreUObject/Private/Internationalization/PackageLocalizationCache.cpp
--- a/Engine/Source/Runtime/CoreUObject/Private/Internationalization/PackageLocalizationCache.cpp
+++ b/Engine/Source/Runtime/CoreUObject/Private/Internationalization/PackageLocalizationCache.cpp
@@ -22,6 +22,12 @@ void FPackageLocalizationCultureCache::ConditionalUpdateCache()
 
 void FPackageLocalizationCultureCache::ConditionalUpdateCache_NoLock()
 {
+	if (!IsInGameThread())
+	{
+		UE_CLOG(PendingSourceRootPathsToSearch.Num() > 0, LogPackageLocalizationCache, Warning, TEXT("Skipping the cache update for %d pending package path(s) due to a cache request from a non-game thread. Some localized packages may be missed for this query."), PendingSourceRootPathsToSearch.Num());
+		return;
+	}
+
 	if (PendingSourceRootPathsToSearch.Num() == 0)
 	{
 		return;
@@ -160,7 +166,9 @@ FName FPackageLocalizationCultureCache::FindLocalizedPackageName(const FName InS
 
 FPackageLocalizationCache::FPackageLocalizationCache()
 {
-	HandleCultureChanged();
+	const FString CurrentCultureName = FInternationalization::Get().GetCurrentCulture()->GetName();
+	CurrentCultureCache = FindOrAddCacheForCulture_NoLock(CurrentCultureName);
+
 	FInternationalization::Get().OnCultureChanged().AddRaw(this, &FPackageLocalizationCache::HandleCultureChanged);
 
 	FPackageName::OnContentPathMounted().AddRaw(this, &FPackageLocalizationCache::HandleContentPathMounted);
@@ -257,4 +265,11 @@ void FPackageLocalizationCache::HandleCultureChanged()
 
 	const FString CurrentCultureName = FInternationalization::Get().GetCurrentCulture()->GetName();
 	CurrentCultureCache = FindOrAddCacheForCulture_NoLock(CurrentCultureName);
+
+	if (CurrentCultureCache.IsValid())
+	{
+		// We expect culture changes to happen on the game thread, so update the cache now while it is likely safe to do so
+		// (ConditionalUpdateCache will internally check that this is currently the game thread before allowing the update)
+		CurrentCultureCache->ConditionalUpdateCache();
+	}
 }
diff --git a/Engine/Source/Runtime/CoreUObject/Private/UObject/UObjectGlobals.cpp b/Engine/Source/Runtime/CoreUObject/Private/UObject/UObjectGlobals.cpp
--- a/Engine/Source/Runtime/CoreUObject/Private/UObject/UObjectGlobals.cpp
+++ b/Engine/Source/Runtime/CoreUObject/Private/UObject/UObjectGlobals.cpp
@@ -1035,10 +1035,16 @@ public:
 	{
 		Package->LinkerLoad = this;
 
-		while (CreateLoader(TFunction<void()>([]() {})) == FLinkerLoad::LINKER_TimedOut)
+		/*while (CreateLoader(TFunction<void()>([]() {})) == FLinkerLoad::LINKER_TimedOut)
 		{
-		}
+		}*/
+
+
+		
 
+		while ( Tick(0.0, false, false) == FLinkerLoad::LINKER_TimedOut ) 
+		{ 
+		}
 
 		FArchive* OtherFile = IFileManager::Get().CreateFileReader(DiffFilename);
 		FDiffFileArchive* DiffArchive = new FDiffFileArchive(Loader, OtherFile);
diff --git a/Engine/Source/Runtime/CoreUObject/Public/Internationalization/PackageLocalizationCache.h b/Engine/Source/Runtime/CoreUObject/Public/Internationalization/PackageLocalizationCache.h
--- a/Engine/Source/Runtime/CoreUObject/Public/Internationalization/PackageLocalizationCache.h
+++ b/Engine/Source/Runtime/CoreUObject/Public/Internationalization/PackageLocalizationCache.h
@@ -142,7 +142,7 @@ protected:
 	void HandleContentPathDismounted(const FString& InAssetPath, const FString& InFilesystemPath);
 
 	/**
-	 * Callback handler for when a new content path is mounted.
+	 * Callback handler for when the active culture is changed.
 	 */
 	void HandleCultureChanged();
 
diff --git a/Engine/Source/Runtime/Engine/Classes/Engine/RendererSettings.h b/Engine/Source/Runtime/Engine/Classes/Engine/RendererSettings.h
--- a/Engine/Source/Runtime/Engine/Classes/Engine/RendererSettings.h
+++ b/Engine/Source/Runtime/Engine/Classes/Engine/RendererSettings.h
@@ -501,6 +501,11 @@ class ENGINE_API URendererOverrideSettings : public UDeveloperSettings
 		ConfigRestartRequired = true))
 		uint32 bSupportAllShaderPermutations : 1;
 
+	UPROPERTY(config, EditAnywhere, Category = Miscellaneous, meta = (
+		ConsoleVariable = "r.SkinCache.ForceRecomputeTangents", DisplayName = "Force all skinned meshes to recompute tangents (also forces Compute SkinCache)",
+		ConfigRestartRequired = true))
+		uint32 bForceRecomputeTangents : 1;
+
 public:
 
 	//~ Begin UObject Interface
diff --git a/Engine/Source/Runtime/Engine/Classes/GameFramework/GameUserSettings.h b/Engine/Source/Runtime/Engine/Classes/GameFramework/GameUserSettings.h
--- a/Engine/Source/Runtime/Engine/Classes/GameFramework/GameUserSettings.h
+++ b/Engine/Source/Runtime/Engine/Classes/GameFramework/GameUserSettings.h
@@ -245,10 +245,14 @@ public:
 	UFUNCTION(BlueprintCallable, Category=Settings)
 	virtual void SetToDefaults();
 
-	/** Gets the desired resolution quality based on DesiredScreenHeight and the current screen resolution */
+	/** Gets the desired resolution quality based on DesiredScreenWidth/Height and the current screen resolution */
 	UFUNCTION(BlueprintCallable, Category=Settings)
 	virtual float GetDefaultResolutionScale();
 
+	/** Gets the recommended resolution quality based on LastRecommendedScreenWidth/Height and the current screen resolution */
+	UFUNCTION(BlueprintCallable, Category = Settings)
+	virtual float GetRecommendedResolutionScale();
+
 	/** Loads the resolution settings before is object is available */
 	static void PreloadResolutionSettings();
 
@@ -283,13 +287,20 @@ public:
 	virtual void ApplyHardwareBenchmarkResults();
 
 	/** Whether the curently running system supports HDR display output */
-	UFUNCTION(BlueprintCallable, Category=Settings)
-	bool SupportsHDRDisplayOutput();
+	UFUNCTION(BlueprintCallable, Category=Settings, meta = (DisplayName = "Supports HDR Display Output"))
+	virtual bool SupportsHDRDisplayOutput() const;
 
 	/** Enables or disables HDR display output. Can be called again to change the desired nit level */
-	UFUNCTION(BlueprintCallable, Category=Settings)
+	UFUNCTION(BlueprintCallable, Category=Settings, meta=(DisplayName = "Enable HDR Display Output"))
 	void EnableHDRDisplayOutput(bool bEnable, int32 DisplayNits = 1000);
 
+	/** Returns 0 if HDR isn't supported or is turned off */
+	UFUNCTION(BlueprintCallable, Category = Settings, meta = (DisplayName = "Get Current HDR Display Nits"))
+	int32 GetCurrentHDRDisplayNits() const;
+
+	UFUNCTION(BlueprintCallable, Category = Settings, meta = (DisplayName = "Is HDR Enabled"))
+	bool IsHDREnabled() const;
+
 	/** Whether to use VSync or not. (public to allow UI to connect to it) */
 	UPROPERTY(config)
 	bool bUseVSync;
diff --git a/Engine/Source/Runtime/Engine/Private/Animation/AnimInstance.cpp b/Engine/Source/Runtime/Engine/Private/Animation/AnimInstance.cpp
--- a/Engine/Source/Runtime/Engine/Private/Animation/AnimInstance.cpp
+++ b/Engine/Source/Runtime/Engine/Private/Animation/AnimInstance.cpp
@@ -1132,7 +1132,8 @@ void UAnimInstance::UpdateCurves(const FBlendedHeapCurve& InCurve)
 		const TArray<SmartName::UID_Type>& UIDList = *InCurve.UIDList;
 		for (int32 CurveId = 0; CurveId < InCurve.UIDList->Num(); ++CurveId)
 		{
-			if (ensureAlwaysMsgf(InCurve.Elements.IsValidIndex(CurveId), TEXT("%s Animation Instance contains out of bound UIDList."), *GetClass()->GetName()))
+			if (ensureAlwaysMsgf(InCurve.Elements.IsValidIndex(CurveId), TEXT("%s Animation Instance contains out of bound UIDList."), *GetClass()->GetName())
+				&& InCurve.Elements[CurveId].IsValid())
 			{
 				// had to add to another data type
 				AddCurveValue(UIDList[CurveId], InCurve.Elements[CurveId].Value);
