Adjusting VREditor analytics events, including adding heartbeat

diff --git a/Engine/Source/Editor/UnrealEd/Private/UnrealEdMisc.cpp b/Engine/Source/Editor/UnrealEd/Private/UnrealEdMisc.cpp
--- a/Engine/Source/Editor/UnrealEd/Private/UnrealEdMisc.cpp
+++ b/Engine/Source/Editor/UnrealEd/Private/UnrealEdMisc.cpp
@@ -77,6 +77,7 @@
 #include "DesktopPlatformModule.h"
 #include "UserActivityTracking.h"
 #include "Widgets/Docking/SDockTab.h"
+#include "IVREditorModule.h"
 
 #define USE_UNIT_TESTS 0
 
@@ -586,6 +587,29 @@ void FUnrealEdMisc::InitEngineAnalytics()
 	}
 }
 
+/*
+* @EventName Editor.Usage.Heartbeat
+*
+* @Trigger Every minute of non-idle time in the editor
+*
+* @Type Dynamic
+*
+* @EventParam Idle (bool) Whether the user is idle
+* @EventParam AverageFrameTime (float) Average frame time
+* @EventParam AverageGameThreadTime (float) Average game thread time
+* @EventParam AverageRenderThreadTime (float) Average render thread time
+* @EventParam AverageGPUFrameTime (float) Average GPU frame time
+* @EventParam IsVanilla (bool) Whether the editor is vanilla launcher install with no marketplace plugins
+* @EventParam IntervalSec (int32) The time since the last heartbeat
+* @EventParam IsDebugger (bool) Whether the debugger is currently present
+* @EventParam WasDebuggerPresent (bool) Whether the debugger was present previously
+* @EventParam IsInVRMode (bool) If the current heartbeat occurred while VR mode was active
+*
+* @Source Editor
+*
+* @Owner Matt.Kuhlenschmidt
+*
+*/
 void FUnrealEdMisc::EditorAnalyticsHeartbeat()
 {
 	// Don't attempt to send the heartbeat if analytics isn't available
@@ -602,7 +626,7 @@ void FUnrealEdMisc::EditorAnalyticsHeartbeat()
 	{
 		bWasDebuggerPresent = bIsDebuggerPresent;
 	}
-	
+	const bool bInVRMode = IVREditorModule::Get().IsVREditorModeActive();
 	double LastInteractionTime = FSlateApplication::Get().GetLastUserInteractionTime();
 	
 	// Did the user interact since the last heartbeat
@@ -624,6 +648,7 @@ void FUnrealEdMisc::EditorAnalyticsHeartbeat()
 	Attributes.Add(FAnalyticsEventAttribute(TEXT("IntervalSec"), UnrealEdMiscDefs::HeartbeatIntervalSeconds));
 	Attributes.Add(FAnalyticsEventAttribute(TEXT("IsDebugger"), bIsDebuggerPresent));
 	Attributes.Add(FAnalyticsEventAttribute(TEXT("WasDebuggerPresent"), bWasDebuggerPresent));
+	Attributes.Add(FAnalyticsEventAttribute(TEXT("IsInVRMode"), bInVRMode));
 	FEngineAnalytics::GetProvider().RecordEvent(TEXT("Editor.Usage.Heartbeat"), Attributes);
 	
 	LastHeartbeatTime = FPlatformTime::Seconds();
diff --git a/Engine/Source/Editor/VREditor/VREditorMode.cpp b/Engine/Source/Editor/VREditor/VREditorMode.cpp
--- a/Engine/Source/Editor/VREditor/VREditorMode.cpp
+++ b/Engine/Source/Editor/VREditor/VREditorMode.cpp
@@ -109,7 +109,7 @@ void UVREditorMode::Init()
 	// Take note of VREditor activation
 	if( FEngineAnalytics::IsAvailable() )
 	{
-		FEngineAnalytics::GetProvider().RecordEvent( TEXT( "Editor.Usage.EnterVREditorMode" ) );
+		FEngineAnalytics::GetProvider().RecordEvent( TEXT( "Editor.Usage.InitVREditorMode" ) );
 	}
 
 	// Setting up colors
@@ -139,6 +139,20 @@ void UVREditorMode::Init()
 	bIsFullyInitialized = true;
 }
 
+/*
+* @EventName Editor.Usage.EnterVRMode
+*
+* @Trigger Entering VR editing mode
+*
+* @Type Static
+*
+* @EventParam HMDDevice (string) The name of the HMD Device type
+*
+* @Source Editor
+*
+* @Owner Lauren.Ridge
+*
+*/
 void UVREditorMode::Shutdown()
 {
 	bIsFullyInitialized = false;
@@ -207,6 +221,15 @@ void UVREditorMode::Enter()
 			WorldInteraction->SetTransformGizmoScale(GetDefault<UVRModeSettings>()->GizmoScale);
 			WorldInteraction->SetShouldSuppressExistingCursor(true);
 			WorldInteraction->SetInVR(true);
+
+			// Take note of VREditor entering (only if actually in VR)
+			if (FEngineAnalytics::IsAvailable())
+			{
+				TArray< FAnalyticsEventAttribute > Attributes;
+				FString HMDName = GEditor->HMDDevice->GetDeviceName().ToString();
+				Attributes.Add(FAnalyticsEventAttribute(TEXT("HMDDevice"), HMDName));
+				FEngineAnalytics::GetProvider().RecordEvent(TEXT("Editor.Usage.EnterVRMode"), Attributes);
+			}
 		}
 	}
 
@@ -314,6 +337,12 @@ void UVREditorMode::Exit(const bool bShouldDisableStereo)
 				// Restore gizmo size
 				WorldInteraction->SetTransformGizmoScale( SavedEditorState.TransformGizmoScale );
 				WorldInteraction->SetShouldSuppressExistingCursor(false);
+
+				// Take note of VREditor exiting (only if actually in VR)
+				if (FEngineAnalytics::IsAvailable())
+				{
+					FEngineAnalytics::GetProvider().RecordEvent(TEXT("Editor.Usage.ExitVRMode"));
+				}
 			}
 
 			CloseViewport( bShouldDisableStereo );
