Fixing incremental linux build errors

diff --git a/Engine/Source/Editor/UnrealEd/Public/EditorWorldExtension.h b/Engine/Source/Editor/UnrealEd/Public/EditorWorldExtension.h
--- a/Engine/Source/Editor/UnrealEd/Public/EditorWorldExtension.h
+++ b/Engine/Source/Editor/UnrealEd/Public/EditorWorldExtension.h
@@ -5,6 +5,7 @@
 #include "CoreMinimal.h"
 #include "InputCoreTypes.h"
 #include "Templates/SubclassOf.h"
+#include "Templates/Casts.h"
 #include "Engine/EngineBaseTypes.h"
 #include "UObject/GCObject.h"
 #include "EditorWorldExtension.generated.h"
diff --git a/Engine/Source/Editor/VREditor/VREditorActions.h b/Engine/Source/Editor/VREditor/VREditorActions.h
--- a/Engine/Source/Editor/VREditor/VREditorActions.h
+++ b/Engine/Source/Editor/VREditor/VREditorActions.h
@@ -5,10 +5,10 @@
 #include "SlateTypes.h"
 #include "InputCoreTypes.h"
 #include "Editor.h"
+#include "UnrealWidget.h"
 
 // Forward declarations 
 class UVREditorMode;
-enum ECoordSystem;
 enum class EGizmoHandleTypes : uint8;
 
 /**
diff --git a/Engine/Source/Editor/ViewportInteraction/Public/ViewportTransformable.h b/Engine/Source/Editor/ViewportInteraction/Public/ViewportTransformable.h
--- a/Engine/Source/Editor/ViewportInteraction/Public/ViewportTransformable.h
+++ b/Engine/Source/Editor/ViewportInteraction/Public/ViewportTransformable.h
@@ -18,6 +18,10 @@ public:
 	{
 	}
 
+	virtual ~FViewportTransformable()
+	{
+	}
+
 	/** Gets the current transform of this object */
 	virtual const FTransform GetTransform() const = 0;
 
