Fixing cooking error - don't create EditorWorldExtension in EditorModeManager if one isn't found

diff --git a/Engine/Source/Editor/UnrealEd/Private/EditorModeManager.cpp b/Engine/Source/Editor/UnrealEd/Private/EditorModeManager.cpp
--- a/Engine/Source/Editor/UnrealEd/Private/EditorModeManager.cpp
+++ b/Engine/Source/Editor/UnrealEd/Private/EditorModeManager.cpp
@@ -223,13 +223,17 @@ void FEditorModeTools::SetPivotLocation( const FVector& Location, const bool bIn
 ECoordSystem FEditorModeTools::GetCoordSystem(bool bGetRawValue)
 {
 	bool bAligningToActors = false;
-	if (GEditor->GetEditorWorldExtensionsManager() != nullptr)
+	if (GEditor->GetEditorWorldExtensionsManager() != nullptr
+		&& GetWorld() != nullptr)
 	{
-
-		UViewportWorldInteraction* ViewportWorldInteraction = Cast<UViewportWorldInteraction>(GEditor->GetEditorWorldExtensionsManager()->GetEditorWorldExtensions(GetWorld())->FindExtension(UViewportWorldInteraction::StaticClass()));
-		if (ViewportWorldInteraction != nullptr && ViewportWorldInteraction->AreAligningToActors() == true)
+		UEditorWorldExtensionCollection* WorldExtensionCollection = GEditor->GetEditorWorldExtensionsManager()->GetEditorWorldExtensions(GetWorld(), false);
+		if (WorldExtensionCollection != nullptr)
 		{
-			bAligningToActors = true;
+			UViewportWorldInteraction* ViewportWorldInteraction = Cast<UViewportWorldInteraction>(WorldExtensionCollection->FindExtension(UViewportWorldInteraction::StaticClass()));
+			if (ViewportWorldInteraction != nullptr && ViewportWorldInteraction->AreAligningToActors() == true)
+			{
+				bAligningToActors = true;
+			}
 		}
 	}
 	if (!bGetRawValue && 
@@ -246,17 +250,22 @@ ECoordSystem FEditorModeTools::GetCoordSystem(bool bGetRawValue)
 void FEditorModeTools::SetCoordSystem(ECoordSystem NewCoordSystem)
 {
 	// If we are trying to enter world space but are aligning to actors, turn off aligning to actors
-	if (GEditor->GetEditorWorldExtensionsManager() != nullptr &&
-		NewCoordSystem == COORD_World)
+	if (GEditor->GetEditorWorldExtensionsManager() != nullptr
+		&& GetWorld() != nullptr
+		&& NewCoordSystem == COORD_World)
 	{
-		UViewportWorldInteraction* ViewportWorldInteraction = Cast<UViewportWorldInteraction>(GEditor->GetEditorWorldExtensionsManager()->GetEditorWorldExtensions(GetWorld())->FindExtension(UViewportWorldInteraction::StaticClass()));
-		if (ViewportWorldInteraction != nullptr && ViewportWorldInteraction->AreAligningToActors() == true)
+		UEditorWorldExtensionCollection* WorldExtensionCollection = GEditor->GetEditorWorldExtensionsManager()->GetEditorWorldExtensions(GetWorld(), false);
+		if (WorldExtensionCollection != nullptr)
 		{
-			if (ViewportWorldInteraction->HasCandidatesSelected())
+			UViewportWorldInteraction* ViewportWorldInteraction = Cast<UViewportWorldInteraction>(WorldExtensionCollection->FindExtension(UViewportWorldInteraction::StaticClass()));
+			if (ViewportWorldInteraction != nullptr && ViewportWorldInteraction->AreAligningToActors() == true)
 			{
-				ViewportWorldInteraction->SetSelectionAsCandidates();
+				if (ViewportWorldInteraction->HasCandidatesSelected())
+				{
+					ViewportWorldInteraction->SetSelectionAsCandidates();
+				}
+				GUnrealEd->Exec(GetWorld(), TEXT("VI.EnableGuides 0"));
 			}
-			GUnrealEd->Exec(GetWorld(), TEXT("VI.EnableGuides 0"));
 		}
 	}
 	CoordSystem = NewCoordSystem;
