Copying //UE4/Orion-Staging to //UE4/Main (Source: //Orion/Dev-General @ 3358916)

Change 3358916 on 2017/03/22 by [REDACTED_PERSON]

	Merging //Orion/Main to Dev-General (//Orion/Dev-General)
	#!tests #!rb na

Change 3357395 on 2017/03/21 by [REDACTED_PERSON]

	Added some more custom stats to the cooker.
	Only cook the english cook culture when we are running local builds.
	#!rb Trivial
	#!test Iterative shared cooked builds paragon

Change 3357377 on 2017/03/21 by [REDACTED_PERSON]

	Added support for packages which fail to load to the package dependency info module
	#!rb Trivial
	#!test Cook paragon

Change 3356838 on 2017/03/21 by [REDACTED_PERSON]

	Merging //Orion/Main to Dev-General (//Orion/Dev-General)
	#!3rb #!tests na

Change 3355306 on 2017/03/20 by [REDACTED_PERSON]

	Switched PackageDependencyInfo to using Guid instead of entire package hash when generating dependency info.
	Stopped cooker from collecting garbage while in the editor.
	Iterative cooks don't resolve string asset references for startup packages.
	#!rb Trivial
	#!test Shared precooked build paragon

Change 3354527 on 2017/03/20 by [REDACTED_PERSON]

	AnalyticsProvider::SetUserID will now flush any pending events before changing the ID. #!jira AN-1660
	#!fyi josh.markiewicz,david.nikdel
	#!rb josh.markiewicz
	#!tests ran client connected to Solo vs. AI server

Change 3353852 on 2017/03/20 by [REDACTED_PERSON]

	Speculative fix for clothing crashes using Mambo. It was possible that the skeletal mesh component could have triggered deletion or creation of simulation state objects while the simulation was in flight on another thread, added tracking and waiting for outstanding tasks.
	#!jira OR-36843, UE-42975
	#!rb Martin.Wilson
	#!tests Editor PIE, -game hero gallery

Change 3353048 on 2017/03/18 by [REDACTED_PERSON]

	#!ORION_DG - Merge MAIN [REDACTED_HANDLE] 3353033

Change 3352845 on 2017/03/17 by [REDACTED_PERSON]

	Renamed the ConvertRenderTargetToTexture2D function so that it's obvious it's a editor only feature.
	#!rb Daniel.Wright
	#!test Editor paragon

Change 3352544 on 2017/03/17 by [REDACTED_PERSON]

	ADded support for ignoring ini settings incompatbilities when using shared cooked builds.
	#!rb Trivial
	#!test Shared cooked build paragon

Change 3352285 on 2017/03/17 by [REDACTED_PERSON]

	Fix client side compilation error to do with render texture conversion function
	#!rb Trivial
	#!test Compile Paragon

Change 3352141 on 2017/03/17 by [REDACTED_PERSON]

	Added support for blueprint function to convert a rendertexture to a texture.
	#!rb Daniel.Wright
	#!test Run in the editor

Change 3351612 on 2017/03/17 by [REDACTED_PERSON]

	Expand EngineDir and ProjectDir variables during AppLocal deployment
	#!tests Jamie verified packaging Orion via the editor works now
	#!rb Jamie.Dale

Change 3350470 on 2017/03/16 by [REDACTED_PERSON]

	Fix for PS4 compile.

	#!rb none
	#!tests PS4 + non unity

Change 3350237 on 2017/03/16 by [REDACTED_PERSON]

	Pak-mounting fix from Dev-Core for OR-36896
	#!tests na
	#!rb GIl.Gribb

Change 3350079 on 2017/03/16 by [REDACTED_PERSON]

	Added 'AnimNotify_PlayMontageNotify' and 'AnimNotify_PlayMontageNotifyWindow' to forward notifies Begin/End to 'PlayMontage' AsyncTask.

	#!rb lina.halper
	#!tests Yin's BP

Change 3349694 on 2017/03/16 by robomerge

	#!ROBOMERGE-AUTHOR: dan.hertzka
	Exposing copy/paste actions for properties embedded within IDetailGroup header rows

	#!rb Matt.Kuhlenschmidt
	#!tests Copy/paste on skin variant primary override rows

	#!ROBOMERGE-SOURCE: CL 3349513 in //Orion/Dev-REGS/... via CL 3349675
	#!ROBOMERGE-BOT: ORION (Main -> Dev-General)

Change 3349560 on 2017/03/16 by [REDACTED_PERSON]

	Update GameplayTagReferenceHelper to pass in raw data for owner struct (Rather than having caller pass raw 'this' to delegate). Fixes crashes with resizing lists while making calling code less crappy (avoid having to implement copy cstor and operator to fixup delegate).

	Added GameplayTagReferenceHelper to gameplay cue classes.

	#!rb none
	#!tests editor

Change 3349305 on 2017/03/16 by [REDACTED_PERSON]

	Merging //Orion/Main to Dev-General (//Orion/Dev-General)
	#!tests compiled
	#!rb na

Change 3349189 on 2017/03/16 by [REDACTED_PERSON]

	Fixed clothing not running in PS4 packaged builds
	#!rb Martin.Wilson
	#!jira OR-36680
	#!tests PS4 cooked OrionEntry with Shinbi

Change 3348659 on 2017/03/15 by [REDACTED_PERSON]

	Fix compilation errors.
	#!rb None

Change 3348646 on 2017/03/15 by [REDACTED_PERSON]

	Unshelved from pending changelist '3347778':

	<description: restricted, no permission to view>

Change 3348636 on 2017/03/15 by [REDACTED_PERSON]

	Fixed issue with rebuildlighting commandlet not checking out separate lighting files.
	#!rb None
	#!test ResavePackages commandlet

Change 3348559 on 2017/03/15 by [REDACTED_PERSON]

	Fixed up some iterative ini settings blacklist configs.
	#!rb Trivial
	#!test Iterative Cook paragon

Change 3348379 on 2017/03/15 by [REDACTED_PERSON]

	Added simple Async Node 'Play Montage' to use outside of gameplay abilities.

	#!rb none
	#!tests none

Change 3348035 on 2017/03/15 by [REDACTED_PERSON]

	Switch automationcheckpoint to being a .log file. Unblocks running on packaged builds in paragon.
	#!rb none
	#!tests ran oh so very many tests with the changes.

Change 3345982 on 2017/03/14 by [REDACTED_PERSON]

	#!orion - OR-36422: Clamp client net send rate for character movement to 60Hz (down from 90). Integrates CL 3345771 from Dev-Framework which adds engine support for specifying the rate parameters, and sets them in Orion DefaultGame.ini to 1/60 second.

	#!jira OR-36422
	#!tests multi-PIE dedicated server, various framerates, net lag, etc.
	#!rb Laurent.Delayen
	#!codereview Laurent.Delayen

Change 3345134 on 2017/03/14 by [REDACTED_PERSON]

	mono work

Change 3344857 on 2017/03/14 by [REDACTED_PERSON]

	Missing includes for transactor header

	#!rb none

Change 3341860 on 2017/03/10 by [REDACTED_PERSON]

	Partial revert of CL 3339904. Fixed material translation error with multiple connections from custom interpolator nodes.
	#!rb None
	#!tests Editor, Known trouble materials with interpolator nodes, With/without material functions

Change 3341759 on 2017/03/10 by [REDACTED_PERSON]

	Fixed up NetworkCompatible version so that it works with UGS.
	#!rb Trivial
	#!test Cook ps4 paragon.

Change 3341616 on 2017/03/10 by [REDACTED_PERSON]

	#!UE4 - added define for OGS feature
	#!rb none
	#!codereview sam.zamani
	#!tests compiles

Change 3341612 on 2017/03/10 by [REDACTED_PERSON]

	#!UE4 - removed old define
	#!tests compiles

Change 3340180 on 2017/03/09 by [REDACTED_PERSON]

	Integrate fix for sync loading from main to Dev General.
	#!rb Ben.Zeigler

Change 3339904 on 2017/03/09 by [REDACTED_PERSON]

	Fixed material translation error when custom interpolator node hooked to multiple function outputs.
	#!rb None
	#!tests Editor

Change 3339280 on 2017/03/09 by [REDACTED_PERSON]

	#!UE4 - removed WebBrowser moduel dependency on OnlineSubsystem
	- added 2 functions to online engine interface
	#!codereview sam.zamani, ben.marsh

Change 3338654 on 2017/03/08 by [REDACTED_PERSON]

	Fixed up some issues with iterative ini settings.
	Added support for target platforms exposing which audio formats they use so they can match up supported formats with different machines.
	#!rb None
	#!test Cook paragon iteratively

Change 3336989 on 2017/03/08 by [REDACTED_PERSON]

	Merging CL 3336693 from Dev-Core: Use shared PCHs for game plugins by default, to reduce time spent generating individual PCHs.

	#!rb none

Change 3336135 on 2017/03/07 by [REDACTED_PERSON]

	Hide GameLayerManager's title bar on exiting PIE

	#!rb Dan.Hertzka
	#!tests Tested in the editor on Windows

Change 3335324 on 2017/03/07 by [REDACTED_PERSON]

	Chat;

	Adding AddedItem, CompletedItem, and DiscardedItem to the chat message type enum so we can control the color for each. Set the colors in the Social asset.
	Creating client record settings for turning on/off the added item, completed item, and discarded item in chat. Put these in the gameplay settings menu.
	Added horizontal boxes to the gameplay settings menu because we are running out of space.
	Added a vertical scroll bar to the gameplay settings menu but it doesn't seem to show. Also fixed the horizontal scroll bar at the bottom to be horizontal instead of vertical.

	#!rb Matt.Schembari
	#!tests MCP, PIE
	#!lockdown Nicholas.Davies
	#!RN

Change 3333541 on 2017/03/06 by [REDACTED_PERSON]

	#!ORION_DG - Merge MAIN @ CL 3333512

	#!RB:none
	#!Tests:none

	#!codeReview: cameron.winston

Change 3332578 on 2017/03/04 by [REDACTED_PERSON]

	Temp Disabled wrong-looking warning
	#!tests #!rb na
	#!ROBOMERGE: Main

Change 3332555 on 2017/03/04 by [REDACTED_PERSON]

	Proper fix for Tencent DLL issue
	#!tests #!rb na
	#!ROBOMERGE: Main

Change 3332552 on 2017/03/04 by [REDACTED_PERSON]

	Fix for Tencent DLL issue while staging
	#!tests none
	#!rb none
	#!ROBOMERGE: Main

Change 3332216 on 2017/03/03 by [REDACTED_PERSON]

	#!ORION_DG - Merge MAIN @ CL 3332168

	#!RB:none
	#!Tests:none

Change 3332060 on 2017/03/03 by [REDACTED_PERSON]

	Fixed issue with AsyncLoading code eventually flushing async loading while in async loading...
	This causes all kinds of cool stuff like objects on the stack corruption and also deleted memory accesses.

	#!rb Gil.Gribb.
	#!test Editor and -game

Change 3331680 on 2017/03/03 by [REDACTED_PERSON]

	#!ORION_MAIN - Merge MAIN @ CL 3331636

	#!RB:none
	#!Tests:none

	#!codeReview: andrew.grant

Change 3331412 on 2017/03/03 by [REDACTED_PERSON]

	#!orion Rebuilt OpenSSL libs for PS4 to fix process termination due to SIGPIPE on closing websockets

	Source change committed in CL#!3331380

	#!jira OR-36274

	#!fyi Paul.Moore

Change 3331375 on 2017/03/03 by [REDACTED_PERSON]

	fix dll path for tenproxy

	#!rb none
	#!tests none

Change 3330953 on 2017/03/02 by [REDACTED_PERSON]

	#!ORION_DG - Merge MAIN @ CL 3330924

	[STOMPED ChestOpeningScreen.uasset]

	#!RB:none
	#!Tests:none

	#!codeReview: bryan.rathman, phil.buuck, matt.schembari, andrew.grant

Change 3330646 on 2017/03/02 by [REDACTED_PERSON]

	Warning and non-unity fix
	#!tests compiled
	#!rb none

Change 3330388 on 2017/03/02 by [REDACTED_PERSON]

	Merging //Orion/Main to Dev-General (//Orion/Dev-General)
	#!tests #!rb na

Change 3329982 on 2017/03/02 by [REDACTED_PERSON]

	fixed updated module rules

	#!rb none
	#!tests regen projects

Change 3329964 on 2017/03/02 by [REDACTED_PERSON]

	Copying //Tasks/Orion/Dev-Online-Tencent to Dev-General (//Orion/Dev-General)

	3245325 Adding new OSS for Tencent online platform

	3245448 tencent third party SDK
	TCLS proxy functionality

	#!rb none

	3245474 missing include

	#!rb none

	3249585 TCLS tenproxy.dll in thirdparty bin folder

	#!rb none

	3249726 Load TenProxy.dll for TCLS integration
	New OSS Tencent

	#!rb none

	3255571 tencent configs

	#!rb none

	3255826 Tencent TCLS paragon launcher

	#!rb none

	3256168 TCLS launch batch update cmd line options

	#!rb none

	3256170 Added "TencentLive,TencentDev" MCP config entries

	#!rb none

	3256504 xmpp config update

	#!rb none

	3273168 skip login steps for tencent
	config update

	#!rb none

	3279427 #!xmpp

	add option to use plain text auth

	3279428 disable ssl and use plain text auth for XMPP connection
	temporary until we have a valid cert setup on Tigase deployment

	3281566 enabled OSS tencent

	this will also be the toggle for detecting when to enable tencent functionality at runtime

	3283103 differentiate between tencent dev/live environments
	disable QoS region selection for tencentdev

	3283106 lower http verbosity

	3283734 config updates

	3285066 disable replays and mtx for tencent build

	3291005 #!online,mcp
	service config bEnabled flag to toggle individual services as needed

	3291006 explicitly mark unneeded Mcp services as disabled

	3291108 allow replay tab to be disabled via UOrionRuntimeOptions.bEnableReplays=false

	3291492 disable recording of replays for tencent mode

	3292750 disable replay tab based on bEnableReplays=false

	3292753 new orion runtime option bDisallowCoinPurchases
	if true, prevents coins from being available for purchase

	3292755 diable mtx coin offers if bDisallowCoinPurchases=true

	3292759 missing header

	3293246 disable query for available friend codes if bEnableFriendCodes=false

	3293250 temp usage of NULL analytics provider

	3298025 Adding optional RegionTencent plugin for overriding config files

	3298027 ability to override config cache values via plugin config files

	3311016 default to TencentDev backend when running in tencent mode

	3311017 CMS tencent config

	3311022 Rename RegionTencent to RegionCN

	3312470 disable links for tencent build

	3313014 move tenproxy.dll to \OrionGame\Binaries\ThirdParty\Tencent

	3314861 tenproxy 2.0.2.7 update

	3314878 default RegionCN plugin to disabled
	this will only be enabled once the RegionCN.pak is loaded

	3314879 TCLS launcher pointing at UE4Editor.exe for development

	3315257 missing file

	3323573 remove TCLS launcher

	3326006 Tencent TLOG SDK

	3326277 wrapper singleton class for tenproxy connection

	3329180 Tencent support for login flow

	3329181 WIP tenproxy connection usage in identity

	3329624 wip tcls proxy

	#!rb none
	#!tests none

Change 3329651 on 2017/03/02 by [REDACTED_PERSON]

	Merging from //UE4/Main @ 3322856 through Orion-Staging
	#!tests QA
	#!rb na

Change 3329411 on 2017/03/02 by robomerge

	#!ROBOMERGE-AUTHOR: dan.hertzka
	Duplicating CL 3303733 from Dev-Editor (simple fix for a massive issue)
	- This will prevent any TAssetPtr property from getting stomped by undo/redo (you know those ridiculous store and card art issues? Fixed!)

	#!lockdown Jason.Bestimt
	#!rb none
	#!tests Undo on an item definition asset

	#!ROBOMERGE-SOURCE: CL 3329404 in //Orion/Release-38.3/... via CL 3329405
	#!ROBOMERGE-BOT: ORION (Main -> Dev-General)

Change 3328858 on 2017/03/01 by [REDACTED_PERSON]

	Fixed crash on importing animation that was edited before

	#!rb: none
	#!tests: reimport

Change 3328459 on 2017/03/01 by [REDACTED_PERSON]

	When adding new ddc back ends to the hierarchcial ddc make sure to update the async backends lists.
	#!codereview Gil.Gribb
	#!test None
	#!rb Trivial

Change 3328182 on 2017/03/01 by [REDACTED_PERSON]

	Unshelved from pending changelist '3318009':

	Adding support for shared cooked builds to be downloaded from the network.
	Included CookedAssetRegistry in the p:\ published builds.
	#!rb Ben.Marsh

Change 3327856 on 2017/03/01 by [REDACTED_PERSON]

	Added velocity overrides to FRK4SpringInterpolator;

	#!RB None
	#!codeReview Laurent.Delayen
	#!Tests PIE

Change 3327096 on 2017/03/01 by [REDACTED_PERSON]

	Added generic reference viewer details customization for gameplay tags. Added it to GameplayStatsMetaData.

	#!rb none
	#!tests editor

Change 3326177 on 2017/02/28 by [REDACTED_PERSON]

	Added some more debugging information to help track down live issue.
	#!rb Chris.Bunner
	#!test Ran editor.

Change 3324951 on 2017/02/28 by [REDACTED_PERSON]

	UDataTable: added AddRow/RemoveRow native functions.
	#!rb JB
	#!tests na

Change 3323852 on 2017/02/27 by [REDACTED_PERSON]

	Fix ::RequestAllGameplayTags OnlyIncludeDictionaryTags option

	#!codereview Ben.Zeigler
	#!rb #!tests na

Change 3323706 on 2017/02/27 by [REDACTED_PERSON]

	#!ORION_DG - Merge MAIN @ CL 3323694

	#!RB:none
	#!Tests:none

Change 3321945 on 2017/02/24 by [REDACTED_PERSON]

	OR-36258

	- fixing an issue where gameplay effects that are set to not refresh the period should not allow the execution of a period effect on application.

	#!RB David.Ratti
	#!tests golden path
	#!codeReview: Billy.Bramer, Fred.Kimberley
	#!RNX

Change 3321876 on 2017/02/24 by [REDACTED_PERSON]

	Fixed erroronEngineContentUse flag not being set properly.
	#!rb Trivial
	#!test Cook Paragon.

Change 3321591 on 2017/02/24 by [REDACTED_PERSON]

	#!ORION_DG - MAIN @ CL 3321563

	#!RB:none
	#!Tests:none

Change 3321260 on 2017/02/24 by [REDACTED_PERSON]

	Fixed issue that was causing missing string references to not show their referencer
	#!rb none

Change 3321040 on 2017/02/24 by [REDACTED_PERSON]

	Merging changes 3316253 and 3319134 from Dev-Core: fixes to file log hangs and crashes.

	#!rb none
	#!tests Cooked Win64 server and client, played cooked Win64 build

Change 3319413 on 2017/02/23 by [REDACTED_PERSON]

	#!ORION_DG - Merge MAIN @ CL 3319394

	#!RB:none
	#!Tests:none

Change 3317905 on 2017/02/22 by [REDACTED_PERSON]

	Integrate CL 3238291 from Odin

	Add Plugin content to the asset registry
	Change the location of AssetRegistry.bin when cooking a plugin as DLC
	Include AssetRegistry.bin in the cooked plugin staging process
	Add function to PluginManager to keep list of any plugins that loaded a pak file
	Use list of plugins with pak files to merge their AssetRegistry.bin files into the main AssetRegistry when it's created
	#!rb Ben.Marsh
	#!codereview Chance.Ivey, Daniel.Lamb

Change 3317648 on 2017/02/22 by [REDACTED_PERSON]

	Instead of popping an external web browser, we use the SWebBrowser widget on GFN.

	#!rb DanH
	#!codereview Andrew.Grant, Dan.Hertzka, Matt.Schembari
	#!tests PIE

Change 3317289 on 2017/02/22 by [REDACTED_PERSON]

	#!ORION_DG - Merge MAIN @ CL 3317254

	#!RB:none
	#!Tests:none

Change 3317186 on 2017/02/22 by [REDACTED_PERSON]

	Fixed items that have been force-scored by an EQS test as 'failed' getting discarted even if the test is being run in scoring-only mode #!UE4

	#!test golden path
	#!rb Lukasz.Furman
	#!codereview Daniel.Broder, John.Abercrombie

Change 3317005 on 2017/02/22 by [REDACTED_PERSON]

	Submitted wrong version of my file.
	#!rb Trivial
	#!test Compile

Change 3316958 on 2017/02/22 by [REDACTED_PERSON]

	Added support in buildcookrun for shared cooked builds.
	#!rb Trivial
	#!test BuildCookRun iterative script

Change 3316942 on 2017/02/22 by [REDACTED_PERSON]

	DLC cooking optimization.
	Optimization to determining package dependency tree, now is async.
	Fixes for iterate shared cooked build.  Added fallback when using shared cooked build to local build if local build is newer.
	Added DLC cooking warning if you are overriding output directories.
	Removed previous release packages names from DLC asset registry.
	Only generate manifest for additional assets instead of all assets.
	Minor optimization to worst case resolving of string asset references.  Only resolve those that haven't been resolved before (only happens when GC thrashing happens).
	#!rb Andrew.Grant
	#!test Cook paragon

diff --git a/Engine/Config/BaseEditor.ini b/Engine/Config/BaseEditor.ini
--- a/Engine/Config/BaseEditor.ini
+++ b/Engine/Config/BaseEditor.ini
@@ -386,6 +386,12 @@ bUseFullGCAssetClassNames=true
 +ConfigSettingBlacklist=*.SourceControlSettings
 +ConfigSettingBlacklist=*.SourceControl.SourceControlSettings
 +ConfigSettingBlacklist=*.Engine:WindowsApplication.Accessibility
++ConfigSettingBlacklist=*.EditorKeyBindings
++ConfigSettingBlacklist=*.GameplayAbilities:CoreRedirects
++ConfigSettingBlacklist=*.Engine.MemReportFullCommands
++ConfigSettingBlacklist=*:CoreRedirects
++ConfigSettingBlacklist=*.OnlineGameplayFramework:/Script/Engine.Engine
+
 
 +VersionedIntRValues=r.AllowStaticLighting
 +VersionedIntRValues=r.GBuffer
diff --git a/Engine/Config/BaseEngine.ini b/Engine/Config/BaseEngine.ini
--- a/Engine/Config/BaseEngine.ini
+++ b/Engine/Config/BaseEngine.ini
@@ -905,7 +905,7 @@ ConfigClass=
 +Cmd="DumpParticleMem"
 +Cmd="ConfigMem"
 +Cmd="r.DumpRenderTargetPoolMemory"
-+Cmd="ListTextures -alphasort"
++Cmd="ListTextures"
 +Cmd="ListSounds -alphasort"
 +Cmd="ListParticleSystems -alphasort"
 +Cmd="obj list class=SoundWave -alphasort"
diff --git a/Engine/Extras/VisualStudioDebugging/InstallVisualizers.bat b/Engine/Extras/VisualStudioDebugging/InstallVisualizers.bat
--- a/Engine/Extras/VisualStudioDebugging/InstallVisualizers.bat
+++ b/Engine/Extras/VisualStudioDebugging/InstallVisualizers.bat
@@ -25,10 +25,16 @@ if "%VS110COMNTOOLS%" neq "" (
   echo Installed visualizers for Visual Studio 2012
 )
 
+if "%SCE_ORBIS_SDK_DIR%" neq "" (
+  mkdir "%USERPROFILE%\Documents\SCE\orbis-debugger"
+  copy PS4UE4.natvis "%USERPROFILE%\Documents\SCE\orbis-debugger"
+  echo Installed Visualizers for Orbis
+)
+
 if "%FoundVSInstance%" equ "0" (
   echo ERROR: Could not find a valid version of Visual Studio installed (2012, 2013, 2015)
 )
 
 popd
 pause
-exit
\ No newline at end of file
+exit /b
\ No newline at end of file
diff --git a/Engine/Plugins/Editor/GameplayTagsEditor/Source/GameplayTagsEditor/Private/GameplayTagReferenceHelperDetails.cpp b/Engine/Plugins/Editor/GameplayTagsEditor/Source/GameplayTagsEditor/Private/GameplayTagReferenceHelperDetails.cpp
new file mode 100644
--- /dev/null
+++ b/Engine/Plugins/Editor/GameplayTagsEditor/Source/GameplayTagsEditor/Private/GameplayTagReferenceHelperDetails.cpp
@@ -0,0 +1,196 @@
+// Copyright 1998-2016 Epic Games, Inc. All Rights Reserved.
+
+#include "GameplayTagReferenceHelperDetails.h"
+#include "AssetRegistryModule.h"
+#include "SHyperlink.h"
+#include "Toolkits/AssetEditorManager.h"
+#include "ObjectTools.h"
+#include "GameplayTagContainer.h"
+#include "UObjectHash.h"
+#include "UnrealType.h"
+
+#define LOCTEXT_NAMESPACE "GameplayTagReferenceHelperDetails"
+
+static const FName GameplayTagColumnName("GameplayTagColumn");
+
+TSharedRef<IPropertyTypeCustomization> FGameplayTagReferenceHelperDetails::MakeInstance()
+{
+	return MakeShareable(new FGameplayTagReferenceHelperDetails());
+}
+
+void FGameplayTagReferenceHelperDetails::CustomizeHeader( TSharedRef<IPropertyHandle> StructPropertyHandle, class FDetailWidgetRow& HeaderRow, IPropertyTypeCustomizationUtils& StructCustomizationUtils )
+{
+	PropertyHandle = StructPropertyHandle;
+
+	TreeItems.Reset();
+	if (FGameplayTagReferenceHelper* Helper = GetValue())
+	{
+		// We need the raw data pointer to the struct (UStruct or UClass) that owns the FGameplayTagReferenceHelper property.
+		// Its not enough to just bind the raw 'this' pointer in the owning struct's cstor, since lists or data tables of structs
+		// will be copied around as the list changes sizes (overloading copy and assignment operators on the owning struct to clean/update 
+		// the delegate is also a major pain).
+		// 
+		// We cheat a bit here and use GetOffset_ForGC to work backwards up the property change and get the raw, castable, address of the owning
+		// structure so that the delegate can just do a static_cast and do whatever they want.
+		//
+		//
+		// Note: this currently does NOT handle the owning struct changing and auto updating. This is a bit tricky since we don't know, in this context, when
+		// an update has to happen, since this thing is not tied directly to a tag uproperty. (E.g, a data table row where the row's key name is the tag tag name)
+
+		void* OwnerStructRawData = nullptr;
+		if (UProperty* MyProperty = PropertyHandle->GetProperty())
+		{
+			if (UStruct* OwnerStruct = MyProperty->GetOwnerStruct())
+			{				
+				// Manually calculate the owning struct raw data pointer
+				checkf(MyProperty->ArrayDim == 1, TEXT("FGameplayTagReferenceHelper should never be in an array"));
+				uint8* PropertyRawData = (uint8*)GetValue();
+				OwnerStructRawData = (void*) (PropertyRawData - MyProperty->GetOffset_ForGC());
+			}
+		}
+		
+		if (ensureMsgf(OwnerStructRawData != nullptr, TEXT("Unable to get outer struct's raw data")))
+		{
+			FName TagName = Helper->OnGetGameplayTagName.Execute(OwnerStructRawData);
+		
+			FAssetIdentifier TagId = FAssetIdentifier(FGameplayTag::StaticStruct(), TagName);
+			TArray<FAssetIdentifier> Referencers;
+
+			FAssetRegistryModule& AssetRegistryModule = FModuleManager::LoadModuleChecked<FAssetRegistryModule>("AssetRegistry");
+			AssetRegistryModule.Get().GetReferencers(TagId, Referencers, EAssetRegistryDependencyType::SearchableName);
+
+			for (FAssetIdentifier& AssetIdentifier : Referencers)
+			{
+				TSharedPtr<FGameplayTagReferenceTreeItem> Item =  TSharedPtr<FGameplayTagReferenceTreeItem>(new FGameplayTagReferenceTreeItem());
+				Item->GameplayTagName = TagName;
+				Item->AssetIdentifier = AssetIdentifier;
+				TreeItems.Add(Item);
+			}
+		}
+	}
+
+	HeaderRow
+	.NameContent()
+	[
+		SNew(SBorder)
+		.BorderImage(FEditorStyle::GetBrush("ToolBar.Background"))
+		[
+			SNew(SGameplayTagReferenceTree)
+			.ItemHeight(24)
+			.TreeItemsSource(&TreeItems)
+			.OnGenerateRow(this, &FGameplayTagReferenceHelperDetails::OnGenerateWidgetForGameplayCueListView)
+			.OnGetChildren( SGameplayTagReferenceTree::FOnGetChildren::CreateLambda([](TSharedPtr<FGameplayTagReferenceTreeItem> Item, TArray< TSharedPtr<FGameplayTagReferenceTreeItem> >& Children) { } ) )
+			.HeaderRow
+			(
+				SNew(SHeaderRow)
+				+ SHeaderRow::Column(GameplayTagColumnName)
+				.DefaultLabel(NSLOCTEXT("GameplayTagReferenceHelper", "GameplayTagReferenceHelperColumn", "GameplayTag Referencers (does not include native code)"))
+				.FillWidth(0.50)
+			)
+		]
+	];
+}
+
+void FGameplayTagReferenceHelperDetails::CustomizeChildren( TSharedRef<IPropertyHandle> StructPropertyHandle, class IDetailChildrenBuilder& StructBuilder, IPropertyTypeCustomizationUtils& StructCustomizationUtils )
+{
+	
+}
+
+FGameplayTagReferenceHelper* FGameplayTagReferenceHelperDetails::GetValue()
+{
+	TArray<void*> RawData;
+	PropertyHandle->AccessRawData(RawData);
+	if (RawData.Num() != 1)
+	{
+		UE_LOG(LogTemp, Error, TEXT("Unexpected raw data count of %d"), RawData.Num());
+		return nullptr;
+	}
+	return static_cast<FGameplayTagReferenceHelper*>(RawData[0]);
+}
+
+/** Builds widget for rows in the GameplayCue Editor tab */
+TSharedRef<ITableRow> FGameplayTagReferenceHelperDetails::OnGenerateWidgetForGameplayCueListView(TSharedPtr< FGameplayTagReferenceTreeItem > InItem, const TSharedRef<STableViewBase>& OwnerTable)
+{
+	class SGameplayTagWidgetItem : public SMultiColumnTableRow< TSharedPtr< FGameplayTagReferenceTreeItem > >
+	{
+	public:
+		SLATE_BEGIN_ARGS(SGameplayTagWidgetItem){}
+		SLATE_END_ARGS()
+
+		void Construct(const FArguments& InArgs, const TSharedRef<STableViewBase>& InOwnerTable, TSharedPtr<FGameplayTagReferenceTreeItem> InListItem)
+		{
+			Item = InListItem;
+			SMultiColumnTableRow< TSharedPtr< FGameplayTagReferenceTreeItem > >::Construct(
+				FSuperRowType::FArguments()
+			, InOwnerTable);
+		}
+	private:
+
+		virtual TSharedRef<SWidget> GenerateWidgetForColumn(const FName& ColumnName) override
+		{
+			if (ColumnName == GameplayTagColumnName)
+			{
+				if (Item->AssetIdentifier.ToString().IsEmpty() == false)
+				{
+					return
+					SNew(SBox)
+					.HAlign(HAlign_Left)
+					[
+						SNew(SHyperlink)
+						.Style(FEditorStyle::Get(), "Common.GotoBlueprintHyperlink")
+						.Text(FText::FromString(Item->AssetIdentifier.ToString()))
+						.OnNavigate(this, &SGameplayTagWidgetItem::NavigateToReference)
+					];
+				}
+				else
+				{
+					return
+					SNew(SBox)
+					.HAlign(HAlign_Left);
+				}
+			}
+			else
+			{
+				return SNew(STextBlock).Text(LOCTEXT("UnknownColumn", "Unknown Column"));
+			}
+		}
+
+		void NavigateToReference()
+		{
+			UPackage* Pkg = FindPackage(nullptr, *Item->AssetIdentifier.PackageName.ToString());
+			if (Pkg == nullptr)
+			{
+				Pkg = LoadPackage(nullptr, *Item->AssetIdentifier.PackageName.ToString(), LOAD_None);
+			}
+
+			if (Pkg)
+			{
+				ForEachObjectWithOuter(Pkg,[](UObject* Obj)
+				{
+					if (ObjectTools::IsObjectBrowsable(Obj))
+					{
+						FAssetEditorManager::Get().OpenEditorForAsset(Obj);
+					}
+				});
+			}
+		}
+
+		TSharedPtr< FGameplayTagReferenceTreeItem > Item;
+	};
+
+	if ( InItem.IsValid() )
+	{
+		return SNew(SGameplayTagWidgetItem, OwnerTable, InItem);
+	}
+	else
+	{
+		return
+		SNew(STableRow< TSharedPtr<FGameplayTagReferenceTreeItem> >, OwnerTable)
+		[
+			SNew(STextBlock)
+			.Text(LOCTEXT("UnknownItemType", "Unknown Item Type"))
+		];
+	}
+}
+
+#undef LOCTEXT_NAMESPACE
\ No newline at end of file
diff --git a/Engine/Plugins/Editor/GameplayTagsEditor/Source/GameplayTagsEditor/Private/GameplayTagReferenceHelperDetails.h b/Engine/Plugins/Editor/GameplayTagsEditor/Source/GameplayTagsEditor/Private/GameplayTagReferenceHelperDetails.h
new file mode 100644
--- /dev/null
+++ b/Engine/Plugins/Editor/GameplayTagsEditor/Source/GameplayTagsEditor/Private/GameplayTagReferenceHelperDetails.h
@@ -0,0 +1,48 @@
+// Copyright 1998-2017 Epic Games, Inc. All Rights Reserved.
+
+#pragma once
+
+#include "Editor/PropertyEditor/Public/PropertyEditorModule.h"
+#include "IDetailCustomization.h"
+#include "AssetData.h"
+#include "DetailWidgetRow.h"
+#include "Widgets/Views/STableViewBase.h"
+#include "Widgets/Views/STableRow.h"
+#include "Widgets/Views/STreeView.h"
+#include "PropertyHandle.h"
+
+class IPropertyHandle;
+
+class FGameplayTagReferenceHelperDetails : public IPropertyTypeCustomization
+{
+public:
+	static TSharedRef<IPropertyTypeCustomization> MakeInstance();
+
+	/** IPropertyTypeCustomization interface */
+	virtual void CustomizeHeader( TSharedRef<class IPropertyHandle> StructPropertyHandle, class FDetailWidgetRow& HeaderRow, IPropertyTypeCustomizationUtils& StructCustomizationUtils ) override;
+	virtual void CustomizeChildren( TSharedRef<class IPropertyHandle> StructPropertyHandle, class IDetailChildrenBuilder& StructBuilder, IPropertyTypeCustomizationUtils& StructCustomizationUtils ) override;              
+
+private:
+
+	struct FGameplayTagReferenceTreeItem : public TSharedFromThis<FGameplayTagReferenceTreeItem>
+	{
+		FGameplayTagReferenceTreeItem()
+		{
+		}
+
+		FName GameplayTagName;
+		FAssetIdentifier AssetIdentifier;	
+	};
+
+	typedef STreeView< TSharedPtr< FGameplayTagReferenceTreeItem > > SGameplayTagReferenceTree;
+
+	TArray< TSharedPtr<FGameplayTagReferenceTreeItem> > TreeItems;
+
+	TSharedRef<ITableRow> OnGenerateWidgetForGameplayCueListView(TSharedPtr< FGameplayTagReferenceTreeItem > InItem, const TSharedRef<STableViewBase>& OwnerTable);
+
+	TSharedPtr<IPropertyHandle> StatBackendNameProp;
+
+	TSharedPtr<class IPropertyHandle> PropertyHandle;
+
+	struct FGameplayTagReferenceHelper* GetValue();
+};
\ No newline at end of file
diff --git a/Engine/Plugins/Editor/GameplayTagsEditor/Source/GameplayTagsEditor/Private/GameplayTagsEditorModule.cpp b/Engine/Plugins/Editor/GameplayTagsEditor/Source/GameplayTagsEditor/Private/GameplayTagsEditorModule.cpp
--- a/Engine/Plugins/Editor/GameplayTagsEditor/Source/GameplayTagsEditor/Private/GameplayTagsEditorModule.cpp
+++ b/Engine/Plugins/Editor/GameplayTagsEditor/Source/GameplayTagsEditor/Private/GameplayTagsEditorModule.cpp
@@ -28,6 +28,7 @@
 #include "HAL/PlatformFile.h"
 #include "HAL/PlatformFilemanager.h"
 #include "Stats/StatsMisc.h"
+#include "GameplayTagReferenceHelperDetails.h"
 
 #define LOCTEXT_NAMESPACE "GameplayTagEditor"
 
@@ -50,6 +51,8 @@ public:
 
 			PropertyModule.RegisterCustomClassLayout(UGameplayTagsList::StaticClass()->GetFName(), FOnGetDetailCustomizationInstance::CreateStatic(&FGameplayTagsSettingsCustomization::MakeInstance));
 
+			PropertyModule.RegisterCustomPropertyTypeLayout("GameplayTagReferenceHelper", FOnGetPropertyTypeCustomizationInstance::CreateStatic(&FGameplayTagReferenceHelperDetails::MakeInstance));
+
 			PropertyModule.NotifyCustomizationModuleChanged();
 		}
 
diff --git a/Engine/Plugins/Runtime/GameplayAbilities/Source/GameplayAbilities/Private/AbilitySystemComponent_Abilities.cpp b/Engine/Plugins/Runtime/GameplayAbilities/Source/GameplayAbilities/Private/AbilitySystemComponent_Abilities.cpp
--- a/Engine/Plugins/Runtime/GameplayAbilities/Source/GameplayAbilities/Private/AbilitySystemComponent_Abilities.cpp
+++ b/Engine/Plugins/Runtime/GameplayAbilities/Source/GameplayAbilities/Private/AbilitySystemComponent_Abilities.cpp
@@ -30,11 +30,15 @@
 #include "Abilities/GameplayAbilityTargetActor.h"
 #include "TickableAttributeSetInterface.h"
 #include "GameplayTagResponseTable.h"
+#include "Engine/DemoNetDriver.h"
+
 #define LOCTEXT_NAMESPACE "AbilitySystemComponent"
 
 /** Enable to log out all render state create, destroy and updatetransform events */
 #define LOG_RENDER_STATE 0
 
+static TAutoConsoleVariable<float> CVarReplayMontageErrorThreshold(TEXT("replay.MontageErrorThreshold"), 0.5f, TEXT("Tolerance level for when montage playback position correction occurs in replays"));
+
 void UAbilitySystemComponent::InitializeComponent()
 {
 	Super::InitializeComponent();
@@ -2355,7 +2359,11 @@ bool UAbilitySystemComponent::IsReadyForReplicatedMontage()
 /**	Replicated Event for AnimMontages */
 void UAbilitySystemComponent::OnRep_ReplicatedAnimMontage()
 {
-	static const float MONTAGE_REP_POS_ERR_THRESH = 0.1f;
+	UWorld* World = GetWorld();
+
+	const bool bIsPlayingReplay = World && World->DemoNetDriver && World->DemoNetDriver->IsPlaying();
+
+	const float MONTAGE_REP_POS_ERR_THRESH = bIsPlayingReplay ? CVarReplayMontageErrorThreshold.GetValueOnGameThread() : 0.1f;
 
 	UAnimInstance* AnimInstance = AbilityActorInfo.IsValid() ? AbilityActorInfo->GetAnimInstance() : nullptr;
 	if (AnimInstance == nullptr || !IsReadyForReplicatedMontage())
diff --git a/Engine/Plugins/Runtime/GameplayAbilities/Source/GameplayAbilities/Private/GameplayCueNotify_Actor.cpp b/Engine/Plugins/Runtime/GameplayAbilities/Source/GameplayAbilities/Private/GameplayCueNotify_Actor.cpp
--- a/Engine/Plugins/Runtime/GameplayAbilities/Source/GameplayAbilities/Private/GameplayCueNotify_Actor.cpp
+++ b/Engine/Plugins/Runtime/GameplayAbilities/Source/GameplayAbilities/Private/GameplayCueNotify_Actor.cpp
@@ -31,6 +31,12 @@ AGameplayCueNotify_Actor::AGameplayCueNotify_Actor(const FObjectInitializer& Obj
 
 	WarnIfLatentActionIsStillRunning = true;
 	WarnIfTimelineIsStillRunning = true;
+
+	ReferenceHelper.OnGetGameplayTagName.BindLambda([](void* RawData)
+	{
+		AGameplayCueNotify_Actor* ThisData = static_cast<AGameplayCueNotify_Actor*>(RawData);
+		return ThisData->GameplayCueTag.GetTagName();
+	});
 }
 
 void AGameplayCueNotify_Actor::EndPlay(const EEndPlayReason::Type EndPlayReason)
diff --git a/Engine/Plugins/Runtime/GameplayAbilities/Source/GameplayAbilities/Private/GameplayCueNotify_Static.cpp b/Engine/Plugins/Runtime/GameplayAbilities/Source/GameplayAbilities/Private/GameplayCueNotify_Static.cpp
--- a/Engine/Plugins/Runtime/GameplayAbilities/Source/GameplayAbilities/Private/GameplayCueNotify_Static.cpp
+++ b/Engine/Plugins/Runtime/GameplayAbilities/Source/GameplayAbilities/Private/GameplayCueNotify_Static.cpp
@@ -11,6 +11,13 @@ UGameplayCueNotify_Static::UGameplayCueNotify_Static(const FObjectInitializer& P
 : Super(PCIP)
 {
 	IsOverride = true;
+
+	
+	ReferenceHelper.OnGetGameplayTagName.BindLambda([](void* RawData)
+	{
+		UGameplayCueNotify_Static* ThisData = static_cast<UGameplayCueNotify_Static*>(RawData);
+		return ThisData->GameplayCueTag.GetTagName();
+	});
 }
 
 #if WITH_EDITOR
diff --git a/Engine/Plugins/Runtime/GameplayAbilities/Source/GameplayAbilities/Private/GameplayEffect.cpp b/Engine/Plugins/Runtime/GameplayAbilities/Source/GameplayAbilities/Private/GameplayEffect.cpp
--- a/Engine/Plugins/Runtime/GameplayAbilities/Source/GameplayAbilities/Private/GameplayEffect.cpp
+++ b/Engine/Plugins/Runtime/GameplayAbilities/Source/GameplayAbilities/Private/GameplayEffect.cpp
@@ -2740,7 +2740,7 @@ FActiveGameplayEffect* FActiveGameplayEffectsContainer::ApplyGameplayEffectSpec(
 	}
 	
 	// Register period callbacks with the timer manager
-	if (Owner && (AppliedEffectSpec.GetPeriod() != UGameplayEffect::NO_PERIOD))
+	if (bSetPeriod && Owner && (AppliedEffectSpec.GetPeriod() != UGameplayEffect::NO_PERIOD))
 	{
 		FTimerManager& TimerManager = Owner->GetWorld()->GetTimerManager();
 		FTimerDelegate Delegate = FTimerDelegate::CreateUObject(Owner, &UAbilitySystemComponent::ExecutePeriodicEffect, AppliedActiveGE->Handle);
@@ -2751,10 +2751,7 @@ FActiveGameplayEffect* FActiveGameplayEffectsContainer::ApplyGameplayEffectSpec(
 			TimerManager.SetTimerForNextTick(Delegate);
 		}
 
-		if (bSetPeriod)
-		{
-			TimerManager.SetTimer(AppliedActiveGE->PeriodHandle, Delegate, AppliedEffectSpec.GetPeriod(), true);
-		}
+		TimerManager.SetTimer(AppliedActiveGE->PeriodHandle, Delegate, AppliedEffectSpec.GetPeriod(), true);
 	}
 
 	if (InPredictionKey.IsLocalClientKey() == false || IsNetAuthority())	// Clients predicting a GameplayEffect must not call MarkItemDirty
diff --git a/Engine/Plugins/Runtime/GameplayAbilities/Source/GameplayAbilities/Public/GameplayCueNotify_Actor.h b/Engine/Plugins/Runtime/GameplayAbilities/Source/GameplayAbilities/Public/GameplayCueNotify_Actor.h
--- a/Engine/Plugins/Runtime/GameplayAbilities/Source/GameplayAbilities/Public/GameplayCueNotify_Actor.h
+++ b/Engine/Plugins/Runtime/GameplayAbilities/Source/GameplayAbilities/Public/GameplayCueNotify_Actor.h
@@ -93,6 +93,9 @@ public:
 	UPROPERTY(EditDefaultsOnly, Category=GameplayCue, meta=(Categories="GameplayCue"))
 	FGameplayTag	GameplayCueTag;
 
+	UPROPERTY(EditAnywhere, Category=GameplayCue)
+	FGameplayTagReferenceHelper ReferenceHelper;
+
 	/** Mirrors GameplayCueTag in order to be asset registry searchable */
 	UPROPERTY(AssetRegistrySearchable)
 	FName GameplayCueName;
diff --git a/Engine/Plugins/Runtime/GameplayAbilities/Source/GameplayAbilities/Public/GameplayCueNotify_Static.h b/Engine/Plugins/Runtime/GameplayAbilities/Source/GameplayAbilities/Public/GameplayCueNotify_Static.h
--- a/Engine/Plugins/Runtime/GameplayAbilities/Source/GameplayAbilities/Public/GameplayCueNotify_Static.h
+++ b/Engine/Plugins/Runtime/GameplayAbilities/Source/GameplayAbilities/Public/GameplayCueNotify_Static.h
@@ -54,6 +54,9 @@ class GAMEPLAYABILITIES_API UGameplayCueNotify_Static : public UObject
 	UPROPERTY(EditDefaultsOnly, Category = GameplayCue, meta=(Categories="GameplayCue"))
 	FGameplayTag	GameplayCueTag;
 
+	UPROPERTY(EditAnywhere, Category="GameplayCue")
+	FGameplayTagReferenceHelper ReferenceHelper;
+
 	/** Mirrors GameplayCueTag in order to be asset registry searchable */
 	UPROPERTY(AssetRegistrySearchable)
 	FName GameplayCueName;
diff --git a/Engine/Source/Developer/AllDesktopTargetPlatform/Private/AllDesktopTargetPlatform.cpp b/Engine/Source/Developer/AllDesktopTargetPlatform/Private/AllDesktopTargetPlatform.cpp
--- a/Engine/Source/Developer/AllDesktopTargetPlatform/Private/AllDesktopTargetPlatform.cpp
+++ b/Engine/Source/Developer/AllDesktopTargetPlatform/Private/AllDesktopTargetPlatform.cpp
@@ -70,6 +70,11 @@ void FAllDesktopTargetPlatform::GetTextureFormats( const UTexture* Texture, TArr
 	OutFormats.Add(GetDefaultTextureFormatName(this, Texture, EngineSettings, false));
 }
 
+void FAllDesktopTargetPlatform::GetAllTextureFormats(TArray<FName>& OutFormats) const
+{
+	GetAllDefaultTextureFormats(this, OutFormats, false);
+}
+
 
 FName FAllDesktopTargetPlatform::GetWaveFormat( const class USoundWave* Wave ) const
 {
@@ -85,4 +90,14 @@ FName FAllDesktopTargetPlatform::GetWaveFormat( const class USoundWave* Wave ) c
 	return NAME_OGG;
 }
 
+
+void FAllDesktopTargetPlatform::GetAllWaveFormats(TArray<FName>& OutFormats) const
+{
+	static FName NAME_OGG(TEXT("OGG"));
+	static FName NAME_OPUS(TEXT("OPUS"));
+	OutFormats.Add(NAME_OGG);
+	OutFormats.Add(NAME_OPUS);
+
+}
+
 #endif // WITH_ENGINE
diff --git a/Engine/Source/Developer/AllDesktopTargetPlatform/Private/AllDesktopTargetPlatform.h b/Engine/Source/Developer/AllDesktopTargetPlatform/Private/AllDesktopTargetPlatform.h
--- a/Engine/Source/Developer/AllDesktopTargetPlatform/Private/AllDesktopTargetPlatform.h
+++ b/Engine/Source/Developer/AllDesktopTargetPlatform/Private/AllDesktopTargetPlatform.h
@@ -135,6 +135,9 @@ public:
 	}
 
 	virtual void GetTextureFormats(const UTexture* Texture, TArray<FName>& OutFormats) const override;
+
+	virtual void GetAllTextureFormats(TArray<FName>& OutFormats) const override;
+
 	
 	virtual const UTextureLODSettings& GetTextureLODSettings() const override
 	{
@@ -147,6 +150,7 @@ public:
 	}
 
 	virtual FName GetWaveFormat( const class USoundWave* Wave ) const override;
+	virtual void GetAllWaveFormats(TArray<FName>& OutFormats) const override;
 #endif // WITH_ENGINE
 
 
diff --git a/Engine/Source/Developer/Android/AndroidTargetPlatform/Private/AndroidTargetPlatform.h b/Engine/Source/Developer/Android/AndroidTargetPlatform/Private/AndroidTargetPlatform.h
--- a/Engine/Source/Developer/Android/AndroidTargetPlatform/Private/AndroidTargetPlatform.h
+++ b/Engine/Source/Developer/Android/AndroidTargetPlatform/Private/AndroidTargetPlatform.h
@@ -162,6 +162,8 @@ public:
 
 	virtual void GetTextureFormats( const UTexture* InTexture, TArray<FName>& OutFormats ) const override;
 
+	virtual void GetAllTextureFormats(TArray<FName>& OutFormats) const override;
+
 	virtual const UTextureLODSettings& GetTextureLODSettings() const override;
 
 	virtual void RegisterTextureLODSettings(const UTextureLODSettings* InTextureLODSettings) override
@@ -170,6 +172,7 @@ public:
 	}
 
 	virtual FName GetWaveFormat( const class USoundWave* Wave ) const override;
+	virtual void GetAllWaveFormats( TArray<FName>& OutFormats) const override;
 #endif //WITH_ENGINE
 
 	virtual bool SupportsVariants() const override;
diff --git a/Engine/Source/Developer/Android/AndroidTargetPlatform/Private/AndroidTargetPlatform.inl b/Engine/Source/Developer/Android/AndroidTargetPlatform/Private/AndroidTargetPlatform.inl
--- a/Engine/Source/Developer/Android/AndroidTargetPlatform/Private/AndroidTargetPlatform.inl
+++ b/Engine/Source/Developer/Android/AndroidTargetPlatform/Private/AndroidTargetPlatform.inl
@@ -339,6 +339,47 @@ inline void FAndroidTargetPlatform<TPlatformProperties>::GetTextureFormats( cons
 }
 
 
+
+template<class TPlatformProperties>
+inline void FAndroidTargetPlatform<TPlatformProperties>::GetAllTextureFormats(TArray<FName>& OutFormats) const
+{
+
+
+	OutFormats.Add(AndroidTexFormat::NameG8);
+	OutFormats.Add(AndroidTexFormat::NameRGBA16F);
+	OutFormats.Add(AndroidTexFormat::NameBGRA8);
+	OutFormats.Add(AndroidTexFormat::NameRGBA16F);
+	OutFormats.Add(AndroidTexFormat::NameRGBA16F);
+	OutFormats.Add(AndroidTexFormat::NameBGRA8);
+	OutFormats.Add(AndroidTexFormat::NameG8);
+	OutFormats.Add(AndroidTexFormat::NameG8);
+	OutFormats.Add(AndroidTexFormat::NameG8);
+
+	auto AddAllTextureFormatIfSupports = [=, &OutFormats](bool bIsNonPOT) 
+	{
+		AddTextureFormatIfSupports(AndroidTexFormat::NameAutoPVRTC, OutFormats, bIsNonPOT); 
+		AddTextureFormatIfSupports(AndroidTexFormat::NamePVRTC2, OutFormats, bIsNonPOT);
+		AddTextureFormatIfSupports(AndroidTexFormat::NamePVRTC4, OutFormats, bIsNonPOT);
+	
+		AddTextureFormatIfSupports(AndroidTexFormat::NameAutoDXT, OutFormats, bIsNonPOT);
+		AddTextureFormatIfSupports(AndroidTexFormat::NameDXT1, OutFormats, bIsNonPOT);
+		AddTextureFormatIfSupports(AndroidTexFormat::NameDXT5, OutFormats, bIsNonPOT);
+	
+		AddTextureFormatIfSupports(AndroidTexFormat::NameATC_RGB, OutFormats, bIsNonPOT); 
+		AddTextureFormatIfSupports(AndroidTexFormat::NameATC_RGBA_I, OutFormats, bIsNonPOT);
+	
+		AddTextureFormatIfSupports(AndroidTexFormat::NameAutoETC1, OutFormats, bIsNonPOT); 
+		AddTextureFormatIfSupports(AndroidTexFormat::NameAutoETC2, OutFormats, bIsNonPOT);
+	
+		AddTextureFormatIfSupports(AndroidTexFormat::NameAutoATC, OutFormats, bIsNonPOT);
+	};
+
+	AddAllTextureFormatIfSupports(true);
+	AddAllTextureFormatIfSupports(false);
+
+}
+
+
 template<class TPlatformProperties>
 void FAndroidTargetPlatform<TPlatformProperties>::GetReflectionCaptureFormats( TArray<FName>& OutFormats ) const
 {
@@ -399,6 +440,17 @@ FName FAndroidTargetPlatform<TPlatformProperties>::GetWaveFormat( const class US
 	return NAME_FORMAT;
 }
 
+
+template<class TPlatformProperties>
+void FAndroidTargetPlatform<TPlatformProperties>::GetAllWaveFormats(TArray<FName>& OutFormats) const
+{
+	static FName NAME_OGG(TEXT("OGG"));
+	static FName NAME_ADPCM(TEXT("ADPCM"));
+
+	OutFormats.Add(NAME_OGG); 
+	OutFormats.Add(NAME_ADPCM);
+}
+
 #endif //WITH_ENGINE
 
 template<class TPlatformProperties>
diff --git a/Engine/Source/Developer/Android/Android_ASTCTargetPlatform/Private/Android_ASTCTargetPlatformModule.cpp b/Engine/Source/Developer/Android/Android_ASTCTargetPlatform/Private/Android_ASTCTargetPlatformModule.cpp
--- a/Engine/Source/Developer/Android/Android_ASTCTargetPlatform/Private/Android_ASTCTargetPlatformModule.cpp
+++ b/Engine/Source/Developer/Android/Android_ASTCTargetPlatform/Private/Android_ASTCTargetPlatformModule.cpp
@@ -101,6 +101,32 @@ class FAndroid_ASTCTargetPlatform
 			OutFormats.Add(TextureFormatName);
 		}
 	}
+
+
+	virtual void GetAllTextureFormats(TArray<FName>& OutFormats) const override
+	{
+		// we remap some of the defaults (with PVRTC and ASTC formats)
+		static FName FormatRemap[][2] =
+		{
+			// Default format:				ASTC format:
+			{ { FName(TEXT("DXT1")) },{ FName(TEXT("ASTC_RGB")) } },
+			{ { FName(TEXT("DXT5")) },{ FName(TEXT("ASTC_RGBA")) } },
+			{ { FName(TEXT("DXT5n")) },{ FName(TEXT("ASTC_NormalAG")) } },
+			{ { FName(TEXT("BC5")) },{ FName(TEXT("ASTC_NormalRG")) } },
+			{ { FName(TEXT("BC6H")) },{ FName(TEXT("ASTC_RGB")) } },
+			{ { FName(TEXT("BC7")) },{ FName(TEXT("ASTC_RGBAuto")) } },
+			{ { FName(TEXT("AutoDXT")) },{ FName(TEXT("ASTC_RGBAuto")) } },
+		};
+
+		GetAllDefaultTextureFormats(this, OutFormats, false);
+
+		for (int32 RemapIndex = 0; RemapIndex < ARRAY_COUNT(FormatRemap); ++RemapIndex)
+		{
+			OutFormats.Remove(FormatRemap[RemapIndex][0]);
+			OutFormats.AddUnique(FormatRemap[RemapIndex][1]);
+		}
+	}
+
 #endif
 
 	virtual bool SupportedByExtensionsString( const FString& ExtensionsString, const int GLESVersion ) const override
diff --git a/Engine/Source/Developer/Android/Android_MultiTargetPlatform/Private/Android_MultiTargetPlatformModule.cpp b/Engine/Source/Developer/Android/Android_MultiTargetPlatform/Private/Android_MultiTargetPlatformModule.cpp
--- a/Engine/Source/Developer/Android/Android_MultiTargetPlatform/Private/Android_MultiTargetPlatformModule.cpp
+++ b/Engine/Source/Developer/Android/Android_MultiTargetPlatform/Private/Android_MultiTargetPlatformModule.cpp
@@ -112,6 +112,20 @@ class FAndroid_MultiTargetPlatform
 			}
 		}
 	}
+
+	virtual void GetAllTextureFormats(TArray<FName>& OutFormats) const override
+	{
+		// Ask each platform variant to choose texture formats
+		for (ITargetPlatform* Platform : FormatTargetPlatforms)
+		{
+			TArray<FName> PlatformFormats;
+			Platform->GetAllTextureFormats(PlatformFormats);
+			for (FName Format : PlatformFormats)
+			{
+				OutFormats.AddUnique(Format);
+			}
+		}
+	}
 #endif	
 
 	virtual FText GetVariantDisplayName() const override
diff --git a/Engine/Source/Developer/AutomationController/Private/AutomationControllerManger.cpp b/Engine/Source/Developer/AutomationController/Private/AutomationControllerManger.cpp
--- a/Engine/Source/Developer/AutomationController/Private/AutomationControllerManger.cpp
+++ b/Engine/Source/Developer/AutomationController/Private/AutomationControllerManger.cpp
@@ -1056,13 +1056,16 @@ void FAutomationControllerManager::UpdateDeviceGroups( )
 TArray<FString> FAutomationControllerManager::GetCheckpointFileContents()
 {
 	TestsRun.Empty();
-	FString CheckpointFileName = FString::Printf(TEXT("%sautomationcheckpoint.txt"), *FPaths::AutomationDir());
+	FString CheckpointFileName = FString::Printf(TEXT("%sautomationcheckpoint.log"), *FPaths::AutomationDir());
 	if (IFileManager::Get().FileExists(*CheckpointFileName))
 	{
 		FString FileData;
 		FFileHelper::LoadFileToString(FileData, *CheckpointFileName);
 		FileData.ParseIntoArrayLines(TestsRun);
-		
+		for (int i = 0; i < TestsRun.Num(); i++)
+		{
+			GLog->Log(TEXT("AutomationCheckpoint"), ELogVerbosity::Log, *TestsRun[i]);
+		}
 	}
 	return TestsRun;
 }
@@ -1071,7 +1074,7 @@ FArchive* FAutomationControllerManager::GetCheckpointFileForWrite()
 {
 	if (!CheckpointFile)
 	{
-		FString CheckpointFileName = FString::Printf(TEXT("%sautomationcheckpoint.txt"), *FPaths::AutomationDir());
+		FString CheckpointFileName = FString::Printf(TEXT("%sautomationcheckpoint.log"), *FPaths::AutomationDir());
 		CheckpointFile = IFileManager::Get().CreateFileWriter(*CheckpointFileName, 8);
 	}
 	return CheckpointFile;
@@ -1084,7 +1087,7 @@ void FAutomationControllerManager::CleanUpCheckpointFile()
 		CheckpointFile->Close();
 		CheckpointFile = nullptr;
 	}
-	FString CheckpointFileName = FString::Printf(TEXT("%sautomationcheckpoint.txt"), *FPaths::AutomationDir());
+	FString CheckpointFileName = FString::Printf(TEXT("%sautomationcheckpoint.log"), *FPaths::AutomationDir());
 	if (IFileManager::Get().FileExists(*CheckpointFileName))
 	{
 		IFileManager::Get().Delete(*CheckpointFileName);
@@ -1094,10 +1097,14 @@ void FAutomationControllerManager::CleanUpCheckpointFile()
 void FAutomationControllerManager::WriteLoadedCheckpointDataToFile()
 {
 	GetCheckpointFileForWrite();
-	for (int i = 0; i < TestsRun.Num(); i++)
+	if (CheckpointFile)
 	{
-		FString LineToWrite = FString::Printf(TEXT("%s\n"), *TestsRun[i]);
-		CheckpointFile->Serialize(TCHAR_TO_ANSI(*LineToWrite), LineToWrite.Len());
+		for (int i = 0; i < TestsRun.Num(); i++)
+		{
+			FString LineToWrite = FString::Printf(TEXT("%s\r\n"), *TestsRun[i]);
+			CheckpointFile->Serialize(TCHAR_TO_ANSI(*LineToWrite), LineToWrite.Len());
+			CheckpointFile->Flush();
+		}
 	}
 }
 
@@ -1106,7 +1113,7 @@ void FAutomationControllerManager::WriteLineToCheckpointFile(FString StringToWri
 	GetCheckpointFileForWrite();
 	if (CheckpointFile)
 	{
-		FString LineToWrite = FString::Printf(TEXT("%s\n"), *StringToWrite);
+		FString LineToWrite = FString::Printf(TEXT("%s\r\n"), *StringToWrite);
 		CheckpointFile->Serialize(TCHAR_TO_ANSI(*LineToWrite), LineToWrite.Len());
 		CheckpointFile->Flush();
 	}
diff --git a/Engine/Source/Developer/DerivedDataCache/Private/HierarchicalDerivedDataBackend.h b/Engine/Source/Developer/DerivedDataCache/Private/HierarchicalDerivedDataBackend.h
--- a/Engine/Source/Developer/DerivedDataCache/Private/HierarchicalDerivedDataBackend.h
+++ b/Engine/Source/Developer/DerivedDataCache/Private/HierarchicalDerivedDataBackend.h
@@ -27,6 +27,12 @@ public:
 		, bIsWritable(false)
 	{
 		check(InnerBackends.Num() > 1); // if it is just one, then you don't need this wrapper
+		UpdateAsyncInnerBackends();
+	}
+
+	void UpdateAsyncInnerBackends()
+	{
+		bIsWritable = false;
 		for (int32 CacheIndex = 0; CacheIndex < InnerBackends.Num(); CacheIndex++)
 		{
 			if (InnerBackends[CacheIndex]->IsWritable())
@@ -48,12 +54,16 @@ public:
 	void AddInnerBackend(FDerivedDataBackendInterface* InInner) 
 	{
 		InnerBackends.Add(InInner);
+		AsyncPutInnerBackends.Empty();
+		UpdateAsyncInnerBackends();
 	}
 
 	/** Removes inner backend. */
 	bool RemoveInnerBackend(FDerivedDataBackendInterface* InInner) 
 	{
-		const int32 NumRemoved = InnerBackends.Remove(InInner);
+		int32 NumRemoved = InnerBackends.Remove(InInner);
+		AsyncPutInnerBackends.Empty();
+		UpdateAsyncInnerBackends();
 		return NumRemoved != 0;
 	}
 
diff --git a/Engine/Source/Developer/HTML5/HTML5TargetPlatform/Private/HTML5TargetPlatform.cpp b/Engine/Source/Developer/HTML5/HTML5TargetPlatform/Private/HTML5TargetPlatform.cpp
--- a/Engine/Source/Developer/HTML5/HTML5TargetPlatform/Private/HTML5TargetPlatform.cpp
+++ b/Engine/Source/Developer/HTML5/HTML5TargetPlatform/Private/HTML5TargetPlatform.cpp
@@ -233,6 +233,33 @@ void FHTML5TargetPlatform::GetTextureFormats( const UTexture* Texture, TArray<FN
 	OutFormats.Add( TextureFormatName);
 }
 
+void FHTML5TargetPlatform::GetAllTextureFormats(TArray<FName>& OutFormats) const
+{
+
+#if WITH_EDITOR
+	// Supported texture format names.
+	static FName NameDXT1(TEXT("DXT1"));
+	static FName NameDXT3(TEXT("DXT3"));
+	static FName NameDXT5(TEXT("DXT5"));
+	static FName NameDXT5n(TEXT("DXT5n"));
+	static FName NameAutoDXT(TEXT("AutoDXT"));
+	static FName NameBGRA8(TEXT("BGRA8"));
+	static FName NameG8(TEXT("G8"));
+	static FName NameRGBA16F(TEXT("RGBA16F"));
+	static FName NameRGBA8(TEXT("RGBA8"));
+
+	// Supported texture format names.
+	OutFormats.Add(NameDXT1);
+	OutFormats.Add(NameDXT3);
+	OutFormats.Add(NameDXT5);
+	OutFormats.Add(NameDXT5n);
+	OutFormats.Add(NameAutoDXT);
+	OutFormats.Add(NameBGRA8);
+	OutFormats.Add(NameG8);
+	OutFormats.Add(NameRGBA16F);
+	OutFormats.Add(NameRGBA8);
+#endif
+}
 
 const UTextureLODSettings& FHTML5TargetPlatform::GetTextureLODSettings() const
 {
@@ -246,6 +273,14 @@ FName FHTML5TargetPlatform::GetWaveFormat( const USoundWave* Wave ) const
 	return NAME_OGG;
 }
 
+
+
+void FHTML5TargetPlatform::GetAllWaveFormats(TArray<FName>& OutFormats) const
+{
+	static FName NAME_OGG(TEXT("OGG"));
+	OutFormats.Add(NAME_OGG);
+}
+
 #endif // WITH_ENGINE
 
 void FHTML5TargetPlatform::RefreshHTML5Setup()
diff --git a/Engine/Source/Developer/HTML5/HTML5TargetPlatform/Private/HTML5TargetPlatform.h b/Engine/Source/Developer/HTML5/HTML5TargetPlatform/Private/HTML5TargetPlatform.h
--- a/Engine/Source/Developer/HTML5/HTML5TargetPlatform/Private/HTML5TargetPlatform.h
+++ b/Engine/Source/Developer/HTML5/HTML5TargetPlatform/Private/HTML5TargetPlatform.h
@@ -81,6 +81,8 @@ public:
 
 	virtual void GetTextureFormats( const UTexture* InTexture, TArray<FName>& OutFormats ) const override;
 
+	virtual void GetAllTextureFormats(TArray<FName>& OutFormats) const override;
+
 	virtual void GetReflectionCaptureFormats( TArray<FName>& OutFormats ) const override
 	{
 		OutFormats.Add(FName(TEXT("EncodedHDR")));
@@ -94,6 +96,7 @@ public:
 	}
 
 	virtual FName GetWaveFormat( const class USoundWave* Wave ) const override;
+	virtual void GetAllWaveFormats(TArray<FName>& OutFormats) const override;
 #endif // WITH_ENGINE
 
 	DECLARE_DERIVED_EVENT(FHTML5TargetPlatform, ITargetPlatform::FOnTargetDeviceDiscovered, FOnTargetDeviceDiscovered);
diff --git a/Engine/Source/Developer/Linux/LinuxTargetPlatform/Private/LinuxTargetPlatform.h b/Engine/Source/Developer/Linux/LinuxTargetPlatform/Private/LinuxTargetPlatform.h
--- a/Engine/Source/Developer/Linux/LinuxTargetPlatform/Private/LinuxTargetPlatform.h
+++ b/Engine/Source/Developer/Linux/LinuxTargetPlatform/Private/LinuxTargetPlatform.h
@@ -276,6 +276,16 @@ public:
 		}
 	}
 
+
+	virtual void GetAllTextureFormats(TArray<FName>& OutFormats) const override
+	{
+		if (!IS_DEDICATED_SERVER)
+		{
+			// just use the standard texture format name for this texture
+			GetAllDefaultTextureFormats(this, OutFormats, false);
+		}
+	}
+
 	virtual const UTextureLODSettings& GetTextureLODSettings() const override
 	{
 		return *TextureLODSettings;
@@ -298,6 +308,14 @@ public:
 
 		return NAME_OGG;
 	}
+
+	virtual void GetAllWaveFormats(TArray<FName>& OutFormats) const override
+	{
+		static FName NAME_OGG(TEXT("OGG"));
+		static FName NAME_OPUS(TEXT("OPUS"));
+		OutFormats.Add(NAME_OGG);
+		OutFormats.Add(NAME_OPUS);
+	}
 #endif //WITH_ENGINE
 
 	virtual bool SupportsVariants() const override
diff --git a/Engine/Source/Developer/Mac/MacTargetPlatform/Private/GenericMacTargetPlatform.h b/Engine/Source/Developer/Mac/MacTargetPlatform/Private/GenericMacTargetPlatform.h
--- a/Engine/Source/Developer/Mac/MacTargetPlatform/Private/GenericMacTargetPlatform.h
+++ b/Engine/Source/Developer/Mac/MacTargetPlatform/Private/GenericMacTargetPlatform.h
@@ -202,6 +202,17 @@ return TSuper::SupportsFeature(Feature);
 	}
 
 
+	virtual void GetAllTextureFormats(TArray<FName>& OutFormats) const override
+	{
+		if (!IS_DEDICATED_SERVER)
+		{
+			// just use the standard texture format name for this texture (with no DX11 support)
+			GetAllDefaultTextureFormats(this, OutFormats, false);
+		}
+	}
+
+
+
 	virtual const UTextureLODSettings& GetTextureLODSettings() const override
 	{
 		return *TextureLODSettings;
@@ -225,6 +236,14 @@ return TSuper::SupportsFeature(Feature);
 
 		return NAME_OGG;
 	}
+
+	virtual void GetAllWaveFormats(TArray<FName>& OutFormats) const override
+	{
+		static FName NAME_OGG(TEXT("OGG"));
+		static FName NAME_OPUS(TEXT("OPUS"));
+		OutFormats.Add(NAME_OGG);
+		OutFormats.Add(NAME_OPUS);
+	}
 #endif //WITH_ENGINE
 
 
diff --git a/Engine/Source/Developer/MaterialUtilities/Private/MaterialUtilities.cpp b/Engine/Source/Developer/MaterialUtilities/Private/MaterialUtilities.cpp
--- a/Engine/Source/Developer/MaterialUtilities/Private/MaterialUtilities.cpp
+++ b/Engine/Source/Developer/MaterialUtilities/Private/MaterialUtilities.cpp
@@ -2365,25 +2365,6 @@ void FMaterialUtilities::ClearRenderTargetPool()
 	RenderTargetPool.Empty();
 }
 
-void FMaterialUtilities::FullyLoadMaterialStatic(UMaterialInterface* Material)
-{
-	FLinkerLoad* Linker = Material->GetLinker();
-	if (Linker == nullptr)
-	{
-		return;
-	}
-	// Load material and all expressions.
-	Linker->LoadAllObjects(true);
-	Material->ConditionalPostLoad();
-	// Now load all used textures.
-	TArray<UTexture*> Textures;
-	Material->GetUsedTextures(Textures, EMaterialQualityLevel::Num, true, ERHIFeatureLevel::SM5, true);
-	for (UTexture* Texture : Textures)
-	{
-		Linker->Preload(Texture);
-	}
-}
-
 void FMaterialUtilities::OptimizeSampleArray(TArray<FColor>& InSamples, FIntPoint& InSampleSize)
 {
 	if (InSamples.Num() > 1)
diff --git a/Engine/Source/Developer/MaterialUtilities/Public/MaterialUtilities.h b/Engine/Source/Developer/MaterialUtilities/Public/MaterialUtilities.h
--- a/Engine/Source/Developer/MaterialUtilities/Public/MaterialUtilities.h
+++ b/Engine/Source/Developer/MaterialUtilities/Public/MaterialUtilities.h
@@ -490,11 +490,6 @@ private:
 	/** Clears the pool with available render targets */
 	static void ClearRenderTargetPool();	
 
-	/** Helper function to make sure all the Material and accompanying textures are fully loaded. This function will do nothing useful if material is already loaded.
-	It's intended to work when we're baking a new material during engine startup, inside a PostLoad call - in this case we could have ExportMaterial() call for 
-	material which has not all components loaded yet.*/
-	static void FullyLoadMaterialStatic(UMaterialInterface* Material);
-
 	/** Call back for garbage collector, cleans up the RenderTargetPool if CurrentlyRendering is set to false */
 	void OnPreGarbageCollect();
 
diff --git a/Engine/Source/Developer/PackageDependencyInfo/Private/PackageDependencyInfo.cpp b/Engine/Source/Developer/PackageDependencyInfo/Private/PackageDependencyInfo.cpp
--- a/Engine/Source/Developer/PackageDependencyInfo/Private/PackageDependencyInfo.cpp
+++ b/Engine/Source/Developer/PackageDependencyInfo/Private/PackageDependencyInfo.cpp
@@ -356,7 +356,8 @@ bool FPackageDependencyInfo::InitializeDependencyInfo( const TCHAR* InPackageNam
 
 		FMD5 Hash;
 		FDateTime EarliestTime = DepPkgInfo->TimeStamp;
-		RecursiveResolveDependentHash(DepPkgInfo, Hash, EarliestTime);
+		TSet<FPackageDependencyTrackingInfo*> ReferencedPkgInfo;
+		RecursiveResolveDependentHash(DepPkgInfo, ReferencedPkgInfo, Hash, EarliestTime);
 		DepPkgInfo->DependentHash.Set(Hash);
 		DepPkgInfo->DependentTimeStamp = EarliestTime;
 		DepPkgInfo->bInitializedHashes = true;
@@ -655,12 +656,13 @@ void FPackageDependencyInfo::AsyncDetermineAllDependentPackageInfo(const TArray<
 	{
 		auto Task = new ResolvePackageDependenciesTask(ThreadSafeQueue, this);
 		CurrentTasks.Add(Task);
-		Task->StartBackgroundTask(
+		Task->StartSynchronousTask();
+		/*Task->StartBackgroundTask(
 #if WITH_EDITOR
 			// use the large thread pool when it exists (editor)
 			GLargeThreadPool
 #endif
-		);
+		);*/
 	}
 
 	while ( true )
@@ -695,8 +697,66 @@ void FPackageDependencyInfo::AsyncDetermineAllDependentPackageInfo(const TArray<
 		check(DepPkgInfo->bInitializedDependencies == true || DepPkgInfo->bValid == false);
 	}
 
+
+	ResolveDependentHashes( true, AllPackages, NumAsyncDependencyTasks);
+
+
+
 	for (auto& DepPkgInfo : AllPackages)
 	{
+		check(DepPkgInfo->bBeingProcessed == false);
+		check((DepPkgInfo->bInitializedHashes == true) || (DepPkgInfo->bValid == false));
+	}
+	// need to clean up linkers created in DeterminePackageDependencies as they were created with NoVerify and can mess up cooker operations
+	CollectGarbage(RF_NoFlags, true);
+}
+
+void FPackageDependencyInfo::ResolveDependentHashes(bool Async, TSet<FPackageDependencyTrackingInfo*> PackageInfoSet, int32 NumAsyncDependencyTasks)
+{
+
+	if (!Async)
+	{
+		for (auto& DepPkgInfo : PackageInfoSet)
+		{
+			if (DepPkgInfo->bInitializedHashes || (DepPkgInfo->bValid == false))
+			{
+				continue;
+			}
+
+			FMD5 Hash;
+			FDateTime EarliestTime = DepPkgInfo->TimeStamp;
+			TSet<FPackageDependencyTrackingInfo*> ReferencedPackageInfo;
+			RecursiveResolveDependentHash(DepPkgInfo, ReferencedPackageInfo, Hash, EarliestTime);
+			DepPkgInfo->DependentHash.Set(Hash);
+			DepPkgInfo->DependentTimeStamp = EarliestTime;
+			DepPkgInfo->bInitializedHashes = true;
+		}
+		return;
+	}
+	
+	TThreadSafeWorkList<FPackageDependencyTrackingInfo*> ThreadSafeQueue(PackageInfoSet);
+
+	class FResolveDependencyHashes : public FNonAbandonableTask
+	{
+	private:
+		TThreadSafeWorkList<FPackageDependencyTrackingInfo*>& ThreadSafeQueue;
+		FPackageDependencyInfo* PackageDependencyInfo;
+	public:
+
+		/** Constructor */
+		FResolveDependencyHashes(TThreadSafeWorkList<FPackageDependencyTrackingInfo*>& InThreadSafeQueue, FPackageDependencyInfo* InPackageDependencyInfo) :
+			ThreadSafeQueue(InThreadSafeQueue),
+			PackageDependencyInfo(InPackageDependencyInfo)
+		{
+		}
+
+		/** Performs work on thread */
+		void DoWork()
+		{
+			FPackageDependencyTrackingInfo* DepPkgInfo = nullptr;
+			while (ThreadSafeQueue.Pop(DepPkgInfo))
+			{
+				TSet<FPackageDependencyTrackingInfo*> AllPackages;
 		if (DepPkgInfo->bInitializedHashes || (DepPkgInfo->bValid == false))
 		{
 			continue;
@@ -704,19 +764,58 @@ void FPackageDependencyInfo::AsyncDetermineAllDependentPackageInfo(const TArray<
 
 		FMD5 Hash;
 		FDateTime EarliestTime = DepPkgInfo->TimeStamp;
-		RecursiveResolveDependentHash(DepPkgInfo, Hash, EarliestTime);
+				TSet<FPackageDependencyTrackingInfo*> ReferencedPackageInfo;
+				PackageDependencyInfo->RecursiveResolveDependentHash(DepPkgInfo, ReferencedPackageInfo, Hash, EarliestTime);
 		DepPkgInfo->DependentHash.Set(Hash);
 		DepPkgInfo->DependentTimeStamp = EarliestTime;
 		DepPkgInfo->bInitializedHashes = true;
 	}
+		}
+		FORCEINLINE TStatId GetStatId() const
+		{
+			RETURN_QUICK_DECLARE_CYCLE_STAT(FResolveDependencyHashes, STATGROUP_ThreadPoolAsyncTasks);
+		}
+	};
 
-	for (auto& DepPkgInfo : AllPackages)
+	typedef FAsyncTask<FResolveDependencyHashes> ResolveDependencyHashesTask;
+	TArray<ResolveDependencyHashesTask *> CurrentTasks;
+
+	for (int I = 0; I < NumAsyncDependencyTasks; ++I)
 	{
-		check(DepPkgInfo->bBeingProcessed == false);
-		check((DepPkgInfo->bInitializedHashes == true) || (DepPkgInfo->bValid == false));
+		auto Task = new ResolveDependencyHashesTask(ThreadSafeQueue, this);
+		CurrentTasks.Add(Task);
+		Task->StartBackgroundTask(
+#if WITH_EDITOR
+			// use the large thread pool when it exists (editor)
+			GLargeThreadPool
+#endif
+		);
 	}
-	// need to clean up linkers created in DeterminePackageDependencies as they were created with NoVerify and can mess up cooker operations
-	CollectGarbage(RF_NoFlags, true);
+
+
+	while (true)
+	{
+		bool bIsDone = true;
+		for (auto& Task : CurrentTasks)
+		{
+			if (!Task->IsDone())
+			{
+				bIsDone = false;
+				break;
+			}
+		}
+
+		if (bIsDone)
+			break;
+
+		FPlatformProcess::Sleep(0.01f);
+	}
+	for (auto& Task : CurrentTasks)
+	{
+		check(Task->IsDone());
+		delete Task;
+	}
+
 }
 
 
@@ -1065,7 +1164,7 @@ bool FPackageDependencyInfo::DeterminePackageDependencies(FPackageDependencyTrac
 	check(PkgInfo->bBeingProcessed == false);
 
 
-	if ( PkgInfo->bInitializedDependencies )
+	if ( PkgInfo->bInitializedDependencies || (PkgInfo->bValid == false) )
 	{
 		return true;
 	}
@@ -1081,12 +1180,24 @@ bool FPackageDependencyInfo::DeterminePackageDependencies(FPackageDependencyTrac
 		PkgInfo->bBeingProcessed = true;
 		{
 			// generate hash for file
+
+			if ( Linker->Summary.Guid.IsValid() )
+			{
+				FMD5 PackageSourceHash;
+				FGuid PackageGuid = Linker->Summary.Guid;
+				PackageSourceHash.Update((uint8*)(&PackageGuid), sizeof(FGuid));
+				PkgInfo->FullPackageHash.Set(PackageSourceHash);
+			}
+			else
+			{
+				UE_LOG(LogPackageDependencyInfo, Warning, TEXT("Found package missing GUID %s calculating hash slow way"), *PkgInfo->PackageName);
 			FArchive* Loader = Linker->Loader;
 			int64 OriginalPos = Loader->Tell();
 			Loader->Seek(0);
 			PkgInfo->FullPackageHash = FMD5Hash::HashFileFromArchive(Loader);
 			Loader->Seek(OriginalPos);
 		}
+		}
 		// Start off with setting the dependent time to the package itself
 		PkgInfo->TimeStamp = PkgInfo->TimeStamp;
 
@@ -1167,6 +1278,8 @@ bool FPackageDependencyInfo::DeterminePackageDependencies(FPackageDependencyTrac
 			}
 		}
 
+		DependentPackages.StableSort();
+
 		for (int32 DependentIdx = 0; DependentIdx < DependentPackages.Num(); DependentIdx++)
 		{
 			FString PkgName = DependentPackages[DependentIdx].ToString();
@@ -1184,6 +1297,17 @@ bool FPackageDependencyInfo::DeterminePackageDependencies(FPackageDependencyTrac
 			}
 
 			FPackageDependencyTrackingInfo** pDepPkgInfo = PackageInformation.Find(LongName);
+
+			if ( pDepPkgInfo == nullptr )
+			{
+				pDepPkgInfo = PackageInformation.Find(LongName + FPackageName::GetMapPackageExtension());
+			}
+
+			if (pDepPkgInfo == nullptr)
+			{
+				pDepPkgInfo = PackageInformation.Find(LongName + FPackageName::GetAssetPackageExtension());
+			}
+
 			if ((pDepPkgInfo == NULL) || (*pDepPkgInfo == NULL))
 			{
 				continue;
@@ -1194,9 +1318,6 @@ bool FPackageDependencyInfo::DeterminePackageDependencies(FPackageDependencyTrac
 			AllPackages.Add(DepPkgInfo);
 		}
 
-		DependentPackages.StableSort();
-		
-
 		PkgInfo->bInitializedDependencies = true;
 		PkgInfo->bBeingProcessed = false;
 
@@ -1214,22 +1335,21 @@ bool FPackageDependencyInfo::DeterminePackageDependencies(FPackageDependencyTrac
 	return true;
 }
 
-void FPackageDependencyInfo::RecursiveResolveDependentHash(FPackageDependencyTrackingInfo* PkgInfo, FMD5 &OutHash, FDateTime& EarliestTime)
+void FPackageDependencyInfo::RecursiveResolveDependentHash(FPackageDependencyTrackingInfo* PkgInfo, TSet<FPackageDependencyTrackingInfo*>& ReferencedPkgInfo, FMD5 &OutHash, FDateTime& EarliestTime)
 {
 	if (PkgInfo != NULL)
 	{
-		if ( PkgInfo->bBeingProcessed )
+		if ( PkgInfo->bValid == false )
 		{
-			// circular reference
 			return;
 		}
-		if ( PkgInfo->bValid == false )
+
+		if ( ReferencedPkgInfo.Contains(PkgInfo) )
 		{
-			// package info is invalid so we don't care about it
+			// circular reference or double reference
 			return;
 		}
-		
-		PkgInfo->bBeingProcessed = true;
+		ReferencedPkgInfo.Add(PkgInfo);
 
 
 		check(PkgInfo->bInitializedDependencies);
@@ -1253,10 +1373,8 @@ void FPackageDependencyInfo::RecursiveResolveDependentHash(FPackageDependencyTra
 
 		for (FPackageDependencyTrackingInfo* DependentPackage : SortedDependentPackages)
 		{
-			RecursiveResolveDependentHash(DependentPackage, OutHash, EarliestTime );
+			RecursiveResolveDependentHash(DependentPackage, ReferencedPkgInfo, OutHash, EarliestTime );
 		}
-
-		PkgInfo->bBeingProcessed = false;
 	}
 }
 
diff --git a/Engine/Source/Developer/PackageDependencyInfo/Private/PackageDependencyInfoPrivate.h b/Engine/Source/Developer/PackageDependencyInfo/Private/PackageDependencyInfoPrivate.h
--- a/Engine/Source/Developer/PackageDependencyInfo/Private/PackageDependencyInfoPrivate.h
+++ b/Engine/Source/Developer/PackageDependencyInfo/Private/PackageDependencyInfoPrivate.h
@@ -140,8 +140,9 @@ protected:
 	 * @param PkgInfo package we want to generate the dependent hash for
 	 * @param OutHash generated dependent hash
 	 */
-	void RecursiveResolveDependentHash(FPackageDependencyTrackingInfo* PkgInfo, FMD5 &OutHash, FDateTime& EarliestTime);
+	void RecursiveResolveDependentHash(FPackageDependencyTrackingInfo* PkgInfo, TSet<FPackageDependencyTrackingInfo*>& ReferencedPkgInfo, FMD5 &OutHash, FDateTime& EarliestTime);
 
+	void ResolveDependentHashes(bool Async, TSet<FPackageDependencyTrackingInfo*> PackageInfoSet, int32 NumAsyncDependencyTasks);
 
 
 	/** The newest time stamp of the shader source files. Used when a package contains a material */
diff --git a/Engine/Source/Developer/TargetPlatform/Public/Interfaces/ITargetPlatform.h b/Engine/Source/Developer/TargetPlatform/Public/Interfaces/ITargetPlatform.h
--- a/Engine/Source/Developer/TargetPlatform/Public/Interfaces/ITargetPlatform.h
+++ b/Engine/Source/Developer/TargetPlatform/Public/Interfaces/ITargetPlatform.h
@@ -329,6 +329,13 @@ public:
 	 */
 	virtual void GetTextureFormats( const class UTexture* Texture, TArray<FName>& OutFormats ) const = 0;
 
+	/**
+	 * Gets the texture formats this platform can use
+	 * 
+	 * @param OutFormats will contain all the texture formats which are possible for this platform
+	 */
+	virtual void GetAllTextureFormats( TArray<FName>& OutFormats ) const = 0;
+
 	/**
 	 * Gets the format to use for a particular piece of audio.
 	 *
@@ -337,6 +344,14 @@ public:
 	 */
 	virtual FName GetWaveFormat( const class USoundWave* Wave ) const = 0;
 
+
+	/**
+	* Gets all the formats which can be returned from GetWaveFormat
+	*
+	* @param output array of all the formats
+	*/
+	virtual void GetAllWaveFormats( TArray<FName>& OutFormats ) const = 0;
+
 	/**
 	 * Gets the texture LOD settings used by this platform.
 	 *
diff --git a/Engine/Source/Developer/Windows/WindowsTargetPlatform/Private/GenericWindowsTargetPlatform.h b/Engine/Source/Developer/Windows/WindowsTargetPlatform/Private/GenericWindowsTargetPlatform.h
--- a/Engine/Source/Developer/Windows/WindowsTargetPlatform/Private/GenericWindowsTargetPlatform.h
+++ b/Engine/Source/Developer/Windows/WindowsTargetPlatform/Private/GenericWindowsTargetPlatform.h
@@ -197,6 +197,14 @@ public:
 		}
 	}
 
+	virtual void GetAllTextureFormats(TArray<FName>& OutFormats) const override
+	{
+		if (!IS_DEDICATED_SERVER)
+		{
+			GetAllDefaultTextureFormats(this, OutFormats, bSupportDX11TextureFormats);
+		}
+	}
+
 	virtual const UTextureLODSettings& GetTextureLODSettings() const override
 	{
 		return *TextureLODSettings;
@@ -220,6 +228,14 @@ public:
 		return NAME_OGG;
 	}
 
+	virtual void GetAllWaveFormats(TArray<FName>& OutFormats) const override
+	{
+		static FName NAME_OGG(TEXT("OGG"));
+		static FName NAME_OPUS(TEXT("OPUS"));
+		OutFormats.Add(NAME_OGG);
+		OutFormats.Add(NAME_OPUS);
+	}
+
 #endif //WITH_ENGINE
 
 	virtual bool SupportsVariants() const override
diff --git a/Engine/Source/Developer/iOS/IOSTargetPlatform/Private/IOSTargetPlatform.cpp b/Engine/Source/Developer/iOS/IOSTargetPlatform/Private/IOSTargetPlatform.cpp
--- a/Engine/Source/Developer/iOS/IOSTargetPlatform/Private/IOSTargetPlatform.cpp
+++ b/Engine/Source/Developer/iOS/IOSTargetPlatform/Private/IOSTargetPlatform.cpp
@@ -460,22 +460,24 @@ void FIOSTargetPlatform::GetAllTargetedShaderFormats( TArray<FName>& OutFormats
 	GetAllPossibleShaderFormats(OutFormats);
 }
 
+// we remap some of the defaults (with PVRTC and ASTC formats)
+static FName FormatRemap[] =
+{
+	// original				PVRTC						ASTC
+	FName(TEXT("DXT1")),	FName(TEXT("PVRTC2")),		FName(TEXT("ASTC_RGB")),
+	FName(TEXT("DXT5")),	FName(TEXT("PVRTC4")),		FName(TEXT("ASTC_RGBA")),
+	FName(TEXT("DXT5n")),	FName(TEXT("PVRTCN")),		FName(TEXT("ASTC_NormalAG")),
+	FName(TEXT("BC5")),		FName(TEXT("PVRTCN")),		FName(TEXT("ASTC_NormalRG")),
+	FName(TEXT("AutoDXT")),	FName(TEXT("AutoPVRTC")),	FName(TEXT("ASTC_RGBAuto")),
+	FName(TEXT("BC4")),		FName(TEXT("G8")),			FName(TEXT("G8")),
+};
+static FName NameBGRA8(TEXT("BGRA8"));
+static FName NameG8 = FName(TEXT("G8"));
+
 void FIOSTargetPlatform::GetTextureFormats( const UTexture* Texture, TArray<FName>& OutFormats ) const
 {
 	check(Texture);
 
-	// we remap some of the defaults (with PVRTC and ASTC formats)
-	static FName FormatRemap[] =
-	{
-		// original				PVRTC						ASTC
-		FName(TEXT("DXT1")),	FName(TEXT("PVRTC2")),		FName(TEXT("ASTC_RGB")),
-		FName(TEXT("DXT5")),	FName(TEXT("PVRTC4")),		FName(TEXT("ASTC_RGBA")),
-		FName(TEXT("DXT5n")),	FName(TEXT("PVRTCN")),		FName(TEXT("ASTC_NormalAG")),
-		FName(TEXT("BC5")),		FName(TEXT("PVRTCN")),		FName(TEXT("ASTC_NormalRG")),
-		FName(TEXT("AutoDXT")),	FName(TEXT("AutoPVRTC")),	FName(TEXT("ASTC_RGBAuto")),
-		FName(TEXT("BC4")),		FName(TEXT("G8")),			FName(TEXT("G8")),
-	};
-	static FName NameBGRA8(TEXT("BGRA8"));
 	static FName NamePOTERROR(TEXT("POTERROR"));
 
 	FName TextureFormatName = NAME_None;
@@ -483,7 +485,7 @@ void FIOSTargetPlatform::GetTextureFormats( const UTexture* Texture, TArray<FNam
 	// forward rendering only needs one channel for shadow maps
 	if (Texture->LODGroup == TEXTUREGROUP_Shadowmap && !SupportsMetalMRT())
 	{
-		TextureFormatName = FName(TEXT("G8"));
+		TextureFormatName = NameG8;
 	}
 
 	// if we didn't assign anything specially, then use the defaults
@@ -532,6 +534,36 @@ void FIOSTargetPlatform::GetTextureFormats( const UTexture* Texture, TArray<FNam
 	}
 }
 
+void FIOSTargetPlatform::GetAllTextureFormats(TArray<FName>& OutFormats) const 
+{
+	bool bFoundRemap = false;
+	bool bIncludePVRTC = !bIsTVOS && CookPVRTC();
+	bool bIncludeASTC = bIsTVOS || CookASTC();
+
+	GetAllDefaultTextureFormats(this, OutFormats, false);
+
+	for (int32 RemapIndex = 0; RemapIndex < ARRAY_COUNT(FormatRemap); RemapIndex += 3)
+	{
+		OutFormats.Remove(FormatRemap[RemapIndex+0]);
+	}
+
+	// include the formats we want (use ASTC first so that it is preferred at runtime if they both exist and it's supported)
+	if (bIncludeASTC)
+	{
+		for (int32 RemapIndex = 0; RemapIndex < ARRAY_COUNT(FormatRemap); RemapIndex += 3)
+		{
+			OutFormats.AddUnique(FormatRemap[RemapIndex + 2]);
+		}
+	}
+	if (bIncludePVRTC)
+	{
+		for (int32 RemapIndex = 0; RemapIndex < ARRAY_COUNT(FormatRemap); RemapIndex += 3)
+		{
+			OutFormats.AddUnique(FormatRemap[RemapIndex + 1]);
+		}
+	}
+}
+
 
 const UTextureLODSettings& FIOSTargetPlatform::GetTextureLODSettings() const
 {
@@ -545,5 +577,12 @@ FName FIOSTargetPlatform::GetWaveFormat( const class USoundWave* Wave ) const
 	return NAME_ADPCM;
 }
 
+
+void FIOSTargetPlatform::GetAllWaveFormats(TArray<FName>& OutFormat) const
+{
+	static FName NAME_ADPCM(TEXT("ADPCM"));
+	OutFormat.Add(NAME_ADPCM);
+}
+
 #endif // WITH_ENGINE
 
diff --git a/Engine/Source/Developer/iOS/IOSTargetPlatform/Private/IOSTargetPlatform.h b/Engine/Source/Developer/iOS/IOSTargetPlatform/Private/IOSTargetPlatform.h
--- a/Engine/Source/Developer/iOS/IOSTargetPlatform/Private/IOSTargetPlatform.h
+++ b/Engine/Source/Developer/iOS/IOSTargetPlatform/Private/IOSTargetPlatform.h
@@ -115,6 +115,8 @@ public:
 
 	virtual void GetTextureFormats( const UTexture* Texture, TArray<FName>& OutFormats ) const override;
 
+	virtual void GetAllTextureFormats( TArray<FName>& OutFormats) const override;
+
 	virtual const UTextureLODSettings& GetTextureLODSettings() const override;
 
 	virtual void RegisterTextureLODSettings(const UTextureLODSettings* InTextureLODSettings) override
@@ -123,6 +125,7 @@ public:
 	}
 
 	virtual FName GetWaveFormat( const class USoundWave* Wave ) const override;
+	virtual void GetAllWaveFormats(TArray<FName>& OutFormat) const override;
 #endif // WITH_ENGINE
 
 	virtual void GetBuildProjectSettingKeys(FString& OutSection, TArray<FString>& InBoolKeys, TArray<FString>& InIntKeys, TArray<FString>& InStringKeys) const override
diff --git a/Engine/Source/Editor/AnimGraph/Classes/AnimGraphNode_ScaleChainLength.h b/Engine/Source/Editor/AnimGraph/Classes/AnimGraphNode_ScaleChainLength.h
new file mode 100644
--- /dev/null
+++ b/Engine/Source/Editor/AnimGraph/Classes/AnimGraphNode_ScaleChainLength.h
@@ -0,0 +1,28 @@
+// Copyright 1998-2017 Epic Games, Inc. All Rights Reserved.
+
+#pragma once
+
+#include "Editor/AnimGraph/Classes/AnimGraphNode_Base.h"
+#include "BoneControllers/AnimNode_ScaleChainLength.h"
+#include "AnimGraphNode_ScaleChainLength.generated.h"
+
+class USkeletalMeshComponent;
+
+UCLASS(MinimalAPI)
+class UAnimGraphNode_ScaleChainLength : public UAnimGraphNode_Base
+{
+	GENERATED_UCLASS_BODY()
+
+	UPROPERTY(EditAnywhere, Category = Settings)
+	FAnimNode_ScaleChainLength Node;
+
+	//~ Begin UEdGraphNode Interface.
+	virtual FLinearColor GetNodeTitleColor() const override;
+	virtual FText GetTooltipText() const override;
+	virtual FText GetNodeTitle(ENodeTitleType::Type TitleType) const override;
+	//~ End UEdGraphNode Interface.
+
+	//~ Begin UAnimGraphNode_Base Interface
+	virtual FString GetNodeCategory() const override;
+	//~ End UAnimGraphNode_Base Interface
+};
diff --git a/Engine/Source/Editor/AnimGraph/Classes/K2Node_PlayMontage.h b/Engine/Source/Editor/AnimGraph/Classes/K2Node_PlayMontage.h
new file mode 100644
--- /dev/null
+++ b/Engine/Source/Editor/AnimGraph/Classes/K2Node_PlayMontage.h
@@ -0,0 +1,24 @@
+// Copyright 1998-2017 Epic Games, Inc. All Rights Reserved.
+
+#pragma once
+
+#include "CoreMinimal.h"
+#include "UObject/ObjectMacros.h"
+#include "K2Node_BaseAsyncTask.h"
+#include "K2Node_PlayMontage.generated.h"
+
+UCLASS()
+class UK2Node_PlayMontage : public UK2Node_BaseAsyncTask
+{
+	GENERATED_UCLASS_BODY()
+
+	//~ Begin UEdGraphNode Interface
+	virtual FText GetTooltipText() const override;
+	virtual FText GetNodeTitle(ENodeTitleType::Type TitleType) const override;
+	virtual void GetPinHoverText(const UEdGraphPin& Pin, FString& HoverTextOut) const override;
+	//~ End UEdGraphNode Interface
+
+	//~ Begin UK2Node Interface
+	virtual FText GetMenuCategory() const override;
+	//~ End UK2Node Interface
+};
diff --git a/Engine/Source/Editor/AnimGraph/Private/AnimGraphNode_ScaleChainLength.cpp b/Engine/Source/Editor/AnimGraph/Private/AnimGraphNode_ScaleChainLength.cpp
new file mode 100644
--- /dev/null
+++ b/Engine/Source/Editor/AnimGraph/Private/AnimGraphNode_ScaleChainLength.cpp
@@ -0,0 +1,36 @@
+// Copyright 1998-2017 Epic Games, Inc. All Rights Reserved.
+
+#include "AnimGraphNode_ScaleChainLength.h"
+#include "Animation/AnimInstance.h"
+
+/////////////////////////////////////////////////////
+// UAnimGraphNode_ScaleChainLength
+
+#define LOCTEXT_NAMESPACE "AnimGraph_ScaleChainLength"
+
+UAnimGraphNode_ScaleChainLength::UAnimGraphNode_ScaleChainLength(const FObjectInitializer& ObjectInitializer)
+	: Super(ObjectInitializer)
+{
+}
+
+FLinearColor UAnimGraphNode_ScaleChainLength::GetNodeTitleColor() const
+{
+	return FLinearColor(0.7f, 0.7f, 0.7f);
+}
+
+FText UAnimGraphNode_ScaleChainLength::GetTooltipText() const
+{
+	return LOCTEXT("AnimGraphNode_ScaleChainLength_Tooltip", "Scales a bone chain based on distance from StartBone to TargetLocation");
+}
+
+FText UAnimGraphNode_ScaleChainLength::GetNodeTitle(ENodeTitleType::Type TitleType) const
+{
+	return LOCTEXT("AnimGraphNode_ScaleChainLength", "Scale Chain Length");
+}
+
+FString UAnimGraphNode_ScaleChainLength::GetNodeCategory() const
+{
+	return TEXT("AnimNode");
+}
+
+#undef LOCTEXT_NAMESPACE
diff --git a/Engine/Source/Editor/AnimGraph/Private/K2Node_PlayMontage.cpp b/Engine/Source/Editor/AnimGraph/Private/K2Node_PlayMontage.cpp
new file mode 100644
--- /dev/null
+++ b/Engine/Source/Editor/AnimGraph/Private/K2Node_PlayMontage.cpp
@@ -0,0 +1,51 @@
+// Copyright 1998-2017 Epic Games, Inc. All Rights Reserved.
+
+#include "K2Node_PlayMontage.h"
+#include "PlayMontageCallbackProxy.h"
+
+#define LOCTEXT_NAMESPACE "K2Node"
+
+UK2Node_PlayMontage::UK2Node_PlayMontage(const FObjectInitializer& ObjectInitializer)
+	: Super(ObjectInitializer)
+{
+	ProxyFactoryFunctionName = GET_FUNCTION_NAME_CHECKED(UPlayMontageCallbackProxy, CreateProxyObjectForPlayMontage);
+	ProxyFactoryClass = UPlayMontageCallbackProxy::StaticClass();
+	ProxyClass = UPlayMontageCallbackProxy::StaticClass();
+}
+
+FText UK2Node_PlayMontage::GetTooltipText() const
+{
+	return LOCTEXT("K2Node_PlayMontage_Tooltip", "Plays a Montage on a SkeletalMeshComponent");
+}
+
+FText UK2Node_PlayMontage::GetNodeTitle(ENodeTitleType::Type TitleType) const
+{
+	return LOCTEXT("PlayMontage", "Play Montage");
+}
+
+FText UK2Node_PlayMontage::GetMenuCategory() const
+{
+	return LOCTEXT("PlayMontageCategory", "Animation|Montage");
+}
+
+void UK2Node_PlayMontage::GetPinHoverText(const UEdGraphPin& Pin, FString& HoverTextOut) const
+{
+	Super::GetPinHoverText(Pin, HoverTextOut);
+
+	static FName NAME_OnNotifyBegin = FName(TEXT("OnNotifyBegin"));
+	static FName NAME_OnNotifyEnd = FName(TEXT("OnNotifyEnd"));
+
+	const FName PinName = FName(*Pin.PinName);
+	if (PinName == NAME_OnNotifyBegin)
+	{
+		FText ToolTipText = LOCTEXT("K2Node_PlayMontage_OnNotifyBegin_Tooltip", "Event called when using a PlayMontageNotify or PlayMontageNotifyWindow Notify in a Montage.");
+		HoverTextOut = FString::Printf(TEXT("%s\n%s"), *ToolTipText.ToString(), *HoverTextOut);
+	}
+	else if (PinName == NAME_OnNotifyEnd)
+	{
+		FText ToolTipText = LOCTEXT("K2Node_PlayMontage_OnNotifyEnd_Tooltip", "Event called when using a PlayMontageNotifyWindow Notify in a Montage.");
+		HoverTextOut = FString::Printf(TEXT("%s\n%s"), *ToolTipText.ToString(), *HoverTextOut);
+	}
+}
+
+#undef LOCTEXT_NAMESPACE
diff --git a/Engine/Source/Editor/PropertyEditor/Private/DetailGroup.cpp b/Engine/Source/Editor/PropertyEditor/Private/DetailGroup.cpp
--- a/Engine/Source/Editor/PropertyEditor/Private/DetailGroup.cpp
+++ b/Engine/Source/Editor/PropertyEditor/Private/DetailGroup.cpp
@@ -89,6 +89,16 @@ bool FDetailGroup::GetExpansionState() const
 	return false;
 }
 
+TSharedPtr<FDetailPropertyRow> FDetailGroup::GetHeaderPropertyRow() const
+{
+	return HeaderCustomization.IsValid() ? HeaderCustomization->PropertyRow : nullptr;
+}
+
+TSharedPtr<FPropertyNode> FDetailGroup::GetHeaderPropertyNode() const
+{
+	return HeaderCustomization.IsValid() ? HeaderCustomization->GetPropertyNode() : nullptr;
+}
+
 bool FDetailGroup::HasColumns() const
 {
 	if( HeaderCustomization.IsValid() && HeaderCustomization->HasPropertyNode() && HeaderCustomization->PropertyRow->HasColumns() )
diff --git a/Engine/Source/Editor/PropertyEditor/Private/DetailGroup.h b/Engine/Source/Editor/PropertyEditor/Private/DetailGroup.h
--- a/Engine/Source/Editor/PropertyEditor/Private/DetailGroup.h
+++ b/Engine/Source/Editor/PropertyEditor/Private/DetailGroup.h
@@ -30,6 +30,9 @@ public:
 	virtual void ToggleExpansion( bool bExpand ) override;
 	virtual bool GetExpansionState() const override;
 
+	TSharedPtr<FDetailPropertyRow> GetHeaderPropertyRow() const;
+	TSharedPtr<FPropertyNode> GetHeaderPropertyNode() const;
+
 	/** @return The name of the group */
 	FName GetGroupName() const { return GroupName; }
 
diff --git a/Engine/Source/Editor/PropertyEditor/Private/PropertyHandleImpl.cpp b/Engine/Source/Editor/PropertyEditor/Private/PropertyHandleImpl.cpp
--- a/Engine/Source/Editor/PropertyEditor/Private/PropertyHandleImpl.cpp
+++ b/Engine/Source/Editor/PropertyEditor/Private/PropertyHandleImpl.cpp
@@ -2073,7 +2073,18 @@ TSharedPtr<IPropertyHandle> FPropertyHandleBase::GetParentHandle() const
 	TSharedPtr<FPropertyNode> ParentNode = Implementation->GetPropertyNode()->GetParentNodeSharedPtr();
 	if( ParentNode.IsValid() )
 	{
-		return PropertyEditorHelpers::GetPropertyHandle( ParentNode.ToSharedRef(), Implementation->GetNotifyHook(), Implementation->GetPropertyUtilities() ).ToSharedRef();
+		return PropertyEditorHelpers::GetPropertyHandle( ParentNode.ToSharedRef(), Implementation->GetNotifyHook(), Implementation->GetPropertyUtilities() );
+	}
+
+	return nullptr;
+}
+
+TSharedPtr<IPropertyHandle> FPropertyHandleBase::GetKeyHandle() const
+{
+	TSharedPtr<FPropertyNode> KeyNode = Implementation->GetPropertyNode()->GetPropertyKeyNode();
+	if (KeyNode.IsValid())
+	{
+		return PropertyEditorHelpers::GetPropertyHandle(KeyNode.ToSharedRef(), Implementation->GetNotifyHook(), Implementation->GetPropertyUtilities());
 	}
 
 	return nullptr;
diff --git a/Engine/Source/Editor/PropertyEditor/Private/PropertyHandleImpl.h b/Engine/Source/Editor/PropertyEditor/Private/PropertyHandleImpl.h
--- a/Engine/Source/Editor/PropertyEditor/Private/PropertyHandleImpl.h
+++ b/Engine/Source/Editor/PropertyEditor/Private/PropertyHandleImpl.h
@@ -441,6 +441,7 @@ public:
 	virtual TSharedPtr<IPropertyHandle> GetChildHandle( uint32 ChildIndex ) const override;
 	virtual TSharedPtr<IPropertyHandle> GetChildHandle( FName ChildName, bool bRecurse = true ) const override;
 	virtual TSharedPtr<IPropertyHandle> GetParentHandle() const override;
+	virtual TSharedPtr<IPropertyHandle> GetKeyHandle() const override;
 	virtual void EnumerateRawData( const EnumerateRawDataFuncRef& InRawDataCallback ) override;
 	virtual void EnumerateConstRawData( const EnumerateConstRawDataFuncRef& InRawDataCallback ) const override;
 	virtual void AccessRawData( TArray<void*>& RawData ) override;
diff --git a/Engine/Source/Editor/PropertyEditor/Private/SDetailSingleItemRow.cpp b/Engine/Source/Editor/PropertyEditor/Private/SDetailSingleItemRow.cpp
--- a/Engine/Source/Editor/PropertyEditor/Private/SDetailSingleItemRow.cpp
+++ b/Engine/Source/Editor/PropertyEditor/Private/SDetailSingleItemRow.cpp
@@ -9,6 +9,7 @@
 #include "IDetailKeyframeHandler.h"
 #include "IDetailPropertyExtensionHandler.h"
 #include "DetailPropertyRow.h"
+#include "DetailGroup.h"
 
 namespace DetailWidgetConstants
 {
@@ -372,11 +373,16 @@ bool SDetailSingleItemRow::OnContextMenuOpening( FMenuBuilder& MenuBuilder )
 		CopyAction = Customization->GetWidgetRow().CopyMenuAction;
 		PasteAction = Customization->GetWidgetRow().PasteMenuAction;
 	}
-	else if(Customization->HasPropertyNode())
+	else
 	{
-		static const FName DisableCopyPasteMetaDataName("DisableCopyPaste");
+		TSharedPtr<FPropertyNode> PropertyNode = Customization->GetPropertyNode();
+		if (!PropertyNode.IsValid() && Customization->DetailGroup.IsValid())
+		{
+			PropertyNode = Customization->DetailGroup->GetHeaderPropertyNode();
+		}
 
-		if(!Customization->GetPropertyNode()->ParentOrSelfHasMetaData(DisableCopyPasteMetaDataName))
+		static const FName DisableCopyPasteMetaDataName("DisableCopyPaste");
+		if (PropertyNode.IsValid() && !PropertyNode->ParentOrSelfHasMetaData(DisableCopyPasteMetaDataName))
 		{
 			CopyAction.ExecuteAction = FExecuteAction::CreateSP(this, &SDetailSingleItemRow::OnCopyProperty);
 			PasteAction.ExecuteAction = FExecuteAction::CreateSP(this, &SDetailSingleItemRow::OnPasteProperty);
@@ -414,14 +420,22 @@ void SDetailSingleItemRow::OnLeftColumnResized( float InNewWidth )
 
 void SDetailSingleItemRow::OnCopyProperty()
 {
-	if( Customization->HasPropertyNode() && OwnerTreeNode.IsValid() )
+	if (OwnerTreeNode.IsValid())
 	{
-		TSharedPtr<IPropertyHandle> Handle = PropertyEditorHelpers::GetPropertyHandle( Customization->GetPropertyNode().ToSharedRef(), OwnerTreeNode.Pin()->GetDetailsView().GetNotifyHook(),  OwnerTreeNode.Pin()->GetDetailsView().GetPropertyUtilities() );
-
-		FString Value;
-		if( Handle->GetValueAsFormattedString(Value, PPF_Copy) == FPropertyAccess::Success )
+		TSharedPtr<FPropertyNode> PropertyNode = Customization->GetPropertyNode();
+		if (!PropertyNode.IsValid() && Customization->DetailGroup.IsValid())
 		{
-			FPlatformMisc::ClipboardCopy(*Value);
+			PropertyNode = Customization->DetailGroup->GetHeaderPropertyNode();
+		}
+		if (PropertyNode.IsValid())
+		{
+			TSharedPtr<IPropertyHandle> Handle = PropertyEditorHelpers::GetPropertyHandle(PropertyNode.ToSharedRef(), OwnerTreeNode.Pin()->GetDetailsView().GetNotifyHook(), OwnerTreeNode.Pin()->GetDetailsView().GetPropertyUtilities());
+
+			FString Value;
+			if (Handle->GetValueAsFormattedString(Value, PPF_Copy) == FPropertyAccess::Success)
+			{
+				FPlatformMisc::ClipboardCopy(*Value);
+			}
 		}
 	}
 }
@@ -431,27 +445,37 @@ void SDetailSingleItemRow::OnPasteProperty()
 	FString ClipboardContent;
 	FPlatformMisc::ClipboardPaste(ClipboardContent);
 
-	if( !ClipboardContent.IsEmpty() && Customization->HasPropertyNode() && OwnerTreeNode.IsValid() )
+	if (!ClipboardContent.IsEmpty() && OwnerTreeNode.IsValid())
 	{
-		TSharedPtr<IPropertyHandle> Handle = PropertyEditorHelpers::GetPropertyHandle(Customization->GetPropertyNode().ToSharedRef(), OwnerTreeNode.Pin()->GetDetailsView().GetNotifyHook(), OwnerTreeNode.Pin()->GetDetailsView().GetPropertyUtilities());
+		TSharedPtr<FPropertyNode> PropertyNode = Customization->GetPropertyNode();
+		if (!PropertyNode.IsValid() && Customization->DetailGroup.IsValid())
+		{
+			PropertyNode = Customization->DetailGroup->GetHeaderPropertyNode();
+		}
+		if (PropertyNode.IsValid())
+		{
+			TSharedPtr<IPropertyHandle> Handle = PropertyEditorHelpers::GetPropertyHandle(PropertyNode.ToSharedRef(), OwnerTreeNode.Pin()->GetDetailsView().GetNotifyHook(), OwnerTreeNode.Pin()->GetDetailsView().GetPropertyUtilities());
 
-		Handle->SetValueFromFormattedString( ClipboardContent );
+			Handle->SetValueFromFormattedString(ClipboardContent);
+		}
 	}
-
 }
 
 bool SDetailSingleItemRow::CanPasteProperty() const
 {
 	// Prevent paste from working if the property's edit condition is not met.
-	if (Customization->PropertyRow.IsValid())
+	TSharedPtr<FDetailPropertyRow> PropertyRow = Customization->PropertyRow;
+	if (!PropertyRow.IsValid() && Customization->DetailGroup.IsValid())
 	{
-		FPropertyEditor* PropertyEditor = Customization->PropertyRow->GetPropertyEditor().Get();
+		PropertyRow = Customization->DetailGroup->GetHeaderPropertyRow();
+	}
+
+	if (PropertyRow.IsValid())
+	{
+		FPropertyEditor* PropertyEditor = PropertyRow->GetPropertyEditor().Get();
 		if (PropertyEditor)
 		{
-			if ((PropertyEditor->HasEditCondition() && !PropertyEditor->IsEditConditionMet()) || PropertyEditor->IsEditConst())
-			{
-				return false;
-			}
+			return !PropertyEditor->IsEditConst() && (!PropertyEditor->HasEditCondition() || PropertyEditor->IsEditConditionMet());
 		}
 	}
 
diff --git a/Engine/Source/Editor/PropertyEditor/Public/PropertyHandle.h b/Engine/Source/Editor/PropertyEditor/Public/PropertyHandle.h
--- a/Engine/Source/Editor/PropertyEditor/Public/PropertyHandle.h
+++ b/Engine/Source/Editor/PropertyEditor/Public/PropertyHandle.h
@@ -345,6 +345,11 @@ public:
 	 */
 	virtual TSharedPtr<IPropertyHandle> GetParentHandle() const = 0;
 
+	/**
+	 * @return The property handle to the key element for this value if this is a map element
+	 */
+	virtual TSharedPtr<IPropertyHandle> GetKeyHandle() const = 0;
+
 	/**
 	 * @return The number of children the property handle has
 	 */
diff --git a/Engine/Source/Editor/UnrealEd/Classes/CookOnTheSide/CookOnTheFlyServer.h b/Engine/Source/Editor/UnrealEd/Classes/CookOnTheSide/CookOnTheFlyServer.h
--- a/Engine/Source/Editor/UnrealEd/Classes/CookOnTheSide/CookOnTheFlyServer.h
+++ b/Engine/Source/Editor/UnrealEd/Classes/CookOnTheSide/CookOnTheFlyServer.h
@@ -43,7 +43,8 @@ enum class ECookInitializationFlags
 	TestCook =									0x00002000, // test the cooker garbage collection process and cooking (cooker will never end just keep testing).
 	IterateOnHash =								0x00004000, // when using iterative cooking use hashes of original files instead of timestamps
 	LogDebugInfo =								0x00008000, // enables additional debug log information
-	IterateSharedBuild =								0x00010000,
+	IterateSharedBuild =						0x00010000, // iterate from a build in teh SharedIterativeBuild directory 
+	IgnoreIniSettingsOutOfDate =				0x00020000, // if the inisettings say the cook is out of date keep using the previously cooked build
 };
 ENUM_CLASS_FLAGS(ECookInitializationFlags);
 
diff --git a/Engine/Source/Editor/UnrealEd/Private/Commandlets/ContentCommandlets.cpp b/Engine/Source/Editor/UnrealEd/Private/Commandlets/ContentCommandlets.cpp
--- a/Engine/Source/Editor/UnrealEd/Private/Commandlets/ContentCommandlets.cpp
+++ b/Engine/Source/Editor/UnrealEd/Private/Commandlets/ContentCommandlets.cpp
@@ -1108,13 +1108,20 @@ void UResavePackagesCommandlet::PerformAdditionalOperations(class UWorld* World,
 						{
 							if (IFileManager::Get().FileExists(*MapBuildDataFilename))
 							{
-								CheckoutFile(MapBuildDataFilename, true);
+								if ( CheckoutFile(MapBuildDataFilename, true) )
+								{
+									SublevelFilenames.Add(MapBuildDataFilename);
+								}
+
 								SavePackageHelper(MapBuildDataPackage, MapBuildDataFilename); 
 							}
 							else
 							{
 								SavePackageHelper(MapBuildDataPackage, MapBuildDataFilename);
-								CheckoutFile(MapBuildDataFilename, true);
+								if (CheckoutFile(MapBuildDataFilename, true))
+								{
+									SublevelFilenames.Add(MapBuildDataFilename);
+								}
 							}
 						}
 					}
diff --git a/Engine/Source/Editor/UnrealEd/Private/Commandlets/CookCommandlet.cpp b/Engine/Source/Editor/UnrealEd/Private/Commandlets/CookCommandlet.cpp
--- a/Engine/Source/Editor/UnrealEd/Private/Commandlets/CookCommandlet.cpp
+++ b/Engine/Source/Editor/UnrealEd/Private/Commandlets/CookCommandlet.cpp
@@ -599,14 +599,14 @@ bool UCookCommandlet::CookByTheBook( const TArray<ITargetPlatform*>& Platforms,
 	bool bTestCook = FParse::Param(*Params, TEXT("Testcook"));
 	CookFlags |= bTestCook ? ECookInitializationFlags::TestCook : ECookInitializationFlags::None;
 	CookFlags |= FParse::Param(*Params, TEXT("iteratehash")) ? ECookInitializationFlags::IterateOnHash : ECookInitializationFlags::None;
-	CookFlags |= FParse::Param(*Params, TEXT("iteratesharedbuild")) ? ECookInitializationFlags::IterateSharedBuild : ECookInitializationFlags::None;
 	CookFlags |= FParse::Param(*Params, TEXT("logdebuginfo")) ? ECookInitializationFlags::LogDebugInfo : ECookInitializationFlags::None;
 
 	// shared cooked build flags
 	bool bIterateSharedCookedBuild = Switches.Contains(TEXT("iteratesharedcookedbuild"));
-	ECookInitializationFlags IterateSharedCookedBuildFlags = ECookInitializationFlags::IterateOnHash | ECookInitializationFlags::Iterative;
+	ECookInitializationFlags IterateSharedCookedBuildFlags = ECookInitializationFlags::IterateOnHash | ECookInitializationFlags::Iterative | ECookInitializationFlags::IterateSharedBuild;
 	CookFlags |= bIterateSharedCookedBuild ? IterateSharedCookedBuildFlags : ECookInitializationFlags::None;
-	
+
+	CookFlags |= FParse::Param(*Params, TEXT("Ignoreinisettingsoutofdate")) ? ECookInitializationFlags::IgnoreIniSettingsOutOfDate : ECookInitializationFlags::None;
 
 	TArray<UClass*> FullGCAssetClasses;
 	if (FullGCAssetClassNames.Num())
diff --git a/Engine/Source/Editor/UnrealEd/Private/CookOnTheFlyServer.cpp b/Engine/Source/Editor/UnrealEd/Private/CookOnTheFlyServer.cpp
--- a/Engine/Source/Editor/UnrealEd/Private/CookOnTheFlyServer.cpp
+++ b/Engine/Source/Editor/UnrealEd/Private/CookOnTheFlyServer.cpp
@@ -761,6 +761,7 @@ UCookOnTheFlyServer::UCookOnTheFlyServer(const FObjectInitializer& ObjectInitial
 	CookByTheBookOptions(NULL),
 	CookFlags(ECookInitializationFlags::None),
 	bIsInitializingSandbox(false),
+	bIgnoreMarkupPackageAlreadyLoaded(false),
 	bIsSavingPackage(false)
 {
 }
@@ -2797,7 +2798,7 @@ void UCookOnTheFlyServer::SaveCookedPackage(UPackage* Package, uint32 SaveFlags,
 
 		SCOPE_TIMER(ResolveRedirectors);
 		COOK_STAT(FScopedDurationTimer ResolveRedirectorsTimer(DetailedCookStats::TickCookOnTheSideResolveRedirectorsTimeSec));
-		GRedirectCollector.ResolveStringAssetReference(Package->GetPathName());
+		GRedirectCollector.ResolveStringAssetReference(Package->GetPathName(), false);
 
 		GIsCookerLoadingPackage = false;
 	}
@@ -3244,7 +3245,19 @@ FString UCookOnTheFlyServer::GetOutputDirectoryOverride() const
 	else if (!OutputDirectory.Contains(TEXT("[Platform]"), ESearchCase::IgnoreCase, ESearchDir::FromEnd) )
 	{
 		// Output directory needs to contain [Platform] token to be able to cook for multiple targets.
-		OutputDirectory = FPaths::Combine(*OutputDirectory, TEXT("[Platform]"));
+		if ( IsCookByTheBookMode() )
+		{
+			ITargetPlatformManagerModule& TPM = GetTargetPlatformManagerRef();
+			const TArray<ITargetPlatform*>& TargetPlatforms = TPM.GetCookingTargetPlatforms();
+
+			// more then one target platform specified append "[platform]" to output directory so that multiple platforms can be cooked
+			check( TargetPlatforms.Num() == 1 );
+		}
+		else
+		{
+			// cook on the fly we have to platform because we don't know which platforms we are cooking for up front
+			OutputDirectory = FPaths::Combine(*OutputDirectory, TEXT("[Platform]"));
+		}
 	}
 	FPaths::NormalizeDirectoryName(OutputDirectory);
 
@@ -3252,7 +3265,7 @@ FString UCookOnTheFlyServer::GetOutputDirectoryOverride() const
 }
 
 template<class T>
-void GetVersionFormatNumbersForIniVersionStrings( TArray<FString>& IniVersionStrings, const FString& FormatName, const TArray<T> &FormatArray )
+void GetVersionFormatNumbersForIniVersionStrings( TArray<FString>& IniVersionStrings, const FString& FormatName, const TArray<const T> &FormatArray )
 {
 	for ( const T& Format : FormatArray )
 	{
@@ -3303,7 +3316,8 @@ void GetAdditionalCurrentIniVersionStrings( const ITargetPlatform* TargetPlatfor
 	}
 
 
-	const UTextureLODSettings& LodSettings = TargetPlatform->GetTextureLODSettings();
+	// these settings are covered by the general search through accessed ini settings 
+	/*const UTextureLODSettings& LodSettings = TargetPlatform->GetTextureLODSettings();
 
 	UEnum* TextureGroupEnum = FindObject<UEnum>(NULL, TEXT("Engine.TextureGroup"));
 	UEnum* TextureMipGenSettingsEnum = FindObject<UEnum>(NULL, TEXT("Engine.TextureMipGenSettings"));
@@ -3325,15 +3339,73 @@ void GetAdditionalCurrentIniVersionStrings( const ITargetPlatform* TargetPlatfor
 		// FString MaxMipVersionString = FString::Printf(TEXT("TextureLODGroupMaxMipCount:%s#%d"), *TextureGroupEnum->GetNameStringByIndex(I), MaxMipCount);
 		//IniVersionStrings.Emplace(MoveTemp(MaxMipVersionString));
 		IniVersionMap.Add(FString::Printf(TEXT("TextureLODGroupMaxMipCount:%s"), *TextureGroupString), FString::Printf(TEXT("%d"), MaxMipCount));
-	}
+	}*/
 
 	// save off the ddc version numbers also
 	ITargetPlatformManagerModule* TPM = GetTargetPlatformManager();
 	check(TPM);
 
-	GetVersionFormatNumbersForIniVersionStrings(IniVersionMap, TEXT("AudioFormat"), TPM->GetAudioFormats());
-	GetVersionFormatNumbersForIniVersionStrings(IniVersionMap, TEXT("TextureFormat"), TPM->GetTextureFormats());
-	GetVersionFormatNumbersForIniVersionStrings(IniVersionMap, TEXT("ShaderFormat"), TPM->GetShaderFormats());
+	{
+		TArray<FName> AllWaveFormatNames;
+		TargetPlatform->GetAllWaveFormats(AllWaveFormatNames);
+		TArray<const IAudioFormat*> SupportedWaveFormats;
+		for ( const auto& WaveName : AllWaveFormatNames )
+		{
+			const IAudioFormat* AudioFormat = TPM->FindAudioFormat(WaveName);
+			if (AudioFormat)
+			{
+				SupportedWaveFormats.Add(AudioFormat);
+			}
+			else
+			{
+				UE_LOG(LogCook, Warning, TEXT("Unable to find audio format \"%s\" which is required by \"%s\""), *WaveName.ToString(), *TargetPlatform->PlatformName());
+			}
+			
+		}
+		GetVersionFormatNumbersForIniVersionStrings(IniVersionMap, TEXT("AudioFormat"), SupportedWaveFormats);
+	}
+
+	{
+		TArray<FName> AllTextureFormats;
+		TargetPlatform->GetAllTextureFormats(AllTextureFormats);
+		TArray<const ITextureFormat*> SupportedTextureFormats;
+		for (const auto& TextureName : AllTextureFormats)
+		{
+			const ITextureFormat* TextureFormat = TPM->FindTextureFormat(TextureName);
+			if ( TextureFormat )
+			{
+				SupportedTextureFormats.Add(TextureFormat);
+			}
+			else
+			{
+				UE_LOG(LogCook, Warning, TEXT("Unable to find texture format \"%s\" which is required by \"%s\""), *TextureName.ToString(), *TargetPlatform->PlatformName());
+			}
+		}
+		GetVersionFormatNumbersForIniVersionStrings(IniVersionMap, TEXT("TextureFormat"), SupportedTextureFormats);
+	}
+
+	{
+		TArray<FName> AllFormatNames;
+		TargetPlatform->GetAllTargetedShaderFormats(AllFormatNames);
+		TArray<const IShaderFormat*> SupportedFormats;
+		for (const auto& FormatName : AllFormatNames)
+		{
+			const IShaderFormat* Format = TPM->FindShaderFormat(FormatName);
+			if ( Format )
+			{
+				SupportedFormats.Add(Format);
+			}
+			else
+			{
+				UE_LOG(LogCook, Warning, TEXT("Unable to find shader \"%s\" which is required by format \"%s\""), *FormatName.ToString(), *TargetPlatform->PlatformName());
+			}
+		}
+		GetVersionFormatNumbersForIniVersionStrings(IniVersionMap, TEXT("ShaderFormat"), SupportedFormats);
+	}
+
+
+	// TODO: Add support for physx version tracking, currently this happens so infrequently that invalidating a cook based on it is not essentual
+	//GetVersionFormatNumbersForIniVersionStrings(IniVersionMap, TEXT("PhysXFormat"), TPM->GetPhysXFormats());
 
 
 	if ( FParse::Param( FCommandLine::Get(), TEXT("fastcook") ) )
@@ -3354,9 +3426,13 @@ void GetAdditionalCurrentIniVersionStrings( const ITargetPlatform* TargetPlatfor
 	FString UE4Value = FString::Printf(TEXT("%d"), GPackageFileLicenseeUE4Version);
 	IniVersionMap.Add(UE4Ver, UE4Value);
 
-	FString UE4NetVer = TEXT("NetworkCompatibleCL");
-	FString UE4NetValue = FString::Printf(TEXT("%u"), FNetworkVersion::GetNetworkCompatibleChangelist());
-	IniVersionMap.Add(UE4NetVer, UE4NetValue);
+	FString UE4EngineVersionCompatibleName = TEXT("EngineVersionCompatibleWith");
+	FString UE4EngineVersionCompatible = FEngineVersion::CompatibleWith().ToString();
+	
+	if (UE4EngineVersionCompatible.Len())
+	{
+		IniVersionMap.Add(UE4EngineVersionCompatibleName, UE4EngineVersionCompatible);
+	}
 
 	IniVersionMap.Add(TEXT("MaterialShaderMapDDCVersion"), *GetMaterialShaderMapDDCKey());
 	IniVersionMap.Add(TEXT("GlobalDDCVersion"), *GetGlobalShaderMapDDCKey());
@@ -3594,13 +3670,27 @@ void UCookOnTheFlyServer::ProcessAccessedIniSettings(const FConfigFile* Config,
 		TSet<FName> ProcessedValues; 
 		const FName SectionName = FName(*ConfigSection.Key);
 
+		if ( SectionName.GetPlainNameString().Contains(TEXT(":")) )
+		{
+			UE_LOG(LogCook, Verbose, TEXT("Ignoring ini section checking for section name %s because it contains ':'"), *SectionName.ToString() );
+			continue;
+		}
+
 		for ( auto& ConfigValue : ConfigSection.Value )
 		{
 			const FName& ValueName = ConfigValue.Key;
 			if ( ProcessedValues.Contains(ValueName) )
 				continue;
 
-			ProcessedValues.Add( ValueName );
+			ProcessedValues.Add(ValueName);
+
+			if (ValueName.GetPlainNameString().Contains(TEXT(":")))
+			{
+				UE_LOG(LogCook, Verbose, TEXT("Ignoring ini section checking for section name %s because it contains ':'"), *ValueName.ToString());
+				continue;
+			}
+
+			
 			TArray<FConfigValue> ValueArray;
 			ConfigSection.Value.MultiFind( ValueName, ValueArray, true );
 
@@ -3655,7 +3745,7 @@ bool UCookOnTheFlyServer::IniSettingsOutOfDate(const ITargetPlatform* TargetPlat
 		const FString* CurrentValue = CurrentAdditionalSettings.Find(OldIniSetting.Key);
 		if ( !CurrentValue )
 		{
-			UE_LOG(LogCook, Display, TEXT("Missing Additional Setting from previous cook %s"), *OldIniSetting.Key);
+			UE_LOG(LogCook, Display, TEXT("Previous cook had additional ini setting: %s current cook is missing this setting."), *OldIniSetting.Key);
 			return true;
 		}
 
@@ -3725,9 +3815,12 @@ bool UCookOnTheFlyServer::IniSettingsOutOfDate(const ITargetPlatform* TargetPlat
 			const FName& SectionName = OldIniSection.Key;
 			const FConfigSection* IniSection = ConfigFile->Find( SectionName.ToString() );
 
+			const FString BlackListSetting = *FString::Printf(TEXT("%s.%s:%s"), *PlatformName, *Filename, *SectionName.ToString());
+
 			if ( IniSection == nullptr )
 			{
 				UE_LOG(LogCook, Display, TEXT("Inisetting is different for %s, Current section doesn't exist"), *FString::Printf(TEXT("%s %s %s"), *PlatformName, *Filename, *SectionName.ToString()));
+				UE_LOG(LogCook, Display, TEXT("To avoid this add blacklist setting to DefaultEditor.ini [CookSettings] %s"), *BlackListSetting);
 				return true;
 			}
 
@@ -3741,6 +3834,7 @@ bool UCookOnTheFlyServer::IniSettingsOutOfDate(const ITargetPlatform* TargetPlat
 				if ( CurrentValues.Num() != OldIniValue.Value.Num() )
 				{
 					UE_LOG(LogCook, Display, TEXT("Inisetting is different for %s, missmatched num array elements %d != %d "), *FString::Printf(TEXT("%s %s %s %s"), *PlatformName, *Filename, *SectionName.ToString(), *ValueName.ToString()), CurrentValues.Num(), OldIniValue.Value.Num());
+					UE_LOG(LogCook, Display, TEXT("To avoid this add blacklist setting to DefaultEditor.ini [CookSettings] %s"), *BlackListSetting);
 					return true;
 				}
 				for ( int Index = 0; Index < CurrentValues.Num(); ++Index )
@@ -3749,6 +3843,7 @@ bool UCookOnTheFlyServer::IniSettingsOutOfDate(const ITargetPlatform* TargetPlat
 					if ( FilteredCurrentValue != OldIniValue.Value[Index] )
 					{
 						UE_LOG(LogCook, Display, TEXT("Inisetting is different for %s, value %s != %s invalidating cook"),  *FString::Printf(TEXT("%s %s %s %s %d"),*PlatformName, *Filename, *SectionName.ToString(), *ValueName.ToString(), Index), *CurrentValues[Index].GetSavedValue(), *OldIniValue.Value[Index] );
+						UE_LOG(LogCook, Display, TEXT("To avoid this add blacklist setting to DefaultEditor.ini [CookSettings] %s"), *BlackListSetting);
 						return true;
 					}
 				}
@@ -4042,7 +4137,24 @@ bool UCookOnTheFlyServer::SaveCookedAssetRegistry(const FString& InCookedAssetRe
 		AllCookedPackages.Add(FailedStandardName );
 		FailedToSavePackagesSet.Add(FailedStandardName);
 	}
+
+	{
+		SCOPE_TIMER(AsyncDetermineAllDependentPackageInfo);
 		
+		TArray<FString> RelativePaths;
+		for ( const auto& FNamePath : AllCookedPackages )
+		{
+			RelativePaths.Add(FNamePath.ToString());
+		}
+		// resolving dependencies for packages
+		int32 NumThreadsToDeterminePackageInfo = FMath::Max(GLargeThreadPool->GetNumThreads() - 2, 1);
+		check(bIgnoreMarkupPackageAlreadyLoaded==false);
+		bIgnoreMarkupPackageAlreadyLoaded = true;
+		PDInfoModule.AsyncDetermineAllDependentPackageInfo(RelativePaths, NumThreadsToDeterminePackageInfo);
+		bIgnoreMarkupPackageAlreadyLoaded = false;
+	}
+
+
 	TSet<FName> ProcessedCookedPackages; // keep track of which packages are processed
 	for (const FName& Package : AllCookedPackages)
 	{
@@ -4417,16 +4529,59 @@ void UCookOnTheFlyServer::PopulateCookedPackagesFromDisk(const TArray<ITargetPla
 			const FString SandboxCookedAssetRegistryFilename = ConvertToFullSandboxPath(*CookedAssetRegistry, true, Target->PlatformName());
 
 			bool bIsIterateSharedBuild = IsCookFlagSet(ECookInitializationFlags::IterateSharedBuild);
+
+
 			TArray<FCookedAssetPackageInfo> CookedPackagesInfo;
-			if ( ReadCookedAssetRegistry(SandboxCookedAssetRegistryFilename, CookedPackagesInfo) == false )
+			if ( bIsIterateSharedBuild )
 			{
-				if ( IsCookFlagSet(ECookInitializationFlags::IterateSharedBuild) )
+				// see if the shared build is newer then the current cooked content in the local directory
+				FDateTime CurrentLocalCookedBuild = IFileManager::Get().GetTimeStamp(*SandboxCookedAssetRegistryFilename);
+
+				// iterate on the shared build if the option is set
+				FString SharedCookedAssetRegistry = FPaths::Combine(*FPaths::GameSavedDir(), TEXT("SharedIterativeBuild"), *Target->PlatformName(), TEXT("Cooked"), TEXT("CookedAssetRegistry.json"));
+
+				FDateTime CurrentIterativeCookedBuild = IFileManager::Get().GetTimeStamp(*SharedCookedAssetRegistry);
+
+				if ( CurrentIterativeCookedBuild > CurrentLocalCookedBuild )
+				{
+					// clean the sandbox 
+					ClearPlatformCookedData(FName(*Target->PlatformName()));
+					FString SandboxDirectory = GetSandboxDirectory(Target->PlatformName());
+					IFileManager::Get().DeleteDirectory(*SandboxDirectory, false, true);
+
+					// SaveCurrentIniSettings(Target); // use this if we don't care about ini safty.
+					// copy the ini settings from the shared cooked build. 
+					const FString SharedCookedIniFile = FPaths::Combine(*FPaths::GameSavedDir(), TEXT("SharedIterativeBuild"), *Target->PlatformName(), TEXT("Cooked"), TEXT("CookedIniVersion.txt"));
+					const FString SandboxCookedIniFile = ConvertToFullSandboxPath(*(FPaths::GameDir() / TEXT("CookedIniVersion.txt")), true).Replace(TEXT("[Platform]"), *Target->PlatformName());
+					
+					IFileManager::Get().Copy(*SandboxCookedIniFile, *SharedCookedIniFile );
+
+					bool bIniSettingsOutOfDate = IniSettingsOutOfDate(Target);
+					if ( bIniSettingsOutOfDate && !IsCookFlagSet(ECookInitializationFlags::IgnoreIniSettingsOutOfDate) )
+					{
+						UE_LOG(LogCook, Display, TEXT("Shared iterative build ini settings out of date, not using shared cooked build"));
+					}
+					else
+					{
+						if ( bIniSettingsOutOfDate )
+						{
+							UE_LOG(LogCook, Display, TEXT("Shared iterative build ini settings out of date, but we don't care"));
+						}
+
+						UE_LOG(LogCook, Display, TEXT("Shared iterative build is newer then local cooked build, iteratively cooking from shared build "));
+						ReadCookedAssetRegistry(SharedCookedAssetRegistry, CookedPackagesInfo);
+					}
+				}
+				else
 				{
-					// iterate on the shared build if the option is set
-					FString SharedCookedAssetRegistry = FPaths::Combine( *FPaths::GameSavedDir(), TEXT("SharedIterativeBuild"), *Target->PlatformName(), TEXT("CookedAssetRegistry.json"));
-					ReadCookedAssetRegistry( SharedCookedAssetRegistry, CookedPackagesInfo );
+					UE_LOG(LogCook, Display, TEXT("Local cook is newer then shared cooked build, iterativly cooking from local build"));
 				}
+
 			}
+
+			// if shared builds are enabled this will silently fail and no bad things will come of it
+			ReadCookedAssetRegistry(SandboxCookedAssetRegistryFilename, CookedPackagesInfo);
+
 			TMap<FName, FCookedAssetPackageInfo*> CookedPackageInfoLookup;
 			for (FCookedAssetPackageInfo& PackageInfo : CookedPackagesInfo)
 			{
@@ -4468,8 +4623,11 @@ void UCookOnTheFlyServer::PopulateCookedPackagesFromDisk(const TArray<ITargetPla
 #endif
 					//AllUncookedFilenames.Empty();
 
-					int32 NumThreadsToDeterminePackageInfo = GLargeThreadPool->GetNumThreads() - 2;
+					int32 NumThreadsToDeterminePackageInfo = FMath::Max(GLargeThreadPool->GetNumThreads() - 2,1);
+					check(bIgnoreMarkupPackageAlreadyLoaded == false);
+					bIgnoreMarkupPackageAlreadyLoaded = true;
 					PDInfoModule.AsyncDetermineAllDependentPackageInfo(AllUncookedFilenames, NumThreadsToDeterminePackageInfo);
+					bIgnoreMarkupPackageAlreadyLoaded = false;
 				}
 			}
 #if OUTPUT_TIMING
@@ -4537,16 +4695,23 @@ void UCookOnTheFlyServer::PopulateCookedPackagesFromDisk(const TArray<ITargetPla
 					if (UE_LOG_ACTIVE(LogCook, Verbose) )
 					// if (IsCookFlagSet(ECookInitializationFlags::LogDebugInfo))
 					{
-						UE_LOG(LogCook, Verbose, TEXT("Tracing package dependency which caused iterative cook to fail for: %s could be slow"), *UncookedFilenameString);
+						UE_LOG(LogCook, Display, TEXT("Tracing package dependency which caused iterative cook to fail for: %s could be slow Dependent package hashes %s != %s"), *UncookedFilenameString, *Lex::ToString(CookedPackageInfo->DependentPackageHash), *Lex::ToString(PackageDependentHash));
 						TMap<FString, FMD5Hash> CurrentPackageHashes;
 						PDInfoModule.RecursiveGetDependentPackageHashes(*UncookedFilenameString, CurrentPackageHashes);
 						for (const auto& CurrentPackageHash : CurrentPackageHashes)
 						{
 							const FName PackageName(*CurrentPackageHash.Key);
 							FCookedAssetPackageInfo* CurrentCookedPackageInfo = CookedPackageInfoLookup.FindRef(PackageName);
-							if (CurrentCookedPackageInfo && CurrentCookedPackageInfo->FullPackageHash != CurrentPackageHash.Value)
+							if (CurrentCookedPackageInfo )
+							{
+								if ( CurrentCookedPackageInfo->FullPackageHash != CurrentPackageHash.Value )
+								{
+									UE_LOG(LogCook, Display, TEXT("\tinvalidated because full package hash for package: %s doesn't match"), *CurrentPackageHash.Key);
+								}
+							}
+							else
 							{
-								UE_LOG(LogCook, Verbose, TEXT("\tinvalidated because full package hash for package: %s doesn't match"), *CurrentPackageHash.Key);
+								UE_LOG(LogCook, Display, TEXT("\tUnable to find cooked package info for: %s"), *CurrentPackageHash.Key);
 							}
 						}
 
@@ -4559,17 +4724,17 @@ void UCookOnTheFlyServer::PopulateCookedPackagesFromDisk(const TArray<ITargetPla
 							DependencyString += TEXT("\n");
 							DependencyString += Dependency;
 						}
-						DependencyString += TEXT("\nCurrentPackageInfo");
+						DependencyString += TEXT("\nCurrentPackageInfo:");
 						for (const FString& Dependency : CurrentPackageDependencies)
 						{
 							DependencyString += TEXT("\n");
 							DependencyString += Dependency;
 						}
 
-						UE_LOG(LogCook, Verbose, TEXT("Verifing package dependency count and order %d != %d %s"), CookedPackageInfo->DependentPackages.Num(), CurrentPackageDependencies.Num(), *DependencyString);
+						UE_LOG(LogCook, Display, TEXT("Verifing package dependency count and order %d != %d %s"), CookedPackageInfo->DependentPackages.Num(), CurrentPackageDependencies.Num(), *DependencyString);
 
 						
-						UE_LOG(LogCook, Verbose, TEXT("Finish trace of: %s"), *UncookedFilenameString);
+						UE_LOG(LogCook, Display, TEXT("Finish trace of: %s"), *UncookedFilenameString);
 					}
 					bShouldKeep = false;
 				}
@@ -5036,7 +5201,10 @@ void UCookOnTheFlyServer::CleanSandbox(const bool bIterative)
 		{
 			for (const ITargetPlatform* Target : Platforms)
 			{
-				if (IniSettingsOutOfDate(Target))
+				bool bIniSettingsOutOfDate = IniSettingsOutOfDate(Target);
+				if (bIniSettingsOutOfDate)
+				{
+					if ( !IsCookFlagSet(ECookInitializationFlags::IgnoreIniSettingsOutOfDate) )
 				{
 					UE_LOG(LogCook, Display, TEXT("Cook invalidated for platform %s ini settings don't match from last cook, clearing all cooked content"), *Target->PlatformName());
 
@@ -5048,6 +5216,11 @@ void UCookOnTheFlyServer::CleanSandbox(const bool bIterative)
 
 					SaveCurrentIniSettings(Target);
 				}
+					else
+					{
+						UE_LOG(LogCook, Display, TEXT("Inisettings were out of date for platform %s but we are going with it anyway because IgnoreIniSettingsOutOfDate is set"), *Target->PlatformName());
+					}
+				}
 			}
 
 			if (bIsInitializingSandbox)
@@ -5068,7 +5241,10 @@ void UCookOnTheFlyServer::CleanSandbox(const bool bIterative)
 
 
 			// Collect garbage to ensure we don't have any packages hanging around from dependent time stamp determination
-			CollectGarbage(RF_NoFlags);
+			if ( IsCookingInEditor() )
+			{
+				CollectGarbage(RF_NoFlags);
+			}
 		}
 	}
 #if OUTPUT_TIMING
@@ -5187,7 +5363,7 @@ void UCookOnTheFlyServer::AddFileToCook( TArray<FName>& InOutFilesToCook, const
 void UCookOnTheFlyServer::CollectFilesToCook(TArray<FName>& FilesInPath, const TArray<FString>& CookMaps, const TArray<FString>& InCookDirectories, const TArray<FString> &CookCultures, const TArray<FString> &IniMapSections, ECookByTheBookOptions FilesToCookFlags)
 {
 #if OUTPUT_TIMING
-	SCOPE_TIMER(StartupPackagesResolveStringAssetReferences);
+	SCOPE_TIMER(CollectFilesToCook);
 #endif
 	UProjectPackagingSettings* PackagingSettings = Cast<UProjectPackagingSettings>(UProjectPackagingSettings::StaticClass()->GetDefaultObject());
 
@@ -5288,7 +5464,7 @@ void UCookOnTheFlyServer::CollectFilesToCook(TArray<FName>& FilesInPath, const T
 	if (!(FilesToCookFlags & ECookByTheBookOptions::NoGameAlwaysCookPackages))
 	{
 		COOK_STAT(FScopedDurationTimer TickTimer(DetailedCookStats::GameCookModificationDelegateTimeSec));
-
+		SCOPE_TIMER(CookModificationDelegate);
 #define DEBUG_COOKMODIFICATIONDELEGATE 0
 #if DEBUG_COOKMODIFICATIONDELEGATE
 		TSet<UPackage*> LoadedPackages;
@@ -5332,6 +5508,7 @@ void UCookOnTheFlyServer::CollectFilesToCook(TArray<FName>& FilesInPath, const T
 
 	for ( const FString& CurrEntry : CookMaps )
 	{
+		SCOPE_TIMER(SearchForPackageOnDisk);
 		if (FPackageName::IsShortPackageName(CurrEntry))
 		{
 			FString OutFilename;
@@ -5672,6 +5849,14 @@ FString UCookOnTheFlyServer::ConvertToFullSandboxPath( const FString &FileName,
 const FString UCookOnTheFlyServer::GetSandboxAssetRegistryFilename()
 {
 	static const FString RegistryFilename = FPaths::GameDir() / TEXT("AssetRegistry.bin");
+
+	if (IsCookingDLC())
+	{
+		check(IsCookByTheBookMode());
+		const FString DLCRegistryFilename = FPaths::Combine(*FPaths::GamePluginsDir(), *CookByTheBookOptions->DlcName, TEXT("AssetRegistry.bin"));
+		return ConvertToFullSandboxPath(*DLCRegistryFilename, true);
+	}
+
 	const FString SandboxRegistryFilename = ConvertToFullSandboxPath(*RegistryFilename, true);
 	return SandboxRegistryFilename;
 }
@@ -5783,6 +5968,20 @@ void UCookOnTheFlyServer::CookByTheBookFinished()
 					verify(FPackageName::TryConvertFilenameToLongPackageName(CookedPackageFileName.ToString(), LongPackageName));
 					LongCookedPackageNames.Add(FName(*LongPackageName));
 				}
+				if ( IsCookingDLC() )
+				{
+					SCOPE_TIMER(RemovingOldManifestEntries);
+					// remove the previous release cooked packages from the new asset registry
+
+					const TArray<FName>* PreviousReleaseCookedPackages = CookByTheBookOptions->BasedOnReleaseCookedPackages.Find(Manifest.Key);
+					if ( PreviousReleaseCookedPackages ) 
+					{
+						for ( const auto& PreviousReleaseCookedPackage : *PreviousReleaseCookedPackages )
+						{
+							LongCookedPackageNames.Remove(PreviousReleaseCookedPackage);
+						}
+					}
+				}
 				{
 					SCOPE_TIMER(BuildChunkManifest);
 					Manifest.Value->BuildChunkManifest(LongCookedPackageNames, SandboxFile.Get(), CookByTheBookOptions->bGenerateStreamingInstallManifests);
@@ -6082,6 +6281,8 @@ void UCookOnTheFlyServer::ValidateCookByTheBookSettings() const
 
 void UCookOnTheFlyServer::StartCookByTheBook( const FCookByTheBookStartupOptions& CookByTheBookStartupOptions )
 {
+	SCOPE_TIMER(StartCookByTheBookTime);
+
 	const TArray<ITargetPlatform*>& TargetPlatforms = CookByTheBookStartupOptions.TargetPlatforms;
 	const TArray<FString>& CookMaps = CookByTheBookStartupOptions.CookMaps;
 	const TArray<FString>& CookDirectories = CookByTheBookStartupOptions.CookDirectories;
@@ -6173,7 +6374,6 @@ void UCookOnTheFlyServer::StartCookByTheBook( const FCookByTheBookStartupOptions
 	{
 		// we are going to change the state of dlc we need to clean out our package filename cache (the generated filename cache is dependent on this key)
 		CookByTheBookOptions->DlcName = DLCName;
-		check( OutputDirectoryOverride.Len() == 0 );
 
 		TermSandbox();
 	}
@@ -6317,7 +6517,11 @@ void UCookOnTheFlyServer::StartCookByTheBook( const FCookByTheBookStartupOptions
 	
 	TArray<FName> FilesInPath;
 
-	if (!CookByTheBookOptions->bDisableUnsolicitedPackages)
+	if (IsCookFlagSet(ECookInitializationFlags::IterateOnHash))
+	{
+		// if we are iterative cooking we don't need to resolve startup packages  they will get resolved if they are saved
+	}
+	else if (!CookByTheBookOptions->bDisableUnsolicitedPackages)
 	{
 		GIsCookerLoadingPackage = true;
 #if OUTPUT_TIMING
@@ -6325,7 +6529,7 @@ void UCookOnTheFlyServer::StartCookByTheBook( const FCookByTheBookStartupOptions
 #endif
 		for (const FName& StartupPackage : CookByTheBookOptions->StartupPackages)
 		{
-			GRedirectCollector.ResolveStringAssetReference(StartupPackage.ToString());
+			GRedirectCollector.ResolveStringAssetReference(StartupPackage.ToString(), false);
 		}
 		GIsCookerLoadingPackage = false;
 	}
@@ -6728,6 +6932,11 @@ void UCookOnTheFlyServer::MaybeMarkPackageAsAlreadyLoaded(UPackage *Package)
 	check(IsCookingInEditor()==false);
 	check(IsCookByTheBookMode());
 
+	if (bIgnoreMarkupPackageAlreadyLoaded == true)
+	{
+		return;
+	}
+
 	if (bIsInitializingSandbox)
 	{
 		return;
diff --git a/Engine/Source/Programs/AutomationTool/Scripts/CookCommand.Automation.cs b/Engine/Source/Programs/AutomationTool/Scripts/CookCommand.Automation.cs
--- a/Engine/Source/Programs/AutomationTool/Scripts/CookCommand.Automation.cs
+++ b/Engine/Source/Programs/AutomationTool/Scripts/CookCommand.Automation.cs
@@ -35,6 +35,133 @@ public partial class Project : CommandUtils
         }
     }
 
+	static void CopySharedCookedBuildForTarget(ProjectParams Params,TargetPlatformDescriptor TargetPlatform, string CookPlatform)
+	{
+
+		string ProjectPath = Params.RawProjectPath.FullName;
+		var LocalPath = CombinePaths(GetDirectoryName(ProjectPath), "Saved", "SharedIterativeBuild", CookPlatform);
+
+		// get network location 
+		ConfigHierarchy Hierarchy = ConfigCache.ReadHierarchy(ConfigHierarchyType.Engine, DirectoryReference.FromFile(Params.RawProjectPath), TargetPlatform.Type);
+		string CookedBuildPath;
+		if (Hierarchy.GetString("SharedCookedBuildSettings", "SharedCookedBuildPath", out CookedBuildPath) == false)
+		{
+			Log("Unable to copy shared cooked build: SharedCookedBuildPath not set in Engine.ini SharedCookedBuildSettings");
+			return ;
+		}
+
+		string BuildRoot = P4Enabled ? P4Env.BuildRootP4.Replace("/", "+") : "";
+		int RecentCL = P4Enabled ? P4Env.Changelist : 0;
+
+		BuildVersion Version;
+		if (BuildVersion.TryRead(out Version))
+		{
+			RecentCL = Version.Changelist;
+			BuildRoot = Version.BranchName;
+		}
+
+		// check to see if we have already synced this build ;)
+		var SyncedBuildFile = CombinePaths(LocalPath, "SyncedBuild.txt");
+		string BuildCL = "Invalid";
+		if ( File.Exists(SyncedBuildFile))
+		{
+			BuildCL = File.ReadAllText(SyncedBuildFile);
+		}
+
+		if (RecentCL == 0 && CookedBuildPath.Contains("[CL]") )
+		{
+			Log("Unable to copy shared cooked build: Unable to determine CL number from P4 or UGS, and is required by SharedCookedBuildPath");
+			return;
+		}
+
+		if (RecentCL == 0 && CookedBuildPath.Contains("[BRANCHNAME]"))
+		{
+			Log("Unable to copy shared cooked build: Unable to determine BRANCHNAME number from P4 or UGS, and is required by SharedCookedBuildPath");
+			return;
+		}
+
+
+		CookedBuildPath = CookedBuildPath.Replace("[CL]", RecentCL.ToString());
+		CookedBuildPath = CookedBuildPath.Replace("[BRANCHNAME]", BuildRoot);
+		CookedBuildPath = CookedBuildPath.Replace("[PLATFORM]", CookPlatform);
+
+		if ( Directory.Exists(CookedBuildPath) == false )
+		{
+			Log("Unable to copy shared cooked build: Unable to find shared build at location {0} check SharedCookedBuildPath in Engine.ini SharedCookedBuildSettings is correct", CookedBuildPath);
+			return;
+		}
+
+		Log("Attempting download of latest shared build CL {0} from location {1}", RecentCL, CookedBuildPath);
+
+		if (BuildCL == RecentCL.ToString())
+		{
+			Log("Already downloaded latest shared build at CL {0}", RecentCL);
+			return;
+		}
+		// delete all the stuff
+		Log("Deleting previous shared build because it was out of date");
+		CommandUtils.DeleteDirectory(LocalPath);
+		Directory.CreateDirectory(LocalPath);
+
+
+		// find all the files in the staged directory
+		string CookedBuildStagedDirectory = Path.GetFullPath(Path.Combine( CookedBuildPath, "Staged" ));
+		string LocalBuildStagedDirectory = Path.GetFullPath(Path.Combine(LocalPath, "Staged"));
+		if (Directory.Exists(CookedBuildStagedDirectory))
+		{
+			foreach (string FileName in Directory.EnumerateFiles(CookedBuildStagedDirectory, "*.*", SearchOption.AllDirectories))
+			{
+				string SourceFileName = Path.GetFullPath(FileName);
+				string DestFileName = SourceFileName.Replace(CookedBuildStagedDirectory, LocalBuildStagedDirectory);
+				Directory.CreateDirectory(Path.GetDirectoryName(DestFileName));
+				File.Copy(SourceFileName, DestFileName);
+			}
+		}
+
+
+		string CookedBuildCookedDirectory = Path.Combine(CookedBuildPath, "Cooked");
+		CookedBuildCookedDirectory = Path.GetFullPath(CookedBuildCookedDirectory);
+		string LocalBuildCookedDirectory = Path.Combine(LocalPath, "Cooked");
+		LocalBuildCookedDirectory = Path.GetFullPath(LocalBuildCookedDirectory);
+		if (Directory.Exists(CookedBuildCookedDirectory))
+		{
+			foreach (string FileName in Directory.EnumerateFiles(CookedBuildCookedDirectory, "*.*", SearchOption.AllDirectories))
+			{
+				string SourceFileName = Path.GetFullPath(FileName);
+				string DestFileName = SourceFileName.Replace(CookedBuildCookedDirectory, LocalBuildCookedDirectory);
+				Directory.CreateDirectory(Path.GetDirectoryName(DestFileName));
+				File.Copy(SourceFileName, DestFileName);
+			}
+		}
+		File.WriteAllText(SyncedBuildFile, RecentCL.ToString());
+		return;
+	}
+
+	static void CopySharedCookedBuild(ProjectParams Params)
+	{
+
+		if (!Params.NoClient)
+		{ 
+			foreach (var ClientPlatform in Params.ClientTargetPlatforms)
+			{
+				// Use the data platform, sometimes we will copy another platform's data
+				var DataPlatformDesc = Params.GetCookedDataPlatformForClientTarget(ClientPlatform);
+				string PlatformToCook = Platform.Platforms[DataPlatformDesc].GetCookPlatform(false, Params.Client);
+				CopySharedCookedBuildForTarget(Params, ClientPlatform, PlatformToCook);
+			}
+		}
+		if (Params.DedicatedServer)
+		{
+			foreach (var ServerPlatform in Params.ServerTargetPlatforms)
+			{
+				// Use the data platform, sometimes we will copy another platform's data
+				var DataPlatformDesc = Params.GetCookedDataPlatformForServerTarget(ServerPlatform);
+				string PlatformToCook = Platform.Platforms[DataPlatformDesc].GetCookPlatform(true, false);
+				CopySharedCookedBuildForTarget(Params, ServerPlatform, PlatformToCook);
+			}
+		}
+	}
+
     public static void Cook(ProjectParams Params)
 	{
 		if ((!Params.Cook && !(Params.CookOnTheFly && !Params.SkipServer)) || Params.SkipCook)
@@ -70,7 +197,7 @@ public partial class Project : CommandUtils
 				String COTFCommandLine = Params.RunCommandline;
 				if (Params.IterativeCooking)
 				{
-					COTFCommandLine += " -iterate";
+					COTFCommandLine += " -iterate -iteratehash";
 				}
 				if (Params.UseDebugParamForEditorExe)
 				{
@@ -179,11 +306,12 @@ public partial class Project : CommandUtils
                 }
                 if (Params.IterativeCooking)
                 {
-                    CommandletParams += " -iterate";
+                    CommandletParams += " -iterate -iterateshash";
                 }
 				if ( Params.IterateSharedCookedBuild)
 				{
-					CommandletParams += " -iteratesharedcookedbuild";
+					CopySharedCookedBuild(Params);
+					CommandletParams += " -iteratesharedcookedbuild";					
 				}
 
 				if (Params.CookMapsOnly)
diff --git a/Engine/Source/Programs/AutomationTool/Scripts/CopyBuildToStagingDirectory.Automation.cs b/Engine/Source/Programs/AutomationTool/Scripts/CopyBuildToStagingDirectory.Automation.cs
--- a/Engine/Source/Programs/AutomationTool/Scripts/CopyBuildToStagingDirectory.Automation.cs
+++ b/Engine/Source/Programs/AutomationTool/Scripts/CopyBuildToStagingDirectory.Automation.cs
@@ -269,7 +269,7 @@ public partial class Project : CommandUtils
 			//check( Params.HasDLCName == false );
 
 			// stage all the previously staged files
-			SC.StageFiles(StagedFileType.NonUFS, CombinePaths(SC.ProjectRoot, "Saved", "SharedIterativeBuild", SC.CookPlatform, "StagedBuild"), "*", true, null, "", true); // remap to the root directory
+			SC.StageFiles(StagedFileType.NonUFS, CombinePaths(SC.ProjectRoot, "Saved", "SharedIterativeBuild", SC.CookPlatform, "Staged"), "*", true, null, "", true); // remap to the root directory
 		}
         if (Params.HasDLCName)
         {
@@ -284,7 +284,7 @@ public partial class Project : CommandUtils
 
             // Put all of the cooked dir into the staged dir
             string PlatformCookDir = CombinePaths(SC.ProjectRoot, "Plugins", DLCName, "Saved", "Cooked", SC.CookPlatform);
-            string[] ExcludeWildCards = {"AssetRegistry.bin"};
+			string[] ExcludeWildCards = null; //{"AssetRegistry.bin"};
 
             // Stage any loose files in the root folder
 			// TODO: not sure if we should stage the loose files if we have pak files enabled... 
diff --git a/Engine/Source/Programs/AutomationTool/Scripts/RebuildLightMapsCommand.Automation.cs b/Engine/Source/Programs/AutomationTool/Scripts/RebuildLightMapsCommand.Automation.cs
--- a/Engine/Source/Programs/AutomationTool/Scripts/RebuildLightMapsCommand.Automation.cs
+++ b/Engine/Source/Programs/AutomationTool/Scripts/RebuildLightMapsCommand.Automation.cs
@@ -277,12 +277,14 @@ namespace AutomationScripts.Automation
 				Message.To.Add(new MailAddress(NextStakeHolder));
 			}
 
-			Message.CC.Add(new MailAddress("[REDACTED_EMAIL]"));
-            Message.CC.Add(new MailAddress("[REDACTED_EMAIL]"));
-			Message.CC.Add(new MailAddress("[REDACTED_EMAIL]"));
-			Message.CC.Add(new MailAddress("[REDACTED_EMAIL]"));
-			Message.CC.Add(new MailAddress("[REDACTED_EMAIL]"));
-			Message.CC.Add(new MailAddress("[REDACTED_EMAIL]"));
+			ConfigHierarchy Ini = ConfigCache.ReadHierarchy(ConfigHierarchyType.EditorPerProjectUserSettings, DirectoryReference.FromFile(ProjectFullPath), UnrealTargetPlatform.Win64);
+			List<string> Emails;
+			Ini.GetArray("RebuildlightingSettings", "EmailNotification", out Emails);
+			foreach (var Email in Emails)
+			{
+				Message.CC.Add(new MailAddress(Email));
+			}
+
 			Message.Subject = String.Format("Nightly lightmap rebuild {0} for {1}", bWasSuccessful ? "[SUCCESS]" : "[FAILED]", Branch);
 			Message.Body = MessageBody;
             /*Attachment Attach = new Attachment();
diff --git a/Engine/Source/Programs/AutomationTool/Win/WinPlatform.Automation.cs b/Engine/Source/Programs/AutomationTool/Win/WinPlatform.Automation.cs
--- a/Engine/Source/Programs/AutomationTool/Win/WinPlatform.Automation.cs
+++ b/Engine/Source/Programs/AutomationTool/Win/WinPlatform.Automation.cs
@@ -275,7 +275,13 @@ public abstract class BaseWinPlatform : Platform
 
 	public void StageAppLocalDependencies(ProjectParams Params, DeploymentContext SC, string PlatformDir)
 	{
-		string BaseAppLocalDependenciesPath = Path.IsPathRooted(Params.AppLocalDirectory) ? CombinePaths(Params.AppLocalDirectory, PlatformDir) : CombinePaths(SC.ProjectRoot, Params.AppLocalDirectory, PlatformDir);
+		Dictionary<string, string> PathVariables = new Dictionary<string, string>();
+		PathVariables["EngineDir"] = Path.Combine(SC.LocalRoot, "Engine");
+		PathVariables["ProjectDir"] = SC.ProjectRoot;
+
+		string ExpandedAppLocalDir = Utils.ExpandVariables(Params.AppLocalDirectory, PathVariables);
+
+		string BaseAppLocalDependenciesPath = Path.IsPathRooted(ExpandedAppLocalDir) ? CombinePaths(ExpandedAppLocalDir, PlatformDir) : CombinePaths(SC.ProjectRoot, ExpandedAppLocalDir, PlatformDir);
 		if (Directory.Exists(BaseAppLocalDependenciesPath))
 		{
 			string ProjectBinaryPath = new DirectoryReference(SC.ProjectBinariesFolder).MakeRelativeTo(new DirectoryReference(CombinePaths(SC.ProjectRoot, "..")));
diff --git a/Engine/Source/Programs/AutomationToolLauncher/app.config b/Engine/Source/Programs/AutomationToolLauncher/app.config
--- a/Engine/Source/Programs/AutomationToolLauncher/app.config
+++ b/Engine/Source/Programs/AutomationToolLauncher/app.config
@@ -1,6 +1,14 @@
 <?xml version="1.0"?>
 <configuration>
-   <runtime>
-      <loadFromRemoteSources enabled="true"/>
-   </runtime>
-<startup><supportedRuntime version="v4.0" sku=".NETFramework,Version=v4.5"/></startup></configuration>
+	<startup>
+		<supportedRuntime version="v4.0" sku=".NETFramework,Version=v4.5" />
+	</startup>
+	<runtime>
+		<assemblyBinding xmlns="urn:schemas-microsoft-com:asm.v1">
+			<dependentAssembly>
+				<assemblyIdentity name="Newtonsoft.Json" publicKeyToken="30ad4fe6b2a6aeed" culture="neutral" />
+				<bindingRedirect oldVersion="0.0.0.0-8.0.0.0" newVersion="9.0.0.0" />
+			</dependentAssembly>
+		</assemblyBinding>
+	</runtime>
+</configuration>
\ No newline at end of file
diff --git a/Engine/Source/Programs/UnrealPak/Private/UnrealPak.cpp b/Engine/Source/Programs/UnrealPak/Private/UnrealPak.cpp
--- a/Engine/Source/Programs/UnrealPak/Private/UnrealPak.cpp
+++ b/Engine/Source/Programs/UnrealPak/Private/UnrealPak.cpp
@@ -1126,7 +1126,7 @@ bool CreatePakFile(const TCHAR* Filename, TArray<FPakInputPair>& FilesToAdd, con
 					const uint64 ContiguousGroupedFilePatchWarningThreshhold = 50 * 1024 * 1024;
 					if (ContiguousTotalSizeSmallerThanBlockSize > ContiguousGroupedFilePatchWarningThreshhold)
 					{
-						UE_LOG(LogPakFile, Warning, TEXT("%i small files (%i) totaling %llu contiguous bytes found before first 'large' file.  Changes to any of these files could cause the whole group to be 'dirty' in a per-file binary diff based patching system."), ContiguousFilesSmallerThanBlockSize, RequiredPatchPadding, ContiguousTotalSizeSmallerThanBlockSize);
+						UE_LOG(LogPakFile, Display, TEXT("%i small files (%i) totaling %llu contiguous bytes found before first 'large' file.  Changes to any of these files could cause the whole group to be 'dirty' in a per-file binary diff based patching system."), ContiguousFilesSmallerThanBlockSize, RequiredPatchPadding, ContiguousTotalSizeSmallerThanBlockSize);
 					}
 					ContiguousTotalSizeSmallerThanBlockSize = 0;
 					ContiguousFilesSmallerThanBlockSize = 0;
diff --git a/Engine/Source/Runtime/AIModule/Private/EnvironmentQuery/EnvQueryInstance.cpp b/Engine/Source/Runtime/AIModule/Private/EnvironmentQuery/EnvQueryInstance.cpp
--- a/Engine/Source/Runtime/AIModule/Private/EnvironmentQuery/EnvQueryInstance.cpp
+++ b/Engine/Source/Runtime/AIModule/Private/EnvironmentQuery/EnvQueryInstance.cpp
@@ -552,6 +552,7 @@ void FEnvQueryInstance::FItemIterator::StoreTestResult()
 	if (Instance.IsInSingleItemFinalSearch())
 	{
 		// handle SingleResult mode
+		// this also implies we're not in 'score-only' mode
 		if (bPassed)
 		{
 			if (bForced)
@@ -570,7 +571,7 @@ void FEnvQueryInstance::FItemIterator::StoreTestResult()
 	}
 	else
 	{
-		if (!bPassed)
+		if (!bPassed && bIsFiltering)
 		{
 			HandleFailedTestResult();
 		}
diff --git a/Engine/Source/Runtime/Analytics/AnalyticsET/Private/IAnalyticsProviderET.cpp b/Engine/Source/Runtime/Analytics/AnalyticsET/Private/IAnalyticsProviderET.cpp
--- a/Engine/Source/Runtime/Analytics/AnalyticsET/Private/IAnalyticsProviderET.cpp
+++ b/Engine/Source/Runtime/Analytics/AnalyticsET/Private/IAnalyticsProviderET.cpp
@@ -281,6 +281,42 @@ private:
 	void EventRequestComplete(FHttpRequestPtr HttpRequest, FHttpResponsePtr HttpResponse, bool bSucceeded, TSharedPtr< TArray<FAnalyticsEventEntry> > FlushedEvents);
 };
 
+class FAnalyticsProviderETNULL :
+	public IAnalyticsProviderET,
+	public TSharedFromThis<FAnalyticsProviderETNULL>
+{
+public:
+	FAnalyticsProviderETNULL(const FAnalyticsET::Config& ConfigValues) {};
+
+	// IAnalyticsProvider
+
+	virtual bool StartSession(const TArray<FAnalyticsEventAttribute>& Attributes) override { return true; }
+	virtual bool StartSession(TArray<FAnalyticsEventAttribute>&& Attributes) override { return true; }
+	virtual void EndSession() override { }
+	virtual void FlushEvents() override { }
+
+	virtual void SetAppID(const FString&& AppID) override { APIKey = AppID; }
+	virtual const FString& GetAppID() const override { return APIKey; }
+	virtual void SetUserID(const FString& InUserID) override { UserID = InUserID; }
+	virtual FString GetUserID() const override { return UserID; }
+
+	virtual FString GetSessionID() const override { return SessionID; }
+	virtual bool SetSessionID(const FString& InSessionID) override { SessionID = InSessionID; return true; }
+
+	virtual void RecordEvent(const FString& EventName, const TArray<FAnalyticsEventAttribute>& Attributes) override {}
+	virtual void RecordEvent(FString EventName, TArray<FAnalyticsEventAttribute>&& Attributes) override {}
+	virtual void RecordEventJson(FString EventName, TArray<FAnalyticsEventAttribute>&& AttributesJson) override {}
+	virtual void SetDefaultEventAttributes(TArray<FAnalyticsEventAttribute>&& Attributes) override {}
+
+	virtual ~FAnalyticsProviderETNULL() {};
+
+	FString GetAPIKey() const { return APIKey; }
+
+	FString APIKey;
+	FString UserID;
+	FString SessionID;
+};
+
 TSharedPtr<IAnalyticsProviderET> FAnalyticsET::CreateAnalyticsProvider(const Config& ConfigValues) const
 {
 	// If we didn't have a proper APIKey, return NULL
@@ -289,7 +325,12 @@ TSharedPtr<IAnalyticsProviderET> FAnalyticsET::CreateAnalyticsProvider(const Con
 		UE_LOG(LogAnalytics, Warning, TEXT("CreateAnalyticsProvider config not contain required parameter %s"), *Config::GetKeyNameForAPIKey());
 		return NULL;
 	}
+	//@todo sz
+#if 0
+	return TSharedPtr<IAnalyticsProviderET>(new FAnalyticsProviderETNULL(ConfigValues));
+#else
 	return TSharedPtr<IAnalyticsProviderET>(new FAnalyticsProviderET(ConfigValues));
+#endif
 }
 
 /**
@@ -671,12 +712,8 @@ void FAnalyticsProviderET::SetUserID(const FString& InUserID)
 	if (!FParse::Value(FCommandLine::Get(), TEXT("ANALYTICSUSERID="), UserID, false))
 	{
 		UE_LOG(LogAnalytics, Log, TEXT("[%s] SetUserId %s"), *APIKey, *InUserID);
-		// Deliberately DO NOT Flush any cached events that would be using the old UserID.
-		// There are cases in our games where we create a provider before we've logged into the backend, so we can't set the final UserID in advance.
-		// Once the login is complete (generally before the first flush), we can set the UserID such that when the flush finally happens, all the 
-		// events come in with the correct UserID.
-		// It's kind of a hack, and subject to incorrect behavior.
-		// @todo add an overload of this method that explicitly avoids the flush and have the game code call that instead.
+		// Flush any cached events that would be using the old UserID.
+		FlushEvents();
 		UserID = InUserID;
 	}
 	else if (UserID != InUserID)
diff --git a/Engine/Source/Runtime/AnimGraphRuntime/Private/AnimNodes/AnimNode_AimOffsetLookAt.cpp b/Engine/Source/Runtime/AnimGraphRuntime/Private/AnimNodes/AnimNode_AimOffsetLookAt.cpp
--- a/Engine/Source/Runtime/AnimGraphRuntime/Private/AnimNodes/AnimNode_AimOffsetLookAt.cpp
+++ b/Engine/Source/Runtime/AnimGraphRuntime/Private/AnimNodes/AnimNode_AimOffsetLookAt.cpp
@@ -20,44 +20,28 @@ void FAnimNode_AimOffsetLookAt::Initialize(const FAnimationInitializeContext& Co
 {
 	FAnimNode_BlendSpacePlayer::Initialize(Context);
 	BasePose.Initialize(Context);
-
-	// SkeletalMesh possibly changed, update cached socket info.
-	bCachedSocketInfo = false;
 }
 
-void FAnimNode_AimOffsetLookAt::PreUpdate(const UAnimInstance* InAnimInstance)
+void FAnimNode_AimOffsetLookAt::RootInitialize(const FAnimInstanceProxy* InProxy)
 {
-	FAnimNode_BlendSpacePlayer::PreUpdate(InAnimInstance);
+	FAnimNode_BlendSpacePlayer::RootInitialize(InProxy);
 
-	if (!bCachedSocketInfo)
+	SocketBoneReference.BoneName = NAME_None;
+	if (USkeletalMeshComponent* SkelMeshComp = InProxy->GetSkelMeshComponent())
 	{
-		bCachedSocketInfo = true;
-		SocketBoneReference.BoneName = NAME_None;
-		if (USkeletalMeshComponent* SkelMeshComp = InAnimInstance->GetSkelMeshComponent())
+		if (USkeletalMesh* SkelMesh = SkelMeshComp->SkeletalMesh)
 		{
-			if (USkeletalMesh* SkelMesh = SkelMeshComp->SkeletalMesh)
+			if (const USkeletalMeshSocket* Socket = SkelMesh->FindSocket(SourceSocketName))
 			{
-				if (const USkeletalMeshSocket* Socket = SkelMesh->FindSocket(SourceSocketName))
-				{
-					SocketLocalTransform = Socket->GetSocketLocalTransform();
-					SocketBoneReference.BoneName = Socket->BoneName;
-				}
-
-				if (const USkeletalMeshSocket* Socket = SkelMesh->FindSocket(PivotSocketName))
-				{
-					PivotSocketLocalTransform = Socket->GetSocketLocalTransform();
-					PivotSocketBoneReference.BoneName = Socket->BoneName;
-				}
+				SocketLocalTransform = Socket->GetSocketLocalTransform();
+				SocketBoneReference.BoneName = Socket->BoneName;
 			}
-		}
-	}
 
-	if (USkeletalMeshComponent* SkelMeshComp = InAnimInstance->GetSkelMeshComponent())
-	{
-		ComponentToWorld = SkelMeshComp->ComponentToWorld;
-		if (const AActor* Owner = SkelMeshComp->GetOwner())
-		{
-			ActorTransform = Owner->GetTransform();
+			if (const USkeletalMeshSocket* Socket = SkelMesh->FindSocket(PivotSocketName))
+			{
+				PivotSocketLocalTransform = Socket->GetSocketLocalTransform();
+				PivotSocketBoneReference.BoneName = Socket->BoneName;
+			}
 		}
 	}
 }
@@ -85,11 +69,8 @@ void FAnimNode_AimOffsetLookAt::CacheBones(const FAnimationCacheBonesContext& Co
 	FAnimNode_BlendSpacePlayer::CacheBones(Context);
 	BasePose.CacheBones(Context);
 
-	if (bCachedSocketInfo)
-	{
-		SocketBoneReference.Initialize(Context.AnimInstanceProxy->GetRequiredBones());
-		PivotSocketBoneReference.Initialize(Context.AnimInstanceProxy->GetRequiredBones());
-	}
+	SocketBoneReference.Initialize(Context.AnimInstanceProxy->GetRequiredBones());
+	PivotSocketBoneReference.Initialize(Context.AnimInstanceProxy->GetRequiredBones());
 }
 
 void FAnimNode_AimOffsetLookAt::Evaluate(FPoseContext& Context)
@@ -118,7 +99,7 @@ void FAnimNode_AimOffsetLookAt::UpdateFromLookAtTarget(FPoseContext& LocalPoseCo
 	FVector BlendInput(X, Y, Z);
 
 	const FBoneContainer& RequiredBones = LocalPoseContext.Pose.GetBoneContainer();
-	if (bCachedSocketInfo && BlendSpace && SocketBoneReference.IsValid(RequiredBones))
+	if (BlendSpace && SocketBoneReference.IsValid(RequiredBones))
 	{
 		FCSPose<FCompactPose> GlobalPose;
 		GlobalPose.InitPose(LocalPoseContext.Pose);
@@ -135,7 +116,10 @@ void FAnimNode_AimOffsetLookAt::UpdateFromLookAtTarget(FPoseContext& LocalPoseCo
 			SourceComponentTransform.SetTranslation(PivotBoneComponentTransform.GetTranslation());
 		}
 
-		const FTransform SourceWorldTransform = SourceComponentTransform * ComponentToWorld;
+		FAnimInstanceProxy* AnimProxy = LocalPoseContext.AnimInstanceProxy;
+		check(AnimProxy);
+		const FTransform SourceWorldTransform = SourceComponentTransform * AnimProxy->GetSkelMeshCompLocalToWorld();
+		const FTransform ActorTransform = AnimProxy->GetSkelMeshCompOwnerTransform();
 
 		// Convert Target to Actor Space
 		const FTransform TargetWorldTransform(LookAtLocation);
@@ -155,7 +139,6 @@ void FAnimNode_AimOffsetLookAt::UpdateFromLookAtTarget(FPoseContext& LocalPoseCo
 #if ENABLE_DRAW_DEBUG
 		if (CVarAimOffsetLookAtDebug.GetValueOnAnyThread() == 1)
 		{
-			FAnimInstanceProxy* AnimProxy = LocalPoseContext.AnimInstanceProxy;
 			AnimProxy->AnimDrawDebugLine(SourceWorldTransform.GetLocation(), TargetWorldTransform.GetLocation(), FColor::Green);
 			AnimProxy->AnimDrawDebugLine(SourceWorldTransform.GetLocation(), SourceWorldTransform.GetLocation() + SourceWorldTransform.GetUnitAxis(EAxis::X) * (TargetWorldTransform.GetLocation() - SourceWorldTransform.GetLocation()).Size(), FColor::Red);
 			AnimProxy->AnimDrawDebugCoordinateSystem(ActorTransform.GetLocation(), ActorTransform.GetRotation().Rotator(), 100.f);
diff --git a/Engine/Source/Runtime/AnimGraphRuntime/Private/AnimNotify_PlayMontageNotify.cpp b/Engine/Source/Runtime/AnimGraphRuntime/Private/AnimNotify_PlayMontageNotify.cpp
new file mode 100644
--- /dev/null
+++ b/Engine/Source/Runtime/AnimGraphRuntime/Private/AnimNotify_PlayMontageNotify.cpp
@@ -0,0 +1,81 @@
+// Copyright 1998-2016 Epic Games, Inc. All Rights Reserved.
+
+#include "AnimNotify_PlayMontageNotify.h"
+#include "Animation/AnimInstance.h"
+
+//////////////////////////////////////////////////////////////////////////
+// UAnimNotify_PlayMontageNotify
+//////////////////////////////////////////////////////////////////////////
+
+UAnimNotify_PlayMontageNotify::UAnimNotify_PlayMontageNotify(const FObjectInitializer& ObjectInitializer)
+	: Super(ObjectInitializer)
+{
+	bIsNativeBranchingPoint = true;
+}
+
+
+void UAnimNotify_PlayMontageNotify::BranchingPointNotify(FBranchingPointNotifyPayload& BranchingPointPayload)
+{
+	Super::BranchingPointNotify(BranchingPointPayload);
+
+	if (USkeletalMeshComponent* MeshComp = BranchingPointPayload.SkelMeshComponent)
+	{
+		if (UAnimInstance* AnimInstance = MeshComp->GetAnimInstance())
+		{
+			AnimInstance->OnPlayMontageNotifyBegin.Broadcast(NotifyName, BranchingPointPayload);
+		}
+	}
+}
+
+#if WITH_EDITOR
+bool UAnimNotify_PlayMontageNotify::CanBePlaced(UAnimSequenceBase* Animation) const
+{
+	return (Animation && Animation->IsA(UAnimMontage::StaticClass()));
+}
+#endif // WITH_EDITOR
+
+
+//////////////////////////////////////////////////////////////////////////
+// UAnimNotify_PlayMontageNotifyWindow
+//////////////////////////////////////////////////////////////////////////
+
+UAnimNotify_PlayMontageNotifyWindow::UAnimNotify_PlayMontageNotifyWindow(const FObjectInitializer& ObjectInitializer)
+	: Super(ObjectInitializer)
+{
+	bIsNativeBranchingPoint = true;
+}
+
+
+void UAnimNotify_PlayMontageNotifyWindow::BranchingPointNotifyBegin(FBranchingPointNotifyPayload& BranchingPointPayload)
+{
+	Super::BranchingPointNotifyBegin(BranchingPointPayload);
+
+	if (USkeletalMeshComponent* MeshComp = BranchingPointPayload.SkelMeshComponent)
+	{
+		if (UAnimInstance* AnimInstance = MeshComp->GetAnimInstance())
+		{
+			AnimInstance->OnPlayMontageNotifyBegin.Broadcast(NotifyName, BranchingPointPayload);
+		}
+	}
+}
+
+
+void UAnimNotify_PlayMontageNotifyWindow::BranchingPointNotifyEnd(FBranchingPointNotifyPayload& BranchingPointPayload)
+{
+	Super::BranchingPointNotifyEnd(BranchingPointPayload);
+
+	if (USkeletalMeshComponent* MeshComp = BranchingPointPayload.SkelMeshComponent)
+	{
+		if (UAnimInstance* AnimInstance = MeshComp->GetAnimInstance())
+		{
+			AnimInstance->OnPlayMontageNotifyEnd.Broadcast(NotifyName, BranchingPointPayload);
+		}
+	}
+}
+
+#if WITH_EDITOR
+bool UAnimNotify_PlayMontageNotifyWindow::CanBePlaced(UAnimSequenceBase* Animation) const
+{
+	return (Animation && Animation->IsA(UAnimMontage::StaticClass()));
+}
+#endif // WITH_EDITOR
\ No newline at end of file
diff --git a/Engine/Source/Runtime/AnimGraphRuntime/Private/BoneControllers/AnimNode_LookAt.cpp b/Engine/Source/Runtime/AnimGraphRuntime/Private/BoneControllers/AnimNode_LookAt.cpp
--- a/Engine/Source/Runtime/AnimGraphRuntime/Private/BoneControllers/AnimNode_LookAt.cpp
+++ b/Engine/Source/Runtime/AnimGraphRuntime/Private/BoneControllers/AnimNode_LookAt.cpp
@@ -146,6 +146,23 @@ void FAnimNode_LookAt::EvaluateSkeletalControl_AnyThread(FComponentSpacePoseCont
 #endif
 }
 
+void FAnimNode_LookAt::EvaluateComponentSpaceInternal(FComponentSpacePoseContext& Context)
+{
+	Super::EvaluateComponentSpaceInternal(Context);
+
+#if !(UE_BUILD_SHIPPING || UE_BUILD_TEST)
+	/*if (bEnableDebug)
+	{
+		const FTransform LocalToWorld = Context.AnimInstanceProxy->GetSkelMeshCompLocalToWorld();
+		FVector TargetWorldLoc = LocalToWorld.TransformPosition(CachedCurrentTargetLocation);
+		FVector SourceWorldLoc = LocalToWorld.TransformPosition(CachedComponentBoneLocation);
+
+		Context.AnimInstanceProxy->AnimDrawDebugLine(SourceWorldLoc, TargetWorldLoc, FColor::Green);
+	}*/
+
+#endif // !(UE_BUILD_SHIPPING || UE_BUILD_TEST)
+}
+
 bool FAnimNode_LookAt::IsValidToEvaluate(const USkeleton* Skeleton, const FBoneContainer& RequiredBones) 
 {
 	// if both bones are valid
diff --git a/Engine/Source/Runtime/AnimGraphRuntime/Private/BoneControllers/AnimNode_ScaleChainLength.cpp b/Engine/Source/Runtime/AnimGraphRuntime/Private/BoneControllers/AnimNode_ScaleChainLength.cpp
new file mode 100644
--- /dev/null
+++ b/Engine/Source/Runtime/AnimGraphRuntime/Private/BoneControllers/AnimNode_ScaleChainLength.cpp
@@ -0,0 +1,119 @@
+// Copyright 1998-2017 Epic Games, Inc. All Rights Reserved.
+
+#include "BoneControllers/AnimNode_ScaleChainLength.h"
+#include "SceneManagement.h"
+#include "Engine/SkeletalMeshSocket.h"
+#include "Animation/AnimInstanceProxy.h"
+
+/////////////////////////////////////////////////////
+// FAnimNode_ScaleChainLength
+
+FAnimNode_ScaleChainLength::FAnimNode_ScaleChainLength()
+{
+}
+
+void FAnimNode_ScaleChainLength::Initialize(const FAnimationInitializeContext& Context)
+{
+	FAnimNode_Base::Initialize(Context);
+	InputPose.Initialize(Context);
+}
+
+void FAnimNode_ScaleChainLength::Update(const FAnimationUpdateContext& Context)
+{
+	EvaluateGraphExposedInputs.Execute(Context);
+	InputPose.Update(Context);
+}
+
+void FAnimNode_ScaleChainLength::CacheBones(const FAnimationCacheBonesContext& Context)
+{
+	InputPose.CacheBones(Context);
+
+	// LOD change, recache bone indices.
+	bBoneIndicesCached = false;
+}
+
+void FAnimNode_ScaleChainLength::Evaluate(FPoseContext& Output)
+{
+	// Evaluate incoming pose into our output buffer.
+	InputPose.Evaluate(Output);
+
+	const float FinalAlpha = AlphaScaleBias.ApplyTo(Alpha);
+	if (!FAnimWeight::IsRelevant(FinalAlpha))
+	{
+		return;
+	}
+
+	const FBoneContainer& BoneContainer = Output.Pose.GetBoneContainer();
+
+	if (!bBoneIndicesCached)
+	{
+		bBoneIndicesCached = true;
+
+		ChainStartBone.Initialize(BoneContainer);
+		ChainEndBone.Initialize(BoneContainer);
+		ChainBoneIndices.Reset();
+
+		// Make sure we have valid start/end bones, and that end is a child of start.
+		// Cache this, so we only evaluate on init and LOD changes.
+		const bool bBoneSetupIsValid = ChainStartBone.IsValid(BoneContainer) && ChainEndBone.IsValid(BoneContainer) &&
+			BoneContainer.BoneIsChildOf(ChainEndBone.GetCompactPoseIndex(BoneContainer), ChainStartBone.GetCompactPoseIndex(BoneContainer));
+
+		if (bBoneSetupIsValid)
+		{
+			FCompactPoseBoneIndex StartBoneIndex = ChainStartBone.GetCompactPoseIndex(BoneContainer);
+			FCompactPoseBoneIndex BoneIndex = ChainEndBone.GetCompactPoseIndex(BoneContainer);
+			ChainBoneIndices.Add(BoneIndex);
+			if (BoneIndex != INDEX_NONE)
+			{
+				FCompactPoseBoneIndex ParentBoneIndex = BoneContainer.GetParentBoneIndex(BoneIndex);
+				while ((ParentBoneIndex != INDEX_NONE) && (ParentBoneIndex != StartBoneIndex))
+				{
+					BoneIndex = ParentBoneIndex;
+					ParentBoneIndex = BoneContainer.GetParentBoneIndex(BoneIndex);
+					ChainBoneIndices.Insert(BoneIndex, 0);
+				};
+				ChainBoneIndices.Insert(StartBoneIndex, 0);
+			}
+		}
+	}
+
+	// Need at least Start/End bones to be valid.
+	if (ChainBoneIndices.Num() < 2)
+	{
+		return;
+	}
+
+	const FVector TargetLocationCompSpace = Output.AnimInstanceProxy->GetSkelMeshCompLocalToWorld().InverseTransformPosition(TargetLocation);
+
+	// Allocate transforms to get component space transform of chain start bone.
+	FCSPose<FCompactPose> CSPose;
+	CSPose.InitPose(Output.Pose);
+
+	const FTransform StartTransformCompSpace = CSPose.GetComponentSpaceTransform(ChainBoneIndices[0]);
+
+	const float DesiredChainLength = (TargetLocationCompSpace - StartTransformCompSpace.GetLocation()).Size();
+	const float ChainLengthScale = !FMath::IsNearlyZero(DefaultChainLength) ? (DesiredChainLength / DefaultChainLength) : 1.f;
+	const float ChainLengthScaleWithAlpha = FMath::LerpStable(1.f, ChainLengthScale, FinalAlpha);
+
+	// If we're not going to scale anything, early out.
+	if (FMath::IsNearlyEqual(ChainLengthScaleWithAlpha, 1.f))
+	{
+		return;
+	}
+
+	// Scale translation of all bones.
+	FCompactPose& BSPose = Output.Pose;
+	for (const FCompactPoseBoneIndex& BoneIndex : ChainBoneIndices)
+	{
+		// Get bone space transform, scale transition.
+		BSPose[BoneIndex].ScaleTranslation(ChainLengthScaleWithAlpha);
+	}
+}
+
+void FAnimNode_ScaleChainLength::GatherDebugData(FNodeDebugData& DebugData)
+{
+	FString DebugLine = DebugData.GetNodeName(this);
+
+	DebugData.AddDebugItem(DebugLine);
+	InputPose.GatherDebugData(DebugData);
+}
diff --git a/Engine/Source/Runtime/AnimGraphRuntime/Private/PlayMontageCallbackProxy.cpp b/Engine/Source/Runtime/AnimGraphRuntime/Private/PlayMontageCallbackProxy.cpp
new file mode 100644
--- /dev/null
+++ b/Engine/Source/Runtime/AnimGraphRuntime/Private/PlayMontageCallbackProxy.cpp
@@ -0,0 +1,133 @@
+// Copyright 1998-2017 Epic Games, Inc. All Rights Reserved.
+
+#include "PlayMontageCallbackProxy.h"
+
+//////////////////////////////////////////////////////////////////////////
+// UPlayMontageCallbackProxy
+
+UPlayMontageCallbackProxy::UPlayMontageCallbackProxy(const FObjectInitializer& ObjectInitializer)
+	: Super(ObjectInitializer)
+	, MontageInstanceID(INDEX_NONE)
+{
+}
+
+UPlayMontageCallbackProxy* UPlayMontageCallbackProxy::CreateProxyObjectForPlayMontage(
+	class USkeletalMeshComponent* InSkeletalMeshComponent,
+	class UAnimMontage* MontageToPlay,
+	float PlayRate,
+	float StartingPosition,
+	FName StartingSection)
+{
+	UPlayMontageCallbackProxy* Proxy = NewObject<UPlayMontageCallbackProxy>();
+	Proxy->SetFlags(RF_StrongRefOnFrame);
+	Proxy->PlayMontage(InSkeletalMeshComponent, MontageToPlay, PlayRate, StartingPosition, StartingSection);
+	return Proxy;
+}
+
+
+void UPlayMontageCallbackProxy::PlayMontage(class USkeletalMeshComponent* InSkeletalMeshComponent, 
+	class UAnimMontage* MontageToPlay, 
+	float PlayRate, 
+	float StartingPosition, 
+	FName StartingSection)
+{
+	bool bPlayedSuccessfully = false;
+	if (InSkeletalMeshComponent)
+	{
+		if (UAnimInstance* AnimInstance = InSkeletalMeshComponent->GetAnimInstance())
+		{
+			const float MontageLength = AnimInstance->Montage_Play(MontageToPlay, PlayRate, EMontagePlayReturnType::MontageLength, StartingPosition);
+			bPlayedSuccessfully = (MontageLength > 0.f);
+
+			if (bPlayedSuccessfully)
+			{
+				AnimInstancePtr = AnimInstance;
+				if (FAnimMontageInstance* MontageInstance = AnimInstance->GetActiveInstanceForMontage(MontageToPlay))
+				{
+					MontageInstanceID = MontageInstance->GetInstanceID();
+				}
+
+				if (StartingSection != NAME_None)
+				{
+					AnimInstance->Montage_JumpToSection(StartingSection, MontageToPlay);
+				}
+
+				BlendingOutDelegate.BindUObject(this, &UPlayMontageCallbackProxy::OnMontageBlendingOut);
+				AnimInstance->Montage_SetBlendingOutDelegate(BlendingOutDelegate, MontageToPlay);
+
+				MontageEndedDelegate.BindUObject(this, &UPlayMontageCallbackProxy::OnMontageEnded);
+				AnimInstance->Montage_SetEndDelegate(MontageEndedDelegate, MontageToPlay);
+
+				AnimInstance->OnPlayMontageNotifyBegin.AddDynamic(this, &UPlayMontageCallbackProxy::OnNotifyBeginReceived);
+				AnimInstance->OnPlayMontageNotifyEnd.AddDynamic(this, &UPlayMontageCallbackProxy::OnNotifyEndReceived);
+			}
+		}
+	}
+
+	if (!bPlayedSuccessfully)
+	{
+		OnInterrupted.Broadcast(NAME_None);
+	}
+}
+
+bool UPlayMontageCallbackProxy::IsNotifyValid(FName NotifyName, const FBranchingPointNotifyPayload& BranchingPointNotifyPayload) const
+{
+	return ((MontageInstanceID != INDEX_NONE) && (BranchingPointNotifyPayload.MontageInstanceID == MontageInstanceID));
+}
+
+
+void UPlayMontageCallbackProxy::OnNotifyBeginReceived(FName NotifyName, const FBranchingPointNotifyPayload& BranchingPointNotifyPayload)
+{
+	if (IsNotifyValid(NotifyName, BranchingPointNotifyPayload))
+	{
+		OnNotifyBegin.Broadcast(NotifyName);
+	}
+}
+
+
+void UPlayMontageCallbackProxy::OnNotifyEndReceived(FName NotifyName, const FBranchingPointNotifyPayload& BranchingPointNotifyPayload)
+{
+	if (IsNotifyValid(NotifyName, BranchingPointNotifyPayload))
+	{
+		OnNotifyEnd.Broadcast(NotifyName);
+	}
+}
+
+
+void UPlayMontageCallbackProxy::OnMontageBlendingOut(UAnimMontage* Montage, bool bInterrupted)
+{
+	if (bInterrupted)
+	{
+		OnInterrupted.Broadcast(NAME_None);
+	}
+	else
+	{
+		OnBlendOut.Broadcast(NAME_None);
+	}
+}
+
+void UPlayMontageCallbackProxy::OnMontageEnded(UAnimMontage* Montage, bool bInterrupted)
+{
+	if (!bInterrupted)
+	{
+		OnCompleted.Broadcast(NAME_None);
+	}
+
+	UnbindDelegates();
+}
+
+void UPlayMontageCallbackProxy::UnbindDelegates()
+{
+	if (UAnimInstance* AnimInstance = AnimInstancePtr.Get())
+	{
+		AnimInstance->OnPlayMontageNotifyBegin.RemoveDynamic(this, &UPlayMontageCallbackProxy::OnNotifyBeginReceived);
+		AnimInstance->OnPlayMontageNotifyEnd.RemoveDynamic(this, &UPlayMontageCallbackProxy::OnNotifyEndReceived);
+	}
+}
+
+void UPlayMontageCallbackProxy::BeginDestroy()
+{
+	UnbindDelegates();
+
+	Super::BeginDestroy();
+}
\ No newline at end of file
diff --git a/Engine/Source/Runtime/AnimGraphRuntime/Public/AnimNodes/AnimNode_AimOffsetLookAt.h b/Engine/Source/Runtime/AnimGraphRuntime/Public/AnimNodes/AnimNode_AimOffsetLookAt.h
--- a/Engine/Source/Runtime/AnimGraphRuntime/Public/AnimNodes/AnimNode_AimOffsetLookAt.h
+++ b/Engine/Source/Runtime/AnimGraphRuntime/Public/AnimNodes/AnimNode_AimOffsetLookAt.h
@@ -41,9 +41,6 @@ struct ANIMGRAPHRUNTIME_API FAnimNode_AimOffsetLookAt : public FAnimNode_BlendSp
 	UPROPERTY(EditAnywhere, BlueprintReadWrite, Category = LookAt, meta = (PinShownByDefault))
 	float Alpha;
 
-	UPROPERTY(Transient)
-	bool bCachedSocketInfo;
-
 	UPROPERTY()
 	FBoneReference SocketBoneReference;
 
@@ -56,25 +53,16 @@ struct ANIMGRAPHRUNTIME_API FAnimNode_AimOffsetLookAt : public FAnimNode_BlendSp
 	UPROPERTY()
 	FTransform PivotSocketLocalTransform;
 
-	UPROPERTY()
-	FTransform ActorTransform;
-
-	UPROPERTY()
-	FTransform ComponentToWorld;
-
 public:
 	FAnimNode_AimOffsetLookAt();
 
 	// FAnimNode_Base interface
+	virtual void RootInitialize(const FAnimInstanceProxy* InProxy) override;
 	virtual void Initialize(const FAnimationInitializeContext& Context) override;
 	virtual void CacheBones(const FAnimationCacheBonesContext& Context) override;
 	virtual void UpdateAssetPlayer(const FAnimationUpdateContext& Context) override;
 	virtual void Evaluate(FPoseContext& Output) override;
 	virtual void GatherDebugData(FNodeDebugData& DebugData) override;
-
-	virtual bool HasPreUpdate() const override { return true; }
-	virtual void PreUpdate(const UAnimInstance* InAnimInstance) override;
-
 	// End of FAnimNode_Base interface
 
 	void UpdateFromLookAtTarget(FPoseContext& LocalPoseContext);
diff --git a/Engine/Source/Runtime/AnimGraphRuntime/Public/AnimNotifies/AnimNotify_PlayMontageNotify.h b/Engine/Source/Runtime/AnimGraphRuntime/Public/AnimNotifies/AnimNotify_PlayMontageNotify.h
new file mode 100644
--- /dev/null
+++ b/Engine/Source/Runtime/AnimGraphRuntime/Public/AnimNotifies/AnimNotify_PlayMontageNotify.h
@@ -0,0 +1,53 @@
+// Copyright 1998-2016 Epic Games, Inc. All Rights Reserved.
+#pragma once
+#include "Animation/AnimNotifies/AnimNotify.h"
+#include "Animation/AnimNotifies/AnimNotifyState.h"
+#include "AnimNotify_PlayMontageNotify.generated.h"
+
+
+//////////////////////////////////////////////////////////////////////////
+// UAnimNotify_PlayMontageNotify
+//////////////////////////////////////////////////////////////////////////
+
+UCLASS(editinlinenew, const, hidecategories = Object, collapsecategories, MinimalAPI, meta = (DisplayName = "Montage Notify"))
+class UAnimNotify_PlayMontageNotify : public UAnimNotify
+{
+	GENERATED_UCLASS_BODY()
+
+public:
+
+	virtual void BranchingPointNotify(FBranchingPointNotifyPayload& BranchingPointPayload) override;
+
+#if WITH_EDITOR
+	virtual bool CanBePlaced(UAnimSequenceBase* Animation) const override;
+#endif
+protected:
+
+	// Name of notify that is passed to the PlayMontage K2Node.
+	UPROPERTY(EditAnywhere, BlueprintReadOnly, Category = "Notify")
+	FName NotifyName;
+};
+
+
+//////////////////////////////////////////////////////////////////////////
+// UAnimNotify_PlayMontageNotifyWindow
+//////////////////////////////////////////////////////////////////////////
+
+UCLASS(editinlinenew, const, hidecategories = Object, collapsecategories, MinimalAPI, meta = (DisplayName = "Montage Notify Window"))
+class UAnimNotify_PlayMontageNotifyWindow : public UAnimNotifyState
+{
+	GENERATED_UCLASS_BODY()
+
+public:
+	virtual void BranchingPointNotifyBegin(FBranchingPointNotifyPayload& BranchingPointPayload) override;
+	virtual void BranchingPointNotifyEnd(FBranchingPointNotifyPayload& BranchingPointPayload) override;
+
+#if WITH_EDITOR
+	virtual bool CanBePlaced(UAnimSequenceBase* Animation) const override;
+#endif
+protected:
+
+	// Name of notify that is passed to ability.
+	UPROPERTY(EditAnywhere, BlueprintReadOnly, Category = "Notify")
+	FName NotifyName;
+};
diff --git a/Engine/Source/Runtime/AnimGraphRuntime/Public/BoneControllers/AnimNode_LookAt.h b/Engine/Source/Runtime/AnimGraphRuntime/Public/BoneControllers/AnimNode_LookAt.h
--- a/Engine/Source/Runtime/AnimGraphRuntime/Public/BoneControllers/AnimNode_LookAt.h
+++ b/Engine/Source/Runtime/AnimGraphRuntime/Public/BoneControllers/AnimNode_LookAt.h
@@ -103,6 +103,7 @@ struct ANIMGRAPHRUNTIME_API FAnimNode_LookAt : public FAnimNode_SkeletalControlB
 	// End of FAnimNode_Base interface
 
 	// FAnimNode_SkeletalControlBase interface
+	virtual void EvaluateComponentSpaceInternal(FComponentSpacePoseContext& Context) override;
 	virtual void EvaluateSkeletalControl_AnyThread(FComponentSpacePoseContext& Output, TArray<FBoneTransform>& OutBoneTransforms) override;
 	virtual bool IsValidToEvaluate(const USkeleton* Skeleton, const FBoneContainer& RequiredBones) override;
 	// End of FAnimNode_SkeletalControlBase interface
diff --git a/Engine/Source/Runtime/AnimGraphRuntime/Public/BoneControllers/AnimNode_ScaleChainLength.h b/Engine/Source/Runtime/AnimGraphRuntime/Public/BoneControllers/AnimNode_ScaleChainLength.h
new file mode 100644
--- /dev/null
+++ b/Engine/Source/Runtime/AnimGraphRuntime/Public/BoneControllers/AnimNode_ScaleChainLength.h
@@ -0,0 +1,60 @@
+// Copyright 1998-2017 Epic Games, Inc. All Rights Reserved.
+
+#pragma once
+
+#include "CoreMinimal.h"
+#include "UObject/ObjectMacros.h"
+#include "BoneIndices.h"
+#include "BoneContainer.h"
+#include "BonePose.h"
+#include "BoneControllers/AnimNode_SkeletalControlBase.h"
+#include "AnimNode_ScaleChainLength.generated.h"
+
+class USkeletalMeshComponent;
+
+/**
+ *	
+ */
+USTRUCT()
+struct ANIMGRAPHRUNTIME_API FAnimNode_ScaleChainLength : public FAnimNode_Base
+{
+	GENERATED_USTRUCT_BODY()
+
+	UPROPERTY(EditAnywhere, BlueprintReadWrite, Category = Links)
+	FPoseLink InputPose;
+
+	/** Default chain length, as animated. */
+	UPROPERTY(EditAnywhere, BlueprintReadWrite, Category = ScaleChainLength, meta = (PinHiddenByDefault))
+	float DefaultChainLength;
+
+	UPROPERTY(EditAnywhere, BlueprintReadWrite, Category = ScaleChainLength)
+	FBoneReference ChainStartBone;
+
+	UPROPERTY(EditAnywhere, BlueprintReadWrite, Category = ScaleChainLength)
+	FBoneReference ChainEndBone;
+
+	UPROPERTY(EditAnywhere, BlueprintReadWrite, Category = ScaleChainLength, meta = (PinShownByDefault))
+	FVector TargetLocation;
+
+	UPROPERTY(EditAnywhere, BlueprintReadWrite, Category = Settings, meta = (PinShownByDefault))
+	mutable float Alpha;
+
+	UPROPERTY(EditAnywhere, BlueprintReadWrite, Category = Settings)
+	FInputScaleBias AlphaScaleBias;
+
+	UPROPERTY(Transient)
+	bool bBoneIndicesCached;
+
+	TArray<FCompactPoseBoneIndex> ChainBoneIndices;
+
+public:
+	FAnimNode_ScaleChainLength();
+
+	// FAnimNode_Base interface
+	virtual void Initialize(const FAnimationInitializeContext& Context) override;
+	virtual void Update(const FAnimationUpdateContext& Context) override;
+	virtual void CacheBones(const FAnimationCacheBonesContext& Context) override;
+	virtual void Evaluate(FPoseContext& Output) override;
+	virtual void GatherDebugData(FNodeDebugData& DebugData) override;
+	// End of FAnimNode_Base interface
+};
diff --git a/Engine/Source/Runtime/AnimGraphRuntime/Public/PlayMontageCallbackProxy.h b/Engine/Source/Runtime/AnimGraphRuntime/Public/PlayMontageCallbackProxy.h
new file mode 100644
--- /dev/null
+++ b/Engine/Source/Runtime/AnimGraphRuntime/Public/PlayMontageCallbackProxy.h
@@ -0,0 +1,79 @@
+// Copyright 1998-2017 Epic Games, Inc. All Rights Reserved.
+#pragma once
+
+#include "CoreMinimal.h"
+#include "UObject/ObjectMacros.h"
+#include "UObject/Object.h"
+#include "UObject/ScriptMacros.h"
+#include "Animation/AnimInstance.h"
+#include "PlayMontageCallbackProxy.generated.h"
+
+DECLARE_DYNAMIC_MULTICAST_DELEGATE_OneParam(FOnMontagePlayDelegate, FName, NotifyName);
+
+UCLASS(MinimalAPI)
+class UPlayMontageCallbackProxy : public UObject
+{
+	GENERATED_UCLASS_BODY()
+
+	// Called when Montage finished playing and wasn't interrupted
+	UPROPERTY(BlueprintAssignable)
+	FOnMontagePlayDelegate OnCompleted;
+
+	// Called when Montage starts blending out and is not interrupted
+	UPROPERTY(BlueprintAssignable)
+	FOnMontagePlayDelegate OnBlendOut;
+
+	// Called when Montage has been interrupted (or failed to play)
+	UPROPERTY(BlueprintAssignable)
+	FOnMontagePlayDelegate OnInterrupted;
+
+	UPROPERTY(BlueprintAssignable)
+	FOnMontagePlayDelegate OnNotifyBegin;
+
+	UPROPERTY(BlueprintAssignable)
+	FOnMontagePlayDelegate OnNotifyEnd;
+
+	// Called to perform the query internally
+	UFUNCTION(BlueprintCallable, meta = (BlueprintInternalUseOnly = "true"))
+	static UPlayMontageCallbackProxy* CreateProxyObjectForPlayMontage(
+		class USkeletalMeshComponent* InSkeletalMeshComponent, 
+		class UAnimMontage* MontageToPlay, 
+		float PlayRate = 1.f, 
+		float StartingPosition = 0.f, 
+		FName StartingSection = NAME_None);
+
+public:
+	//~ Begin UObject Interface
+	virtual void BeginDestroy() override;
+	//~ End UObject Interface
+
+protected:
+	UFUNCTION()
+	void OnMontageBlendingOut(UAnimMontage* Montage, bool bInterrupted);
+
+	UFUNCTION()
+	void OnMontageEnded(UAnimMontage* Montage, bool bInterrupted);
+
+	UFUNCTION()
+	void OnNotifyBeginReceived(FName NotifyName, const FBranchingPointNotifyPayload& BranchingPointNotifyPayload);
+
+	UFUNCTION()
+	void OnNotifyEndReceived(FName NotifyName, const FBranchingPointNotifyPayload& BranchingPointNotifyPayload);
+
+private:
+	TWeakObjectPtr<UAnimInstance> AnimInstancePtr;
+	int32 MontageInstanceID;
+
+	bool IsNotifyValid(FName NotifyName, const FBranchingPointNotifyPayload& BranchingPointNotifyPayload) const;
+	void UnbindDelegates();
+
+	FOnMontageBlendingOutStarted BlendingOutDelegate;
+	FOnMontageEnded MontageEndedDelegate;
+
+	void PlayMontage(
+		class USkeletalMeshComponent* InSkeletalMeshComponent,
+		class UAnimMontage* MontageToPlay,
+		float PlayRate = 1.f,
+		float StartingPosition = 0.f,
+		FName StartingSection = NAME_None);
+};
diff --git a/Engine/Source/Runtime/AssetRegistry/AssetRegistry.Build.cs b/Engine/Source/Runtime/AssetRegistry/AssetRegistry.Build.cs
--- a/Engine/Source/Runtime/AssetRegistry/AssetRegistry.Build.cs
+++ b/Engine/Source/Runtime/AssetRegistry/AssetRegistry.Build.cs
@@ -9,7 +9,8 @@ public class AssetRegistry : ModuleRules
 		PrivateDependencyModuleNames.AddRange(
 			new string[] {
 				"Core",
-				"CoreUObject"
+				"CoreUObject",
+				"Projects",
 			}
 			);
 
diff --git a/Engine/Source/Runtime/AssetRegistry/Private/AssetRegistry.cpp b/Engine/Source/Runtime/AssetRegistry/Private/AssetRegistry.cpp
--- a/Engine/Source/Runtime/AssetRegistry/Private/AssetRegistry.cpp
+++ b/Engine/Source/Runtime/AssetRegistry/Private/AssetRegistry.cpp
@@ -15,6 +15,7 @@
 #include "DependsNode.h"
 #include "PackageReader.h"
 #include "GenericPlatform/GenericPlatformChunkInstall.h"
+#include "IPluginManager.h"
 
 #if WITH_EDITOR
 #include "IDirectoryWatcher.h"
@@ -126,6 +127,15 @@ FAssetRegistry::FAssetRegistry()
 			// serialize the data with the memory reader (will convert FStrings to FNames, etc)
 			Serialize(SerializedAssetData);
 		}
+		TArray<TSharedRef<IPlugin>> PakPlugins = IPluginManager::Get().GetPluginsWithPakFile();
+		for (TSharedRef<IPlugin> PakPlugin : PakPlugins)
+		{
+			if (FFileHelper::LoadFileToArray(SerializedAssetData, *(PakPlugin->GetBaseDir() / TEXT("AssetRegistry.bin"))))
+			{
+				SerializedAssetData.Seek(0);
+				Serialize(SerializedAssetData);
+			}
+		}
 	}
 
 	// Report startup time. This does not include DirectoryWatcher startup time.
diff --git a/Engine/Source/Runtime/Core/Private/HAL/FileManagerGeneric.cpp b/Engine/Source/Runtime/Core/Private/HAL/FileManagerGeneric.cpp
--- a/Engine/Source/Runtime/Core/Private/HAL/FileManagerGeneric.cpp
+++ b/Engine/Source/Runtime/Core/Private/HAL/FileManagerGeneric.cpp
@@ -699,6 +699,13 @@ bool FArchiveFileReaderGeneric::InternalPrecache( int64 PrecacheOffset, int64 Pr
 
 void FArchiveFileReaderGeneric::Serialize( void* V, int64 Length )
 {
+	if (Pos + Length > Size)
+	{
+		ArIsError = true;
+		UE_LOG(LogFileManager, Error, TEXT("Requested read of %d bytes when %d bytes remain (file=%s, size=%d)"), Length, Size-Pos, *Filename, Size);
+		return;
+	}
+
 	while( Length>0 )
 	{
 		int64 Copy = FMath::Min( Length, BufferBase+BufferCount-Pos );
diff --git a/Engine/Source/Runtime/Core/Public/Containers/UnrealString.h b/Engine/Source/Runtime/Core/Public/Containers/UnrealString.h
--- a/Engine/Source/Runtime/Core/Public/Containers/UnrealString.h
+++ b/Engine/Source/Runtime/Core/Public/Containers/UnrealString.h
@@ -15,6 +15,7 @@
 #include "Templates/UnrealTemplate.h"
 #include "Math/NumericLimits.h"
 #include "Containers/Array.h"
+#include "Containers/Set.h"
 #include "Misc/CString.h"
 #include "Misc/Crc.h"
 #include "Math/UnrealMathUtility.h"
@@ -1697,6 +1698,36 @@ public:
 
 		return Result;
 	}
+
+	/**
+	* Joins a Set of 'something that can be concatentated to strings with +=' together into a single string with separators.
+	*
+	* @param	Set		The Set of 'things' to concatenate.
+	* @param	Separator	The string used to separate each element.
+	*
+	* @return	The final, joined, separated string.
+	*/
+	template <typename T, typename Allocator>
+	static FString Join(const TSet<T, Allocator>& Set, const TCHAR* Separator)
+	{
+		FString Result;
+		bool    First = true;
+		for (const T& Element : Set)
+		{
+			if (First)
+			{
+				First = false;
+			}
+			else
+			{
+				Result += Separator;
+			}
+
+			Result += Element;
+		}
+
+		return Result;
+	}
 };
 
 template<>
diff --git a/Engine/Source/Runtime/CoreUObject/Private/Misc/RedirectCollector.cpp b/Engine/Source/Runtime/CoreUObject/Private/Misc/RedirectCollector.cpp
--- a/Engine/Source/Runtime/CoreUObject/Private/Misc/RedirectCollector.cpp
+++ b/Engine/Source/Runtime/CoreUObject/Private/Misc/RedirectCollector.cpp
@@ -135,7 +135,7 @@ void FRedirectCollector::LogTimers() const
 	LOG_REDIRECT_TIMERS();
 }
 
-void FRedirectCollector::ResolveStringAssetReference(FString FilterPackage)
+void FRedirectCollector::ResolveStringAssetReference(FString FilterPackage, bool bProcessAlreadyResolvedPackages)
 {
 	SCOPE_REDIRECT_TIMER(ResolveTimeTotal);
 	
@@ -180,7 +180,15 @@ void FRedirectCollector::ResolveStringAssetReference(FString FilterPackage)
 					continue;
 				}
 			}
-			
+#if WITH_EDITOR
+			if ( (bProcessAlreadyResolvedPackages == false) && 
+				AlreadyRemapped.Contains(ToLoadFName) )
+			{
+				// don't load this package again because we already did it once before
+				// once is enough for somethings.
+				continue;
+			}
+#endif
 			if (ToLoad.Len() > 0 )
 			{
 				SCOPE_REDIRECT_TIMER(ResolveTimeLoad);
@@ -190,6 +198,9 @@ void FRedirectCollector::ResolveStringAssetReference(FString FilterPackage)
 
 				StringAssetRefFilenameStack.Push(RefFilenameAndProperty.GetPackage().ToString());
 
+				// LoadObject will mark any failed loads as "KnownMissing", so since we want to print out the referencer if it was missing and not-known-missing
+				// then check here before attempting to load (TODO: Can we just not even do the load? Seems like we should be able to...)
+
 				int32 DotIndex = ToLoad.Find(TEXT("."));
 				FString PackageName = DotIndex != INDEX_NONE ? ToLoad.Left(DotIndex) : ToLoad;
 
@@ -227,6 +238,9 @@ void FRedirectCollector::ResolveStringAssetReference(FString FilterPackage)
 					{
 						StringAssetRemap.Add(ToLoad, Dest);
 					}
+#if WITH_EDITOR
+					AlreadyRemapped.Add(ToLoadFName);
+#endif
 				}
 				else
 				{
diff --git a/Engine/Source/Runtime/CoreUObject/Public/Misc/RedirectCollector.h b/Engine/Source/Runtime/CoreUObject/Public/Misc/RedirectCollector.h
--- a/Engine/Source/Runtime/CoreUObject/Public/Misc/RedirectCollector.h
+++ b/Engine/Source/Runtime/CoreUObject/Public/Misc/RedirectCollector.h
@@ -136,8 +136,10 @@ public:
 	/**
 	 * Load the string asset references to resolve them, add that to the remap table, and empty the array
 	 * @param FilterPackage If set, only fixup references that were created by FilterPackage. If empty, clear all of them
+	 * @param bProcessAlreadyResolvedPackages If a package has already been resolved in a previous call of this function, should we process this package again and resolve it again?  
+	 *             True = process it again, False = don't process it again
 	 */
-	void ResolveStringAssetReference(FString FilterPackage = FString());
+	void ResolveStringAssetReference(FString FilterPackage = FString(), bool bProcessAlreadyResolvedPackages=true);
 
 	/**
 	 * Do we have any references to resolve.
@@ -169,6 +171,9 @@ private:
 
 	/** When saving, apply this remapping to all string asset references */
 	TMap<FString, FString> StringAssetRemap;
+#if WITH_EDITOR
+	TSet<FName> AlreadyRemapped;
+#endif
 };
 
 // global redirect collector callback structure
diff --git a/Engine/Source/Runtime/Engine/Classes/Animation/AnimInstance.h b/Engine/Source/Runtime/Engine/Classes/Animation/AnimInstance.h
--- a/Engine/Source/Runtime/Engine/Classes/Animation/AnimInstance.h
+++ b/Engine/Source/Runtime/Engine/Classes/Animation/AnimInstance.h
@@ -13,6 +13,7 @@
 #include "BonePose.h"
 #include "Components/SkeletalMeshComponent.h"
 #include "Animation/AnimNotifyQueue.h"
+#include "Animation/AnimNotifies/AnimNotify.h"
 #include "AnimInstance.generated.h"
 
 class FDebugDisplayInfo;
@@ -80,6 +81,8 @@ DECLARE_DELEGATE_ThreeParams(FOnGraphStateChanged, const struct FAnimNode_StateM
 /** Delegate that allows users to insert custom animation curve values - for now, it's only single, not sure how to make this to multi delegate and retrieve value sequentially, so */
 DECLARE_DELEGATE_OneParam(FOnAddCustomAnimationCurves, UAnimInstance*)
 
+/** Delegate called by 'PlayMontageNotify' and 'PlayMontageNotifyWindow' **/
+DECLARE_DYNAMIC_MULTICAST_DELEGATE_TwoParams(FPlayMontageAnimNotifyDelegate, FName, NotifyName, const FBranchingPointNotifyPayload&, BranchingPointPayload);
 
 
 USTRUCT()
@@ -1242,4 +1245,11 @@ protected:
 protected:
 	/** Proxy object, nothing should access this from an externally-callable API as it is used as a scratch area on worker threads */
 	mutable FAnimInstanceProxy* AnimInstanceProxy;
+
+public:
+	/** Called when a montage hits a 'PlayMontageNotify' or 'PlayMontageNotifyWindow' begin */
+	FPlayMontageAnimNotifyDelegate OnPlayMontageNotifyBegin;
+
+	/** Called when a montage hits a 'PlayMontageNotify' or 'PlayMontageNotifyWindow' end */
+	FPlayMontageAnimNotifyDelegate OnPlayMontageNotifyEnd;
 };
diff --git a/Engine/Source/Runtime/Engine/Classes/Engine/DataAsset.h b/Engine/Source/Runtime/Engine/Classes/Engine/DataAsset.h
--- a/Engine/Source/Runtime/Engine/Classes/Engine/DataAsset.h
+++ b/Engine/Source/Runtime/Engine/Classes/Engine/DataAsset.h
@@ -6,6 +6,7 @@
 #include "UObject/ObjectMacros.h"
 #include "UObject/Object.h"
 #include "AssetBundleData.h"
+#include "SubclassOf.h"
 #include "DataAsset.generated.h"
 
 /**
@@ -14,12 +15,16 @@
 UCLASS(abstract, MinimalAPI)
 class UDataAsset : public UObject
 {
-	GENERATED_BODY()
+	GENERATED_UCLASS_BODY()
 public:
 	// UObject interface
 #if WITH_EDITORONLY_DATA
 	ENGINE_API virtual void Serialize(FArchive& Ar) override;
 #endif
+
+private:
+	UPROPERTY(AssetRegistrySearchable)
+	TSubclassOf<UDataAsset> NativeClass;
 };
 
 /**
diff --git a/Engine/Source/Runtime/Engine/Classes/Engine/DataTable.h b/Engine/Source/Runtime/Engine/Classes/Engine/DataTable.h
--- a/Engine/Source/Runtime/Engine/Classes/Engine/DataTable.h
+++ b/Engine/Source/Runtime/Engine/Classes/Engine/DataTable.h
@@ -186,6 +186,12 @@ class UDataTable
 
 	ENGINE_API TArray<FName> GetRowNames() const;
 
+	/** Removes a single row from the DataTable by name. Just returns if row is not found. */
+	ENGINE_API void RemoveRow(FName RowName);
+
+	/** Copies RowData into table. That is: create Row if not found and copy data into the RowMap based on RowData. This is a "copy in" operation, so changing the passed in RowData after the fact does nothing. */
+	ENGINE_API void AddRow(FName RowName, const FTableRowBase& RowData);
+
 #if WITH_EDITOR
 
 private:
@@ -253,6 +259,9 @@ private:
 	 * @param OutCollectedImportProblems	[OUT] Array of strings of import problems
 	 */
 	void OnPostDataImported(OUT TArray<FString>& OutCollectedImportProblems);
+
+
+	UScriptStruct& GetEmptyUsingStruct() const;
 };
 
 
diff --git a/Engine/Source/Runtime/Engine/Classes/Engine/GameInstance.h b/Engine/Source/Runtime/Engine/Classes/Engine/GameInstance.h
--- a/Engine/Source/Runtime/Engine/Classes/Engine/GameInstance.h
+++ b/Engine/Source/Runtime/Engine/Classes/Engine/GameInstance.h
@@ -346,7 +346,6 @@ public:
 	 *
 	 * @return true if able or attempting to travel, false otherwise
 	 */
-#define ADDED_CLIENTTRAVELTOSESSION 1
 	virtual bool ClientTravelToSession(int32 ControllerId, FName InSessionName);
 
 	/** Broadcast a notification that travel is occurring */
diff --git a/Engine/Source/Runtime/Engine/Classes/Engine/NetSerialization.h b/Engine/Source/Runtime/Engine/Classes/Engine/NetSerialization.h
--- a/Engine/Source/Runtime/Engine/Classes/Engine/NetSerialization.h
+++ b/Engine/Source/Runtime/Engine/Classes/Engine/NetSerialization.h
@@ -846,12 +846,6 @@ bool FFastArraySerializer::FastArrayDeltaSerialize( TArray<Type> &Items, FNetDel
 
 		UE_LOG( LogNetFastTArray, Log, TEXT( "Read NumChanged: %d NumDeletes: %d." ), NumChanged, NumDeletes );
 
-		// Increment keys so that a client can re-serialize the array if needed, such as for client replay recording
-		if ( NumDeletes > 0 || NumChanged > 0 )
-		{
-			ArraySerializer.IncrementArrayReplicationKey();
-		}
-
 		TArray<int32, TInlineAllocator<8> > DeleteIndices;
 		TArray<int32, TInlineAllocator<8> > AddedIndices;
 		TArray<int32, TInlineAllocator<8> > ChangedIndices;
@@ -1005,6 +999,13 @@ bool FFastArraySerializer::FastArrayDeltaSerialize( TArray<Type> &Items, FNetDel
 			}
 		}
 
+		// Increment keys so that a client can re-serialize the array if needed, such as for client replay recording.
+		// Must check the size of DeleteIndices instead of NumDeletes to handle implicit deletes.
+		if ( DeleteIndices.Num() > 0 || NumChanged > 0 )
+		{
+			ArraySerializer.IncrementArrayReplicationKey();
+		}
+
 		// ---------------------------------------------------------
 		// Invoke all callbacks: removed -> added -> changed
 		// ---------------------------------------------------------
diff --git a/Engine/Source/Runtime/Engine/Classes/Engine/SpringInterpolator.h b/Engine/Source/Runtime/Engine/Classes/Engine/SpringInterpolator.h
--- a/Engine/Source/Runtime/Engine/Classes/Engine/SpringInterpolator.h
+++ b/Engine/Source/Runtime/Engine/Classes/Engine/SpringInterpolator.h
@@ -204,6 +204,17 @@ public:
 		return State.Velocity;
 	}
 
+	void SetVelocity(const T& NewVelocity)
+	{
+		State.Velocity = NewVelocity;
+		bIsInMotion = bIsInMotion || !IsVelocityZero();
+	}
+
+	void OffsetVelocity(const T& Offset)
+	{
+		SetVelocity(State.Velocity + Offset);
+	}
+
 	bool IsPositionEqualTo(const T& OtherPosition) const
 	{
 		return FRK4SpringInterpolatorUtils::AreEqual(GetPosition(), OtherPosition, RK4_SPRING_INTERPOLATOR_POSITION_TOLERANCE);
diff --git a/Engine/Source/Runtime/Engine/Classes/Kismet/KismetRenderingLibrary.h b/Engine/Source/Runtime/Engine/Classes/Kismet/KismetRenderingLibrary.h
--- a/Engine/Source/Runtime/Engine/Classes/Kismet/KismetRenderingLibrary.h
+++ b/Engine/Source/Runtime/Engine/Classes/Kismet/KismetRenderingLibrary.h
@@ -57,6 +57,13 @@ class UKismetRenderingLibrary : public UBlueprintFunctionLibrary
 	UFUNCTION(BlueprintCallable, Category="Rendering", meta=(Keywords="DrawMaterialToRenderTarget", WorldContext="WorldContextObject", UnsafeDuringActorConstruction="true"))
 	static ENGINE_API void DrawMaterialToRenderTarget(UObject* WorldContextObject, UTextureRenderTarget2D* TextureRenderTarget, UMaterialInterface* Material);
 
+	/**
+	 * Copies the contents of a render target to a UTexture2D
+	 * Only works in the editor
+	 */
+	UFUNCTION(BlueprintCallable, Category = "Rendering", meta = (Keywords = "ConvertRenderTarget", WorldContext = "WorldContextObject", UnsafeDuringActorConstruction = "true"))
+	static ENGINE_API void ConvertRenderTargetToTexture2DEditorOnly(UObject* WorldContextObject, UTextureRenderTarget2D* RenderTarget, UTexture2D* Texture);
+
 	/**
 	* Exports a render target as a HDR image onto the disk.
 	*/
diff --git a/Engine/Source/Runtime/Engine/Private/Animation/AnimInstanceProxy.cpp b/Engine/Source/Runtime/Engine/Private/Animation/AnimInstanceProxy.cpp
--- a/Engine/Source/Runtime/Engine/Private/Animation/AnimInstanceProxy.cpp
+++ b/Engine/Source/Runtime/Engine/Private/Animation/AnimInstanceProxy.cpp
@@ -220,8 +220,18 @@ void FAnimInstanceProxy::PreUpdate(UAnimInstance* InAnimInstance, float DeltaSec
 
 	InitializeObjects(InAnimInstance);
 
-	// Save off LOD level that we're currently using.
-	LODLevel = InAnimInstance->GetSkelMeshComponent()->PredictedLODLevel;
+	if (USkeletalMeshComponent* SkelMeshComp = InAnimInstance->GetSkelMeshComponent())
+	{
+		// Save off LOD level that we're currently using.
+		LODLevel = SkelMeshComp->PredictedLODLevel;
+
+		// Cache these transforms, so nodes don't have to pull it off the gamethread manually.
+		SkelMeshCompLocalToWorld = SkelMeshComp->ComponentToWorld;
+		if (const AActor* Owner = SkelMeshComp->GetOwner())
+		{
+			SkelMeshCompOwnerTransform = Owner->GetTransform();
+		}
+	}
 
 	NotifyQueue.Reset(InAnimInstance->GetSkelMeshComponent());
 
diff --git a/Engine/Source/Runtime/Engine/Private/Animation/AnimSequence.cpp b/Engine/Source/Runtime/Engine/Private/Animation/AnimSequence.cpp
--- a/Engine/Source/Runtime/Engine/Private/Animation/AnimSequence.cpp
+++ b/Engine/Source/Runtime/Engine/Private/Animation/AnimSequence.cpp
@@ -257,7 +257,7 @@ void UAnimSequence::Serialize(FArchive& Ar)
 		const bool bIsTransacting = Ar.IsTransacting();
 		const bool bIsCookingForDedicatedServer = bIsCooking && Ar.CookingTarget()->IsServerOnly();
 		const bool bIsCountingMemory = Ar.IsCountingMemory();
-		const bool bCookingTargetNeedsCompressedData = bIsCooking && (!UAnimationSettings::Get()->bStripAnimationDataOnDedicatedServer || !bIsCookingForDedicatedServer);
+		const bool bCookingTargetNeedsCompressedData = bIsCooking && (!UAnimationSettings::Get()->bStripAnimationDataOnDedicatedServer || !bIsCookingForDedicatedServer || bEnableRootMotion);
 
 		bool bSerializeCompressedData = bCookingTargetNeedsCompressedData || bIsDuplicating || bIsTransacting || bIsCountingMemory;
 		Ar << bSerializeCompressedData;
@@ -2573,8 +2573,7 @@ void UAnimSequence::RecycleAnimSequence()
 	CompressedTrackOffsets.Empty(0);
 	CompressedByteStream.Empty(0);
 	CompressedScaleOffsets.Empty(0);
-	SourceRawAnimationData.Reset();
-
+	SourceRawAnimationData.Empty(0);
 #endif // WITH_EDITORONLY_DATA
 }
 bool UAnimSequence::CopyAnimSequenceProperties(UAnimSequence* SourceAnimSeq, UAnimSequence* DestAnimSeq, bool bSkipCopyingNotifies)
diff --git a/Engine/Source/Runtime/Engine/Private/Components/SkinnedMeshComponent.cpp b/Engine/Source/Runtime/Engine/Private/Components/SkinnedMeshComponent.cpp
--- a/Engine/Source/Runtime/Engine/Private/Components/SkinnedMeshComponent.cpp
+++ b/Engine/Source/Runtime/Engine/Private/Components/SkinnedMeshComponent.cpp
@@ -449,15 +449,7 @@ void USkinnedMeshComponent::CreateRenderState_Concurrent()
 		if(MeshObject)
 		{
 			// Identify current LOD
-			int32 UseLOD;
-			if(MasterPoseComponent.IsValid())
-			{
-				UseLOD = FMath::Clamp(MasterPoseComponent->PredictedLODLevel, 0, MeshObject->GetSkeletalMeshResource().LODModels.Num()-1);
-			}
-			else
-			{
-				UseLOD = FMath::Clamp(PredictedLODLevel, 0, MeshObject->GetSkeletalMeshResource().LODModels.Num()-1);
-			}
+			const int32 UseLOD = FMath::Clamp(PredictedLODLevel, 0, MeshObject->GetSkeletalMeshResource().LODModels.Num()-1);
 
 			// If we have a valid LOD, set up required data, during reimport we may try to create data before we have all the LODs
 			// imported, in that case we skip until we have all the LODs
@@ -2293,11 +2285,16 @@ bool USkinnedMeshComponent::UpdateLODStatus()
 		}
 		else
 		{
-			// If no MeshObject - just assume lowest LOD.
-			if (MeshObject)
+			// Match LOD of MasterPoseComponent if it exists.
+			if (MasterPoseComponent.IsValid())
+			{
+				PredictedLODLevel = FMath::Clamp(MasterPoseComponent->PredictedLODLevel, 0, MaxLODIndex);
+			}
+			else if (MeshObject)
 			{
 				PredictedLODLevel = FMath::Clamp(MeshObject->MinDesiredLODLevel + GSkeletalMeshLODBias, 0, MaxLODIndex);
 			}
+			// If no MeshObject - just assume lowest LOD.
 			else
 			{
 				PredictedLODLevel = MaxLODIndex;
diff --git a/Engine/Source/Runtime/Engine/Private/DataAsset.cpp b/Engine/Source/Runtime/Engine/Private/DataAsset.cpp
--- a/Engine/Source/Runtime/Engine/Private/DataAsset.cpp
+++ b/Engine/Source/Runtime/Engine/Private/DataAsset.cpp
@@ -4,6 +4,11 @@
 #include "Misc/PackageName.h"
 #include "UObject/Package.h"
 #include "Engine/AssetManager.h"
+UDataAsset::UDataAsset(const FObjectInitializer& ObjectInitializer)
+	: Super(ObjectInitializer)
+{
+	NativeClass = GetClass();
+}
 
 #if WITH_EDITORONLY_DATA
 void UDataAsset::Serialize(FArchive& Ar)
diff --git a/Engine/Source/Runtime/Engine/Private/DataTable.cpp b/Engine/Source/Runtime/Engine/Private/DataTable.cpp
--- a/Engine/Source/Runtime/Engine/Private/DataTable.cpp
+++ b/Engine/Source/Runtime/Engine/Private/DataTable.cpp
@@ -244,7 +244,7 @@ void UDataTable::PostLoad()
 }
 #endif
 
-void UDataTable::EmptyTable()
+UScriptStruct& UDataTable::GetEmptyUsingStruct() const
 {
 	UScriptStruct* EmptyUsingStruct = RowStruct;
 	if (!EmptyUsingStruct)
@@ -256,11 +256,18 @@ void UDataTable::EmptyTable()
 		EmptyUsingStruct = FTableRowBase::StaticStruct();
 	}
 
+	return *EmptyUsingStruct;
+}
+
+void UDataTable::EmptyTable()
+{
+	UScriptStruct& EmptyUsingStruct = GetEmptyUsingStruct();
+
 	// Iterate over all rows in table and free mem
 	for (auto RowIt = RowMap.CreateIterator(); RowIt; ++RowIt)
 	{
 		uint8* RowData = RowIt.Value();
-		EmptyUsingStruct->DestroyStruct(RowData);
+		EmptyUsingStruct.DestroyStruct(RowData);
 		FMemory::Free(RowData);
 	}
 
@@ -268,6 +275,36 @@ void UDataTable::EmptyTable()
 	RowMap.Empty();
 }
 
+void UDataTable::RemoveRow(FName RowName)
+{
+	UScriptStruct& EmptyUsingStruct = GetEmptyUsingStruct();
+
+	uint8* RowData = nullptr;
+	RowMap.RemoveAndCopyValue(RowName, RowData);
+		
+	if (RowData)
+	{
+		EmptyUsingStruct.DestroyStruct(RowData);
+		FMemory::Free(RowData);
+	}
+}
+
+	
+void UDataTable::AddRow(FName RowName, const FTableRowBase& RowData)
+{
+	UScriptStruct& EmptyUsingStruct = GetEmptyUsingStruct();
+	RemoveRow(RowName);
+		
+	uint8* NewRawRowData = (uint8*)FMemory::Malloc(EmptyUsingStruct.GetStructureSize());
+	
+	EmptyUsingStruct.InitializeStruct(NewRawRowData);
+	EmptyUsingStruct.CopyScriptStruct(NewRawRowData, &RowData);
+
+	// Add to map
+	RowMap.Add(RowName, NewRawRowData);
+
+}
+
 /** Returns the column property where PropertyName matches the name of the column property. Returns NULL if no match is found or the match is not a supported table property */
 UProperty* UDataTable::FindTableProperty(const FName& PropertyName) const
 {
diff --git a/Engine/Source/Runtime/Engine/Private/GameUserSettings.cpp b/Engine/Source/Runtime/Engine/Private/GameUserSettings.cpp
--- a/Engine/Source/Runtime/Engine/Private/GameUserSettings.cpp
+++ b/Engine/Source/Runtime/Engine/Private/GameUserSettings.cpp
@@ -392,10 +392,22 @@ void UGameUserSettings::ApplyNonResolutionSettings()
 
 	IConsoleManager::Get().CallAllConsoleVariableSinks();
 
-	if (bUseHDRDisplayOutput)
+	bool bWithEditor = false;
+#if WITH_EDITOR
+	if (GIsEditor)
+	{
+		bWithEditor = true;
+	}
+#endif
+
+	if (bUseHDRDisplayOutput && !bWithEditor)
 	{
 		EnableHDRDisplayOutput(true, HDRDisplayOutputNits);
 	}
+	else
+	{
+		EnableHDRDisplayOutput(false, HDRDisplayOutputNits);
+	}
 }
 
 void UGameUserSettings::ApplyResolutionSettings(bool bCheckForCommandLineOverrides)
@@ -494,6 +506,7 @@ void UGameUserSettings::PreloadResolutionSettings()
 	int32 ResolutionY = GetDefaultResolution().Y;
 	EWindowMode::Type WindowMode = GetDefaultWindowMode();
 	bool bUseDesktopResolution = false;
+	bool bUseHDR = false;
 
 	int32 Version=0;
 	if( GConfig->GetInt(*GameUserSettingsCategory, TEXT("Version"), Version, GGameUserSettingsIni ) && Version == UE_GAMEUSERSETTINGS_VERSION )
@@ -518,6 +531,15 @@ void UGameUserSettings::PreloadResolutionSettings()
 			ResolutionY = DisplayMetrics.PrimaryDisplayHeight;
 		}
 #endif
+
+		if (GConfig->GetBool(*GameUserSettingsCategory, TEXT("bUseHDRDisplayOutput"), bUseHDR, GGameUserSettingsIni))
+		{
+			static auto CVarHDROutputEnabled = IConsoleManager::Get().FindConsoleVariable(TEXT("r.HDR.EnableHDROutput"));
+			if (CVarHDROutputEnabled)
+			{
+				CVarHDROutputEnabled->Set(bUseHDR ? 1 : 0, ECVF_SetByGameSetting);
+			}
+		}
 	}
 
 	RequestResolutionChange(ResolutionX, ResolutionY, WindowMode);
diff --git a/Engine/Source/Runtime/Engine/Private/KismetRenderingLibrary.cpp b/Engine/Source/Runtime/Engine/Private/KismetRenderingLibrary.cpp
--- a/Engine/Source/Runtime/Engine/Private/KismetRenderingLibrary.cpp
+++ b/Engine/Source/Runtime/Engine/Private/KismetRenderingLibrary.cpp
@@ -18,6 +18,7 @@
 #include "OneColorShader.h"
 #include "PipelineStateCache.h"
 #include "ClearQuad.h"
+#include "Engine/Texture2D.h"
 
 //////////////////////////////////////////////////////////////////////////
 // UKismetRenderingLibrary
@@ -172,6 +173,56 @@ void UKismetRenderingLibrary::ExportRenderTarget(UObject* WorldContextObject, UT
 		FMessageLog("Blueprint").Warning(LOCTEXT("ExportRenderTarget_InvalidFileName", "ExportRenderTarget: FileName must be non-empty."));
 	}
 }
+/*
+
+void UKismetRenderingLibrary::CreateTexture2DFromRenderTarget(UObject* WorldContextObject, UTextureRenderTarget2D* RenderTarget, const FString &TextureAssetName)
+{
+	if (RenderTarget && Texture)
+	{
+
+		//FImageUtils::CreateTexture2D
+
+		UTexture2D* NewTexture = RenderTarget->ConstructTexture2D(Texture->GetOuter(), Texture->GetName(), RenderTarget->GetMaskedFlags(), CTF_Default, NULL);
+
+		check(NewTexture == Texture);
+		NewTexture->UpdateResource();
+	}
+	else if (!RenderTarget)
+	{
+		FMessageLog("Blueprint").Warning(LOCTEXT("ConvertRenderTargetToTexture2D_InvalidRenderTarget", "ExportRenderTarget: RenderTarget must be non-null."));
+	}
+	else if (!Texture)
+	{
+		FMessageLog("Blueprint").Warning(LOCTEXT("ConvertRenderTargetToTexture2D_InvalidTexture", "ExportRenderTarget: Texture must be non-null."));
+	}
+
+}*/
+
+
+void UKismetRenderingLibrary::ConvertRenderTargetToTexture2DEditorOnly( UObject* WorldContextObject, UTextureRenderTarget2D* RenderTarget, UTexture2D* Texture )
+{
+#if WITH_EDITOR
+	if (RenderTarget && Texture)
+	{
+		UTexture2D* NewTexture = RenderTarget->ConstructTexture2D(Texture->GetOuter(), Texture->GetName(), RenderTarget->GetMaskedFlags(), CTF_Default, NULL);
+
+		check(NewTexture == Texture);
+		NewTexture->UpdateResource();
+		NewTexture->Modify();
+	}
+	else if (!RenderTarget)
+	{
+		FMessageLog("Blueprint").Warning(LOCTEXT("ConvertRenderTargetToTexture2D_InvalidRenderTarget", "ExportRenderTarget: RenderTarget must be non-null."));
+	}
+	else if (!Texture)
+	{
+		FMessageLog("Blueprint").Warning(LOCTEXT("ConvertRenderTargetToTexture2D_InvalidTexture", "ExportRenderTarget: Texture must be non-null."));
+	}
+#else
+	FMessageLog("Blueprint").Error(LOCTEXT("Convert to render target can't be used at run time.", "ConvertRenderTarget: Can't convert render target to texture2d at run time. "));
+#endif
+
+}
 
 void UKismetRenderingLibrary::ExportTexture2D(UObject* WorldContextObject, UTexture2D* Texture, const FString& FilePath, const FString& FileName)
 {
diff --git a/Engine/Source/Runtime/Engine/Private/Materials/MaterialShared.cpp b/Engine/Source/Runtime/Engine/Private/Materials/MaterialShared.cpp
--- a/Engine/Source/Runtime/Engine/Private/Materials/MaterialShared.cpp
+++ b/Engine/Source/Runtime/Engine/Private/Materials/MaterialShared.cpp
@@ -1804,6 +1804,11 @@ FShader* FMaterial::GetShader(FMeshMaterialShaderType* ShaderType, FVertexFactor
 	FShader* Shader = MeshShaderMap ? MeshShaderMap->GetShader(ShaderType) : nullptr;
 	if (!Shader)
 	{
+		// we don't care about thread safety because we are about to crash 
+		const auto CachedGameThreadShaderMap = GameThreadShaderMap;
+		const auto CachedGameMeshShaderMap = CachedGameThreadShaderMap ? CachedGameThreadShaderMap->GetMeshShaderMap(VertexFactoryType) : nullptr;
+		bool bShaderWasFoundInGameShaderMap = CachedGameMeshShaderMap && CachedGameMeshShaderMap->GetShader(ShaderType) != nullptr;
+
 		// Get the ShouldCache results that determine whether the shader should be compiled
 		auto ShaderPlatform = GShaderPlatformForFeatureLevel[GetFeatureLevel()];
 		bool bMaterialShouldCache = ShouldCache(ShaderPlatform, ShaderType, VertexFactoryType);
@@ -1817,10 +1822,14 @@ FShader* FMaterial::GetShader(FMeshMaterialShaderType* ShaderType, FVertexFactor
 		// This is usually the result of an incorrect ShouldCache function.
 		UE_LOG(LogMaterial, Fatal,
 			TEXT("Couldn't find Shader %s for Material Resource %s!\n")
+			TEXT("		RenderMeshShaderMap %d, RenderThreadShaderMap %d\n")
+			TEXT("		GameMeshShaderMap %d, GameThreadShaderMap %d, bShaderWasFoundInGameShaderMap %d\n")
 			TEXT("		With VF=%s, Platform=%s\n")
 			TEXT("		ShouldCache: Mat=%u, VF=%u, Shader=%u \n")
 			TEXT("		MaterialUsageDesc: %s"),
 			ShaderType->GetName(),  *GetFriendlyName(),
+			MeshShaderMap != nullptr, RenderingThreadShaderMap != nullptr,
+			CachedGameMeshShaderMap != nullptr, CachedGameThreadShaderMap != nullptr, bShaderWasFoundInGameShaderMap,
 			VertexFactoryType->GetName(), *LegacyShaderPlatformToShaderFormat(ShaderPlatform).ToString(),
 			bMaterialShouldCache, bVFShouldCache, bShaderShouldCache,
 			*MaterialUsage
diff --git a/Engine/Source/Runtime/Engine/Private/PhysicsEngine/PhysScene.cpp b/Engine/Source/Runtime/Engine/Private/PhysicsEngine/PhysScene.cpp
--- a/Engine/Source/Runtime/Engine/Private/PhysicsEngine/PhysScene.cpp
+++ b/Engine/Source/Runtime/Engine/Private/PhysicsEngine/PhysScene.cpp
@@ -201,6 +201,7 @@ public:
 
 	FORCEINLINE TStatId FindOrCreateStatId(const char* StatName)
 	{
+#if STATS
 		for (int StatIdx = 0; StatIdx < NumStats; ++StatIdx)
 		{
 			FStatLookup& Lookup = Stats[StatIdx];
@@ -232,7 +233,7 @@ public:
 			return NewStat.Stat;
 		}
 		
-
+#endif // STATS
 		return TStatId();
 	}
 
diff --git a/Engine/Source/Runtime/Engine/Private/Slate/SGameLayerManager.cpp b/Engine/Source/Runtime/Engine/Private/Slate/SGameLayerManager.cpp
--- a/Engine/Source/Runtime/Engine/Private/Slate/SGameLayerManager.cpp
+++ b/Engine/Source/Runtime/Engine/Private/Slate/SGameLayerManager.cpp
@@ -216,6 +216,10 @@ void SGameLayerManager::ClearWidgets()
 		}
 	}
 	PlayerLayers.Reset();
+
+	WindowTitleBarContentStack.Empty();
+	bIsWindowTitleBarVisible = false;
+	SetDefaultWindowTitleBarContentAsCurrent();
 }
 
 void SGameLayerManager::Tick(const FGeometry& AllottedGeometry, const double InCurrentTime, const float InDeltaTime)
diff --git a/Engine/Source/Runtime/Engine/Private/Texture.cpp b/Engine/Source/Runtime/Engine/Private/Texture.cpp
--- a/Engine/Source/Runtime/Engine/Private/Texture.cpp
+++ b/Engine/Source/Runtime/Engine/Private/Texture.cpp
@@ -982,6 +982,11 @@ FName GetDefaultTextureFormatName( const ITargetPlatform* TargetPlatform, const
 {
 	FName TextureFormatName = NAME_None;
 
+	/**
+	 * IF you add a format to this function don't forget to update GetAllDefaultTextureFormatNames 
+	 */
+
+
 #if WITH_EDITOR
 	// Supported texture format names.
 	static FName NameDXT1(TEXT("DXT1"));
@@ -1114,3 +1119,43 @@ FName GetDefaultTextureFormatName( const ITargetPlatform* TargetPlatform, const
 
 	return TextureFormatName;
 }
+
+
+void GetAllDefaultTextureFormats(const class ITargetPlatform* TargetPlatform, TArray<FName>& OutFormats, bool bSupportDX11TextureFormats)
+{
+#if WITH_EDITOR
+	static FName NameDXT1(TEXT("DXT1"));
+	static FName NameDXT3(TEXT("DXT3"));
+	static FName NameDXT5(TEXT("DXT5"));
+	static FName NameDXT5n(TEXT("DXT5n"));
+	static FName NameAutoDXT(TEXT("AutoDXT"));
+	static FName NameBC4(TEXT("BC4"));
+	static FName NameBC5(TEXT("BC5"));
+	static FName NameBGRA8(TEXT("BGRA8"));
+	static FName NameXGXR8(TEXT("XGXR8"));
+	static FName NameG8(TEXT("G8"));
+	static FName NameVU8(TEXT("VU8"));
+	static FName NameRGBA16F(TEXT("RGBA16F"));
+	static FName NameBC6H(TEXT("BC6H"));
+	static FName NameBC7(TEXT("BC7"));
+
+	OutFormats.Add(NameDXT1);
+	OutFormats.Add(NameDXT3);
+	OutFormats.Add(NameDXT5);
+	OutFormats.Add(NameDXT5n);
+	OutFormats.Add(NameAutoDXT);
+	OutFormats.Add(NameBC4);
+	OutFormats.Add(NameBC5);
+	OutFormats.Add(NameBGRA8);
+	OutFormats.Add(NameXGXR8);
+	OutFormats.Add(NameG8);
+	OutFormats.Add(NameVU8);
+	OutFormats.Add(NameRGBA16F);
+	if (bSupportDX11TextureFormats)
+	{
+		OutFormats.Add(NameBC6H);
+		OutFormats.Add(NameBC7);
+	}
+#endif
+}
+
diff --git a/Engine/Source/Runtime/Engine/Private/Texture2D.cpp b/Engine/Source/Runtime/Engine/Private/Texture2D.cpp
--- a/Engine/Source/Runtime/Engine/Private/Texture2D.cpp
+++ b/Engine/Source/Runtime/Engine/Private/Texture2D.cpp
@@ -361,6 +361,8 @@ void UTexture2D::UnlinkStreaming()
 
 void UTexture2D::CancelPendingTextureStreaming()
 {
+	FlushRenderingCommands();
+
 	for( TObjectIterator<UTexture2D> It; It; ++It )
 	{
 		UTexture2D* CurrentTexture = *It;
@@ -426,6 +428,8 @@ void UTexture2D::UpdateResource()
 #if WITH_EDITOR
 	// Recache platform data if the source has changed.
 	CachePlatformData();
+	// clear all the cooked cached platform data if the source could have changed... 
+	ClearAllCachedCookedPlatformData();
 #endif // #if WITH_EDITOR
 
 	// Route to super.
@@ -600,6 +604,12 @@ bool UTexture2D::UpdateStreamingStatus( bool bWaitForMipFading /*= false*/ )
 			}
 #endif // #if WITH_EDITOR
 		}
+		else if (RequestStatus > TexState_InProgress_Upload && RequestStatus < TexState_ReadyFor_Upload)
+		{
+			// Bogus state. Got here by canceling a texture upload at a strange time. Recover by changing to TexState_ReadyFor_Upload.
+			UE_LOG(LogTexture, Warning, TEXT("UpdateStreamingStatus for %s got into a bad state (TexState_InProgress_Upload < RequestStatus < TexState_ReadyFor_Upload). Recovering by changing to TexState_ReadyFor_Upload."), *GetPathName());
+			PendingMipChangeRequestStatus.Set(TexState_ReadyFor_Upload);
+		}
 	}
 	return bHasPendingRequestInFlight;
 }
diff --git a/Engine/Source/Runtime/Engine/Public/Animation/AnimInstanceProxy.h b/Engine/Source/Runtime/Engine/Public/Animation/AnimInstanceProxy.h
--- a/Engine/Source/Runtime/Engine/Public/Animation/AnimInstanceProxy.h
+++ b/Engine/Source/Runtime/Engine/Public/Animation/AnimInstanceProxy.h
@@ -264,6 +264,18 @@ public:
 		return LODLevel;
 	}
 
+	// Cached SkeletalMeshComponent LocalToWorld Transform.
+	const FTransform& GetSkelMeshCompLocalToWorld() const
+	{
+		return SkelMeshCompLocalToWorld;
+	}
+
+	// Cached SkeletalMeshComponent Owner Transform.
+	const FTransform& GetSkelMeshCompOwnerTransform() const
+	{
+		return SkelMeshCompOwnerTransform;
+	}
+
 	/** Get the current skeleton we are using. Note that this will return nullptr outside of pre/post update */
 	USkeleton* GetSkeleton() 
 	{ 
@@ -740,7 +752,14 @@ private:
 	/** LODLevel used by RequiredBones */
 	int32 LODLevel;
 
+	/** Cached SkeletalMeshComponent LocalToWorld transform. */
+	FTransform SkelMeshCompLocalToWorld;
+
+	/** Cached SkeletalMeshComponent Owner Transform */
+	FTransform SkelMeshCompOwnerTransform;
+
 protected:
+
 	/** When RequiredBones mapping has changed, AnimNodes need to update their bones caches. */
 	bool bBoneCachesInvalidated;
 
diff --git a/Engine/Source/Runtime/Engine/Public/Net/OnlineEngineInterface.h b/Engine/Source/Runtime/Engine/Public/Net/OnlineEngineInterface.h
--- a/Engine/Source/Runtime/Engine/Public/Net/OnlineEngineInterface.h
+++ b/Engine/Source/Runtime/Engine/Public/Net/OnlineEngineInterface.h
@@ -19,6 +19,8 @@ DECLARE_DELEGATE_TwoParams(FOnlineSessionStartComplete, FName /*InSessionName*/,
 DECLARE_DELEGATE_TwoParams(FOnlineSessionEndComplete, FName /*InSessionName*/, bool /*bWasSuccessful*/);
 /** Delegate fired when an external UI is opened or closed */
 DECLARE_DELEGATE_OneParam(FOnlineExternalUIChanged, bool /*bInIsOpening*/);
+/** Delegate executed when the web url UI has been closed */
+DECLARE_DELEGATE_OneParam(FOnlineShowWebUrlClosed, const FString& /*FinalUrl*/);
 /** Delegate fired when a PIE login has completed */
 DECLARE_DELEGATE_ThreeParams(FOnPIELoginComplete, int32 /*LocalUserNum*/, bool /*bWasSuccessful*/, const FString& /*Error*/);
 
@@ -126,12 +128,75 @@ public:
 	 * External UI
 	 */
 
+	struct FShowWebUrlParams
+	{
+		/** presented without a frame if embedded enabled */
+		bool bEmbedded;
+		/** Show the built in close button */
+		bool bShowCloseButton;
+		/** Show the built in background */
+		bool bShowBackground;
+		/** Hide the mouse cursor */
+		bool bHideCursor;
+		/** Rest cookies before invoking web browser */
+		bool bResetCookies;
+		/** x offset in pixels from top left */
+		int32 OffsetX;
+		/** y offset in pixels from top left */
+		int32 OffsetY;
+		/** x size in pixels */
+		int32 SizeX;
+		/** y size in pixels */
+		int32 SizeY;
+		/** if specified then restricted to only navigate within these domains */
+		TArray<FString> AllowedDomains;
+		/** portion of url for detecting callback.  Eg. "&code=", "redirect=", etc */
+		FString CallbackPath;
+
+		/**
+		 * Constructor
+		 */
+		FShowWebUrlParams(bool InbEmbedded, int32 InOffsetX, int32 InOffsetY, int32 InSizeX, int32 InSizeY)
+			: bEmbedded(InbEmbedded)
+			, bShowCloseButton(false)
+			, bShowBackground(false)
+			, bHideCursor(false)
+			, bResetCookies(false)
+			, OffsetX(InOffsetX)
+			, OffsetY(InOffsetY)
+			, SizeX(InSizeX)
+			, SizeY(InSizeY)
+		{
+		}
+
+		/**
+		 * Default Constructor
+		 */
+		FShowWebUrlParams()
+			: bEmbedded(false)
+			, bShowCloseButton(false)
+			, bShowBackground(false)
+			, bHideCursor(false)
+			, bResetCookies(false)
+			, OffsetX(0)
+			, OffsetY(0)
+			, SizeX(0)
+			, SizeY(0)
+		{
+		}
+	};
+
 	/** Show an online external leaderboard UI if applicable to the platform */
 	virtual void ShowLeaderboardUI(UWorld* World, const FString& CategoryName) {}
 	/** Show an online external achievements UI if applicable to the platform */
 	virtual void ShowAchievementsUI(UWorld* World, int32 LocalUserNum) {}
 	/** Bind a delegate to the opening/closing of an online platform's external UI */
 	virtual void BindToExternalUIOpening(const FOnlineExternalUIChanged& Delegate) {}
+#define OSS_ADDED_SHOW_WEB 1 
+	/** Displays a web page in the external UI */
+	virtual void ShowWebURL(const FString& CurrentURL, const FShowWebUrlParams& ShowParams, const FOnlineShowWebUrlClosed& CompletionDelegate) {}
+	/** Closes the currently active web external UI */
+	virtual bool CloseWebURL() { return false; }
 
 	/**
 	 * Debug
diff --git a/Engine/Source/Runtime/Engine/Public/TextureResource.h b/Engine/Source/Runtime/Engine/Public/TextureResource.h
--- a/Engine/Source/Runtime/Engine/Public/TextureResource.h
+++ b/Engine/Source/Runtime/Engine/Public/TextureResource.h
@@ -41,9 +41,9 @@ enum ETextureStreamingState
 	// Mip data is in the process of being uploaded to the GPU.
 	TexState_InProgress_Upload			= 3,
 	// Mip data has been loaded in to system memory and is ready to be transferred to the GPU.
-	TexState_ReadyFor_Upload			= 4,
+	TexState_ReadyFor_Upload			= 10,
 	// We're currently loading in mip data.
-	TexState_InProgress_Loading			= 5,
+	TexState_InProgress_Loading			= 11,
 	// ...
 	// States 2+N means we're currently loading in N mips
 	// ...
@@ -895,3 +895,5 @@ private:
 };
 
 ENGINE_API FName GetDefaultTextureFormatName( const class ITargetPlatform* TargetPlatform, const class UTexture* Texture, const class FConfigFile& EngineSettings, bool bSupportDX11TextureFormats );
+// returns all the texture formats which can be returned by GetDefaultTextureFormatName
+ENGINE_API void GetAllDefaultTextureFormats( const class ITargetPlatform* TargetPlatform, TArray<FName>& OutFormats, bool bSupportDX11TextureFormats);
\ No newline at end of file
diff --git a/Engine/Source/Runtime/FriendsAndChat/Private/FriendsChatStyle.cpp b/Engine/Source/Runtime/FriendsAndChat/Private/FriendsChatStyle.cpp
--- a/Engine/Source/Runtime/FriendsAndChat/Private/FriendsChatStyle.cpp
+++ b/Engine/Source/Runtime/FriendsAndChat/Private/FriendsChatStyle.cpp
@@ -76,6 +76,24 @@ FFriendsChatStyle& FFriendsChatStyle::SetGameAdminChatColor(const FLinearColor&
 	return *this;
 }
 
+FFriendsChatStyle& FFriendsChatStyle::SetAddedItemChatColor(const FLinearColor& InFontColor)
+{
+	AddedItemChatColor = InFontColor;
+	return *this;
+}
+
+FFriendsChatStyle& FFriendsChatStyle::SetCompletedItemChatColor(const FLinearColor& InFontColor)
+{
+	CompletedItemChatColor = InFontColor;
+	return *this;
+}
+
+FFriendsChatStyle& FFriendsChatStyle::SetDiscardedItemChatColor(const FLinearColor& InFontColor)
+{
+	DiscardedItemChatColor = InFontColor;
+	return *this;
+}
+
 FFriendsChatStyle& FFriendsChatStyle::SetWhisplerHyperlinkChatColor(const FLinearColor& InFontColor)
 {
 	WhisperHyperlinkChatColor = InFontColor;
diff --git a/Engine/Source/Runtime/FriendsAndChat/Public/FriendsChatStyle.h b/Engine/Source/Runtime/FriendsAndChat/Public/FriendsChatStyle.h
--- a/Engine/Source/Runtime/FriendsAndChat/Public/FriendsChatStyle.h
+++ b/Engine/Source/Runtime/FriendsAndChat/Public/FriendsChatStyle.h
@@ -108,6 +108,18 @@ struct FRIENDSANDCHAT_API FFriendsChatStyle
 	FLinearColor GameAdminChatColor;
 	FFriendsChatStyle& SetGameAdminChatColor(const FLinearColor& InFontColor);
 
+	UPROPERTY(EditAnywhere, Category = Appearance)
+	FLinearColor AddedItemChatColor;
+	FFriendsChatStyle& SetAddedItemChatColor(const FLinearColor& InFontColor);
+
+	UPROPERTY(EditAnywhere, Category = Appearance)
+	FLinearColor CompletedItemChatColor;
+	FFriendsChatStyle& SetCompletedItemChatColor(const FLinearColor& InFontColor);
+
+	UPROPERTY(EditAnywhere, Category = Appearance)
+	FLinearColor DiscardedItemChatColor;
+	FFriendsChatStyle& SetDiscardedItemChatColor(const FLinearColor& InFontColor);
+
 	UPROPERTY(EditAnywhere, Category=Appearance)
 	FLinearColor WhisperHyperlinkChatColor;
 	FFriendsChatStyle& SetWhisplerHyperlinkChatColor(const FLinearColor& InFontColor);
diff --git a/Engine/Source/Runtime/GameplayTags/Classes/GameplayTagContainer.h b/Engine/Source/Runtime/GameplayTags/Classes/GameplayTagContainer.h
--- a/Engine/Source/Runtime/GameplayTags/Classes/GameplayTagContainer.h
+++ b/Engine/Source/Runtime/GameplayTags/Classes/GameplayTagContainer.h
@@ -855,6 +855,31 @@ struct TStructOpsTypeTraits<FGameplayTagContainer> : public TStructOpsTypeTraits
 	};
 };
 
+/**
+ *	Helper struct for viewing tag references (assets that reference a tag). Drop this into a struct and set the OnGetgameplayStatName. A details customization
+ *	will display a tree view of assets referencing the tag
+ */
+USTRUCT()
+struct FGameplayTagReferenceHelper
+{
+	GENERATED_USTRUCT_BODY()
+
+	FGameplayTagReferenceHelper()
+	{
+	}
+
+	/** 
+	 *	Delegate to be called to get the tag we want to inspect. The void* is a raw pointer to the outer struct's data. You can do a static cast to access it. Eg:
+	 *	
+	 *	ReferenceHelper.OnGetGameplayTagName.BindLambda([this](void* RawData) {
+	 *		FMyStructType* ThisData = static_cast<FMyStructType*>(RawData);
+	 *		return ThisData->MyTag.GetTagName();
+	 *	});
+	 *
+	*/
+	DECLARE_DELEGATE_RetVal_OneParam(FName, FOnGetGameplayTagName, void* /**RawOuterStructData*/);
+	FOnGetGameplayTagName OnGetGameplayTagName;
+};
 
 /** Enumerates the list of supported query expression types. */
 UENUM()
diff --git a/Engine/Source/Runtime/GameplayTags/Private/GameplayTagsManager.cpp b/Engine/Source/Runtime/GameplayTags/Private/GameplayTagsManager.cpp
--- a/Engine/Source/Runtime/GameplayTags/Private/GameplayTagsManager.cpp
+++ b/Engine/Source/Runtime/GameplayTags/Private/GameplayTagsManager.cpp
@@ -1033,7 +1033,7 @@ void UGameplayTagsManager::RequestAllGameplayTags(FGameplayTagContainer& TagCont
 	for (const TSharedPtr<FGameplayTagNode>& TagNode : ValueArray)
 	{
 #if WITH_EDITOR
-		bool DictTag = IsDictionaryTag(TagNode->GetSimpleTagName());
+		bool DictTag = IsDictionaryTag(TagNode->GetCompleteTagName());
 #else
 		bool DictTag = false;
 #endif 
diff --git a/Engine/Source/Runtime/Online/XMPP/Private/XmppJingle/XmppConnectionJingle.cpp b/Engine/Source/Runtime/Online/XMPP/Private/XmppJingle/XmppConnectionJingle.cpp
--- a/Engine/Source/Runtime/Online/XMPP/Private/XmppJingle/XmppConnectionJingle.cpp
+++ b/Engine/Source/Runtime/Online/XMPP/Private/XmppJingle/XmppConnectionJingle.cpp
@@ -672,6 +672,7 @@ void FXmppConnectionJingle::Login(const FString& UserId, const FString& Password
 	buzz::XmppClientSettings Settings;
 	Settings.set_host(TCHAR_TO_UTF8(*ServerConfig.Domain));
 	Settings.set_use_tls(ServerConfig.bUseSSL ? buzz::TLS_ENABLED : buzz::TLS_DISABLED);
+	Settings.set_allow_plain(ServerConfig.bUsePlainTextAuth);
 	Settings.set_resource(TCHAR_TO_UTF8(*ServerConfig.ClientResource));
 	Settings.set_server(rtc::SocketAddress(TCHAR_TO_UTF8(*ServerConfig.ServerAddr), ServerConfig.ServerPort));
 
diff --git a/Engine/Source/Runtime/Online/XMPP/Public/XmppConnection.h b/Engine/Source/Runtime/Online/XMPP/Public/XmppConnection.h
--- a/Engine/Source/Runtime/Online/XMPP/Public/XmppConnection.h
+++ b/Engine/Source/Runtime/Online/XMPP/Public/XmppConnection.h
@@ -50,6 +50,7 @@ public:
 	FXmppServer() 
 		: ServerPort(5222) 
 		, bUseSSL(true)
+		, bUsePlainTextAuth(false)
 		, PingInterval(60.0f)
 		, PingTimeout(30.0f)
 		, MaxPingRetries(1)
@@ -70,6 +71,8 @@ public:
 	FString Platform;
 	/** true to enable SSL handshake for connection */
 	bool bUseSSL;
+	/** true to allow the usage of plain text authentication */
+	bool bUsePlainTextAuth;
 	/** seconds between sending server ping requests */
 	float PingInterval;
 	/** seconds before a server ping request is treated as a timeout */
diff --git a/Engine/Source/Runtime/PakFile/Private/IPlatformFilePak.cpp b/Engine/Source/Runtime/PakFile/Private/IPlatformFilePak.cpp
--- a/Engine/Source/Runtime/PakFile/Private/IPlatformFilePak.cpp
+++ b/Engine/Source/Runtime/PakFile/Private/IPlatformFilePak.cpp
@@ -2822,7 +2822,7 @@ void FPakPrecacher::DoSignatureCheck(bool bWasCanceled, IAsyncReadRequest* Reque
 	const uint8* Data = nullptr;
 	int64 RequestSize = 0;
 	int64 RequestOffset = 0;
-	FPakData* PakData = nullptr;
+	uint16 PakIndex;
 
 	{
 		// Try and keep lock for as short a time as possible. Find our request and copy out the data we need
@@ -2838,8 +2838,7 @@ void FPakPrecacher::DoSignatureCheck(bool bWasCanceled, IAsyncReadRequest* Reque
 		RequestOffset = GetRequestOffset(Block.OffsetAndPakIndex);
 		check((RequestOffset % FPakInfo::MaxChunkDataSize) == 0);
 		RequestSize = RequestToLower.RequestSize;
-		int64 PakIndex = GetRequestPakIndex(Block.OffsetAndPakIndex);
-		PakData = &CachedPakData[PakIndex];
+		PakIndex = GetRequestPakIndex(Block.OffsetAndPakIndex);
 		Data = RequestToLower.Memory;
 		SignatureIndex = RequestOffset / FPakInfo::MaxChunkDataSize;
 	}
@@ -2848,7 +2847,6 @@ void FPakPrecacher::DoSignatureCheck(bool bWasCanceled, IAsyncReadRequest* Reque
 	check(NumSignaturesToCheck > 0);
 	check(RequestSize > 0);
 	check(RequestOffset >= 0);
-	check(PakData != nullptr);
 
 	// Hash the contents of the incoming buffer and check that it matches what we expected
 	for (int64 SignedChunkIndex = 0; SignedChunkIndex < NumSignaturesToCheck; ++SignedChunkIndex, ++SignatureIndex)
@@ -2857,7 +2855,14 @@ void FPakPrecacher::DoSignatureCheck(bool bWasCanceled, IAsyncReadRequest* Reque
 
 		{
 			SCOPE_SECONDS_ACCUMULATOR(STAT_PakCache_SigningChunkHashTime);
-			bool bChunkHashesMatch = ComputePakChunkHash(Data, Size) == PakData->ChunkHashes[SignatureIndex];
+
+			TPakChunkHash ThisHash = ComputePakChunkHash(Data, Size);
+			bool bChunkHashesMatch;
+			{
+				FScopeLock Lock(&CachedFilesScopeLock);
+				FPakData* PakData = &CachedPakData[PakIndex];
+				bChunkHashesMatch = ThisHash == PakData->ChunkHashes[SignatureIndex];
+			}
 			ensure(bChunkHashesMatch);
 			if (!ensure(bChunkHashesMatch))
 			{
diff --git a/Engine/Source/Runtime/Projects/Private/PluginManager.cpp b/Engine/Source/Runtime/Projects/Private/PluginManager.cpp
--- a/Engine/Source/Runtime/Projects/Private/PluginManager.cpp
+++ b/Engine/Source/Runtime/Projects/Private/PluginManager.cpp
@@ -3,6 +3,7 @@
 #include "PluginManager.h"
 #include "GenericPlatform/GenericPlatformFile.h"
 #include "HAL/PlatformFilemanager.h"
+#include "HAL/FileManager.h"
 #include "Misc/MessageDialog.h"
 #include "Misc/CommandLine.h"
 #include "Misc/Paths.h"
@@ -11,6 +12,7 @@
 #include "Misc/CoreDelegates.h"
 #include "Misc/App.h"
 #include "Misc/EngineVersion.h"
+#include "Misc/FileHelper.h"
 #include "ProjectDescriptor.h"
 #include "Interfaces/IProjectManager.h"
 #include "Modules/ModuleManager.h"
@@ -247,7 +249,7 @@ void FPluginManager::ReadPluginsInDirectory(const FString& PluginsDirectory, con
 			if ( Descriptor.Load(FileName, LoadedFrom == EPluginLoadedFrom::GameProject, FailureReason) )
 			{
 				TSharedRef<FPlugin> Plugin = MakeShareable(new FPlugin(FileName, Descriptor, LoadedFrom));
-
+				
 				FString FullPath = FPaths::ConvertRelativePathToFull(FileName);
 				UE_LOG(LogPluginManager, Verbose, TEXT("Read plugin descriptor for %s, from %s"), *Plugin->GetName(), *FullPath);
 
@@ -435,13 +437,13 @@ bool FPluginManager::ConfigureEnabledPlugins()
 						Plugin->bEnabled = false;
 						AllEnabledPlugins.Remove(Plugin->Name);
 						UE_LOG(LogPluginManager, Log, TEXT("Disabled plugin '%s' due to incompatibility"), *Plugin->FileName);
-					}
+			}
 					else
 					{
 						UE_LOG(LogPluginManager, Log, TEXT("Enabled plugin '%s' despite being built for CL %d"), *Plugin->FileName, Plugin->Descriptor.CompatibleChangelist);
 					}
 				}
-			}
+		}
 		}
 
 		if (bHasProjectFile)
@@ -566,6 +568,30 @@ bool FPluginManager::ConfigureEnabledPlugins()
 					// Nothing to add, remove from map
 					GConfig->Remove(PluginConfigFilename);
 				}
+
+				if (!GIsEditor)
+				{
+					// override config cache entries with plugin configs (Engine.ini, Game.ini, etc in <PluginDir>\Config\)
+					TArray<FString> PluginConfigs;
+					IFileManager::Get().FindFiles(PluginConfigs, *PluginConfigDir, TEXT("ini"));
+					for (const FString& ConfigFile : PluginConfigs)
+					{
+						FString PlaformName = FPlatformProperties::PlatformName();
+						PluginConfigFilename = FString::Printf(TEXT("%s%s/%s.ini"), *FPaths::GeneratedConfigDir(), *PlaformName, *FPaths::GetBaseFilename(ConfigFile));
+						FConfigFile* FoundConfig = GConfig->Find(PluginConfigFilename, false);
+						if (FoundConfig != nullptr)
+						{
+							FString PluginConfigContent;
+							if (FFileHelper::LoadFileToString(PluginConfigContent, *FPaths::Combine(PluginConfigDir, ConfigFile)))
+							{
+								FoundConfig->CombineFromBuffer(PluginConfigContent);
+								// if plugin config overrides are applied then don't save
+								FoundConfig->NoSave = true;
+							}
+						}
+					}
+				}
+
 			}
 		}
 
@@ -619,56 +645,56 @@ TSharedPtr<FPlugin> FPluginManager::FindPluginInstance(const FString& Name)
 
 
 static bool TryLoadModulesForPlugin( const FPlugin& Plugin, const ELoadingPhase::Type LoadingPhase )
-{
-	TMap<FName, EModuleLoadResult> ModuleLoadFailures;
+		{
+			TMap<FName, EModuleLoadResult> ModuleLoadFailures;
 	FModuleDescriptor::LoadModulesForPhase(LoadingPhase, Plugin.Descriptor.Modules, ModuleLoadFailures);
 
-	FText FailureMessage;
-	for( auto FailureIt( ModuleLoadFailures.CreateConstIterator() ); FailureIt; ++FailureIt )
-	{
-		const FName ModuleNameThatFailedToLoad = FailureIt.Key();
-		const EModuleLoadResult FailureReason = FailureIt.Value();
+			FText FailureMessage;
+			for( auto FailureIt( ModuleLoadFailures.CreateConstIterator() ); FailureIt; ++FailureIt )
+			{
+				const FName ModuleNameThatFailedToLoad = FailureIt.Key();
+				const EModuleLoadResult FailureReason = FailureIt.Value();
 
-		if( FailureReason != EModuleLoadResult::Success )
-		{
+				if( FailureReason != EModuleLoadResult::Success )
+				{
 			const FText PluginNameText = FText::FromString(Plugin.Name);
-			const FText TextModuleName = FText::FromName(FailureIt.Key());
+					const FText TextModuleName = FText::FromName(FailureIt.Key());
 
-			if ( FailureReason == EModuleLoadResult::FileNotFound )
-			{
-				FailureMessage = FText::Format( LOCTEXT("PluginModuleNotFound", "Plugin '{0}' failed to load because module '{1}' could not be found.  Please ensure the plugin is properly installed, otherwise consider disabling the plugin for this project."), PluginNameText, TextModuleName );
-			}
-			else if ( FailureReason == EModuleLoadResult::FileIncompatible )
-			{
-				FailureMessage = FText::Format( LOCTEXT("PluginModuleIncompatible", "Plugin '{0}' failed to load because module '{1}' does not appear to be compatible with the current version of the engine.  The plugin may need to be recompiled."), PluginNameText, TextModuleName );
-			}
-			else if ( FailureReason == EModuleLoadResult::CouldNotBeLoadedByOS )
-			{
-				FailureMessage = FText::Format( LOCTEXT("PluginModuleCouldntBeLoaded", "Plugin '{0}' failed to load because module '{1}' could not be loaded.  There may be an operating system error or the module may not be properly set up."), PluginNameText, TextModuleName );
-			}
-			else if ( FailureReason == EModuleLoadResult::FailedToInitialize )
-			{
-				FailureMessage = FText::Format( LOCTEXT("PluginModuleFailedToInitialize", "Plugin '{0}' failed to load because module '{1}' could not be initialized successfully after it was loaded."), PluginNameText, TextModuleName );
+					if ( FailureReason == EModuleLoadResult::FileNotFound )
+					{
+						FailureMessage = FText::Format( LOCTEXT("PluginModuleNotFound", "Plugin '{0}' failed to load because module '{1}' could not be found.  Please ensure the plugin is properly installed, otherwise consider disabling the plugin for this project."), PluginNameText, TextModuleName );
+					}
+					else if ( FailureReason == EModuleLoadResult::FileIncompatible )
+					{
+						FailureMessage = FText::Format( LOCTEXT("PluginModuleIncompatible", "Plugin '{0}' failed to load because module '{1}' does not appear to be compatible with the current version of the engine.  The plugin may need to be recompiled."), PluginNameText, TextModuleName );
+					}
+					else if ( FailureReason == EModuleLoadResult::CouldNotBeLoadedByOS )
+					{
+						FailureMessage = FText::Format( LOCTEXT("PluginModuleCouldntBeLoaded", "Plugin '{0}' failed to load because module '{1}' could not be loaded.  There may be an operating system error or the module may not be properly set up."), PluginNameText, TextModuleName );
+					}
+					else if ( FailureReason == EModuleLoadResult::FailedToInitialize )
+					{
+						FailureMessage = FText::Format( LOCTEXT("PluginModuleFailedToInitialize", "Plugin '{0}' failed to load because module '{1}' could not be initialized successfully after it was loaded."), PluginNameText, TextModuleName );
+					}
+					else 
+					{
+						ensure(0);	// If this goes off, the error handling code should be updated for the new enum values!
+						FailureMessage = FText::Format( LOCTEXT("PluginGenericLoadFailure", "Plugin '{0}' failed to load because module '{1}' could not be loaded for an unspecified reason.  This plugin's functionality will not be available.  Please report this error."), PluginNameText, TextModuleName );
+					}
+
+					// Don't need to display more than one module load error per plugin that failed to load
+					break;
+				}
 			}
-			else 
+
+			if( !FailureMessage.IsEmpty() )
 			{
-				ensure(0);	// If this goes off, the error handling code should be updated for the new enum values!
-				FailureMessage = FText::Format( LOCTEXT("PluginGenericLoadFailure", "Plugin '{0}' failed to load because module '{1}' could not be loaded for an unspecified reason.  This plugin's functionality will not be available.  Please report this error."), PluginNameText, TextModuleName );
+				FMessageDialog::Open(EAppMsgType::Ok, FailureMessage);
+				return false;
 			}
 
-			// Don't need to display more than one module load error per plugin that failed to load
-			break;
-		}
-	}
-
-	if( !FailureMessage.IsEmpty() )
-	{
-		FMessageDialog::Open(EAppMsgType::Ok, FailureMessage);
-		return false;
-	}
-
 	return true;
-}
+		}
 
 bool FPluginManager::LoadModulesForEnabledPlugins( const ELoadingPhase::Type LoadingPhase )
 {
diff --git a/Engine/Source/Runtime/WebBrowser/Public/SWebBrowser.h b/Engine/Source/Runtime/WebBrowser/Public/SWebBrowser.h
--- a/Engine/Source/Runtime/WebBrowser/Public/SWebBrowser.h
+++ b/Engine/Source/Runtime/WebBrowser/Public/SWebBrowser.h
@@ -206,8 +206,6 @@ public:
 
 	void UnbindAdapter(const TSharedRef<IWebBrowserAdapter>& Adapter);
 
-private:
-
 	/** Returns true if the browser can navigate backwards. */
 	bool CanGoBack() const;
 
diff --git a/Engine/Source/Runtime/WebBrowser/WebBrowser.Build.cs b/Engine/Source/Runtime/WebBrowser/WebBrowser.Build.cs
--- a/Engine/Source/Runtime/WebBrowser/WebBrowser.Build.cs
+++ b/Engine/Source/Runtime/WebBrowser/WebBrowser.Build.cs
@@ -57,9 +57,10 @@ public class WebBrowser : ModuleRules
 			}
 		}
 
-		if (Target.Platform == UnrealTargetPlatform.PS4)
+		if (Target.Platform == UnrealTargetPlatform.PS4 &&
+			UEBuildConfiguration.bCompileAgainstEngine)
 		{
-			PrivateDependencyModuleNames.Add("OnlineSubsystem");
+			PrivateDependencyModuleNames.Add("Engine");
 		}
 
 		bEnableShadowVariableWarnings = false;
