[NO NON-PII COMMIT MESSAGE]

diff --git a/Engine/Source/Editor/UnrealEd/Private/CookOnTheFlyServer.cpp b/Engine/Source/Editor/UnrealEd/Private/CookOnTheFlyServer.cpp
--- a/Engine/Source/Editor/UnrealEd/Private/CookOnTheFlyServer.cpp
+++ b/Engine/Source/Editor/UnrealEd/Private/CookOnTheFlyServer.cpp
@@ -1177,14 +1177,31 @@ bool UCookOnTheFlyServer::ContainsRedirector(const FName& PackageName, TMap<FStr
 			{
 				ConstructorHelpers::StripObjectClass(RedirectedPath);
 				FAssetData DestinationData = AssetRegistry->GetAssetByObjectPath(FName(*RedirectedPath), true);
+				TSet<FString> SeenPaths;
+
+				SeenPaths.Add(RedirectedPath);
 
 				// Need to follow chain of redirectors
 				while (DestinationData.IsRedirector())
 				{
-					if (Asset.GetTagValue("DestinationObject", RedirectedPath))
+					if (DestinationData.GetTagValue("DestinationObject", RedirectedPath))
 					{
 						ConstructorHelpers::StripObjectClass(RedirectedPath);
-						DestinationData = AssetRegistry->GetAssetByObjectPath(FName(*RedirectedPath), true);
+						if (SeenPaths.Contains(RedirectedPath))
+						{
+							// Recursive, bail
+							DestinationData = FAssetData();
+						}
+						else
+						{
+							SeenPaths.Add(RedirectedPath);
+							DestinationData = AssetRegistry->GetAssetByObjectPath(FName(*RedirectedPath), true);
+						}
+					}
+					else
+					{
+						// Can't extract
+						DestinationData = FAssetData();						
 					}
 				}
 
