Copying //UE4/Dev-Framework to //UE4/Dev-Main (Source: //UE4/Dev-Framework @ 3431384)

diff --git a/Engine/Source/Editor/UnrealEd/Classes/Editor/Transactor.h b/Engine/Source/Editor/UnrealEd/Classes/Editor/Transactor.h
--- a/Engine/Source/Editor/UnrealEd/Classes/Editor/Transactor.h
+++ b/Engine/Source/Editor/UnrealEd/Classes/Editor/Transactor.h
@@ -234,13 +234,6 @@ protected:
 				}
 				return *this;
 			}
-			FArchive& operator<<(class FAssetPtr& AssetPtr) override
-			{
-				FStringAssetReference ID;
-				ID.Serialize(*this);
-				AssetPtr = ID;
-				return *this;
-			}
 			void Preload( UObject* InObject ) override
 			{
 				if( Owner )
@@ -330,12 +323,6 @@ protected:
 				}
 				return (FArchive&)*this << ObjectIndex;
 			}
-			FArchive& operator<<(class FAssetPtr& AssetPtr)
-			{
-				FStringAssetReference ID = AssetPtr.GetUniqueID();
-				ID.Serialize(*this);
-				return *this;
-			}
 			TArray<uint8>& Data;
 			ObjectMapType ObjectMap;
 			TArray<FPersistentObjectRef>& ReferencedObjects;
diff --git a/Engine/Source/Editor/UnrealEd/Classes/Editor/UnrealEdTypes.h b/Engine/Source/Editor/UnrealEd/Classes/Editor/UnrealEdTypes.h
--- a/Engine/Source/Editor/UnrealEd/Classes/Editor/UnrealEdTypes.h
+++ b/Engine/Source/Editor/UnrealEd/Classes/Editor/UnrealEdTypes.h
@@ -21,7 +21,7 @@ struct FLightmassParameterValue
 	GENERATED_USTRUCT_BODY()
 
 	/** If true, override the given parameter with the given settings */
-	UPROPERTY(EditAnywhere, BlueprintReadWrite, Category=LightmassParameterValue)
+	UPROPERTY(EditAnywhere, Category=LightmassParameterValue)
 	uint32 bOverride:1;
 
 
@@ -39,7 +39,7 @@ struct FLightmassBooleanParameterValue : public FLightmassParameterValue
 	GENERATED_USTRUCT_BODY()
 
 	/** The boolean value to override the parent value with */
-	UPROPERTY(EditAnywhere, BlueprintReadWrite, Category=LightmassBooleanParameterValue)
+	UPROPERTY(EditAnywhere, Category=LightmassBooleanParameterValue)
 	uint32 ParameterValue:1;
 
 
@@ -57,7 +57,7 @@ struct FLightmassScalarParameterValue : public FLightmassParameterValue
 	GENERATED_USTRUCT_BODY()
 
 	/** The scalar value to override the parent value with */
-	UPROPERTY(EditAnywhere, BlueprintReadWrite, Category=LightmassScalarParameterValue)
+	UPROPERTY(EditAnywhere, Category=LightmassScalarParameterValue)
 	float ParameterValue;
 
 
@@ -78,22 +78,22 @@ struct FLightmassParameterizedMaterialSettings
 	GENERATED_USTRUCT_BODY()
 
 	/** If true, forces translucency to cast static shadows as if the material were masked. */
-	UPROPERTY(EditAnywhere, BlueprintReadWrite, Category=LightmassParameterizedMaterialSettings)
+	UPROPERTY(EditAnywhere, Category=LightmassParameterizedMaterialSettings)
 	struct FLightmassBooleanParameterValue CastShadowAsMasked;
 
 	/** Scales the emissive contribution of this material to static lighting. */
-	UPROPERTY(EditAnywhere, BlueprintReadWrite, Category=LightmassParameterizedMaterialSettings)
+	UPROPERTY(EditAnywhere, Category=LightmassParameterizedMaterialSettings)
 	struct FLightmassScalarParameterValue EmissiveBoost;
 
 	/** Scales the diffuse contribution of this material to static lighting. */
-	UPROPERTY(EditAnywhere, BlueprintReadWrite, Category=LightmassParameterizedMaterialSettings)
+	UPROPERTY(EditAnywhere, Category=LightmassParameterizedMaterialSettings)
 	struct FLightmassScalarParameterValue DiffuseBoost;
 
 	/** 
 	 * Scales the resolution that this material's attributes were exported at. 
 	 * This is useful for increasing material resolution when details are needed.
 	 */
-	UPROPERTY(EditAnywhere, BlueprintReadWrite, Category=LightmassParameterizedMaterialSettings)
+	UPROPERTY(EditAnywhere, Category=LightmassParameterizedMaterialSettings)
 	struct FLightmassScalarParameterValue ExportResolutionScale;
 
 	FLightmassParameterizedMaterialSettings()
diff --git a/Engine/Source/Editor/UnrealEd/Classes/UserDefinedStructure/UserDefinedStructEditorData.h b/Engine/Source/Editor/UnrealEd/Classes/UserDefinedStructure/UserDefinedStructEditorData.h
--- a/Engine/Source/Editor/UnrealEd/Classes/UserDefinedStructure/UserDefinedStructEditorData.h
+++ b/Engine/Source/Editor/UnrealEd/Classes/UserDefinedStructure/UserDefinedStructEditorData.h
@@ -44,41 +44,51 @@ struct FStructVariableDescription
 	FEdGraphTerminalType PinValueType;
 
 	UPROPERTY()
-	bool bIsArray;
+	EPinContainerType ContainerType;
 
+	// DEPRECATED(4.17)
 	UPROPERTY()
-	bool bIsSet;
+	uint8 bIsArray_DEPRECATED:1;
 
+	// DEPRECATED(4.17)
 	UPROPERTY()
-	bool bIsMap;
+	uint8 bIsSet_DEPRECATED:1;
+
+	// DEPRECATED(4.17)
+	UPROPERTY()
+	uint8 bIsMap_DEPRECATED:1;
 
 	UPROPERTY(Transient)
-	bool bInvalidMember;
+	uint8 bInvalidMember:1;
 
-	// CurrentDefaultValue stores the actual default value, after the DefaultValue was changed, and before the struct was recompiled
 	UPROPERTY()
-	FString CurrentDefaultValue;
+	uint8 bDontEditoOnInstance:1;
 
 	UPROPERTY()
-	FString ToolTip;
+	uint8 bEnableMultiLineText:1;
 
 	UPROPERTY()
-	bool bDontEditoOnInstance;
+	uint8 bEnable3dWidget:1;
 
+	// CurrentDefaultValue stores the actual default value, after the DefaultValue was changed, and before the struct was recompiled
 	UPROPERTY()
-	bool bEnableMultiLineText;
+	FString CurrentDefaultValue;
 
 	UPROPERTY()
-	bool bEnable3dWidget;
+	FString ToolTip;
 
 	UNREALED_API bool SetPinType(const struct FEdGraphPinType& VarType);
 
 	UNREALED_API FEdGraphPinType ToPinType() const;
 
+	// DEPRECATED(4.17)
+	void PostSerialize(const FArchive& Ar);
+
 	FStructVariableDescription()
-		: bIsArray(false)
-		, bIsSet(false)
-		, bIsMap(false)
+		: ContainerType(EPinContainerType::None)
+		, bIsArray_DEPRECATED(false)
+		, bIsSet_DEPRECATED(false)
+		, bIsMap_DEPRECATED(false)
 		, bInvalidMember(false)
 		, bDontEditoOnInstance(false)
 		, bEnableMultiLineText(false)
@@ -86,6 +96,15 @@ struct FStructVariableDescription
 	{ }
 };
 
+template<>
+struct TStructOpsTypeTraits< FStructVariableDescription > : public TStructOpsTypeTraitsBase2< FStructVariableDescription >
+{
+	enum 
+	{
+		WithPostSerialize = true,
+	};
+};
+
 class FStructOnScopeMember : public FStructOnScope
 {
 public:
diff --git a/Engine/Source/Editor/UnrealEd/Private/AutoReimport/AssetSourceFilenameCache.cpp b/Engine/Source/Editor/UnrealEd/Private/AutoReimport/AssetSourceFilenameCache.cpp
--- a/Engine/Source/Editor/UnrealEd/Private/AutoReimport/AssetSourceFilenameCache.cpp
+++ b/Engine/Source/Editor/UnrealEd/Private/AutoReimport/AssetSourceFilenameCache.cpp
@@ -7,6 +7,12 @@
 
 FAssetSourceFilenameCache::FAssetSourceFilenameCache()
 {
+	if (GIsRequestingExit)
+	{
+		// This can get created for the first timeon shutdown, if so don't do anything
+		return;
+	}
+
 	IAssetRegistry& AssetRegistry = FModuleManager::LoadModuleChecked<FAssetRegistryModule>("AssetRegistry").Get();
 
 	AssetRegistry.OnAssetAdded().AddRaw(this, &FAssetSourceFilenameCache::HandleOnAssetAdded);
@@ -46,41 +52,15 @@ void FAssetSourceFilenameCache::Shutdown()
 	UAssetImportData::OnImportDataChanged.RemoveAll(this);
 }
 
-TOptional<FAssetImportInfo> FAssetSourceFilenameCache::ExtractAssetImportInfo(const TArray<UObject::FAssetRegistryTag>& InTags)
-{
-	static const FName LegacySourceFilePathName("SourceFile");
-
-	TOptional<FAssetImportInfo> Info;
-
-	for (const auto& Tag : InTags)
-	{
-		if (Tag.Name == UObject::SourceFileTagName())
-		{
-			Info = FAssetImportInfo::FromJson(*Tag.Value);
-			// We're done
-			break;
-		}
-		else if (Tag.Name == LegacySourceFilePathName)
-		{
-			FAssetImportInfo Legacy;
-			Legacy.Insert(Tag.Value);
-			Info = Legacy;
-			// Keep looking for a newer json version
-		}
-	}
-
-	return Info;
-}
-
-TOptional<FAssetImportInfo> FAssetSourceFilenameCache::ExtractAssetImportInfo(const TSharedMapView<FName, FString>& InTags)
+TOptional<FAssetImportInfo> FAssetSourceFilenameCache::ExtractAssetImportInfo(const FAssetData& AssetData)
 {
 	static const FName LegacySourceFilePathName("SourceFile");
 
-	if (const FString* ImportDataString = InTags.Find(UObject::SourceFileTagName()))
+	if (const FString* ImportDataString = AssetData.TagsAndValues.Find(UObject::SourceFileTagName()))
 	{
 		return FAssetImportInfo::FromJson(*ImportDataString);
 	}
-	else if (const FString* LegacyFilename = InTags.Find(LegacySourceFilePathName))
+	else if (const FString* LegacyFilename = AssetData.TagsAndValues.Find(LegacySourceFilePathName))
 	{
 		FAssetImportInfo Legacy;
 		Legacy.Insert(*LegacyFilename);
@@ -94,7 +74,7 @@ TOptional<FAssetImportInfo> FAssetSourceFilenameCache::ExtractAssetImportInfo(co
 
 void FAssetSourceFilenameCache::HandleOnAssetAdded(const FAssetData& AssetData)
 {
-	TOptional<FAssetImportInfo> ImportData = ExtractAssetImportInfo(AssetData.TagsAndValues);
+	TOptional<FAssetImportInfo> ImportData = ExtractAssetImportInfo(AssetData);
 	if (ImportData.IsSet())
 	{
 		for (const auto& SourceFile : ImportData->SourceFiles)
@@ -106,7 +86,7 @@ void FAssetSourceFilenameCache::HandleOnAssetAdded(const FAssetData& AssetData)
 
 void FAssetSourceFilenameCache::HandleOnAssetRemoved(const FAssetData& AssetData)
 {
-	TOptional<FAssetImportInfo> ImportData = ExtractAssetImportInfo(AssetData.TagsAndValues);
+	TOptional<FAssetImportInfo> ImportData = ExtractAssetImportInfo(AssetData);
 	if (ImportData.IsSet())
 	{
 		for (auto& SourceFile : ImportData->SourceFiles)
@@ -126,7 +106,7 @@ void FAssetSourceFilenameCache::HandleOnAssetRemoved(const FAssetData& AssetData
 
 void FAssetSourceFilenameCache::HandleOnAssetRenamed(const FAssetData& AssetData, const FString& OldPath)
 {
-	TOptional<FAssetImportInfo> ImportData = ExtractAssetImportInfo(AssetData.TagsAndValues);
+	TOptional<FAssetImportInfo> ImportData = ExtractAssetImportInfo(AssetData);
 	if (ImportData.IsSet())
 	{
 		FName OldPathName = *OldPath;
@@ -183,7 +163,7 @@ TArray<FAssetData> FAssetSourceFilenameCache::GetAssetsPertainingToFile(const IA
 		for (const FName& Path : *ObjectPaths)
 		{
 			FAssetData Asset = Registry.GetAssetByObjectPath(Path);
-			TOptional<FAssetImportInfo> ImportInfo = ExtractAssetImportInfo(Asset.TagsAndValues);
+			TOptional<FAssetImportInfo> ImportInfo = ExtractAssetImportInfo(Asset);
 			if (ImportInfo.IsSet())
 			{
 				auto AssetPackagePath = FPackageName::LongPackageNameToFilename(Asset.PackagePath.ToString() / TEXT(""));
diff --git a/Engine/Source/Editor/UnrealEd/Private/AutoReimport/AutoReimportManager.cpp b/Engine/Source/Editor/UnrealEd/Private/AutoReimport/AutoReimportManager.cpp
--- a/Engine/Source/Editor/UnrealEd/Private/AutoReimport/AutoReimportManager.cpp
+++ b/Engine/Source/Editor/UnrealEd/Private/AutoReimport/AutoReimportManager.cpp
@@ -398,7 +398,7 @@ void FAutoReimportManager::OnAssetRenamed(const FAssetData& AssetData, const FSt
 	// Additionally, we rename the source file if it matched the name of the asset before the rename/move.
 	//	- If we rename the source file, then we also update the reimport paths for the asset
 
-	TOptional<FAssetImportInfo> ImportInfo = FAssetSourceFilenameCache::ExtractAssetImportInfo(AssetData.TagsAndValues);
+	TOptional<FAssetImportInfo> ImportInfo = FAssetSourceFilenameCache::ExtractAssetImportInfo(AssetData);
 	if (!ImportInfo.IsSet() || ImportInfo->SourceFiles.Num() != 1)
 	{
 		return;
diff --git a/Engine/Source/Editor/UnrealEd/Private/AutoReimport/ContentDirectoryMonitor.cpp b/Engine/Source/Editor/UnrealEd/Private/AutoReimport/ContentDirectoryMonitor.cpp
--- a/Engine/Source/Editor/UnrealEd/Private/AutoReimport/ContentDirectoryMonitor.cpp
+++ b/Engine/Source/Editor/UnrealEd/Private/AutoReimport/ContentDirectoryMonitor.cpp
@@ -59,7 +59,7 @@ DirectoryWatcher::FFileCacheConfig GenerateFileCacheConfig(const FString& InPath
 		// We need to consider this as a changed file if the hash doesn't match any asset imported from that file
 		for (FAssetData& Asset : Assets)
 		{
-			TOptional<FAssetImportInfo> Info = FAssetSourceFilenameCache::ExtractAssetImportInfo(Asset.TagsAndValues);
+			TOptional<FAssetImportInfo> Info = FAssetSourceFilenameCache::ExtractAssetImportInfo(Asset);
 
 			// Check if the source file that this asset last imported was the same as the one we're going to reimport.
 			// If it is, there's no reason to auto-reimport it
diff --git a/Engine/Source/Editor/UnrealEd/Private/Commandlets/AssetRegistryGenerator.cpp b/Engine/Source/Editor/UnrealEd/Private/Commandlets/AssetRegistryGenerator.cpp
--- a/Engine/Source/Editor/UnrealEd/Private/Commandlets/AssetRegistryGenerator.cpp
+++ b/Engine/Source/Editor/UnrealEd/Private/Commandlets/AssetRegistryGenerator.cpp
@@ -237,6 +237,7 @@ bool FAssetRegistryGenerator::GenerateStreamingInstallManifest()
 			{
 				FString Filename = ChunkFilenames[FilenameIndex];
 				FString PakListLine = FPaths::ConvertRelativePathToFull(Filename.Replace(TEXT("[Platform]"), *Platform));
+				if (MaxChunkSize > 0)
 				{
 					TArray<FString> FoundFiles;
 					FString FileSearchString = FString::Printf(TEXT("%s.*"), *PakListLine);
@@ -247,14 +248,14 @@ bool FAssetRegistryGenerator::GenerateStreamingInstallManifest()
 						int64 FileSize = IFileManager::Get().FileSize(*(FPaths::Combine(Path,FoundFile)));
 						CurrentPakSize +=  FileSize > 0 ? FileSize : 0;
 					}
+					if (MaxChunkSize < CurrentPakSize)
+					{
+						// early out if we are over memory limit
+						bFinishedAllFiles = false;
+						break;
+					}
 				}
-				if ((MaxChunkSize > 0) && (MaxChunkSize < CurrentPakSize))
-				{
-					// early out if we are over memory limit
-					bFinishedAllFiles = false;
-					break;
-				}
-
+				
 				PakListLine.ReplaceInline(TEXT("/"), TEXT("\\"));
 				PakListLine += TEXT("\r\n");
 				PakListFile->Serialize(TCHAR_TO_ANSI(*PakListLine), PakListLine.Len());
@@ -519,7 +520,7 @@ void FAssetRegistryGenerator::UpdatePackageSourceHashes()
 	}
 }
 
-void FAssetRegistryGenerator::ComputePackageDifferences(TSet<FName>& ModifiedPackages, TSet<FName>& NewPackages, TSet<FName>& RemovedPackages, TSet<FName>& IdenticalCookedPackages, TSet<FName>& IdenticalUncookedPackages)
+void FAssetRegistryGenerator::ComputePackageDifferences(TSet<FName>& ModifiedPackages, TSet<FName>& NewPackages, TSet<FName>& RemovedPackages, TSet<FName>& IdenticalCookedPackages, TSet<FName>& IdenticalUncookedPackages, bool bRecurseModifications)
 {
 	for (const TPair<FName, const FAssetPackageData*>& PackagePair : State.GetAssetPackageDataMap())
 	{
@@ -562,27 +563,29 @@ void FAssetRegistryGenerator::ComputePackageDifferences(TSet<FName>& ModifiedPac
 		}
 	}
 
-	// Recurse modified packages to their dependencies. In theory a small number of packages are modified at once so this is faster than computing the hashes recursively
-	TArray<FName> ModifiedPackagesToRecurse = ModifiedPackages.Array();
-
-	for (int32 RecurseIndex = 0; RecurseIndex < ModifiedPackagesToRecurse.Num(); RecurseIndex++)
+	if (bRecurseModifications)
 	{
-		FName ModifiedPackage = ModifiedPackagesToRecurse[RecurseIndex];
-		TArray<FAssetIdentifier> Referencers;
-		State.GetReferencers(ModifiedPackage, Referencers, EAssetRegistryDependencyType::Hard);
+		// Recurse modified packages to their dependencies. In theory a small number of packages are modified at once so this is faster than computing the hashes recursively
+		TArray<FName> ModifiedPackagesToRecurse = ModifiedPackages.Array();
 
-		for (const FAssetIdentifier& Referencer : Referencers)
+		for (int32 RecurseIndex = 0; RecurseIndex < ModifiedPackagesToRecurse.Num(); RecurseIndex++)
 		{
-			FName ReferencerPackage = Referencer.PackageName;
-			if (!ModifiedPackages.Contains(ReferencerPackage))
+			FName ModifiedPackage = ModifiedPackagesToRecurse[RecurseIndex];
+			TArray<FAssetIdentifier> Referencers;
+			State.GetReferencers(ModifiedPackage, Referencers, EAssetRegistryDependencyType::Hard);
+
+			for (const FAssetIdentifier& Referencer : Referencers)
 			{
-				// Remove from identical/new list
-				IdenticalCookedPackages.Remove(ReferencerPackage);
-				IdenticalUncookedPackages.Remove(ReferencerPackage);
-				NewPackages.Remove(ReferencerPackage);
+				FName ReferencerPackage = Referencer.PackageName;
+				if (!ModifiedPackages.Contains(ReferencerPackage) && (IdenticalCookedPackages.Contains(ReferencerPackage) || IdenticalUncookedPackages.Contains(ReferencerPackage)))
+				{
+					// Remove from identical list
+					IdenticalCookedPackages.Remove(ReferencerPackage);
+					IdenticalUncookedPackages.Remove(ReferencerPackage);
 
-				ModifiedPackages.Add(ReferencerPackage);
-				ModifiedPackagesToRecurse.Add(ReferencerPackage);
+					ModifiedPackages.Add(ReferencerPackage);
+					ModifiedPackagesToRecurse.Add(ReferencerPackage);
+				}
 			}
 		}
 	}
@@ -769,6 +772,15 @@ bool FAssetRegistryGenerator::SaveAssetRegistry(const FString& SandboxPath)
 	AssetRegistry.InitializeSerializationOptions(DevelopmentSaveOptions, TargetPlatform->IniPlatformName());
 	DevelopmentSaveOptions.ModifyForDevelopment();
 
+	// Write runtime registry, this can be excluded per game/platform
+	FAssetRegistrySerializationOptions SaveOptions;
+	AssetRegistry.InitializeSerializationOptions(SaveOptions, TargetPlatform->IniPlatformName());
+
+	// Flush the asset registry and make sure the asset data is in sync, as it may have been updated during cook
+	AssetRegistry.Tick(-1.0f);
+
+	AssetRegistry.InitializeTemporaryAssetRegistryState(State, SaveOptions, true);
+
 	if (DevelopmentSaveOptions.bSerializeAssetRegistry)
 	{
 		// Create development registry data, used for incremental cook and editor viewing
@@ -782,14 +794,10 @@ bool FAssetRegistryGenerator::SaveAssetRegistry(const FString& SandboxPath)
 		FFileHelper::SaveArrayToFile(SerializedAssetRegistry, *PlatformSandboxPath);
 	}
 
-	// Write runtime registry, this can be excluded per game/platform
-	FAssetRegistrySerializationOptions SaveOptions;
-	AssetRegistry.InitializeSerializationOptions(SaveOptions, TargetPlatform->IniPlatformName());
-
 	if (SaveOptions.bSerializeAssetRegistry)
 	{
 		// Prune out the development only packages
-		State.PruneAssetData(CookedPackages, TSet<FName>());
+		State.PruneAssetData(CookedPackages, TSet<FName>(), SaveOptions.bFilterAssetDataWithNoTags);
 
 		// Create runtime registry data
 		FArrayWriter SerializedAssetRegistry;
diff --git a/Engine/Source/Editor/UnrealEd/Private/Commandlets/ContentCommandlets.cpp b/Engine/Source/Editor/UnrealEd/Private/Commandlets/ContentCommandlets.cpp
--- a/Engine/Source/Editor/UnrealEd/Private/Commandlets/ContentCommandlets.cpp
+++ b/Engine/Source/Editor/UnrealEd/Private/Commandlets/ContentCommandlets.cpp
@@ -48,8 +48,8 @@
 DEFINE_LOG_CATEGORY(LogContentCommandlet);
 
 #include "AssetRegistryModule.h"
-
-
+#include "IDirectoryWatcher.h"
+#include "DirectoryWatcherModule.h"
 #include "Particles/Material/ParticleModuleMeshMaterial.h"
 #include "Particles/ParticleLODLevel.h"
 #include "Particles/ParticleModuleRequired.h"
@@ -237,18 +237,32 @@ int32 UResavePackagesCommandlet::InitializeResaveParameters( const TArray<FStrin
 		}
 	}
 
-	const bool bResaveDirectRefsAndDeps = Switches.Contains(TEXT("resavedirectrefsanddeps"));
-	if (bExplicitPackages && PackageNames.Num() == 1 && bResaveDirectRefsAndDeps)
+	FAssetRegistryModule& AssetRegistryModule = FModuleManager::LoadModuleChecked<FAssetRegistryModule>(TEXT("AssetRegistry"));
+	IAssetRegistry& AssetRegistry = AssetRegistryModule.Get();
+
+	// This option works if a single package is specified, it will resave all packages that reference it, and all packages that it references
+	const bool bResaveDirectRefsAndDeps = Switches.Contains(TEXT("ResaveDirectRefsAndDeps"));
+
+	// This option will filter the package list and only save packages that are redirectors, or that reference redirectors
+	const bool bFixupRedirects = (Switches.Contains(TEXT("FixupRedirects")) || Switches.Contains(TEXT("FixupRedirectors")));
+
+	if (bResaveDirectRefsAndDeps || bFixupRedirects)
 	{
-		FAssetRegistryModule& AssetRegistryModule = FModuleManager::LoadModuleChecked<FAssetRegistryModule>("AssetRegistry");
-		AssetRegistryModule.Get().SearchAllAssets(true);
+		AssetRegistry.SearchAllAssets(true);
 
+		// Force directory watcher tick to register paths
+		FDirectoryWatcherModule& DirectoryWatcherModule = FModuleManager::Get().LoadModuleChecked<FDirectoryWatcherModule>(TEXT("DirectoryWatcher"));
+		DirectoryWatcherModule.Get()->Tick(-1.0f);
+	}
+
+	if (bExplicitPackages && PackageNames.Num() == 1 && bResaveDirectRefsAndDeps)
+	{
 		FName PackageName = FName(*FPackageName::FilenameToLongPackageName(PackageNames[0]));
 
 		TArray<FName> Referencers;
-		AssetRegistryModule.Get().GetReferencers(PackageName, Referencers);
+		AssetRegistry.GetReferencers(PackageName, Referencers);
 		TArray<FName> Dependencies;
-		AssetRegistryModule.Get().GetDependencies(PackageName, Dependencies);
+		AssetRegistry.GetDependencies(PackageName, Dependencies);
 
 		for (FName Ref : Referencers)
 		{
@@ -263,6 +277,54 @@ int32 UResavePackagesCommandlet::InitializeResaveParameters( const TArray<FStrin
 			PackageNames.Add(File);
 		}
 	}
+	else if (bFixupRedirects)
+	{
+		// Look for all packages containing redirects, and their referencers
+		TArray<FAssetData> RedirectAssets;
+		TSet<FString> RedirectPackages;
+		TSet<FString> ReferencerPackages;
+
+		AssetRegistry.GetAssetsByClass(UObjectRedirector::StaticClass()->GetFName(), RedirectAssets);
+
+		for (const FAssetData& AssetData : RedirectAssets)
+		{
+			bool bIsAlreadyInSet = false;
+			FString RedirectFile;
+			FPackageName::SearchForPackageOnDisk(*AssetData.PackageName.ToString(), nullptr, &RedirectFile);
+
+			RedirectPackages.Add(RedirectFile, &bIsAlreadyInSet);
+
+			if (!bIsAlreadyInSet)
+			{
+				TArray<FName> Referencers;
+				AssetRegistry.GetReferencers(AssetData.PackageName, Referencers);
+
+				for (FName Referencer : Referencers)
+				{
+					FString ReferencerFile;
+					FPackageName::SearchForPackageOnDisk(*Referencer.ToString(), nullptr, &ReferencerFile);
+
+					ReferencerPackages.Add(ReferencerFile);
+				}
+			}
+		}
+
+		// Filter packagenames list to packages that are pointing to redirectors, it will probably be much smaller
+		TArray<FString> OldArray = PackageNames;
+		PackageNames.Reset();
+		for (FString& PackageName : OldArray)
+		{
+			if (RedirectPackages.Contains(PackageName))
+			{
+				RedirectorsToFixup.Add(PackageName);
+			}
+
+			if (ReferencerPackages.Contains(PackageName))
+			{
+				PackageNames.Add(PackageName);
+			}
+		}
+	}
 
 	// Check for the min and max versions
 	MinResaveUE4Version = IGNORE_PACKAGE_VERSION;
@@ -518,61 +580,11 @@ void UResavePackagesCommandlet::LoadAndSaveOnePackage(const FString& Filename)
 				if (bIsEmpty)
 				{
 					bSavePackage = false;
+					Package = nullptr;
 					
-					// get file SCC status
 					UE_LOG(LogContentCommandlet, Display, TEXT("Package %s is empty and will be deleted"), *Filename);
-					FString PackageFilename = SourceControlHelpers::PackageFilename(Package);
-
-					ISourceControlProvider& SourceControlProvider = ISourceControlModule::Get().GetProvider();
-					FSourceControlStatePtr SourceControlState = SourceControlProvider.GetState(PackageFilename, EStateCacheUsage::ForceUpdate);
-
-					// Unload package so we can delete it
-					TArray<UPackage *> PackagesToDelete;
-					PackagesToDelete.Add(Package);
-					PackageTools::UnloadPackages(PackagesToDelete);
-					PackagesToDelete.Empty();
-					Package = NULL;
-
-					if (bAutoCheckOut)
-					{
-						if (SourceControlState.IsValid() && (SourceControlState->IsCheckedOut() || SourceControlState->IsAdded()))
-						{
-							UE_LOG(LogContentCommandlet, Display, TEXT("Revert '%s' from source control..."), *Filename);
-							SourceControlProvider.Execute(ISourceControlOperation::Create<FRevert>(), PackageFilename);
 
-							UE_LOG(LogContentCommandlet, Display, TEXT("Deleting '%s' from source control..."), *Filename);
-							SourceControlProvider.Execute(ISourceControlOperation::Create<FDelete>(), PackageFilename);
-
-							FilesToSubmit.Add(Filename);
-						}
-						else if (SourceControlState.IsValid() && SourceControlState->CanCheckout())
-						{
-							UE_LOG(LogContentCommandlet, Display, TEXT("Deleting '%s' from source control..."), *Filename);
-							SourceControlProvider.Execute(ISourceControlOperation::Create<FDelete>(), PackageFilename);
-
-							FilesToSubmit.Add(Filename);
-						}
-						else if (SourceControlState.IsValid() && SourceControlState->IsCheckedOutOther())
-						{
-							UE_LOG(LogContentCommandlet, Warning, TEXT("Couldn't delete '%s' from source control, someone has it checked out, skipping..."), *Filename);
-						}
-						else if (SourceControlState.IsValid() && !SourceControlState->IsSourceControlled())
-						{
-							UE_LOG(LogContentCommandlet, Warning, TEXT("'%s' is not in source control, attempting to delete from disk..."), *Filename);
-							if (!IFileManager::Get().Delete(*Filename, false, true))
-							{
-								UE_LOG(LogContentCommandlet, Warning, TEXT("  ... failed to delete from disk."), *Filename);
-							}
-						}
-						else
-						{
-							UE_LOG(LogContentCommandlet, Warning, TEXT("'%s' is in an unknown source control state, attempting to delete from disk..."), *Filename);
-							if (!IFileManager::Get().Delete(*Filename, false, true))
-							{
-								UE_LOG(LogContentCommandlet, Warning, TEXT("  ... failed to delete from disk."), *Filename);
-							}
-						}
-					}
+					DeleteOnePackage(Filename);
 				}
 			}
 
@@ -617,8 +629,7 @@ void UResavePackagesCommandlet::LoadAndSaveOnePackage(const FString& Filename)
 
 									VerboseMessage(TEXT("Post CheckOut"));
 
-									FString PackageName(FPackageName::FilenameToLongPackageName(Filename));
-									FilesToSubmit.Add(*PackageName);
+									FilesToSubmit.AddUnique(*Filename);
 								}
 							}
 							VerboseMessage(TEXT("Post ForceGetStatus2"));
@@ -667,6 +678,82 @@ void UResavePackagesCommandlet::LoadAndSaveOnePackage(const FString& Filename)
 	}
 }
 
+void UResavePackagesCommandlet::DeleteOnePackage(const FString& Filename)
+{
+	bool bIsReadOnly = IFileManager::Get().IsReadOnly(*Filename);
+
+	if (bVerifyContent)
+	{
+		return;
+	}
+
+	if (bIsReadOnly && !bAutoCheckOut)
+	{
+		if (Verbosity != ONLY_ERRORS)
+		{
+			UE_LOG(LogContentCommandlet, Warning, TEXT("Skipping read-only file %s"), *Filename);
+		}
+		return;
+	}
+
+	FString PackageName;
+	FPackageName::TryConvertFilenameToLongPackageName(Filename, PackageName);
+
+	UPackage* Package = FindPackage(nullptr, *PackageName);
+
+	if (Package)
+	{
+		// Unload package so we can delete it
+		TArray<UPackage *> PackagesToDelete;
+		PackagesToDelete.Add(Package);
+		PackageTools::UnloadPackages(PackagesToDelete);
+		PackagesToDelete.Empty();
+		Package = nullptr;
+	}
+
+	FString PackageFilename = SourceControlHelpers::PackageFilename(Filename);
+	ISourceControlProvider& SourceControlProvider = ISourceControlModule::Get().GetProvider();
+	FSourceControlStatePtr SourceControlState = SourceControlProvider.GetState(PackageFilename, EStateCacheUsage::ForceUpdate);
+
+	if (SourceControlState.IsValid() && (SourceControlState->IsCheckedOut() || SourceControlState->IsAdded()))
+	{
+		UE_LOG(LogContentCommandlet, Display, TEXT("Revert '%s' from source control..."), *Filename);
+		SourceControlProvider.Execute(ISourceControlOperation::Create<FRevert>(), PackageFilename);
+
+		UE_LOG(LogContentCommandlet, Display, TEXT("Deleting '%s' from source control..."), *Filename);
+		SourceControlProvider.Execute(ISourceControlOperation::Create<FDelete>(), PackageFilename);
+
+		FilesToSubmit.AddUnique(Filename);
+	}
+	else if (SourceControlState.IsValid() && SourceControlState->CanCheckout())
+	{
+		UE_LOG(LogContentCommandlet, Display, TEXT("Deleting '%s' from source control..."), *Filename);
+		SourceControlProvider.Execute(ISourceControlOperation::Create<FDelete>(), PackageFilename);
+
+		FilesToSubmit.AddUnique(Filename);
+	}
+	else if (SourceControlState.IsValid() && SourceControlState->IsCheckedOutOther())
+	{
+		UE_LOG(LogContentCommandlet, Warning, TEXT("Couldn't delete '%s' from source control, someone has it checked out, skipping..."), *Filename);
+	}
+	else if (SourceControlState.IsValid() && !SourceControlState->IsSourceControlled())
+	{
+		UE_LOG(LogContentCommandlet, Warning, TEXT("'%s' is not in source control, attempting to delete from disk..."), *Filename);
+		if (!IFileManager::Get().Delete(*Filename, false, true))
+		{
+			UE_LOG(LogContentCommandlet, Warning, TEXT("  ... failed to delete from disk."), *Filename);
+		}
+	}
+	else
+	{
+		UE_LOG(LogContentCommandlet, Warning, TEXT("'%s' is in an unknown source control state, attempting to delete from disk..."), *Filename);
+		if (!IFileManager::Get().Delete(*Filename, false, true))
+		{
+			UE_LOG(LogContentCommandlet, Warning, TEXT("  ... failed to delete from disk."), *Filename);
+		}
+	}
+}
+
 int32 UResavePackagesCommandlet::Main( const FString& Params )
 {
 	const TCHAR* Parms = *Params;
@@ -685,9 +772,9 @@ int32 UResavePackagesCommandlet::Main( const FString& Params )
 	/** if we should only save dirty packages **/
 	bOnlySaveDirtyPackages = Switches.Contains(TEXT("OnlySaveDirtyPackages"));
 	/** if we should auto checkout packages that need to be saved**/
-	bAutoCheckOut = Switches.Contains(TEXT("AutoCheckOutPackages"));
+	bAutoCheckOut = Switches.Contains(TEXT("AutoCheckOutPackages")) || Switches.Contains(TEXT("AutoCheckOut"));
 	/** if we should auto checkin packages that were checked out**/
-	bAutoCheckIn = bAutoCheckOut && Switches.Contains(TEXT("AutoCheckIn"));
+	bAutoCheckIn = bAutoCheckOut && (Switches.Contains(TEXT("AutoCheckIn")) || Switches.Contains(TEXT("AutoSubmit")));
 	/** determine if we are building lighting for the map packages on the pass. **/
 	bShouldBuildLighting = Switches.Contains(TEXT("buildlighting"));
 	/** determine if we are building lighting for the map packages on the pass. **/
@@ -708,7 +795,7 @@ int32 UResavePackagesCommandlet::Main( const FString& Params )
 	}
 
 	// Retrieve list of all packages in .ini paths.
-	if( !PackageNames.Num() )
+	if(!PackageNames.Num() && !RedirectorsToFixup.Num())
 	{
 		return 0;
 	}
@@ -766,6 +853,41 @@ int32 UResavePackagesCommandlet::Main( const FString& Params )
 		}
 	}
 
+	// Force a directory watcher and asset registry tick
+	FDirectoryWatcherModule& DirectoryWatcherModule = FModuleManager::Get().LoadModuleChecked<FDirectoryWatcherModule>(TEXT("DirectoryWatcher"));
+	DirectoryWatcherModule.Get()->Tick(-1.0f);
+
+	FAssetRegistryModule& AssetRegistryModule = FModuleManager::LoadModuleChecked<FAssetRegistryModule>(TEXT("AssetRegistry"));
+	IAssetRegistry& AssetRegistry = AssetRegistryModule.Get();
+	AssetRegistry.Tick(-1.0f);
+
+	// Delete unreferenced redirector packages
+	for (int32 PackageIndex = 0; PackageIndex < RedirectorsToFixup.Num(); PackageIndex++)
+	{
+		const FString& Filename = RedirectorsToFixup[PackageIndex];
+
+		FName PackageName = FName(*FPackageName::FilenameToLongPackageName(Filename));
+
+		// Load and save this package
+		TArray<FName> Referencers;
+
+		AssetRegistry.GetReferencers(PackageName, Referencers);
+
+		if (Referencers.Num() == 0)
+		{
+			if (Verbosity != ONLY_ERRORS)
+			{
+				UE_LOG(LogContentCommandlet, Display, TEXT("Deleting unreferenced redirector [%s]"), *Filename);
+			}
+
+			DeleteOnePackage(Filename);
+		}
+		else if (Verbosity != ONLY_ERRORS)
+		{
+			UE_LOG(LogContentCommandlet, Display, TEXT("Can't delete redirector [%s], unsaved packages reference it"), *Filename);
+		}
+	}
+
 	// Submit the results to source control
 	if( bAutoCheckIn )
 	{
@@ -806,6 +928,10 @@ FText UResavePackagesCommandlet::GetChangelistDescription() const
 	{
 		ChangelistDescription = NSLOCTEXT("ContentCmdlets", "ChangelistDescriptionBuildTextureStreaming", "Rebuild texture streaming.");
 	}
+	else if (RedirectorsToFixup.Num() > 0)
+	{
+		ChangelistDescription = NSLOCTEXT("ContentCmdlets", "ChangelistDescriptionRedirectors", "Fixing Redirectors");
+	}
 	else
 	{
 		ChangelistDescription = NSLOCTEXT("ContentCmdlets", "ChangelistDescription", "Resave Deprecated Packages");
@@ -1174,7 +1300,7 @@ void UResavePackagesCommandlet::PerformAdditionalOperations(class UWorld* World,
 		{
 			for(const auto& SublevelFilename : SublevelFilenames)
 			{
-				FilesToSubmit.Add(SublevelFilename);
+				FilesToSubmit.AddUnique(SublevelFilename);
 			}
 		}
 
diff --git a/Engine/Source/Editor/UnrealEd/Private/Commandlets/DumpBlueprintsInfoCommandlet.cpp b/Engine/Source/Editor/UnrealEd/Private/Commandlets/DumpBlueprintsInfoCommandlet.cpp
--- a/Engine/Source/Editor/UnrealEd/Private/Commandlets/DumpBlueprintsInfoCommandlet.cpp
+++ b/Engine/Source/Editor/UnrealEd/Private/Commandlets/DumpBlueprintsInfoCommandlet.cpp
@@ -1946,7 +1946,15 @@ static bool DumpBlueprintInfoUtils::DumpPinContextActions(uint32 Indent, UEdGrap
 			FileOutWriter->Serialize(TCHAR_TO_ANSI(TEXT(",\n")), 2);
 			PinType.bIsReference = false;
 
-			PinType.bIsArray = true;
+			PinType.ContainerType = EPinContainerType::Array;
+			DumpContextualPinTypeActions(Indent, Graph, PinType, FileOutWriter);
+			FileOutWriter->Serialize(TCHAR_TO_ANSI(TEXT(",\n")), 2);
+
+			PinType.ContainerType = EPinContainerType::Set;
+			DumpContextualPinTypeActions(Indent, Graph, PinType, FileOutWriter);
+			FileOutWriter->Serialize(TCHAR_TO_ANSI(TEXT(",\n")), 2);
+
+			PinType.ContainerType = EPinContainerType::Map;
 			DumpContextualPinTypeActions(Indent, Graph, PinType, FileOutWriter);
 
 			bWroteToFile = true;
@@ -1988,7 +1996,15 @@ static bool DumpBlueprintInfoUtils::DumpTypeTreeActions(uint32 Indent, UEdGraph*
 		FileOutWriter->Serialize(TCHAR_TO_ANSI(TEXT(",\n")), 2);
 		PinType.bIsReference = false;
 
-		PinType.bIsArray = true;
+		PinType.ContainerType = EPinContainerType::Array;
+		DumpContextualPinTypeActions(Indent, Graph, PinType, FileOutWriter);
+		FileOutWriter->Serialize(TCHAR_TO_ANSI(TEXT(",\n")), 2);
+
+		PinType.ContainerType = EPinContainerType::Set;
+		DumpContextualPinTypeActions(Indent, Graph, PinType, FileOutWriter);
+		FileOutWriter->Serialize(TCHAR_TO_ANSI(TEXT(",\n")), 2);
+
+		PinType.ContainerType = EPinContainerType::Map;
 		DumpContextualPinTypeActions(Indent, Graph, PinType, FileOutWriter);
 
 		PendingLineEnding = TEXT(",\n");
@@ -2019,7 +2035,7 @@ static bool DumpBlueprintInfoUtils::DumpTypeTreeActions(uint32 Indent, UEdGraph*
 }
 
 //------------------------------------------------------------------------------
-static void DumpBlueprintInfoUtils::DumpContextualPinTypeActions(uint32 Indent, UEdGraph* Graph, FEdGraphPinType const& PinType, FArchive* FileOutWriter)
+static void DumpBlueprintInfoUtils::DumpContextualPinTypeActions(uint32 Indent, UEdGraph* Graph, const FEdGraphPinType& PinType, FArchive* FileOutWriter)
 {
 	FGraphContextMenuBuilder ContextMenuBuilder(Graph);
 
diff --git a/Engine/Source/Editor/UnrealEd/Private/Commandlets/GenerateBlueprintAPICommandlet.cpp b/Engine/Source/Editor/UnrealEd/Private/Commandlets/GenerateBlueprintAPICommandlet.cpp
--- a/Engine/Source/Editor/UnrealEd/Private/Commandlets/GenerateBlueprintAPICommandlet.cpp
+++ b/Engine/Source/Editor/UnrealEd/Private/Commandlets/GenerateBlueprintAPICommandlet.cpp
@@ -932,10 +932,19 @@ static void GenerateBlueprintAPIUtils::DumpActionMenuItem(uint32 Indent, FGraphA
 						ActionEntry += "," + PinDetailsIndentedNewline + "\"DefaultValue\"         : \"" + MakeJsonString(Pin->DefaultValue) + "\"";
 					}
 
-					if (Pin->PinType.bIsArray)
+					if (Pin->PinType.IsArray())
 					{
 						ActionEntry += "," + PinDetailsIndentedNewline + "\"IsArray\"              : \"true\"";
 					}
+					else if (Pin->PinType.IsSet())
+					{
+						ActionEntry += "," + PinDetailsIndentedNewline + "\"IsSet\"              : \"true\"";
+					}
+					else if (Pin->PinType.IsMap())
+					{
+						ActionEntry += "," + PinDetailsIndentedNewline + "\"IsMap\"              : \"true\"";
+						// TODO: Send the Map value type as well
+					}
 
 					if (Pin->PinType.bIsConst)
 					{
diff --git a/Engine/Source/Editor/UnrealEd/Private/CookOnTheFlyServer.cpp b/Engine/Source/Editor/UnrealEd/Private/CookOnTheFlyServer.cpp
--- a/Engine/Source/Editor/UnrealEd/Private/CookOnTheFlyServer.cpp
+++ b/Engine/Source/Editor/UnrealEd/Private/CookOnTheFlyServer.cpp
@@ -2869,7 +2869,7 @@ void UCookOnTheFlyServer::SaveCookedPackage(UPackage* Package, uint32 SaveFlags,
 	// Don't resolve, just add to request list as needed
 	TSet<FName> StringAssetPackages;
 
-	GRedirectCollector.GetStringAssetReferencePackageList(Package->GetFName(), StringAssetPackages);
+	GRedirectCollector.ProcessStringAssetReferencePackageList(Package->GetFName(), false, StringAssetPackages);
 	
 	for (FName StringAssetPackage : StringAssetPackages)
 	{
@@ -2880,7 +2880,7 @@ void UCookOnTheFlyServer::SaveCookedPackage(UPackage* Package, uint32 SaveFlags,
 		{
 			for (TPair<FString, FString>& RedirectedPath : RedirectedPaths)
 			{
-				GRedirectCollector.AddStringAssetReferenceRedirection(RedirectedPath.Key, RedirectedPath.Value);
+				GRedirectCollector.AddAssetPathRedirection(RedirectedPath.Key, RedirectedPath.Value);
 			}
 		}
 
@@ -4354,7 +4354,6 @@ bool UCookOnTheFlyServer::SaveCookedAssetRegistry(const FString& InCookedAssetRe
 				Json->WriteValue(TEXT("ObjectPath"), AssetData.ObjectPath.ToString());
 				Json->WriteValue(TEXT("PackageName"), AssetData.PackageName.ToString());
 				Json->WriteValue(TEXT("PackagePath"), AssetData.PackagePath.ToString());
-				Json->WriteValue(TEXT("GroupNames"), AssetData.GroupNames.ToString());
 				Json->WriteValue(TEXT("AssetName"), AssetData.AssetName.ToString());
 				Json->WriteValue(TEXT("AssetClass"), AssetData.AssetClass.ToString());
 				Json->WriteObjectStart("TagsAndValues");
@@ -4637,11 +4636,11 @@ void UCookOnTheFlyServer::PopulateCookedPackagesFromDisk(const TArray<ITargetPla
 				FDateTime CurrentLocalCookedBuild = IFileManager::Get().GetTimeStamp(*SandboxCookedAssetRegistryFilename);
 
 				// iterate on the shared build if the option is set
-				FString SharedCookedAssetRegistry = FPaths::Combine(*FPaths::GameSavedDir(), TEXT("SharedIterativeBuild"), *Target->PlatformName(), TEXT("Cooked"), GetAssetRegistryFilename());
+				FString SharedCookedAssetRegistry = FPaths::Combine(*FPaths::GameSavedDir(), TEXT("SharedIterativeBuild"), *Target->PlatformName(), TEXT("Cooked"), TEXT("DevelopmentAssetRegistry.bin"));
 
 				FDateTime CurrentIterativeCookedBuild = IFileManager::Get().GetTimeStamp(*SharedCookedAssetRegistry);
 
-				if (CurrentIterativeCookedBuild > CurrentLocalCookedBuild)
+				if (CurrentIterativeCookedBuild >= CurrentLocalCookedBuild)
 				{
 					// clean the sandbox 
 					ClearPlatformCookedData(FName(*Target->PlatformName()));
@@ -4674,16 +4673,20 @@ void UCookOnTheFlyServer::PopulateCookedPackagesFromDisk(const TArray<ITargetPla
 				else
 				{
 					UE_LOG(LogCook, Display, TEXT("Local cook is newer then shared cooked build, iterativly cooking from local build"));
+					PlatformAssetRegistry->LoadPreviousAssetRegistry(SandboxCookedAssetRegistryFilename);
 				}
-
 			}
-
-			PlatformAssetRegistry->LoadPreviousAssetRegistry(SandboxCookedAssetRegistryFilename);
+			else
+			{
+				PlatformAssetRegistry->LoadPreviousAssetRegistry(SandboxCookedAssetRegistryFilename);
+			}
 
 			// Get list of changed packages
 			TSet<FName> ModifiedPackages, NewPackages, RemovedPackages, IdenticalCookedPackages, IdenticalUncookedPackages;
 
-			PlatformAssetRegistry->ComputePackageDifferences(ModifiedPackages, NewPackages, RemovedPackages, IdenticalCookedPackages, IdenticalUncookedPackages);
+			// We recurse modifications up the reference chain because it is safer, if this ends up being a significant issue in some games we can add a command line flag
+			bool bRecurseModifications = true;
+			PlatformAssetRegistry->ComputePackageDifferences(ModifiedPackages, NewPackages, RemovedPackages, IdenticalCookedPackages, IdenticalUncookedPackages, bRecurseModifications);
 
 			// check the files on disk 
 			TMap<FName, FName> UncookedPathToCookedPath;
@@ -4693,8 +4696,13 @@ void UCookOnTheFlyServer::PopulateCookedPackagesFromDisk(const TArray<ITargetPla
 			const static FName NAME_DummyCookedFilename(TEXT("DummyCookedFilename")); // pls never name a package dummycookedfilename otherwise shit might go wonky
 			if (bIsIterateSharedBuild)
 			{
+				TSet<FName> ExistingPackages = ModifiedPackages;
+				ExistingPackages.Append(RemovedPackages);
+				ExistingPackages.Append(IdenticalCookedPackages);
+				ExistingPackages.Append(IdenticalUncookedPackages);
+
 				// if we are iterating of a shared build the cooked files might not exist in the cooked directory because we assume they are packaged in the pak file (which we don't want to extract)
-				for (FName PackageName : IdenticalCookedPackages)
+				for (FName PackageName : ExistingPackages)
 				{
 					FString Filename;
 					if (FPackageName::DoesPackageExist(PackageName.ToString(), nullptr, &Filename))
@@ -5778,7 +5786,19 @@ void UCookOnTheFlyServer::CollectFilesToCook(TArray<FName>& FilesInPath, const T
 
 		if (UAssetManager::IsValid())
 		{
-			UAssetManager::Get().ModifyCook(FilesInPath);
+			TArray<FName> PackagesToNeverCook;
+
+			UAssetManager::Get().ModifyCook(FilesInPath, PackagesToNeverCook);
+
+			for (FName NeverCookPackage : PackagesToNeverCook)
+			{
+				const FName StandardPackageFilename = GetCachedStandardPackageFileFName(NeverCookPackage);
+
+				if (StandardPackageFilename != NAME_None)
+				{
+					NeverCookPackageList.Add(StandardPackageFilename);
+				}
+			}
 		}
 #if DEBUG_COOKMODIFICATIONDELEGATE
 		for (TObjectIterator<UPackage> It; It; ++It)
@@ -6034,6 +6054,20 @@ void UCookOnTheFlyServer::CollectFilesToCook(TArray<FName>& FilesInPath, const T
 			}
 		}
 	}
+
+	if (CookByTheBookOptions && !(FilesToCookFlags & ECookByTheBookOptions::DisableUnsolicitedPackages))
+	{
+		// Gather initial unsolicited package list, this is needed in iterative mode because it may skip cooking all explicit packages and never hit this code
+		TArray<UPackage*> UnsolicitedPackages;
+		bool ContainsFullAssetGCClasses = false;
+		UE_LOG(LogCook, Verbose, TEXT("Finding initial unsolicited packages"));
+		GetUnsolicitedPackages(UnsolicitedPackages, ContainsFullAssetGCClasses, CookByTheBookOptions->TargetPlatformNames);
+
+		for (UPackage* UnsolicitedPackage : UnsolicitedPackages)
+		{
+			AddFileToCook(FilesInPath, UnsolicitedPackage->GetName());
+		}
+	}
 }
 
 bool UCookOnTheFlyServer::IsCookByTheBookRunning() const
@@ -6831,11 +6865,11 @@ void UCookOnTheFlyServer::StartCookByTheBook( const FCookByTheBookStartupOptions
 	TSet<FName> StartupStringAssetPackages;
 
 	// Get the list of string asset references, for both empty package and all startup packages
-	GRedirectCollector.GetStringAssetReferencePackageList(NAME_None, StartupStringAssetPackages);
+	GRedirectCollector.ProcessStringAssetReferencePackageList(NAME_None, false, StartupStringAssetPackages);
 
 	for (const FName& StartupPackage : CookByTheBookOptions->StartupPackages)
 	{
-		GRedirectCollector.GetStringAssetReferencePackageList(StartupPackage, StartupStringAssetPackages);
+		GRedirectCollector.ProcessStringAssetReferencePackageList(StartupPackage, false, StartupStringAssetPackages);
 	}
 
 	for (FName StringAssetPackage : StartupStringAssetPackages)
@@ -6847,7 +6881,7 @@ void UCookOnTheFlyServer::StartCookByTheBook( const FCookByTheBookStartupOptions
 		{
 			for (TPair<FString, FString>& RedirectedPath : RedirectedPaths)
 			{
-				GRedirectCollector.AddStringAssetReferenceRedirection(RedirectedPath.Key, RedirectedPath.Value);
+				GRedirectCollector.AddAssetPathRedirection(RedirectedPath.Key, RedirectedPath.Value);
 			}
 		}
 
diff --git a/Engine/Source/Editor/UnrealEd/Private/EditorEngine.cpp b/Engine/Source/Editor/UnrealEd/Private/EditorEngine.cpp
--- a/Engine/Source/Editor/UnrealEd/Private/EditorEngine.cpp
+++ b/Engine/Source/Editor/UnrealEd/Private/EditorEngine.cpp
@@ -4679,12 +4679,12 @@ namespace ReattachActorsHelper
 
 				AActor* ChildActor = CurrentAttachmentInfo.AttachedActors[AttachedActorIdx].Actor;
 				ChildActor->Modify();
-				ChildActor->DetachRootComponentFromParent(true);
+				ChildActor->DetachFromActor(FDetachmentTransformRules::KeepWorldTransform);
 			}
 
 			// Modify the actor so undo will reattach it.
 			ActorToReattach->Modify();
-			ActorToReattach->DetachRootComponentFromParent(true);
+			ActorToReattach->DetachFromActor(FDetachmentTransformRules::KeepWorldTransform);
 		}
 	}
 
diff --git a/Engine/Source/Editor/UnrealEd/Private/EditorExporters.cpp b/Engine/Source/Editor/UnrealEd/Private/EditorExporters.cpp
--- a/Engine/Source/Editor/UnrealEd/Private/EditorExporters.cpp
+++ b/Engine/Source/Editor/UnrealEd/Private/EditorExporters.cpp
@@ -491,7 +491,7 @@ bool ULevelExporterT3D::ExportText( const FExportObjectInnerContext* Context, UO
 
 				AActor* ParentActor = Actor->GetAttachParentActor();
 				FName SocketName = Actor->GetAttachParentSocketName();
-				Actor->DetachRootComponentFromParent(true);
+				Actor->DetachFromActor(FDetachmentTransformRules::KeepWorldTransform);
 
 				FString ParentActorString = ( ParentActor ? FString::Printf(TEXT(" ParentActor=%s"), *ParentActor->GetName() ) : TEXT(""));
 				FString SocketNameString = ( (ParentActor && SocketName != NAME_None) ? FString::Printf(TEXT(" SocketName=%s"), *SocketName.ToString() ) : TEXT(""));
@@ -912,7 +912,7 @@ static void AddActorToOBJs(AActor* Actor, TArray<FOBJGeom*>& Objects, TSet<UMate
 		if( StaticMeshComponent->IsRegistered() && StaticMeshComponent->GetStaticMesh()
 			&& StaticMeshComponent->GetStaticMesh()->HasValidRenderData() )
 		{
-			LocalToWorld = StaticMeshComponent->ComponentToWorld.ToMatrixWithScale();
+			LocalToWorld = StaticMeshComponent->GetComponentTransform().ToMatrixWithScale();
 			StaticMesh = StaticMeshComponent->GetStaticMesh();
 			if (StaticMesh)
 			{
diff --git a/Engine/Source/Editor/UnrealEd/Private/EditorObject.cpp b/Engine/Source/Editor/UnrealEd/Private/EditorObject.cpp
--- a/Engine/Source/Editor/UnrealEd/Private/EditorObject.cpp
+++ b/Engine/Source/Editor/UnrealEd/Private/EditorObject.cpp
@@ -496,13 +496,6 @@ static const TCHAR* ImportProperties(
 				}
 				else // handle the non-template case (subobjects and non-template components)
 				{
-					// don't allow Actor-derived subobjects
-					if ( TemplateClass->IsChildOf(AActor::StaticClass()) )
-					{
-						Warn->Logf(ELogVerbosity::Error,TEXT("Cannot create subobjects from Actor-derived classes: %s"), *StrLine);
-						return NULL;
-					}
-
 					ComponentTemplate = FindObject<UObject>(SubobjectOuter, *TemplateName.ToString());
 					if (ComponentTemplate != NULL)
 					{
diff --git a/Engine/Source/Editor/UnrealEd/Private/EditorServer.cpp b/Engine/Source/Editor/UnrealEd/Private/EditorServer.cpp
--- a/Engine/Source/Editor/UnrealEd/Private/EditorServer.cpp
+++ b/Engine/Source/Editor/UnrealEd/Private/EditorServer.cpp
@@ -3544,7 +3544,7 @@ void UEditorEngine::PasteSelectedActorsFromClipboard( UWorld* InWorld, const FTe
 
 					AActor* ParentActor = Actor->GetAttachParentActor();
 					FName SocketName = Actor->GetAttachParentSocketName();
-					Actor->DetachRootComponentFromParent(true);
+					Actor->DetachFromActor(FDetachmentTransformRules::KeepWorldTransform);
 					AttachData.Emplace(ParentActor, SocketName);
 
 					// If this actor is in a group, add it to the list
@@ -6950,7 +6950,7 @@ void UEditorEngine::AutoMergeStaticMeshes()
 			UStaticMeshComponent* Component = MergeComponents[ComponentIndex];
 
 			// calculate a matrix to go from my component space to the owner's component's space
-			FMatrix TransformToOwnerSpace = Component->ComponentToWorld.ToMatrixWithScale() * OwnerComponent->ComponentToWorld.ToMatrixWithScale().Inverse();
+			FMatrix TransformToOwnerSpace = Component->GetComponentTransform().ToMatrixWithScale() * OwnerComponent->GetComponentTransform().ToMatrixWithScale().Inverse();
 
 			// if we have negative scale, we need to munge the matrix and scaling
 			if (TransformToOwnerSpace.Determinant() < 0.0f)
diff --git a/Engine/Source/Editor/UnrealEd/Private/Factories/EditorFactories.cpp b/Engine/Source/Editor/UnrealEd/Private/Factories/EditorFactories.cpp
--- a/Engine/Source/Editor/UnrealEd/Private/Factories/EditorFactories.cpp
+++ b/Engine/Source/Editor/UnrealEd/Private/Factories/EditorFactories.cpp
@@ -260,14 +260,14 @@ class FAssetClassParentFilter : public IClassViewerFilter
 {
 public:
 	FAssetClassParentFilter()
-		: DisallowedClassFlags(0), bDisallowBlueprintBase(false)
+		: DisallowedClassFlags(CLASS_None), bDisallowBlueprintBase(false)
 	{}
 
 	/** All children of these classes will be included unless filtered out by another setting. */
 	TSet< const UClass* > AllowedChildrenOfClasses;
 
 	/** Disallowed class flags. */
-	uint32 DisallowedClassFlags;
+	EClassFlags DisallowedClassFlags;
 
 	/** Disallow blueprint base classes. */
 	bool bDisallowBlueprintBase;
diff --git a/Engine/Source/Editor/UnrealEd/Private/Factories/SlateWidgetStyleAssetFactory.cpp b/Engine/Source/Editor/UnrealEd/Private/Factories/SlateWidgetStyleAssetFactory.cpp
--- a/Engine/Source/Editor/UnrealEd/Private/Factories/SlateWidgetStyleAssetFactory.cpp
+++ b/Engine/Source/Editor/UnrealEd/Private/Factories/SlateWidgetStyleAssetFactory.cpp
@@ -18,7 +18,7 @@ public:
 	TSet< const UClass* > DisallowedClasses;
 
 	/** Disallowed class flags. */
-	uint32 DisallowedClassFlags;
+	EClassFlags DisallowedClassFlags;
 
 	virtual bool IsClassAllowed(const FClassViewerInitializationOptions& InInitOptions, const UClass* InClass, TSharedRef< FClassViewerFilterFuncs > InFilterFuncs ) override
 	{
diff --git a/Engine/Source/Editor/UnrealEd/Private/Fbx/FbxMainExport.cpp b/Engine/Source/Editor/UnrealEd/Private/Fbx/FbxMainExport.cpp
--- a/Engine/Source/Editor/UnrealEd/Private/Fbx/FbxMainExport.cpp
+++ b/Engine/Source/Editor/UnrealEd/Private/Fbx/FbxMainExport.cpp
@@ -3998,9 +3998,9 @@ void FFbxExporter::ExportLandscapeToFbx(ALandscapeProxy* Landscape, const TCHAR*
 
 			FVector Normal, TangentX, TangentY;
 			CDI.GetLocalTangentVectors(VertX, VertY, TangentX, TangentY, Normal);
-			Normal /= Component->ComponentToWorld.GetScale3D(); Normal.Normalize();
-			TangentX /= Component->ComponentToWorld.GetScale3D(); TangentX.Normalize();
-			TangentY /= Component->ComponentToWorld.GetScale3D(); TangentY.Normalize();
+			Normal /= Component->GetComponentTransform().GetScale3D(); Normal.Normalize();
+			TangentX /= Component->GetComponentTransform().GetScale3D(); TangentX.Normalize();
+			TangentY /= Component->GetComponentTransform().GetScale3D(); TangentY.Normalize();
 			FbxVector4 FbxNormal = FbxVector4(Normal.X, -Normal.Y, Normal.Z); FbxNormal.Normalize();
 			Normals.SetAt(BaseVertIndex + VertIndex, FbxNormal);
 			FbxVector4 FbxTangent = FbxVector4(TangentX.X, -TangentX.Y, TangentX.Z); FbxTangent.Normalize();
diff --git a/Engine/Source/Editor/UnrealEd/Private/Kismet2/BlueprintEditorUtils.cpp b/Engine/Source/Editor/UnrealEd/Private/Kismet2/BlueprintEditorUtils.cpp
--- a/Engine/Source/Editor/UnrealEd/Private/Kismet2/BlueprintEditorUtils.cpp
+++ b/Engine/Source/Editor/UnrealEd/Private/Kismet2/BlueprintEditorUtils.cpp
@@ -2131,11 +2131,21 @@ void FBlueprintEditorUtils::MarkBlueprintAsModified(UBlueprint* Blueprint, FProp
 
 bool FBlueprintEditorUtils::ShouldRegenerateBlueprint(UBlueprint* Blueprint)
 {
-	return !GForceDisableBlueprintCompileOnLoad
+	return !IsCompileOnLoadDisabled(Blueprint)
 		&& Blueprint->bRecompileOnLoad
 		&& !Blueprint->bIsRegeneratingOnLoad;
 }
 
+bool FBlueprintEditorUtils::IsCompileOnLoadDisabled(UBlueprint* Blueprint)
+{
+	bool bCompilationDisabled = false;
+	if(Blueprint->GetLinker())
+	{
+		bCompilationDisabled = (Blueprint->GetLinker()->LoadFlags & LOAD_DisableCompileOnLoad) != LOAD_None;
+	}
+	return bCompilationDisabled;
+}
+
 // Helper function to get the blueprint that ultimately owns a node.
 UBlueprint* FBlueprintEditorUtils::FindBlueprintForNode(const UEdGraphNode* Node)
 {
@@ -3974,8 +3984,10 @@ bool FBlueprintEditorUtils::IsPinTypeValid(const FEdGraphPinType& Type)
 // Gets the visible class variable list.  This includes both variables introduced here and in all superclasses.
 void FBlueprintEditorUtils::GetClassVariableList(const UBlueprint* Blueprint, TArray<FName>& VisibleVariables, bool bIncludePrivateVars) 
 {
-	// Existing variables in the parent class and above
-	check(!Blueprint->bHasBeenRegenerated || Blueprint->bIsRegeneratingOnLoad || (Blueprint->SkeletonGeneratedClass != nullptr));
+	// Existing variables in the parent class and above, when using the compilation manager the previous SkeletonGeneratedClass will have been cleared when
+	// we're regenerating the SkeletonGeneratedClass. Using this function in the skeleton pass at all is highly dubious, but I am leaving it until the 
+	// compilation manager is on full time:
+	check(!Blueprint->bHasBeenRegenerated || Blueprint->bIsRegeneratingOnLoad || (Blueprint->SkeletonGeneratedClass != nullptr) || GBlueprintUseCompilationManager);
 	if (Blueprint->SkeletonGeneratedClass != nullptr)
 	{
 		for (TFieldIterator<UProperty> PropertyIt(Blueprint->SkeletonGeneratedClass, EFieldIteratorFlags::IncludeSuper); PropertyIt; ++PropertyIt)
@@ -4501,7 +4513,7 @@ void FBlueprintEditorUtils::ChangeMemberVariableType(UBlueprint* Blueprint, cons
 				{
 					Variable.VarType = NewPinType;
 
-					if(Variable.VarType.bIsSet || Variable.VarType.bIsMap)
+					if(Variable.VarType.IsSet() || Variable.VarType.IsMap())
 					{
 						// Make sure that the variable is no longer tagged for replication, and warn the user if the variable is no
 						// longer goign to be replicated:
@@ -4960,6 +4972,9 @@ void FBlueprintEditorUtils::ChangeLocalVariableType(UBlueprint* InBlueprint, con
 
 				Variable.VarType = NewPinType;
 
+				// Reset the default value
+				Variable.DefaultValue.Empty();
+
 				// Mark the Blueprint as structurally modified so we can reconstruct the node successfully
 				FBlueprintEditorUtils::MarkBlueprintAsStructurallyModified(InBlueprint);
 
@@ -5173,12 +5188,12 @@ bool FBlueprintEditorUtils::ValidateAllFunctionGraphs(UBlueprint* InBlueprint, U
 
 void FBlueprintEditorUtils::FixupVariableDescription(UBlueprint* Blueprint, FBPVariableDescription& VarDesc)
 {
-	if ((VarDesc.PropertyFlags & CPF_Config) != 0 )// commented out for 4.16 (since we cannot modify FBlueprintsObjectVersion in the release branch)... use the version check when merging to Main/Dev-Framework: && Blueprint->GetLinkerCustomVersion(FBlueprintsObjectVersion::GUID) < FBlueprintsObjectVersion::DisallowObjectConfigVars)
+	if ((VarDesc.PropertyFlags & CPF_Config) != 0 && Blueprint->GetLinkerCustomVersion(FBlueprintsObjectVersion::GUID) < FBlueprintsObjectVersion::DisallowObjectConfigVars)
 	{
+		// Synchronized with FBlueprintVarActionDetails::IsConfigCheckBoxEnabled
 		const FEdGraphPinType& VarType = VarDesc.VarType;
 		if (VarType.PinCategory == UEdGraphSchema_K2::PC_Object ||
-			VarType.PinCategory == UEdGraphSchema_K2::PC_Interface ||
-			VarType.PinCategory == UEdGraphSchema_K2::PC_Asset)
+			VarType.PinCategory == UEdGraphSchema_K2::PC_Interface)
 		{
 			VarDesc.PropertyFlags &= ~CPF_Config;
 		}
@@ -5254,86 +5269,6 @@ void FBlueprintEditorUtils::ValidateBlueprintChildVariables(UBlueprint* InBluepr
 
 //////////////////////////////////////////////////////////////////////////
 
-void FBlueprintEditorUtils::ImportKismetDefaultValueToProperty(UEdGraphPin* SourcePin, UProperty* DestinationProperty, uint8* DestinationAddress, UObject* OwnerObject)
-{
-	FString LiteralString = SourcePin->GetDefaultAsString();
-
-	const UEdGraphSchema_K2* K2Schema = GetDefault<UEdGraphSchema_K2>();
-
-	if (UStructProperty* StructProperty = Cast<UStructProperty>(DestinationProperty))
-	{
-		static UScriptStruct* VectorStruct = TBaseStructure<FVector>::Get();
-		static UScriptStruct* RotatorStruct = TBaseStructure<FRotator>::Get();
-		static UScriptStruct* TransformStruct = TBaseStructure<FTransform>::Get();
-
-		if (StructProperty->Struct == VectorStruct)
-		{
-			FDefaultValueHelper::ParseVector(LiteralString, *((FVector*)DestinationAddress));
-			return;
-		}
-		else if (StructProperty->Struct == RotatorStruct)
-		{
-			FDefaultValueHelper::ParseRotator(LiteralString, *((FRotator*)DestinationAddress));
-			return;
-		}
-		else if (StructProperty->Struct == TransformStruct)
-		{
-			(*(FTransform*)DestinationAddress).InitFromString( LiteralString );
-			return;
-		}
-	}
-	else if (UObjectPropertyBase* ObjectProperty = Cast<UObjectPropertyBase>(DestinationProperty))
-	{
-		ObjectProperty->SetObjectPropertyValue(DestinationAddress, SourcePin->DefaultObject);
-		return;
-	}
-
-	const int32 PortFlags = 0;
-	DestinationProperty->ImportText(*LiteralString, DestinationAddress, PortFlags, OwnerObject);
-}
-
-void FBlueprintEditorUtils::ExportPropertyToKismetDefaultValue(UEdGraphPin* TargetPin, UProperty* SourceProperty, uint8* SourceAddress, UObject* OwnerObject)
-{
-	FString LiteralString;
-
-	if (UStructProperty* StructProperty = Cast<UStructProperty>(SourceProperty))
-	{
-		static UScriptStruct* VectorStruct = TBaseStructure<FVector>::Get();
-		static UScriptStruct* RotatorStruct = TBaseStructure<FRotator>::Get();
-		static UScriptStruct* TransformStruct = TBaseStructure<FTransform>::Get();
-
-		if (StructProperty->Struct == VectorStruct)
-		{
-			FVector& SourceVector = *((FVector*)SourceAddress);
-			TargetPin->DefaultValue = FString::Printf(TEXT("%f, %f, %f"), SourceVector.X, SourceVector.Y, SourceVector.Z);
-			return;
-		}
-		else if (StructProperty->Struct == RotatorStruct)
-		{
-			FRotator& SourceRotator = *((FRotator*)SourceAddress);
-			TargetPin->DefaultValue = FString::Printf(TEXT("%f, %f, %f"), SourceRotator.Pitch, SourceRotator.Yaw, SourceRotator.Roll);
-			return;
-		}
-		else if (StructProperty->Struct == TransformStruct)
-		{
-			TargetPin->DefaultValue = (*(FTransform*)SourceAddress).ToString();
-			return;
-		}
-	}
-	else if (UObjectPropertyBase* ObjectProperty = Cast<UObjectPropertyBase>(SourceProperty))
-	{
-		TargetPin->DefaultObject = ObjectProperty->GetObjectPropertyValue(SourceAddress);
-		return;
-	}
-
-	//@TODO: Add support for object literals
-
-	const int32 PortFlags = 0;
-
-	TargetPin->DefaultValue.Empty();
-	SourceProperty->ExportTextItem(TargetPin->DefaultValue, SourceAddress, nullptr, OwnerObject, PortFlags);
-}
-
 void FBlueprintEditorUtils::FindNativizationDependencies(UBlueprint* Blueprint, TArray<UClass*>& NativizeDependenciesOut)
 {
 	FBlueprintEditorUtils::FindImplementedInterfaces(Blueprint, /*bGetAllInterfaces =*/false, NativizeDependenciesOut);
@@ -7468,7 +7403,6 @@ public:
 		return InFilterFuncs->IfInChildOfClassesSet( AllowedChildrenOfClasses, InClass) != EFilterReturn::Failed && 
 			InFilterFuncs->IfInChildOfClassesSet(DisallowedChildrenOfClasses, InClass) != EFilterReturn::Passed && 
 			InFilterFuncs->IfInClassesSet(DisallowedClasses, InClass) != EFilterReturn::Passed &&
-			!InClass->GetCppTypeInfo()->IsAbstract() &&
 			!InClass->HasAnyClassFlags(CLASS_Deprecated | CLASS_NewerVersionExists) &&
 			InClass->HasAnyClassFlags(CLASS_Interface) &&
 			// Here is some loaded classes only logic, Blueprints will never have this info
@@ -7764,89 +7698,93 @@ bool FBlueprintEditorUtils::IsObjectADebugCandidate( AActor* InActorObject, UBlu
 	return bPassesFlags && bCanDebugThisObject;
 }
 
-bool FBlueprintEditorUtils::PropertyValueFromString(const UProperty* Property, const FString& Value, uint8* DefaultObject)
+bool FBlueprintEditorUtils::PropertyValueFromString(const UProperty* Property, const FString& StrValue, uint8* Container)
+{
+	return PropertyValueFromString_Direct(Property, StrValue, Property->ContainerPtrToValuePtr<uint8>(Container));
+}
+
+bool FBlueprintEditorUtils::PropertyValueFromString_Direct(const UProperty* Property, const FString& StrValue, uint8* DirectValue)
 {
 	bool bParseSucceeded = true;
-	if( !Property->IsA(UStructProperty::StaticClass()) )
+	if (!Property->IsA(UStructProperty::StaticClass()))
 	{
-		if( Property->IsA(UIntProperty::StaticClass()) )
+		if (Property->IsA(UIntProperty::StaticClass()))
 		{
 			int32 IntValue = 0;
-			bParseSucceeded = FDefaultValueHelper::ParseInt(Value, IntValue);
-			CastChecked<UIntProperty>(Property)->SetPropertyValue_InContainer(DefaultObject, IntValue);
+			bParseSucceeded = FDefaultValueHelper::ParseInt(StrValue, IntValue);
+			CastChecked<UIntProperty>(Property)->SetPropertyValue(DirectValue, IntValue);
 		}
-		else if( Property->IsA(UFloatProperty::StaticClass()) )
+		else if (Property->IsA(UFloatProperty::StaticClass()))
 		{
 			float FloatValue = 0.0f;
-			bParseSucceeded = FDefaultValueHelper::ParseFloat(Value, FloatValue);
-			CastChecked<UFloatProperty>(Property)->SetPropertyValue_InContainer(DefaultObject, FloatValue);
+			bParseSucceeded = FDefaultValueHelper::ParseFloat(StrValue, FloatValue);
+			CastChecked<UFloatProperty>(Property)->SetPropertyValue(DirectValue, FloatValue);
 		}
 		else if (const UByteProperty* ByteProperty = Cast<const UByteProperty>(Property))
 		{
 			int32 IntValue = 0;
 			if (const UEnum* Enum = ByteProperty->Enum)
 			{
-				IntValue = Enum->GetValueByName(FName(*Value));
+				IntValue = Enum->GetValueByName(FName(*StrValue));
 				bParseSucceeded = (INDEX_NONE != IntValue);
 
 				// If the parse did not succeed, clear out the int to keep the enum value valid
-				if(!bParseSucceeded)
+				if (!bParseSucceeded)
 				{
 					IntValue = 0;
 				}
 			}
 			else
 			{
-				bParseSucceeded = FDefaultValueHelper::ParseInt(Value, IntValue);
+				bParseSucceeded = FDefaultValueHelper::ParseInt(StrValue, IntValue);
 			}
-			bParseSucceeded = bParseSucceeded && ( IntValue <= 255 ) && ( IntValue >= 0 );
-			ByteProperty->SetPropertyValue_InContainer(DefaultObject, IntValue);
+			bParseSucceeded = bParseSucceeded && (IntValue <= 255) && (IntValue >= 0);
+			ByteProperty->SetPropertyValue(DirectValue, IntValue);
 		}
 		else if (const UEnumProperty* EnumProperty = Cast<const UEnumProperty>(Property))
 		{
-			int64 IntValue = EnumProperty->GetEnum()->GetValueByName(FName(*Value));
+			int64 IntValue = EnumProperty->GetEnum()->GetValueByName(FName(*StrValue));
 			bParseSucceeded = (INDEX_NONE != IntValue);
 
 			// If the parse did not succeed, clear out the int to keep the enum value valid
-			if(!bParseSucceeded)
+			if (!bParseSucceeded)
 			{
 				IntValue = 0;
 			}
-			bParseSucceeded = bParseSucceeded && ( IntValue <= 255 ) && ( IntValue >= 0 );
-			void* PropAddr = EnumProperty->ContainerPtrToValuePtr<void>(DefaultObject);
-			EnumProperty->GetUnderlyingProperty()->SetIntPropertyValue(PropAddr, IntValue);
+			bParseSucceeded = bParseSucceeded && (IntValue <= 255) && (IntValue >= 0);
+			EnumProperty->GetUnderlyingProperty()->SetIntPropertyValue(DirectValue, IntValue);
 		}
-		else if( Property->IsA(UStrProperty::StaticClass()) )
+		else if (Property->IsA(UStrProperty::StaticClass()))
 		{
-			CastChecked<UStrProperty>(Property)->SetPropertyValue_InContainer(DefaultObject, Value);
+			CastChecked<UStrProperty>(Property)->SetPropertyValue(DirectValue, StrValue);
 		}
-		else if( Property->IsA(UBoolProperty::StaticClass()) )
+		else if (Property->IsA(UBoolProperty::StaticClass()))
 		{
-			CastChecked<UBoolProperty>(Property)->SetPropertyValue_InContainer(DefaultObject, Value.ToBool());
+			CastChecked<UBoolProperty>(Property)->SetPropertyValue(DirectValue, StrValue.ToBool());
 		}
-		else if( Property->IsA(UNameProperty::StaticClass()) )
+		else if (Property->IsA(UNameProperty::StaticClass()))
 		{
-			CastChecked<UNameProperty>(Property)->SetPropertyValue_InContainer(DefaultObject, FName(*Value));
+			CastChecked<UNameProperty>(Property)->SetPropertyValue(DirectValue, FName(*StrValue));
 		}
-		else if( Property->IsA(UTextProperty::StaticClass()) )
+		else if (Property->IsA(UTextProperty::StaticClass()))
 		{
 			FStringOutputDevice ImportError;
-			const TCHAR* EndOfParsedBuff = Property->ImportText(*Value, Property->ContainerPtrToValuePtr<uint8>(DefaultObject), 0, nullptr, &ImportError);
+			const TCHAR* EndOfParsedBuff = Property->ImportText(*StrValue, DirectValue, 0, nullptr, &ImportError);
 			bParseSucceeded = EndOfParsedBuff && ImportError.IsEmpty();
 		}
 		else
 		{
 			// Empty array-like properties need to use "()" in order to import correctly (as array properties export comma separated within a set of brackets)
-			const TCHAR* const ValueToImport = (Value.IsEmpty() && (Property->IsA(UArrayProperty::StaticClass()) || Property->IsA(UMulticastDelegateProperty::StaticClass())))
+			const TCHAR* const ValueToImport = (StrValue.IsEmpty() && (Property->IsA(UArrayProperty::StaticClass()) || Property->IsA(UMulticastDelegateProperty::StaticClass())))
 				? TEXT("()")
-				: *Value;
+				: *StrValue;
 
 			FStringOutputDevice ImportError;
-			const TCHAR* EndOfParsedBuff = Property->ImportText(*Value, Property->ContainerPtrToValuePtr<uint8>(DefaultObject), 0, nullptr, &ImportError);
+			const TCHAR* EndOfParsedBuff = Property->ImportText(*StrValue, DirectValue, 0, nullptr, &ImportError);
 			bParseSucceeded = EndOfParsedBuff && ImportError.IsEmpty();
 		}
 	}
-	else 
+	else
 	{
 		static UScriptStruct* VectorStruct = TBaseStructure<FVector>::Get();
 		static UScriptStruct* RotatorStruct = TBaseStructure<FRotator>::Get();
@@ -7856,42 +7794,41 @@ bool FBlueprintEditorUtils::PropertyValueFromString(const UProperty* Property, c
 		const UStructProperty* StructProperty = CastChecked<UStructProperty>(Property);
 
 		// Struct properties must be handled differently, unfortunately.  We only support FVector, FRotator, and FTransform
-		if( StructProperty->Struct == VectorStruct )
+		if (StructProperty->Struct == VectorStruct)
 		{
 			FVector V = FVector::ZeroVector;
-			bParseSucceeded = FDefaultValueHelper::ParseVector(Value, V);
-			Property->CopyCompleteValue( Property->ContainerPtrToValuePtr<uint8>(DefaultObject), &V );
+			bParseSucceeded = FDefaultValueHelper::ParseVector(StrValue, V);
+			Property->CopyCompleteValue(DirectValue, &V);
 		}
-		else if( StructProperty->Struct == RotatorStruct )
+		else if (StructProperty->Struct == RotatorStruct)
 		{
 			FRotator R = FRotator::ZeroRotator;
-			bParseSucceeded = FDefaultValueHelper::ParseRotator(Value, R);
-			Property->CopyCompleteValue( Property->ContainerPtrToValuePtr<uint8>(DefaultObject), &R );
+			bParseSucceeded = FDefaultValueHelper::ParseRotator(StrValue, R);
+			Property->CopyCompleteValue(DirectValue, &R);
 		}
-		else if( StructProperty->Struct == TransformStruct )
+		else if (StructProperty->Struct == TransformStruct)
 		{
 			FTransform T = FTransform::Identity;
-			bParseSucceeded = T.InitFromString( Value );
-			Property->CopyCompleteValue( Property->ContainerPtrToValuePtr<uint8>(DefaultObject), &T );
+			bParseSucceeded = T.InitFromString(StrValue);
+			Property->CopyCompleteValue(DirectValue, &T);
 		}
-		else if( StructProperty->Struct == LinearColorStruct )
+		else if (StructProperty->Struct == LinearColorStruct)
 		{
 			FLinearColor Color;
 			// Color form: "(R=%f,G=%f,B=%f,A=%f)"
-			bParseSucceeded = Color.InitFromString(Value);
-			Property->CopyCompleteValue( Property->ContainerPtrToValuePtr<uint8>(DefaultObject), &Color );
+			bParseSucceeded = Color.InitFromString(StrValue);
+			Property->CopyCompleteValue(DirectValue, &Color);
 		}
 		else if (StructProperty->Struct)
 		{
 			const UScriptStruct* Struct = StructProperty->Struct;
 			const int32 StructSize = Struct->GetStructureSize() * StructProperty->ArrayDim;
-			uint8* StructData = Property->ContainerPtrToValuePtr<uint8>(DefaultObject);
-			StructProperty->InitializeValue(StructData);
+			StructProperty->InitializeValue(DirectValue);
 			ensure(1 == StructProperty->ArrayDim);
-			bParseSucceeded = FStructureEditorUtils::Fill_MakeStructureDefaultValue(Cast<const UUserDefinedStruct>(Struct), StructData);
+			bParseSucceeded = FStructureEditorUtils::Fill_MakeStructureDefaultValue(Cast<const UUserDefinedStruct>(Struct), DirectValue);
 
 			FStringOutputDevice ImportError;
-			const TCHAR* EndOfParsedBuff = StructProperty->ImportText(Value.IsEmpty() ? TEXT("()") : *Value, StructData, 0, nullptr, &ImportError);
+			const TCHAR* EndOfParsedBuff = StructProperty->ImportText(StrValue.IsEmpty() ? TEXT("()") : *StrValue, DirectValue, 0, nullptr, &ImportError);
 			bParseSucceeded &= EndOfParsedBuff && ImportError.IsEmpty();
 		}
 	}
@@ -8507,15 +8444,15 @@ const FSlateBrush* FBlueprintEditorUtils::GetIconFromPin( const FEdGraphPinType&
 	
 	const FSlateBrush* IconBrush = FEditorStyle::GetBrush(TEXT("Kismet.VariableList.TypeIcon"));
 	const UObject* PinSubObject = PinType.PinSubCategoryObject.Get();
-	if( PinType.bIsArray && PinType.PinCategory != K2Schema->PC_Exec )
+	if( PinType.IsArray() && PinType.PinCategory != K2Schema->PC_Exec )
 	{
 		IconBrush = FEditorStyle::GetBrush(TEXT("Kismet.VariableList.ArrayTypeIcon"));
 	}
-	else if (PinType.bIsMap && PinType.PinCategory != K2Schema->PC_Exec)
+	else if (PinType.IsMap() && PinType.PinCategory != K2Schema->PC_Exec)
 	{
 		IconBrush = FEditorStyle::GetBrush(TEXT("Kismet.VariableList.MapKeyTypeIcon"));
 	}
-	else if (PinType.bIsSet && PinType.PinCategory != K2Schema->PC_Exec)
+	else if (PinType.IsSet() && PinType.PinCategory != K2Schema->PC_Exec)
 	{
 		if( bIsLarge )
 		{
@@ -8540,7 +8477,7 @@ const FSlateBrush* FBlueprintEditorUtils::GetIconFromPin( const FEdGraphPinType&
 const FSlateBrush* FBlueprintEditorUtils::GetSecondaryIconFromPin(const FEdGraphPinType& PinType)
 {
 	const UEdGraphSchema_K2* K2Schema = GetDefault<UEdGraphSchema_K2>();
-	if (PinType.bIsMap && PinType.PinCategory != K2Schema->PC_Exec)
+	if (PinType.IsMap() && PinType.PinCategory != K2Schema->PC_Exec)
 	{
 		return FEditorStyle::GetBrush(TEXT("Kismet.VariableList.MapValueTypeIcon"));
 	}
diff --git a/Engine/Source/Editor/UnrealEd/Private/Kismet2/CompilerResultsLog.cpp b/Engine/Source/Editor/UnrealEd/Private/Kismet2/CompilerResultsLog.cpp
--- a/Engine/Source/Editor/UnrealEd/Private/Kismet2/CompilerResultsLog.cpp
+++ b/Engine/Source/Editor/UnrealEd/Private/Kismet2/CompilerResultsLog.cpp
@@ -1,6 +1,7 @@
 // Copyright 1998-2017 Epic Games, Inc. All Rights Reserved.
 
 #include "Kismet2/CompilerResultsLog.h"
+#include "Engine/Blueprint.h"
 #include "Modules/ModuleManager.h"
 #include "Misc/PackageName.h"
 #include "Editor/EditorPerProjectUserSettings.h"
@@ -238,6 +239,13 @@ void FCompilerResultsLog::InternalLogSummary()
 /** Update the source backtrack map to note that NewObject was most closely generated/caused by the SourceObject */
 void FCompilerResultsLog::NotifyIntermediateObjectCreation(UObject* NewObject, UObject* SourceObject)
 {
+	if(UEdGraphNode* NewNode = Cast<UEdGraphNode>(NewObject))
+	{
+		if(UEdGraphNode* OldNode = Cast<UEdGraphNode>(SourceObject))
+		{
+			FullSourceBacktrackMap.Add(NewNode, OldNode);
+		}
+	}
 	SourceBacktrackMap.NotifyIntermediateObjectCreation(NewObject, SourceObject);
 }
 
@@ -257,6 +265,86 @@ UObject const* FCompilerResultsLog::FindSourceObject(UObject const* PossiblyDupl
 	return SourceBacktrackMap.FindSourceObject(PossiblyDuplicatedObject);
 }
 
+int32 FCompilerResultsLog::CalculateStableIdentifierForLatentActionManager( const UEdGraphNode* Node )
+{
+	/* 
+		The name of this function is meant to instill a bit of caution:
+		1. The Latent Action Manager uses uint32s to identify actions, so there
+			is some risk of collision, increasing if we aren't able to distribute
+			keys across the whole range of uint32
+		2. We need these identifiers to be stable across blueprint compiles, meaning
+			we can't just create a GUID and hash it
+
+		Meeting these two requirements has proved difficult. The edge cases involve
+		macros and nodes that implement UK2Node::ExpandNode, e.g. LoadAsset/LoadAssetClass
+		nodes in Macros. In order to handle that case we use the full source backtrack map.
+		Typically an intermediate node has a dynamic GUID, which is useless, but source 
+		nodes that came from macros have stable GUIDs, and can be used 
+	*/
+
+	const UEdGraphNode* OriginalNode = Node;
+
+	// first search for a node with a stable GUID (e.g., not a node that was created via SpawnIntermediateNode, 
+	// but we want to include nodes that are created via macro instantiation):
+	bool bNodeHasStableGUID = false;
+	while(!bNodeHasStableGUID)
+	{
+		if(!Node->HasAnyFlags(RF_Transient) || 
+			GetIntermediateTunnelInstance(Node))
+		{
+			bNodeHasStableGUID = true;
+		}
+		else
+		{
+			UEdGraphNode* const* PreviousNode = FullSourceBacktrackMap.Find(Node);
+			if(PreviousNode && *PreviousNode)
+			{
+				Node = *PreviousNode;
+			}
+			else
+			{
+				// we failed to find a source node, bail
+				bNodeHasStableGUID = true;
+				Node = nullptr;
+			}
+		}
+	}
+
+	int32 LatentUUID = 0;
+	
+	if(Node)
+	{
+		LatentUUID = GetTypeHash(Node->NodeGuid);
+
+		const UEdGraphNode* ResultNode = Node;
+		const UEdGraphNode* SourceNode = Cast<UEdGraphNode>(GetIntermediateTunnelInstance(Node));
+		while (SourceNode && SourceNode != ResultNode)
+		{
+			if (SourceNode->NodeGuid.IsValid())
+			{
+				LatentUUID = HashCombine(LatentUUID, GetTypeHash(SourceNode->NodeGuid));
+			}
+			ResultNode = SourceNode;
+			SourceNode = Cast<UEdGraphNode>(GetIntermediateTunnelInstance(ResultNode));
+		}
+	}
+	else
+	{
+		Warning(
+			*LOCTEXT(
+				"UUIDDeterministicCookWarn", 
+				"Failed to produce a deterministic UUID for a node's latent action: @@"
+			).ToString(),
+			OriginalNode
+		);
+
+		static int32 FallbackUUID = 0;
+		LatentUUID = FallbackUUID++;
+	}
+
+	return LatentUUID;
+}
+
 UEdGraphPin* FCompilerResultsLog::FindSourcePin(UEdGraphPin* PossiblyDuplicatedPin)
 {
 	return SourceBacktrackMap.FindSourcePin(PossiblyDuplicatedPin);
@@ -636,10 +724,10 @@ UEdGraphNode* FCompilerResultsLog::GetSourceNode(const UEdGraphNode* Intermediat
 	return Result;
 }
 
-UEdGraphNode* FCompilerResultsLog::GetIntermediateTunnelInstance(const UEdGraphNode* IntermediateNode)
+UEdGraphNode* FCompilerResultsLog::GetIntermediateTunnelInstance(const UEdGraphNode* IntermediateNode) const
 {
 	TWeakObjectPtr<UEdGraphNode> Result;
-	if (TWeakObjectPtr<UEdGraphNode>* IntermediateTunnelInstanceNode = IntermediateTunnelNodeToTunnelInstanceMap.Find(IntermediateNode))
+	if (const TWeakObjectPtr<UEdGraphNode>* IntermediateTunnelInstanceNode = IntermediateTunnelNodeToTunnelInstanceMap.Find(IntermediateNode))
 	{
 		Result = *IntermediateTunnelInstanceNode;
 	}
diff --git a/Engine/Source/Editor/UnrealEd/Private/Kismet2/ComponentEditorUtils.cpp b/Engine/Source/Editor/UnrealEd/Private/Kismet2/ComponentEditorUtils.cpp
--- a/Engine/Source/Editor/UnrealEd/Private/Kismet2/ComponentEditorUtils.cpp
+++ b/Engine/Source/Editor/UnrealEd/Private/Kismet2/ComponentEditorUtils.cpp
@@ -334,7 +334,6 @@ bool FComponentEditorUtils::CanCopyComponents(const TArray<UActorComponent*>& Co
 void FComponentEditorUtils::CopyComponents(const TArray<UActorComponent*>& ComponentsToCopy)
 {
 	FStringOutputDevice Archive;
-	const FExportObjectInnerContext Context;
 
 	// Clear the mark state for saving.
 	UnMarkAllObjects(EObjectMark(OBJECTMARK_TagExp | OBJECTMARK_TagImp));
@@ -345,7 +344,7 @@ void FComponentEditorUtils::CopyComponents(const TArray<UActorComponent*>& Compo
 	for (UActorComponent* Component : ComponentsToCopy)
 	{
 		// Duplicate the component into a temporary object
-		UObject* DuplicatedComponent = StaticDuplicateObject(Component, GetTransientPackage(), Component->GetFName(), RF_AllFlags & ~RF_ArchetypeObject);
+		UObject* DuplicatedComponent = StaticDuplicateObject(Component, GetTransientPackage(), Component->GetFName());
 		if (DuplicatedComponent)
 		{
 			// If the duplicated component is a scene component, wipe its attach parent (to prevent log warnings for referencing a private object in an external package)
@@ -367,6 +366,8 @@ void FComponentEditorUtils::CopyComponents(const TArray<UActorComponent*>& Compo
 		}
 	}
 
+	const FExportObjectInnerContext Context;
+
 	// Export the component object(s) to text for copying
 	for (auto ObjectIt = ObjectMap.CreateIterator(); ObjectIt; ++ObjectIt)
 	{
@@ -626,7 +627,7 @@ UActorComponent* FComponentEditorUtils::DuplicateComponent(UActorComponent* Temp
 				USceneComponent* RootComponent = Actor->GetRootComponent();
 				check(RootComponent);
 
-				// ComponentToWorld is not a UPROPERTY, so make sure the clone has calculated it properly before attachment
+				// GetComponentTransform() is not a UPROPERTY, so make sure the clone has calculated it properly before attachment
 				NewSceneComponent->UpdateComponentToWorld();
 
 				NewSceneComponent->SetupAttachment(RootComponent);
diff --git a/Engine/Source/Editor/UnrealEd/Private/Kismet2/Kismet2.cpp b/Engine/Source/Editor/UnrealEd/Private/Kismet2/Kismet2.cpp
--- a/Engine/Source/Editor/UnrealEd/Private/Kismet2/Kismet2.cpp
+++ b/Engine/Source/Editor/UnrealEd/Private/Kismet2/Kismet2.cpp
@@ -2,6 +2,8 @@
 
 
 #include "CoreMinimal.h"
+
+#include "BlueprintCompilationManager.h"
 #include "Misc/CoreMisc.h"
 #include "Stats/StatsMisc.h"
 #include "Stats/Stats.h"
@@ -89,6 +91,8 @@ DECLARE_CYCLE_STAT(TEXT("Validate Generated Class"), EKismetCompilerStats_Valida
 
 #define LOCTEXT_NAMESPACE "UnrealEd.Editor"
 
+extern COREUOBJECT_API bool GBlueprintUseCompilationManager;
+
 //////////////////////////////////////////////////////////////////////////
 // FArchiveInvalidateTransientRefs
 
@@ -740,6 +744,12 @@ bool FKismetEditorUtilities::IsReferencedByUndoBuffer(UBlueprint* Blueprint)
 
 void FKismetEditorUtilities::CompileBlueprint(UBlueprint* BlueprintObj, EBlueprintCompileOptions CompileFlags, FCompilerResultsLog* pResults)
 {
+	if(GBlueprintUseCompilationManager)
+	{
+		FBlueprintCompilationManager::CompileSynchronously(FBPCompileRequest(BlueprintObj, CompileFlags, pResults));
+		return;
+	}
+
 	const bool bIsRegeneratingOnLoad		= (CompileFlags & EBlueprintCompileOptions::IsRegeneratingOnLoad		) != EBlueprintCompileOptions::None;
 	const bool bSkipGarbageCollection		= (CompileFlags & EBlueprintCompileOptions::SkipGarbageCollection		) != EBlueprintCompileOptions::None;
 	const bool bSaveIntermediateProducts	= (CompileFlags & EBlueprintCompileOptions::SaveIntermediateProducts	) != EBlueprintCompileOptions::None;
@@ -997,6 +1007,12 @@ bool FKismetEditorUtilities::GenerateBlueprintSkeleton(UBlueprint* BlueprintObj,
 void FKismetEditorUtilities::RecompileBlueprintBytecode(UBlueprint* BlueprintObj, TArray<UObject*>* ObjLoaded,  EBlueprintBytecodeRecompileOptions Flags)
 {
 	FSecondsCounterScope Timer(BlueprintCompileAndLoadTimerData); 
+
+	if(FBlueprintEditorUtils::IsCompileOnLoadDisabled(BlueprintObj))
+	{
+		return;
+	}
+
 	bool bBatchCompile = (Flags & EBlueprintBytecodeRecompileOptions::BatchCompile) != EBlueprintBytecodeRecompileOptions::None;
 	bool bSkipReinstancing = (Flags & EBlueprintBytecodeRecompileOptions::SkipReinstancing) != EBlueprintBytecodeRecompileOptions::None;
 
@@ -2144,7 +2160,7 @@ bool FKismetEditorUtilities::CanBlueprintImplementInterface(UBlueprint const* Bl
 	bool bCanImplementInterface = false;
 
 	// if the class is an actual implementable interface
-	if (IsClassABlueprintInterface(Class) && !Class->GetCppTypeInfo()->IsAbstract() && !Class->HasMetaData(FBlueprintMetadata::MD_CannotImplementInterfaceInBlueprint))
+	if (IsClassABlueprintInterface(Class) && !Class->HasMetaData(FBlueprintMetadata::MD_CannotImplementInterfaceInBlueprint))
 	{
 		bCanImplementInterface = true;
 
diff --git a/Engine/Source/Editor/UnrealEd/Private/Kismet2/KismetReinstanceUtilities.cpp b/Engine/Source/Editor/UnrealEd/Private/Kismet2/KismetReinstanceUtilities.cpp
--- a/Engine/Source/Editor/UnrealEd/Private/Kismet2/KismetReinstanceUtilities.cpp
+++ b/Engine/Source/Editor/UnrealEd/Private/Kismet2/KismetReinstanceUtilities.cpp
@@ -107,7 +107,7 @@ struct FReplaceReferenceHelper
 					class ReferenceReplace : public FArchiveReplaceObjectRef < UObject >
 					{
 					public:
-						ReferenceReplace(UObject* InSearchObject, const TMap<UObject*, UObject*>& InReplacementMap, TMap<FStringAssetReference, UObject*> InWeakReferencesMap)
+						ReferenceReplace(UObject* InSearchObject, const TMap<UObject*, UObject*>& InReplacementMap, const TMap<FStringAssetReference, UObject*>& InWeakReferencesMap)
 							: FArchiveReplaceObjectRef<UObject>(InSearchObject, InReplacementMap, false, false, false, true), WeakReferencesMap(InWeakReferencesMap)
 						{
 							SerializeSearchObject();
@@ -247,51 +247,8 @@ FBlueprintCompileReinstancer::FBlueprintCompileReinstancer(UClass* InClassToRein
 		// Remember the initial CDO for the class being resinstanced
 		OriginalCDO = ClassToReinstance->GetDefaultObject();
 
-		// Duplicate the class we're reinstancing into the transient package.  We'll re-class all objects we find to point to this new class
-		GIsDuplicatingClassForReinstancing = true;
-		ClassToReinstance->ClassFlags |= CLASS_NewerVersionExists;
-
-		if(bAvoidCDODuplication)
-		{
-			// clear the CDO so that we don't duplicate it, it's stored in OriginalCDO now
-			ClassToReinstance->ClassDefaultObject = nullptr;
-		}
-
-		const FName ReinstanceName = MakeUniqueObjectName(GetTransientPackage(), ClassToReinstance->GetClass(), *(FString(TEXT("REINST_")) + *ClassToReinstance->GetName()));
-		DuplicatedClass = (UClass*)StaticDuplicateObject(ClassToReinstance, GetTransientPackage(), ReinstanceName, ~RF_Transactional); 
-		// If you compile a blueprint that is part of the rootset, there's no reason for the REINST version to be part of the rootset:
-		DuplicatedClass->RemoveFromRoot();
-
-		ClassToReinstance->ClassFlags &= ~CLASS_NewerVersionExists;
-		GIsDuplicatingClassForReinstancing = false;
-
-		UBlueprintGeneratedClass* BPClassToReinstance = Cast<UBlueprintGeneratedClass>(ClassToReinstance);
-		UBlueprintGeneratedClass* BPGDuplicatedClass = Cast<UBlueprintGeneratedClass>(DuplicatedClass);
-		if (BPGDuplicatedClass && BPClassToReinstance && BPClassToReinstance->OverridenArchetypeForCDO)
-		{
-			BPGDuplicatedClass->OverridenArchetypeForCDO = BPClassToReinstance->OverridenArchetypeForCDO;
-		}
-
-		UFunction* DuplicatedClassUberGraphFunction = BPGDuplicatedClass ? BPGDuplicatedClass->UberGraphFunction : nullptr;
-		if (DuplicatedClassUberGraphFunction)
-		{
-			DuplicatedClassUberGraphFunction->Bind();
-			DuplicatedClassUberGraphFunction->StaticLink(true);
-		}
-
-		// Bind and link the duplicate class, so that it has the proper duplicate property offsets
-		DuplicatedClass->Bind();
-		DuplicatedClass->StaticLink(true);
-		
-		if(bAvoidCDODuplication)
-		{
-			// give the OriginalCDO to the DuplicatedClass:
-			OriginalCDO->Rename(*OriginalCDO->GetName(), DuplicatedClass, REN_DoNotDirty | REN_DontCreateRedirectors | REN_ForceNoResetLoaders);
-			DuplicatedClass->ClassDefaultObject = OriginalCDO;
-		}
+		DuplicatedClass = MoveCDOToNewClass(ClassToReinstance, TMap<UClass*, UClass*>(), bAvoidCDODuplication);
 
-		DuplicatedClass->ClassDefaultObject->SetClass(DuplicatedClass);
-	
 		if(!bAvoidCDODuplication)
 		{
 			ensure( ClassToReinstance->ClassDefaultObject->GetClass() == DuplicatedClass );
@@ -533,7 +490,7 @@ public:
 					Actor->ReregisterAllComponents();
 					Actor->RerunConstructionScripts();
 
-					if (SelectedObjecs.Contains(Obj))
+					if (SelectedObjecs.Contains(Obj) && GEditor)
 					{
 						GEditor->SelectActor(Actor, /*bInSelected =*/true, /*bNotify =*/true, false, true);
 					}
@@ -1170,7 +1127,7 @@ struct FActorReplacementHelper
 			}
 
 			// Save off transform
-			TargetWorldTransform = OldRootComponent->ComponentToWorld;
+			TargetWorldTransform = OldRootComponent->GetComponentTransform();
 			TargetWorldTransform.SetTranslation(OldRootComponent->GetComponentLocation()); // take into account any custom location
 		}
 
@@ -1259,7 +1216,7 @@ void FActorReplacementHelper::Finalize(const TMap<UObject*, UObject*>& OldToNewI
 		NewActor->MarkComponentsRenderStateDirty();
 	}
 
-	if (bSelectNewActor)
+	if (bSelectNewActor && GEditor)
 	{
 		GEditor->SelectActor(NewActor, /*bInSelected =*/true, /*bNotify =*/true);
 	}
@@ -1275,7 +1232,10 @@ void FActorReplacementHelper::Finalize(const TMap<UObject*, UObject*>& OldToNewI
 			}
 		}
 	}
-	GEditor->NotifyToolsOfObjectReplacement(ConstructedComponentReplacementMap);
+	if (GEditor)
+	{
+		GEditor->NotifyToolsOfObjectReplacement(ConstructedComponentReplacementMap);
+	}
 
 	// Make array of component subobjects that have been reinstanced as part of the new Actor.
 	TArray<UObject*> SourceObjects;
@@ -1287,7 +1247,7 @@ void FActorReplacementHelper::Finalize(const TMap<UObject*, UObject*>& OldToNewI
 	
 	// Destroy actor and clear references.
 	NewActor->Modify();
-	if (GEditor->Layers.IsValid())
+	if (GEditor && GEditor->Layers.IsValid())
 	{
 		GEditor->Layers->InitializeNewActorLayers(NewActor);
 	}
@@ -1516,6 +1476,58 @@ void FBlueprintCompileReinstancer::BatchReplaceInstancesOfClass(TMap<UClass*, UC
 	ReplaceInstancesOfClass_Inner(InOldToNewClassMap, nullptr, ObjectsThatShouldUseOldStuff, bClassObjectReplaced, bPreserveRootComponent);
 }
 
+UClass* FBlueprintCompileReinstancer::MoveCDOToNewClass(UClass* OwnerClass, const TMap<UClass*, UClass*>& OldToNewMap, bool bAvoidCDODuplication)
+{
+	GIsDuplicatingClassForReinstancing = true;
+	OwnerClass->ClassFlags |= CLASS_NewerVersionExists;
+
+	UObject* OldCDO = OwnerClass->ClassDefaultObject;
+	if(bAvoidCDODuplication)
+	{
+		// prevent duplication of CDO by hiding it from StaticDuplicateObject
+		OwnerClass->ClassDefaultObject = nullptr;
+	}
+	const FName ReinstanceName = MakeUniqueObjectName(GetTransientPackage(), OwnerClass->GetClass(), *(FString(TEXT("REINST_")) + *OwnerClass->GetName()));
+	UClass* CopyOfOwnerClass = CastChecked<UClass>(StaticDuplicateObject(OwnerClass, GetTransientPackage(), ReinstanceName, ~RF_Transactional));
+
+	CopyOfOwnerClass->RemoveFromRoot();
+	OwnerClass->ClassFlags &= ~CLASS_NewerVersionExists;
+	GIsDuplicatingClassForReinstancing = false;
+
+	UClass * const* OverridenParent = OldToNewMap.Find(CopyOfOwnerClass->GetSuperClass());
+	if(OverridenParent && *OverridenParent)
+	{
+		CopyOfOwnerClass->SetSuperStruct(*OverridenParent);
+	}
+
+	UBlueprintGeneratedClass* BPClassToReinstance = Cast<UBlueprintGeneratedClass>(OwnerClass);
+	UBlueprintGeneratedClass* BPGDuplicatedClass = Cast<UBlueprintGeneratedClass>(CopyOfOwnerClass);
+	if (BPGDuplicatedClass && BPClassToReinstance && BPClassToReinstance->OverridenArchetypeForCDO)
+	{
+		BPGDuplicatedClass->OverridenArchetypeForCDO = BPClassToReinstance->OverridenArchetypeForCDO;
+	}
+
+	UFunction* DuplicatedClassUberGraphFunction = BPGDuplicatedClass ? BPGDuplicatedClass->UberGraphFunction : nullptr;
+	if (DuplicatedClassUberGraphFunction)
+	{
+		DuplicatedClassUberGraphFunction->Bind();
+		DuplicatedClassUberGraphFunction->StaticLink(true);
+	}
+
+	CopyOfOwnerClass->Bind();
+	CopyOfOwnerClass->StaticLink(true);
+	if(OldCDO)
+	{
+		if(bAvoidCDODuplication)
+		{
+			OldCDO->Rename(*OldCDO->GetName(), CopyOfOwnerClass, REN_DoNotDirty | REN_DontCreateRedirectors | REN_ForceNoResetLoaders);
+			CopyOfOwnerClass->ClassDefaultObject = OldCDO;
+		}
+		OldCDO->SetClass(CopyOfOwnerClass);
+	}
+	return CopyOfOwnerClass;
+}
+
 void FBlueprintCompileReinstancer::ReplaceInstancesOfClass_Inner(TMap<UClass*, UClass*>& InOldToNewClassMap, UObject* InOriginalCDO, TSet<UObject*>* ObjectsThatShouldUseOldStuff, bool bClassObjectReplaced, bool bPreserveRootComponent)
 {
 	// If there is an original CDO, we are only reinstancing a single class
@@ -1655,16 +1667,16 @@ void FBlueprintCompileReinstancer::ReplaceInstancesOfClass_Inner(TMap<UClass*, U
 						FRotator Rotation = FRotator::ZeroRotator;
 						if (USceneComponent* OldRootComponent = OldActor->GetRootComponent())
 						{
-							// We need to make sure that the ComponentToWorld transform is up to date, but we don't want to run any initialization logic
+							// We need to make sure that the GetComponentTransform() transform is up to date, but we don't want to run any initialization logic
 							// so we silence the update, cache it off, revert the change (so no events are raised), and then directly update the transform
 							// with the value calculated in ConditionalUpdateComponentToWorld:
 							FScopedMovementUpdate SilenceMovement(OldRootComponent);
 
 							OldRootComponent->ConditionalUpdateComponentToWorld();
-							FTransform OldComponentToWorld = OldRootComponent->ComponentToWorld;
+							FTransform OldComponentToWorld = OldRootComponent->GetComponentTransform();
 							SilenceMovement.RevertMove();
 
-							OldRootComponent->ComponentToWorld = OldComponentToWorld;
+							OldRootComponent->SetComponentToWorld(OldComponentToWorld);
 							Location = OldActor->GetActorLocation();
 							Rotation = OldActor->GetActorRotation();
 						}
@@ -1736,10 +1748,13 @@ void FBlueprintCompileReinstancer::ReplaceInstancesOfClass_Inner(TMap<UClass*, U
 
 						if (OldActor->IsSelected())
 						{
-							GEditor->SelectActor(OldActor, /*bInSelected =*/false, /*bNotify =*/false);
+							if(GEditor)
+							{
+								GEditor->SelectActor(OldActor, /*bInSelected =*/false, /*bNotify =*/false);
+							}
 							bSelectionChanged = true;
 						}
-						if (GEditor->Layers.IsValid())
+						if (GEditor && GEditor->Layers.IsValid())
 						{
 							GEditor->Layers->DisassociateActorFromLayers(OldActor);
 						}
