Fix regression where hidden viewports would redraw causing multiple scenes to redraw every frame.  Fixed in 4.16 by 3434815

diff --git a/Engine/Source/Editor/UnrealEd/Private/EditorServer.cpp b/Engine/Source/Editor/UnrealEd/Private/EditorServer.cpp
--- a/Engine/Source/Editor/UnrealEd/Private/EditorServer.cpp
+++ b/Engine/Source/Editor/UnrealEd/Private/EditorServer.cpp
@@ -280,7 +280,6 @@ void UEditorEngine::RedrawAllViewports(bool bInvalidateHitProxies)
 		FEditorViewportClient* ViewportClient = AllViewportClients[ViewportIndex];
 		if ( ViewportClient && ViewportClient->Viewport )
 		{
-			ViewportClient->RequestRealTimeFrames(1);
 			if ( bInvalidateHitProxies )
 			{
 				// Invalidate hit proxies and display pixels.
