Copying //UE4/Dev-Rendering to //UE4/Dev-Main (Source: //UE4/Dev-Rendering @ 3461187)

diff --git a/Engine/Source/Editor/UnrealEd/Private/EditorLevelUtils.cpp b/Engine/Source/Editor/UnrealEd/Private/EditorLevelUtils.cpp
--- a/Engine/Source/Editor/UnrealEd/Private/EditorLevelUtils.cpp
+++ b/Engine/Source/Editor/UnrealEd/Private/EditorLevelUtils.cpp
@@ -85,9 +85,9 @@ int32 UEditorLevelUtils::MoveActorsToLevel(const TArray<AActor*>& ActorsToMove,
 				}
 
 				bool bIsSourceLevelLocked = FLevelUtils::IsLevelLocked(CurActor);
-
+				
 				if (!bIsSourceLevelLocked)
-				{
+		{
 					if (CurActor->GetLevel() != DestLevel)
 					{
 						FinalMoveList.Add(CurActor);
@@ -110,7 +110,7 @@ int32 UEditorLevelUtils::MoveActorsToLevel(const TArray<AActor*>& ActorsToMove,
 
 
 		if (FinalMoveList.Num() > 0)
-		{
+			{
 			GEditor->SelectNone(false, true, false);
 
 			USelection* ActorSelection = GEditor->GetSelectedActors();
@@ -136,16 +136,16 @@ int32 UEditorLevelUtils::MoveActorsToLevel(const TArray<AActor*>& ActorsToMove,
 				OwningWorld->SetCurrentLevel(DestLevel);
 				const bool bLevelVisible = DestLevel->bIsVisible;
 				if (!bLevelVisible)
-				{
+			{
 					UEditorLevelUtils::SetLevelVisibility(DestLevel, true, false);
-				}
-
+		}
+		
 				// Paste the actors into the new level
 				GEditor->edactPasteSelected(OwningWorld, false, false, false);
 
 				// Restore new level visibility to previous state
 				if (!bLevelVisible)
-				{
+		{
 					UEditorLevelUtils::SetLevelVisibility(DestLevel, false, false);
 				}
 
@@ -156,7 +156,7 @@ int32 UEditorLevelUtils::MoveActorsToLevel(const TArray<AActor*>& ActorsToMove,
 			// The moved (pasted) actors will now be selected
 			NumMovedActors += FinalMoveList.Num();
 		}
-
+		
 		// Restore the original clipboard contents
 		FPlatformMisc::ClipboardCopy(*OriginalClipboardContent);
 	}
@@ -166,283 +166,283 @@ int32 UEditorLevelUtils::MoveActorsToLevel(const TArray<AActor*>& ActorsToMove,
 }
 
 int32 UEditorLevelUtils::MoveSelectedActorsToLevel(ULevelStreaming* DestStreamingLevel)
-{
+	{
 	return MoveSelectedActorsToLevel(DestStreamingLevel->GetLoadedLevel());
 }
 
 int32 UEditorLevelUtils::MoveSelectedActorsToLevel(ULevel* DestLevel)
-{
+		{
 	TArray<AActor*> ActorsToMove;
 	for (FSelectionIterator It(GEditor->GetSelectedActorIterator()); It; ++It)
-	{
+			{
 		if (AActor* Actor = Cast<AActor>(*It))
-		{
+				{
 			ActorsToMove.Add(Actor);
+			}
 		}
-	}
 
 	return MoveActorsToLevel(ActorsToMove, DestLevel);
 	
-}
+	}
 
 ULevel* UEditorLevelUtils::AddLevelsToWorld(UWorld* InWorld, const TArray<FString>& LevelPackageNames, UClass* LevelStreamingClass)
-{
-	if ( !ensure(InWorld) )
 	{
-		return nullptr;
-	}
+		if ( !ensure(InWorld) )
+		{
+			return nullptr;
+		}
 
-	FScopedSlowTask SlowTask(LevelPackageNames.Num(), LOCTEXT("AddLevelsToWorldTask", "Adding Levels to World"));
-	SlowTask.MakeDialog();
+		FScopedSlowTask SlowTask(LevelPackageNames.Num(), LOCTEXT("AddLevelsToWorldTask", "Adding Levels to World"));
+		SlowTask.MakeDialog();
 
-	TArray<FString> PackageNames = LevelPackageNames;
+		TArray<FString> PackageNames = LevelPackageNames;
 
-	// Sort the level packages alphabetically by name.
-	PackageNames.Sort();
+		// Sort the level packages alphabetically by name.
+		PackageNames.Sort();
 
-	// Fire ULevel::LevelDirtiedEvent when falling out of scope.
-	FScopedLevelDirtied LevelDirtyCallback;
+		// Fire ULevel::LevelDirtiedEvent when falling out of scope.
+		FScopedLevelDirtied LevelDirtyCallback;
 
-	// Try to add the levels that were specified in the dialog.
-	ULevel* NewLevel = nullptr;
-	for (const auto& PackageName : PackageNames)
-	{
-		SlowTask.EnterProgressFrame();
+		// Try to add the levels that were specified in the dialog.
+		ULevel* NewLevel = nullptr;
+		for (const auto& PackageName : PackageNames)
+		{
+			SlowTask.EnterProgressFrame();
 
 		ULevelStreaming* NewStreamingLevel = AddLevelToWorld(InWorld, *PackageName, LevelStreamingClass);
 		NewLevel = NewStreamingLevel->GetLoadedLevel();
+			if (NewLevel)
+			{
+				LevelDirtyCallback.Request();
+			}
+
+		} // for each file
+
+		// Set the last loaded level to be the current level
 		if (NewLevel)
 		{
-			LevelDirtyCallback.Request();
+			if (InWorld->SetCurrentLevel(NewLevel))
+			{
+				FEditorDelegates::NewCurrentLevel.Broadcast();
+			}
 		}
 
-	} // for each file
-
-	// Set the last loaded level to be the current level
-	if (NewLevel)
-	{
-		if (InWorld->SetCurrentLevel(NewLevel))
+		// For safety
+		if (GLevelEditorModeTools().IsModeActive(FBuiltinEditorModes::EM_Landscape))
 		{
-			FEditorDelegates::NewCurrentLevel.Broadcast();
+			GLevelEditorModeTools().ActivateDefaultMode();
 		}
-	}
 
-	// For safety
-	if (GLevelEditorModeTools().IsModeActive(FBuiltinEditorModes::EM_Landscape))
-	{
-		GLevelEditorModeTools().ActivateDefaultMode();
-	}
-
-	// refresh editor windows
-	FEditorDelegates::RefreshAllBrowsers.Broadcast();
+		// refresh editor windows
+		FEditorDelegates::RefreshAllBrowsers.Broadcast();
 
-	// Update volume actor visibility for each viewport since we loaded a level which could potentially contain volumes
-	GUnrealEd->UpdateVolumeActorVisibility(NULL);
+		// Update volume actor visibility for each viewport since we loaded a level which could potentially contain volumes
+		GUnrealEd->UpdateVolumeActorVisibility(NULL);
 
-	return NewLevel;
-}
+		return NewLevel;
+	}
 
 ULevelStreaming* UEditorLevelUtils::AddLevelToWorld( UWorld* InWorld, const TCHAR* LevelPackageName, TSubclassOf<ULevelStreaming> LevelStreamingClass)
-{
+	{
 	ULevel* NewLevel = nullptr;
 	ULevelStreaming* StreamingLevel = nullptr;
-	bool bIsPersistentLevel = (InWorld->PersistentLevel->GetOutermost()->GetName() == FString(LevelPackageName));
+		bool bIsPersistentLevel = (InWorld->PersistentLevel->GetOutermost()->GetName() == FString(LevelPackageName));
 
-	if ( bIsPersistentLevel || FLevelUtils::FindStreamingLevel( InWorld, LevelPackageName ) )
-	{
-		// Do nothing if the level already exists in the world.
-		FMessageDialog::Open( EAppMsgType::Ok, NSLOCTEXT("UnrealEd", "LevelAlreadyExistsInWorld", "A level with that name already exists in the world.") );
-	}
-	else
-	{
-		// If the selected class is still NULL, abort the operation.
-		if ( LevelStreamingClass == nullptr)
+		if ( bIsPersistentLevel || FLevelUtils::FindStreamingLevel( InWorld, LevelPackageName ) )
 		{
+			// Do nothing if the level already exists in the world.
+			FMessageDialog::Open( EAppMsgType::Ok, NSLOCTEXT("UnrealEd", "LevelAlreadyExistsInWorld", "A level with that name already exists in the world.") );
+		}
+		else
+		{
+			// If the selected class is still NULL, abort the operation.
+		if ( LevelStreamingClass == nullptr)
+			{
 			return nullptr;
-		}	
+			}	
 
-		const FScopedBusyCursor BusyCursor;
+			const FScopedBusyCursor BusyCursor;
 
 		StreamingLevel = NewObject<ULevelStreaming>(InWorld, LevelStreamingClass, NAME_None, RF_NoFlags, NULL);
 
-		// Associate a package name.
-		StreamingLevel->SetWorldAssetByPackageName(LevelPackageName);
+			// Associate a package name.
+			StreamingLevel->SetWorldAssetByPackageName(LevelPackageName);
 
-		// Seed the level's draw color.
-		StreamingLevel->LevelColor = FLinearColor::MakeRandomColor();
+			// Seed the level's draw color.
+			StreamingLevel->LevelColor = FLinearColor::MakeRandomColor();
 
-		// Add the new level to world.
-		InWorld->StreamingLevels.Add( StreamingLevel );
+			// Add the new level to world.
+			InWorld->StreamingLevels.Add( StreamingLevel );
 
-		// Refresh just the newly created level.
-		TArray<ULevelStreaming*> LevelsForRefresh;
-		LevelsForRefresh.Add( StreamingLevel );
-		InWorld->RefreshStreamingLevels( LevelsForRefresh );
-		InWorld->MarkPackageDirty();
+			// Refresh just the newly created level.
+			TArray<ULevelStreaming*> LevelsForRefresh;
+			LevelsForRefresh.Add( StreamingLevel );
+			InWorld->RefreshStreamingLevels( LevelsForRefresh );
+			InWorld->MarkPackageDirty();
 
-		NewLevel = StreamingLevel->GetLoadedLevel();
-		if ( NewLevel != nullptr)
-		{
-			EditorLevelUtils::SetLevelVisibility( NewLevel, true, true );
-
-			// Levels migrated from other projects may fail to load their world settings
-			// If so we create a new AWorldSettings actor here.
-			if ( NewLevel->GetWorldSettings(false) == nullptr )
+			NewLevel = StreamingLevel->GetLoadedLevel();
+			if ( NewLevel != nullptr)
 			{
-				UWorld* SubLevelWorld = CastChecked<UWorld>(NewLevel->GetOuter());
+				EditorLevelUtils::SetLevelVisibility( NewLevel, true, true );
 
-				FActorSpawnParameters SpawnInfo;
-				SpawnInfo.SpawnCollisionHandlingOverride = ESpawnActorCollisionHandlingMethod::AlwaysSpawn;
-				SpawnInfo.Name = GEngine->WorldSettingsClass->GetFName();
-				AWorldSettings* NewWorldSettings = SubLevelWorld->SpawnActor<AWorldSettings>( GEngine->WorldSettingsClass, SpawnInfo );
+				// Levels migrated from other projects may fail to load their world settings
+				// If so we create a new AWorldSettings actor here.
+				if ( NewLevel->GetWorldSettings(false) == nullptr )
+				{
+					UWorld* SubLevelWorld = CastChecked<UWorld>(NewLevel->GetOuter());
 
-				NewLevel->SetWorldSettings(NewWorldSettings);
+					FActorSpawnParameters SpawnInfo;
+					SpawnInfo.SpawnCollisionHandlingOverride = ESpawnActorCollisionHandlingMethod::AlwaysSpawn;
+					SpawnInfo.Name = GEngine->WorldSettingsClass->GetFName();
+					AWorldSettings* NewWorldSettings = SubLevelWorld->SpawnActor<AWorldSettings>( GEngine->WorldSettingsClass, SpawnInfo );
+
+					NewLevel->SetWorldSettings(NewWorldSettings);
+				}
 			}
 		}
-	}
 
-	if(NewLevel) // if the level was successfully added
-	{
-		FEditorDelegates::OnAddLevelToWorld.Broadcast(NewLevel);
-	}
+		if(NewLevel) // if the level was successfully added
+		{
+			FEditorDelegates::OnAddLevelToWorld.Broadcast(NewLevel);
+		}
 
 	return StreamingLevel;
-}
+	}
 
 ULevelStreaming* UEditorLevelUtils::SetStreamingClassForLevel(ULevelStreaming* InLevel, TSubclassOf<ULevelStreaming> LevelStreamingClass)
-{
-	check(InLevel);
+	{
+		check(InLevel);
 
-	const FScopedBusyCursor BusyCursor;
+		const FScopedBusyCursor BusyCursor;
 
-	// Cache off the package name, as it will be lost when unloading the level
-	const FName CachedPackageName = InLevel->GetWorldAssetPackageFName();
+		// Cache off the package name, as it will be lost when unloading the level
+		const FName CachedPackageName = InLevel->GetWorldAssetPackageFName();
 
-	// First hide and remove the level if it exists
-	ULevel* Level = InLevel->GetLoadedLevel();
-	check(Level);
-	SetLevelVisibility( Level, false, false );
-	check(Level->OwningWorld);
-	UWorld* World = Level->OwningWorld;
+		// First hide and remove the level if it exists
+		ULevel* Level = InLevel->GetLoadedLevel();
+		check(Level);
+		SetLevelVisibility( Level, false, false );
+		check(Level->OwningWorld);
+		UWorld* World = Level->OwningWorld;
 		
-	World->StreamingLevels.Remove(InLevel);
+		World->StreamingLevels.Remove(InLevel);
 		
-	// re-add the level with the desired streaming class
-	AddLevelToWorld(World, *(CachedPackageName.ToString()), LevelStreamingClass);
+		// re-add the level with the desired streaming class
+		AddLevelToWorld(World, *(CachedPackageName.ToString()), LevelStreamingClass);
 
-	// Transfer level streaming settings
-	ULevelStreaming* NewStreamingLevel = FLevelUtils::FindStreamingLevel( Level );
-	if ( NewStreamingLevel )
-	{
-		NewStreamingLevel->LevelTransform = InLevel->LevelTransform;
-		NewStreamingLevel->EditorStreamingVolumes = InLevel->EditorStreamingVolumes;
-		NewStreamingLevel->MinTimeBetweenVolumeUnloadRequests = InLevel->MinTimeBetweenVolumeUnloadRequests;
-		NewStreamingLevel->LevelColor = InLevel->LevelColor;
-		NewStreamingLevel->Keywords = InLevel->Keywords;
-	}
+		// Transfer level streaming settings
+		ULevelStreaming* NewStreamingLevel = FLevelUtils::FindStreamingLevel( Level );
+		if ( NewStreamingLevel )
+		{
+			NewStreamingLevel->LevelTransform = InLevel->LevelTransform;
+			NewStreamingLevel->EditorStreamingVolumes = InLevel->EditorStreamingVolumes;
+			NewStreamingLevel->MinTimeBetweenVolumeUnloadRequests = InLevel->MinTimeBetweenVolumeUnloadRequests;
+			NewStreamingLevel->LevelColor = InLevel->LevelColor;
+			NewStreamingLevel->Keywords = InLevel->Keywords;
+		}
 
-	return NewStreamingLevel;
-}
+		return NewStreamingLevel;
+	}
 
 void UEditorLevelUtils::MakeLevelCurrent(ULevel* InLevel)
-{
-	// Locked levels can't be made current.
-	if ( !FLevelUtils::IsLevelLocked( InLevel ) )
-	{ 
-		// Make current broadcast if it changed
-		if( InLevel->OwningWorld->SetCurrentLevel( InLevel ) )
-		{
-			FEditorDelegates::NewCurrentLevel.Broadcast();
-		}
+	{
+		// Locked levels can't be made current.
+		if ( !FLevelUtils::IsLevelLocked( InLevel ) )
+		{ 
+			// Make current broadcast if it changed
+			if( InLevel->OwningWorld->SetCurrentLevel( InLevel ) )
+			{
+				FEditorDelegates::NewCurrentLevel.Broadcast();
+			}
 			
-		// Deselect all selected builder brushes.
-		bool bDeselectedSomething = false;
-		for ( FSelectionIterator It( GEditor->GetSelectedActorIterator() ) ; It ; ++It )
-		{
-			AActor* Actor = static_cast<AActor*>( *It );
-			checkSlow( Actor->IsA(AActor::StaticClass()) );
-			ABrush* Brush = Cast< ABrush >( Actor );
-			if ( Brush && FActorEditorUtils::IsABuilderBrush(Actor) )
+			// Deselect all selected builder brushes.
+			bool bDeselectedSomething = false;
+			for ( FSelectionIterator It( GEditor->GetSelectedActorIterator() ) ; It ; ++It )
 			{
-				GEditor->SelectActor( Actor, /*bInSelected=*/ false, /*bNotify=*/ false );
-				bDeselectedSomething = true;
+				AActor* Actor = static_cast<AActor*>( *It );
+				checkSlow( Actor->IsA(AActor::StaticClass()) );
+				ABrush* Brush = Cast< ABrush >( Actor );
+				if ( Brush && FActorEditorUtils::IsABuilderBrush(Actor) )
+				{
+					GEditor->SelectActor( Actor, /*bInSelected=*/ false, /*bNotify=*/ false );
+					bDeselectedSomething = true;
+				}
 			}
-		}
 
-		// Send a selection change callback if necessary.
-		if ( bDeselectedSomething )
+			// Send a selection change callback if necessary.
+			if ( bDeselectedSomething )
+			{
+				GEditor->NoteSelectionChange();
+			}
+
+			// Force the current level to be visible.
+			SetLevelVisibility( InLevel, true, false );
+		}
+		else
 		{
-			GEditor->NoteSelectionChange();
+			FMessageDialog::Open( EAppMsgType::Ok, NSLOCTEXT("UnrealEd", "Error_OperationDisallowedOnLockedLevelMakeLevelCurrent", "MakeLevelCurrent: The requested operation could not be completed because the level is locked.") );
 		}
-
-		// Force the current level to be visible.
-		SetLevelVisibility( InLevel, true, false );
 	}
-	else
-	{
-		FMessageDialog::Open( EAppMsgType::Ok, NSLOCTEXT("UnrealEd", "Error_OperationDisallowedOnLockedLevelMakeLevelCurrent", "MakeLevelCurrent: The requested operation could not be completed because the level is locked.") );
-	}
-}
-
+	
 void UEditorLevelUtils::MakeLevelCurrent(ULevelStreaming* InStreamingLevel)
 {
 	MakeLevelCurrent(InStreamingLevel->GetLoadedLevel());
 }
 
 bool UEditorLevelUtils::PrivateRemoveInvalidLevelFromWorld(ULevelStreaming* InLevelStreaming)
-{
-	bool bRemovedLevelStreaming = false;
-	if (InLevelStreaming != NULL)
 	{
-		check(InLevelStreaming->GetLoadedLevel() == NULL); // This method is designed to be used to remove left over references to null levels 
+		bool bRemovedLevelStreaming = false;
+		if (InLevelStreaming != NULL)
+		{
+			check(InLevelStreaming->GetLoadedLevel() == NULL); // This method is designed to be used to remove left over references to null levels 
 
-		InLevelStreaming->Modify();
+			InLevelStreaming->Modify();
 
-		// Disassociate the level from the volume.
-		for ( auto VolIter = InLevelStreaming->EditorStreamingVolumes.CreateIterator(); VolIter; VolIter++ )
-		{
-			ALevelStreamingVolume* LevelStreamingVolume = *VolIter;
-			if ( LevelStreamingVolume )
+			// Disassociate the level from the volume.
+			for ( auto VolIter = InLevelStreaming->EditorStreamingVolumes.CreateIterator(); VolIter; VolIter++ )
 			{
-				LevelStreamingVolume->Modify();
-				LevelStreamingVolume->StreamingLevelNames.Remove( InLevelStreaming->GetWorldAssetPackageFName() );
+				ALevelStreamingVolume* LevelStreamingVolume = *VolIter;
+				if ( LevelStreamingVolume )
+				{
+					LevelStreamingVolume->Modify();
+					LevelStreamingVolume->StreamingLevelNames.Remove( InLevelStreaming->GetWorldAssetPackageFName() );
+				}
 			}
-		}
 
-		// Disassociate the volumes from the level.
-		InLevelStreaming->EditorStreamingVolumes.Empty();
+			// Disassociate the volumes from the level.
+			InLevelStreaming->EditorStreamingVolumes.Empty();
 
-		UWorld* OwningWorld = Cast<UWorld>(InLevelStreaming->GetOuter());
-		if (OwningWorld != NULL)
-		{
-			OwningWorld->StreamingLevels.Remove(InLevelStreaming);
-			OwningWorld->RefreshStreamingLevels();
-			bRemovedLevelStreaming = true;
+			UWorld* OwningWorld = Cast<UWorld>(InLevelStreaming->GetOuter());
+			if (OwningWorld != NULL)
+			{
+				OwningWorld->StreamingLevels.Remove(InLevelStreaming);
+				OwningWorld->RefreshStreamingLevels();
+				bRemovedLevelStreaming = true;
+			}
 		}
+		return bRemovedLevelStreaming;
 	}
-	return bRemovedLevelStreaming;
-}
 
 bool UEditorLevelUtils::RemoveInvalidLevelFromWorld(ULevelStreaming* InLevelStreaming)
-{
-	bool bRemoveSuccessful = PrivateRemoveInvalidLevelFromWorld(InLevelStreaming);
-	if (bRemoveSuccessful)
 	{
-		// Redraw the main editor viewports.
-		FEditorSupportDelegates::RedrawAllViewports.Broadcast();
+		bool bRemoveSuccessful = PrivateRemoveInvalidLevelFromWorld(InLevelStreaming);
+		if (bRemoveSuccessful)
+		{
+			// Redraw the main editor viewports.
+			FEditorSupportDelegates::RedrawAllViewports.Broadcast();
 
-		// refresh editor windows
-		FEditorDelegates::RefreshAllBrowsers.Broadcast();
+			// refresh editor windows
+			FEditorDelegates::RefreshAllBrowsers.Broadcast();
 
-		// Update selection for any selected actors that were in the level and are no longer valid
-		GEditor->NoteSelectionChange();
+			// Update selection for any selected actors that were in the level and are no longer valid
+			GEditor->NoteSelectionChange();
 
-		// Collect garbage to clear out the destroyed level
-		CollectGarbage( GARBAGE_COLLECTION_KEEPFLAGS );
+			// Collect garbage to clear out the destroyed level
+			CollectGarbage( GARBAGE_COLLECTION_KEEPFLAGS );
+		}
+		return bRemoveSuccessful;
 	}
-	return bRemoveSuccessful;
-}
-
+	
 
 ULevelStreaming* UEditorLevelUtils::CreateNewStreamingLevel(TSubclassOf<ULevelStreaming> LevelStreamingClass, const FString& PackagePath /*= TEXT("")*/, bool bMoveSelectedActorsIntoNewLevel /*= false*/)
 {
@@ -531,449 +531,449 @@ ULevelStreaming* UEditorLevelUtils::CreateNewStreamingLevelForWorld(UWorld& InWo
 
 
 bool UEditorLevelUtils::RemoveLevelFromWorld(ULevel* InLevel)
-{
-	if (GEditor->Layers.IsValid())
 	{
-		GEditor->Layers->RemoveLevelLayerInformation( InLevel );
-	}
+		if (GEditor->Layers.IsValid())
+		{
+			GEditor->Layers->RemoveLevelLayerInformation( InLevel );
+		}
 		
-	GEditor->CloseEditedWorldAssets(CastChecked<UWorld>(InLevel->GetOuter()));
+		GEditor->CloseEditedWorldAssets(CastChecked<UWorld>(InLevel->GetOuter()));
 
-	UWorld* OwningWorld = InLevel->OwningWorld;
-	const FName LevelPackageName		= InLevel->GetOutermost()->GetFName();
-	const bool bRemovingCurrentLevel	= InLevel->IsCurrentLevel();
-	const bool bRemoveSuccessful		= PrivateRemoveLevelFromWorld( InLevel );
-	if ( bRemoveSuccessful )
-	{
-		if (bRemovingCurrentLevel)
+		UWorld* OwningWorld = InLevel->OwningWorld;
+		const FName LevelPackageName		= InLevel->GetOutermost()->GetFName();
+		const bool bRemovingCurrentLevel	= InLevel->IsCurrentLevel();
+		const bool bRemoveSuccessful		= PrivateRemoveLevelFromWorld( InLevel );
+		if ( bRemoveSuccessful )
 		{
-			MakeLevelCurrent(OwningWorld->PersistentLevel);
-		}
+			if (bRemovingCurrentLevel)
+			{
+				MakeLevelCurrent(OwningWorld->PersistentLevel);
+			}
 
 
-		EditorDestroyLevel(InLevel);
+			EditorDestroyLevel(InLevel);
 
-		// Redraw the main editor viewports.
-		FEditorSupportDelegates::RedrawAllViewports.Broadcast();
+			// Redraw the main editor viewports.
+			FEditorSupportDelegates::RedrawAllViewports.Broadcast();
 
-		// refresh editor windows
-		FEditorDelegates::RefreshAllBrowsers.Broadcast();
+			// refresh editor windows
+			FEditorDelegates::RefreshAllBrowsers.Broadcast();
 			
-		// Reset transaction buffer and run GC to clear out the destroyed level
-		GEditor->Cleanse(true, false, LOCTEXT("RemoveLevelTransReset", "Removing Levels from World"));
+			// Reset transaction buffer and run GC to clear out the destroyed level
+			GEditor->Cleanse(true, false, LOCTEXT("RemoveLevelTransReset", "Removing Levels from World"));
 
-		// Ensure that world was removed
-		UPackage* LevelPackage = FindObjectFast<UPackage>(NULL, LevelPackageName);
-		if (LevelPackage != nullptr)
-		{
-			UWorld* TheWorld = UWorld::FindWorldInPackage(LevelPackage->GetOutermost());								
-			if (TheWorld != nullptr)
-			{	
-				StaticExec(NULL, *FString::Printf(TEXT("OBJ REFS CLASS=%s NAME=%s shortest"),*TheWorld->GetClass()->GetName(), *TheWorld->GetPathName()));
-				TMap<UObject*,UProperty*>	Route		= FArchiveTraceRoute::FindShortestRootPath( TheWorld, true, GARBAGE_COLLECTION_KEEPFLAGS );
-				FString						ErrorString	= FArchiveTraceRoute::PrintRootPath( Route, TheWorld );
-				UE_LOG(LogStreaming, Fatal, TEXT("%s didn't get garbage collected!") LINE_TERMINATOR TEXT("%s"), *TheWorld->GetFullName(), *ErrorString );
+			// Ensure that world was removed
+			UPackage* LevelPackage = FindObjectFast<UPackage>(NULL, LevelPackageName);
+			if (LevelPackage != nullptr)
+			{
+				UWorld* TheWorld = UWorld::FindWorldInPackage(LevelPackage->GetOutermost());								
+				if (TheWorld != nullptr)
+				{	
+					StaticExec(NULL, *FString::Printf(TEXT("OBJ REFS CLASS=%s NAME=%s shortest"),*TheWorld->GetClass()->GetName(), *TheWorld->GetPathName()));
+					TMap<UObject*,UProperty*>	Route		= FArchiveTraceRoute::FindShortestRootPath( TheWorld, true, GARBAGE_COLLECTION_KEEPFLAGS );
+					FString						ErrorString	= FArchiveTraceRoute::PrintRootPath( Route, TheWorld );
+					UE_LOG(LogStreaming, Fatal, TEXT("%s didn't get garbage collected!") LINE_TERMINATOR TEXT("%s"), *TheWorld->GetFullName(), *ErrorString );
+				}
 			}
 		}
+		return bRemoveSuccessful;
 	}
-	return bRemoveSuccessful;
-}
 
 
 bool UEditorLevelUtils::PrivateRemoveLevelFromWorld(ULevel* InLevel)
-{
-	if ( !InLevel || InLevel->IsPersistentLevel() )
 	{
-		return false;
-	}
+		if ( !InLevel || InLevel->IsPersistentLevel() )
+		{
+			return false;
+		}
 
-	if ( FLevelUtils::IsLevelLocked(InLevel) )
-	{
-		FMessageDialog::Open( EAppMsgType::Ok, NSLOCTEXT("UnrealEd", "Error_OperationDisallowedOnLockedLevelRemoveLevelFromWorld", "RemoveLevelFromWorld: The requested operation could not be completed because the level is locked.") );
-		return false;
-	}
+		if ( FLevelUtils::IsLevelLocked(InLevel) )
+		{
+			FMessageDialog::Open( EAppMsgType::Ok, NSLOCTEXT("UnrealEd", "Error_OperationDisallowedOnLockedLevelRemoveLevelFromWorld", "RemoveLevelFromWorld: The requested operation could not be completed because the level is locked.") );
+			return false;
+		}
 
-	int32 StreamingLevelIndex = INDEX_NONE;
+		int32 StreamingLevelIndex = INDEX_NONE;
 
-	for( int32 LevelIndex = 0 ; LevelIndex < InLevel->OwningWorld->StreamingLevels.Num() ; ++LevelIndex )
-	{
-		ULevelStreaming* StreamingLevel = InLevel->OwningWorld->StreamingLevels[ LevelIndex ];
-		if( StreamingLevel && StreamingLevel->GetLoadedLevel() == InLevel )
+		for( int32 LevelIndex = 0 ; LevelIndex < InLevel->OwningWorld->StreamingLevels.Num() ; ++LevelIndex )
 		{
-			StreamingLevelIndex = LevelIndex;
-			break;
+			ULevelStreaming* StreamingLevel = InLevel->OwningWorld->StreamingLevels[ LevelIndex ];
+			if( StreamingLevel && StreamingLevel->GetLoadedLevel() == InLevel )
+			{
+				StreamingLevelIndex = LevelIndex;
+				break;
+			}
 		}
-	}
 
-	if (StreamingLevelIndex != INDEX_NONE)
-	{
-		InLevel->OwningWorld->StreamingLevels[StreamingLevelIndex]->MarkPendingKill();
-		InLevel->OwningWorld->StreamingLevels.RemoveAt( StreamingLevelIndex );
-		InLevel->OwningWorld->RefreshStreamingLevels();
-	}
-	else if (InLevel->bIsVisible)
-	{
-		InLevel->OwningWorld->RemoveFromWorld(InLevel);
-		check(InLevel->bIsVisible == false);
-	}
+		if (StreamingLevelIndex != INDEX_NONE)
+		{
+			InLevel->OwningWorld->StreamingLevels[StreamingLevelIndex]->MarkPendingKill();
+			InLevel->OwningWorld->StreamingLevels.RemoveAt( StreamingLevelIndex );
+			InLevel->OwningWorld->RefreshStreamingLevels();
+		}
+		else if (InLevel->bIsVisible)
+		{
+			InLevel->OwningWorld->RemoveFromWorld(InLevel);
+			check(InLevel->bIsVisible == false);
+		}
 
-	InLevel->ReleaseRenderingResources();
+		InLevel->ReleaseRenderingResources();
 
-	IStreamingManager::Get().RemoveLevel( InLevel );
-	UWorld* World = InLevel->OwningWorld;
-	World->RemoveLevel(InLevel);
-	InLevel->ClearLevelComponents();
+		IStreamingManager::Get().RemoveLevel( InLevel );
+		UWorld* World = InLevel->OwningWorld;
+		World->RemoveLevel(InLevel);
+		InLevel->ClearLevelComponents();
 
-	// remove all group actors from the world in the level we are removing
-	// otherwise, this will cause group actors to not be garbage collected
-	for (int32 GroupIndex = World->ActiveGroupActors.Num()-1; GroupIndex >= 0; --GroupIndex)
-	{
-		AGroupActor* GroupActor = Cast<AGroupActor>(World->ActiveGroupActors[GroupIndex]);
-		if (GroupActor && GroupActor->IsInLevel(InLevel))
+		// remove all group actors from the world in the level we are removing
+		// otherwise, this will cause group actors to not be garbage collected
+		for (int32 GroupIndex = World->ActiveGroupActors.Num()-1; GroupIndex >= 0; --GroupIndex)
 		{
-			World->ActiveGroupActors.RemoveAt(GroupIndex);
+			AGroupActor* GroupActor = Cast<AGroupActor>(World->ActiveGroupActors[GroupIndex]);
+			if (GroupActor && GroupActor->IsInLevel(InLevel))
+			{
+				World->ActiveGroupActors.RemoveAt(GroupIndex);
+			}
 		}
-	}
 		
-	// Mark all model components as pending kill so GC deletes references to them.
-	for (UModelComponent* ModelComponent : InLevel->ModelComponents)
-	{
-		if (ModelComponent != nullptr)
+		// Mark all model components as pending kill so GC deletes references to them.
+		for (UModelComponent* ModelComponent : InLevel->ModelComponents)
 		{
-			ModelComponent->MarkPendingKill();
+			if (ModelComponent != nullptr)
+			{
+				ModelComponent->MarkPendingKill();
+			}
 		}
-	}
 
-	// Mark all actors and their components as pending kill so GC will delete references to them.
-	for (AActor* Actor : InLevel->Actors)
-	{
-		if (Actor != nullptr)
+		// Mark all actors and their components as pending kill so GC will delete references to them.
+		for (AActor* Actor : InLevel->Actors)
 		{
-			Actor->MarkComponentsAsPendingKill();
-			Actor->MarkPendingKill();
+			if (Actor != nullptr)
+			{
+				Actor->MarkComponentsAsPendingKill();
+				Actor->MarkPendingKill();
+			}
 		}
-	}
 
-	World->MarkPackageDirty();
-	World->BroadcastLevelsChanged();
+		World->MarkPackageDirty();
+		World->BroadcastLevelsChanged();
 
-	return true;
-}
+		return true;
+	}
 
 bool UEditorLevelUtils::EditorDestroyLevel(ULevel* InLevel)
-{
-	UWorld* World = InLevel->OwningWorld;
+	{
+		UWorld* World = InLevel->OwningWorld;
 
-	InLevel->GetOuter()->MarkPendingKill();
-	InLevel->MarkPendingKill();
-	InLevel->GetOuter()->ClearFlags(RF_Public | RF_Standalone);
+		InLevel->GetOuter()->MarkPendingKill();
+		InLevel->MarkPendingKill();
+		InLevel->GetOuter()->ClearFlags(RF_Public | RF_Standalone);
 
-	UPackage* Package = InLevel->GetOutermost();
-	// We want to unconditionally destroy the level, so clear the dirty flag here so it can be unloaded successfully
-	Package->SetDirtyFlag(false);
+		UPackage* Package = InLevel->GetOutermost();
+		// We want to unconditionally destroy the level, so clear the dirty flag here so it can be unloaded successfully
+		Package->SetDirtyFlag(false);
 
-	TArray<UPackage*> Packages;
-	Packages.Add(Package);
-	if (!PackageTools::UnloadPackages(Packages))
-	{
-		FFormatNamedArguments Args;
-		Args.Add(TEXT("Package"), FText::FromString(Package->GetName()));
-		FMessageDialog::Open(EAppMsgType::Ok, FText::Format(LOCTEXT("UnloadPackagesFail", "Unable to unload package '{Package}'."), Args));
-		return false;
-	}
+		TArray<UPackage*> Packages;
+		Packages.Add(Package);
+		if (!PackageTools::UnloadPackages(Packages))
+		{
+			FFormatNamedArguments Args;
+			Args.Add(TEXT("Package"), FText::FromString(Package->GetName()));
+			FMessageDialog::Open(EAppMsgType::Ok, FText::Format(LOCTEXT("UnloadPackagesFail", "Unable to unload package '{Package}'."), Args));
+			return false;
+		}
 
-	return true;
-}
+		return true;
+	}
 
 ULevel* UEditorLevelUtils::CreateNewLevel( UWorld* InWorld, bool bMoveSelectedActorsIntoNewLevel, TSubclassOf<ULevelStreaming> LevelStreamingClass, const FString& DefaultFilename )
-{
+		{
 	ULevelStreaming* StreamingLevel = CreateNewStreamingLevelForWorld(*InWorld, LevelStreamingClass, DefaultFilename, bMoveSelectedActorsIntoNewLevel);
 	return StreamingLevel->GetLoadedLevel();
-}
-	
+		}
+
 void UEditorLevelUtils::DeselectAllSurfacesInLevel(ULevel* InLevel)
-{
-	if(InLevel)
 	{
-		UModel* Model = InLevel->Model;
-		for( int32 SurfaceIndex = 0 ; SurfaceIndex < Model->Surfs.Num() ; ++SurfaceIndex )
+		if(InLevel)
 		{
-			FBspSurf& Surf = Model->Surfs[SurfaceIndex];
-			if( Surf.PolyFlags & PF_Selected )
+			UModel* Model = InLevel->Model;
+			for( int32 SurfaceIndex = 0 ; SurfaceIndex < Model->Surfs.Num() ; ++SurfaceIndex )
 			{
-				Model->ModifySurf( SurfaceIndex, false );
-				Surf.PolyFlags &= ~PF_Selected;
+				FBspSurf& Surf = Model->Surfs[SurfaceIndex];
+				if( Surf.PolyFlags & PF_Selected )
+				{
+					Model->ModifySurf( SurfaceIndex, false );
+					Surf.PolyFlags &= ~PF_Selected;
+				}
 			}
 		}
 	}
-}
 
 void UEditorLevelUtils::SetLevelVisibility(ULevel* Level, bool bShouldBeVisible, bool bForceLayersVisible)
-{	
-	// Nothing to do
-	if ( Level == NULL )
-	{
-		return;
-	}
-
-	// Handle the case of the p-level
-	// The p-level can't be unloaded, so its actors/BSP should just be temporarily hidden/unhidden
-	// Also, intentionally do not force layers visible for the p-level
-	if ( Level->IsPersistentLevel() )
-	{
-		//create a transaction so we can undo the visibilty toggle
-		const FScopedTransaction Transaction( LOCTEXT( "ToggleLevelVisibility", "Toggle Level Visibility" ) );		
-		if ( Level->bIsVisible != bShouldBeVisible )
+	{	
+		// Nothing to do
+		if ( Level == NULL )
 		{
-			Level->Modify();
+			return;
 		}
-		// Set the visibility of each actor in the p-level
-		for ( TArray<AActor*>::TIterator PLevelActorIter( Level->Actors ); PLevelActorIter; ++PLevelActorIter )
+
+		// Handle the case of the p-level
+		// The p-level can't be unloaded, so its actors/BSP should just be temporarily hidden/unhidden
+		// Also, intentionally do not force layers visible for the p-level
+		if ( Level->IsPersistentLevel() )
 		{
-			AActor* CurActor = *PLevelActorIter;
-			if ( CurActor && !FActorEditorUtils::IsABuilderBrush(CurActor) && CurActor->bHiddenEdLevel == bShouldBeVisible )
+			//create a transaction so we can undo the visibilty toggle
+			const FScopedTransaction Transaction( LOCTEXT( "ToggleLevelVisibility", "Toggle Level Visibility" ) );		
+			if ( Level->bIsVisible != bShouldBeVisible )
 			{
-				CurActor->Modify();
-				CurActor->bHiddenEdLevel = !bShouldBeVisible;
-				CurActor->RegisterAllComponents();
-				CurActor->MarkComponentsRenderStateDirty();
+				Level->Modify();
 			}
-		}
-
-		// Set the visibility of each BSP surface in the p-level
-		UModel* CurLevelModel = Level->Model;
-		if ( CurLevelModel )
-		{
-			CurLevelModel->Modify();
-			for ( TArray<FBspSurf>::TIterator SurfaceIterator( CurLevelModel->Surfs ); SurfaceIterator; ++SurfaceIterator )
+			// Set the visibility of each actor in the p-level
+			for ( TArray<AActor*>::TIterator PLevelActorIter( Level->Actors ); PLevelActorIter; ++PLevelActorIter )
 			{
-				FBspSurf& CurSurf = *SurfaceIterator;
-				CurSurf.bHiddenEdLevel = !bShouldBeVisible;
+				AActor* CurActor = *PLevelActorIter;
+				if ( CurActor && !FActorEditorUtils::IsABuilderBrush(CurActor) && CurActor->bHiddenEdLevel == bShouldBeVisible )
+				{
+					CurActor->Modify();
+					CurActor->bHiddenEdLevel = !bShouldBeVisible;
+					CurActor->RegisterAllComponents();
+					CurActor->MarkComponentsRenderStateDirty();
+				}
 			}
-		}
 
-		// Add/remove model components from the scene
-		for(int32 ComponentIndex = 0;ComponentIndex < Level->ModelComponents.Num();ComponentIndex++)
-		{
-			UModelComponent* CurLevelModelCmp = Level->ModelComponents[ComponentIndex];
-			if(CurLevelModelCmp)
+			// Set the visibility of each BSP surface in the p-level
+			UModel* CurLevelModel = Level->Model;
+			if ( CurLevelModel )
 			{
-				if (bShouldBeVisible)
+				CurLevelModel->Modify();
+				for ( TArray<FBspSurf>::TIterator SurfaceIterator( CurLevelModel->Surfs ); SurfaceIterator; ++SurfaceIterator )
 				{
-					CurLevelModelCmp->RegisterComponentWithWorld(Level->OwningWorld);
+					FBspSurf& CurSurf = *SurfaceIterator;
+					CurSurf.bHiddenEdLevel = !bShouldBeVisible;
 				}
-				else if (!bShouldBeVisible && CurLevelModelCmp->IsRegistered())
+			}
+
+			// Add/remove model components from the scene
+			for(int32 ComponentIndex = 0;ComponentIndex < Level->ModelComponents.Num();ComponentIndex++)
+			{
+				UModelComponent* CurLevelModelCmp = Level->ModelComponents[ComponentIndex];
+				if(CurLevelModelCmp)
 				{
-					CurLevelModelCmp->UnregisterComponent();
+					if (bShouldBeVisible)
+					{
+						CurLevelModelCmp->RegisterComponentWithWorld(Level->OwningWorld);
+					}
+					else if (!bShouldBeVisible && CurLevelModelCmp->IsRegistered())
+					{
+						CurLevelModelCmp->UnregisterComponent();
+					}
 				}
 			}
-		}
-
-		Level->GetWorld()->OnLevelsChanged().Broadcast();
-		FEditorSupportDelegates::RedrawAllViewports.Broadcast();
-	}
-	else
-	{
-		ULevelStreaming* StreamingLevel = NULL;
-		if (Level->OwningWorld == NULL || Level->OwningWorld->PersistentLevel != Level )
-		{
-			StreamingLevel = FLevelUtils::FindStreamingLevel( Level );
-		}
-
-		// Create a transaction so we can undo the visibility toggle
-		const FScopedTransaction Transaction( LOCTEXT( "ToggleLevelVisibility", "Toggle Level Visibility" ) );		
-
-		// Handle the case of a streaming level
-		if ( StreamingLevel )
-		{
-			// We need to set the RF_Transactional to make a streaming level serialize itself. so store the original ones, set the flag, and put the original flags back when done
-			EObjectFlags cachedFlags = StreamingLevel->GetFlags();
-			StreamingLevel->SetFlags( RF_Transactional );
-			StreamingLevel->Modify();			
-			StreamingLevel->SetFlags( cachedFlags );
 
-			// Set the visibility state for this streaming level.  
-			StreamingLevel->bShouldBeVisibleInEditor = bShouldBeVisible;
+			Level->GetWorld()->OnLevelsChanged().Broadcast();
+			FEditorSupportDelegates::RedrawAllViewports.Broadcast();
 		}
-
-		if (!bShouldBeVisible && GEditor->Layers.IsValid())
-		{
-			GEditor->Layers->RemoveLevelLayerInformation( Level );
-		}
-
-		// UpdateLevelStreaming sets Level->bIsVisible directly, so we need to make sure it gets saved to the transaction buffer.
-		if ( Level->bIsVisible != bShouldBeVisible )
+		else
 		{
-			Level->Modify();
-		}
+			ULevelStreaming* StreamingLevel = NULL;
+			if (Level->OwningWorld == NULL || Level->OwningWorld->PersistentLevel != Level )
+			{
+				StreamingLevel = FLevelUtils::FindStreamingLevel( Level );
+			}
 
-		if ( StreamingLevel )
-		{
-			Level->OwningWorld->FlushLevelStreaming();
+			// Create a transaction so we can undo the visibility toggle
+			const FScopedTransaction Transaction( LOCTEXT( "ToggleLevelVisibility", "Toggle Level Visibility" ) );		
 
-			// In the Editor we expect this operation will complete in a single call
-			check(Level->bIsVisible == bShouldBeVisible);
-		}
-		else if (Level->OwningWorld != NULL)
-		{
-			// In case we level has no associated StreamingLevel, remove or add to world directly
-			if (bShouldBeVisible)
+			// Handle the case of a streaming level
+			if ( StreamingLevel )
 			{
-				if (!Level->bIsVisible)
-				{
-					Level->OwningWorld->AddToWorld(Level);
-				}
+				// We need to set the RF_Transactional to make a streaming level serialize itself. so store the original ones, set the flag, and put the original flags back when done
+				EObjectFlags cachedFlags = StreamingLevel->GetFlags();
+				StreamingLevel->SetFlags( RF_Transactional );
+				StreamingLevel->Modify();			
+				StreamingLevel->SetFlags( cachedFlags );
+
+				// Set the visibility state for this streaming level.  
+				StreamingLevel->bShouldBeVisibleInEditor = bShouldBeVisible;
 			}
-			else
+
+			if (!bShouldBeVisible && GEditor->Layers.IsValid())
 			{
-				Level->OwningWorld->RemoveFromWorld(Level);
+				GEditor->Layers->RemoveLevelLayerInformation( Level );
 			}
-				
-			// In the Editor we expect this operation will complete in a single call
-			check(Level->bIsVisible == bShouldBeVisible);
-		}
 
-		if (bShouldBeVisible && GEditor->Layers.IsValid())
-		{
-			GEditor->Layers->AddLevelLayerInformation( Level );
-		}
+			// UpdateLevelStreaming sets Level->bIsVisible directly, so we need to make sure it gets saved to the transaction buffer.
+			if ( Level->bIsVisible != bShouldBeVisible )
+			{
+				Level->Modify();
+			}
 
-		// Force the level's layers to be visible, if desired
-		FEditorSupportDelegates::RedrawAllViewports.Broadcast();
+			if ( StreamingLevel )
+			{
+				Level->OwningWorld->FlushLevelStreaming();
 
-		// Iterate over the level's actors, making a list of their layers and unhiding the layers.
-		TArray<AActor*>& Actors = Level->Actors;
-		for ( int32 ActorIndex = 0 ; ActorIndex < Actors.Num() ; ++ActorIndex )
-		{
-			AActor* Actor = Actors[ ActorIndex ];
-			if ( Actor )
+				// In the Editor we expect this operation will complete in a single call
+				check(Level->bIsVisible == bShouldBeVisible);
+			}
+			else if (Level->OwningWorld != NULL)
 			{
-				bool bModified = false;
-				if (bShouldBeVisible && bForceLayersVisible && 
-					GEditor->Layers.IsValid() && 
-					GEditor->Layers->IsActorValidForLayer( Actor ))
+				// In case we level has no associated StreamingLevel, remove or add to world directly
+				if (bShouldBeVisible)
 				{
-					// Make the actor layer visible, if it's not already.
-					if ( Actor->bHiddenEdLayer )
+					if (!Level->bIsVisible)
 					{
-						bModified = Actor->Modify();
-						Actor->bHiddenEdLayer = false;
+						Level->OwningWorld->AddToWorld(Level);
 					}
-
-					const bool bIsVisible = true;
-					GEditor->Layers->SetLayersVisibility( Actor->Layers, bIsVisible );
 				}
+				else
+				{
+					Level->OwningWorld->RemoveFromWorld(Level);
+				}
+				
+				// In the Editor we expect this operation will complete in a single call
+				check(Level->bIsVisible == bShouldBeVisible);
+			}
+
+			if (bShouldBeVisible && GEditor->Layers.IsValid())
+			{
+				GEditor->Layers->AddLevelLayerInformation( Level );
+			}
+
+			// Force the level's layers to be visible, if desired
+			FEditorSupportDelegates::RedrawAllViewports.Broadcast();
 
-				// Set the visibility of each actor in the streaming level
-				if ( !FActorEditorUtils::IsABuilderBrush(Actor) && Actor->bHiddenEdLevel == bShouldBeVisible )
+			// Iterate over the level's actors, making a list of their layers and unhiding the layers.
+			TArray<AActor*>& Actors = Level->Actors;
+			for ( int32 ActorIndex = 0 ; ActorIndex < Actors.Num() ; ++ActorIndex )
+			{
+				AActor* Actor = Actors[ ActorIndex ];
+				if ( Actor )
 				{
-					if ( !bModified )
+					bool bModified = false;
+					if (bShouldBeVisible && bForceLayersVisible && 
+						GEditor->Layers.IsValid() && 
+						GEditor->Layers->IsActorValidForLayer( Actor ))
 					{
-						bModified = Actor->Modify();
+						// Make the actor layer visible, if it's not already.
+						if ( Actor->bHiddenEdLayer )
+						{
+							bModified = Actor->Modify();
+							Actor->bHiddenEdLayer = false;
+						}
+
+						const bool bIsVisible = true;
+						GEditor->Layers->SetLayersVisibility( Actor->Layers, bIsVisible );
 					}
-					Actor->bHiddenEdLevel = !bShouldBeVisible;
 
-					if (bShouldBeVisible)
-					{
-						Actor->ReregisterAllComponents();
-					}
-					else
+					// Set the visibility of each actor in the streaming level
+					if ( !FActorEditorUtils::IsABuilderBrush(Actor) && Actor->bHiddenEdLevel == bShouldBeVisible )
 					{
-						Actor->UnregisterAllComponents();
+						if ( !bModified )
+						{
+							bModified = Actor->Modify();
+						}
+						Actor->bHiddenEdLevel = !bShouldBeVisible;
+
+						if (bShouldBeVisible)
+						{
+							Actor->ReregisterAllComponents();
+						}
+						else
+						{
+							Actor->UnregisterAllComponents();
+						}
 					}
 				}
 			}
 		}
-	}
 
-	FEditorDelegates::RefreshLayerBrowser.Broadcast();
+		FEditorDelegates::RefreshLayerBrowser.Broadcast();
 
-	// Notify the Scene Outliner, as new Actors may be present in the world.
-	GEngine->BroadcastLevelActorListChanged();
+		// Notify the Scene Outliner, as new Actors may be present in the world.
+		GEngine->BroadcastLevelActorListChanged();
 
-	// If the level is being hidden, deselect actors and surfaces that belong to this level.
-	if ( !bShouldBeVisible )
-	{
-		USelection* SelectedActors = GEditor->GetSelectedActors();
-		SelectedActors->Modify();
-		TArray<AActor*>& Actors = Level->Actors;
-		for ( int32 ActorIndex = 0 ; ActorIndex < Actors.Num() ; ++ActorIndex )
+		// If the level is being hidden, deselect actors and surfaces that belong to this level.
+		if ( !bShouldBeVisible )
 		{
-			AActor* Actor = Actors[ ActorIndex ];
-			if ( Actor )
+			USelection* SelectedActors = GEditor->GetSelectedActors();
+			SelectedActors->Modify();
+			TArray<AActor*>& Actors = Level->Actors;
+			for ( int32 ActorIndex = 0 ; ActorIndex < Actors.Num() ; ++ActorIndex )
 			{
-				SelectedActors->Deselect( Actor );
+				AActor* Actor = Actors[ ActorIndex ];
+				if ( Actor )
+				{
+					SelectedActors->Deselect( Actor );
+				}
 			}
-		}
 
-		DeselectAllSurfacesInLevel(Level);
+			DeselectAllSurfacesInLevel(Level);
 
-		// Tell the editor selection status was changed.
-		GEditor->NoteSelectionChange();
-	}
+			// Tell the editor selection status was changed.
+			GEditor->NoteSelectionChange();
+		}
 
-	Level->bIsVisible = bShouldBeVisible;
+		Level->bIsVisible = bShouldBeVisible;
 
-	if (Level->bIsLightingScenario)
-	{
-		Level->OwningWorld->PropagateLightingScenarioChange();
+		if (Level->bIsLightingScenario)
+		{
+			Level->OwningWorld->PropagateLightingScenarioChange(bShouldBeVisible);
+		}
 	}
-}
 
 void UEditorLevelUtils::GetWorlds(UWorld* InWorld, TArray<UWorld*>& OutWorlds, bool bIncludeInWorld, bool bOnlyEditorVisible)
-{
-	OutWorlds.Empty();
-
-	if (!InWorld)
 	{
-		return;
-	}
+		OutWorlds.Empty();
 
-	if ( bIncludeInWorld )
-	{
-		OutWorlds.AddUnique( InWorld );
-	}
+		if (!InWorld)
+		{
+			return;
+		}
 
-	// Iterate over the world's level array to find referenced levels ("worlds"). We don't 
-	for ( int32 LevelIndex = 0 ; LevelIndex < InWorld->StreamingLevels.Num() ; ++LevelIndex )
-	{
-		ULevelStreaming* StreamingLevel = InWorld->StreamingLevels[LevelIndex];
-		if ( StreamingLevel )
+		if ( bIncludeInWorld )
+		{
+			OutWorlds.AddUnique( InWorld );
+		}
+
+		// Iterate over the world's level array to find referenced levels ("worlds"). We don't 
+		for ( int32 LevelIndex = 0 ; LevelIndex < InWorld->StreamingLevels.Num() ; ++LevelIndex )
 		{
-			// If we asked for only sub-levels that are editor-visible, then limit our results appropriately
-			bool bShouldAlwaysBeLoaded = false; // Cast< ULevelStreamingAlwaysLoaded >( StreamingLevel ) != NULL;
-			if( !bOnlyEditorVisible || bShouldAlwaysBeLoaded || StreamingLevel->bShouldBeVisibleInEditor )
+			ULevelStreaming* StreamingLevel = InWorld->StreamingLevels[LevelIndex];
+			if ( StreamingLevel )
 			{
-				const ULevel* Level = StreamingLevel->GetLoadedLevel();
-					
-				// This should always be the case for valid level names as the Editor preloads all packages.
-				if ( Level != NULL )
+				// If we asked for only sub-levels that are editor-visible, then limit our results appropriately
+				bool bShouldAlwaysBeLoaded = false; // Cast< ULevelStreamingAlwaysLoaded >( StreamingLevel ) != NULL;
+				if( !bOnlyEditorVisible || bShouldAlwaysBeLoaded || StreamingLevel->bShouldBeVisibleInEditor )
 				{
-					// Newer levels have their packages' world as the outer.
-					UWorld* World = Cast<UWorld>( Level->GetOuter() );
-					if ( World != NULL )
+					const ULevel* Level = StreamingLevel->GetLoadedLevel();
+					
+					// This should always be the case for valid level names as the Editor preloads all packages.
+					if ( Level != NULL )
 					{
-						OutWorlds.AddUnique( World );
+						// Newer levels have their packages' world as the outer.
+						UWorld* World = Cast<UWorld>( Level->GetOuter() );
+						if ( World != NULL )
+						{
+							OutWorlds.AddUnique( World );
+						}
 					}
 				}
 			}
 		}
-	}
 		
-	// Levels can be loaded directly without StreamingLevel facilities
-	for ( int32 LevelIndex = 0 ; LevelIndex < InWorld->GetLevels().Num() ; ++LevelIndex )
-	{
-		ULevel* Level = InWorld->GetLevel(LevelIndex);
-		if ( Level )
+		// Levels can be loaded directly without StreamingLevel facilities
+		for ( int32 LevelIndex = 0 ; LevelIndex < InWorld->GetLevels().Num() ; ++LevelIndex )
 		{
-			// Newer levels have their packages' world as the outer.
-			UWorld* World = Cast<UWorld>( Level->GetOuter() );
-			if ( World != NULL )
+			ULevel* Level = InWorld->GetLevel(LevelIndex);
+			if ( Level )
 			{
-				OutWorlds.AddUnique( World );
+				// Newer levels have their packages' world as the outer.
+				UWorld* World = Cast<UWorld>( Level->GetOuter() );
+				if ( World != NULL )
+				{
+					OutWorlds.AddUnique( World );
+				}
 			}
 		}
 	}
-}
 
 #undef LOCTEXT_NAMESPACE
diff --git a/Engine/Source/Editor/UnrealEd/Private/EditorViewportClient.cpp b/Engine/Source/Editor/UnrealEd/Private/EditorViewportClient.cpp
--- a/Engine/Source/Editor/UnrealEd/Private/EditorViewportClient.cpp
+++ b/Engine/Source/Editor/UnrealEd/Private/EditorViewportClient.cpp
@@ -3311,7 +3311,7 @@ void FEditorViewportClient::Draw(FViewport* InViewport, FCanvas* Canvas)
 		GetScene(),
 		EngineShowFlags)
 		.SetWorldTimes( TimeSeconds, DeltaTimeSeconds, RealTimeSeconds )
-		.SetRealtimeUpdate( IsRealtime() )
+		.SetRealtimeUpdate( IsRealtime() && FSlateThrottleManager::Get().IsAllowingExpensiveTasks() )
 		.SetViewModeParam( ViewModeParam, ViewModeParamName ) );
 
 	ViewFamily.EngineShowFlags = EngineShowFlags;
diff --git a/Engine/Source/Editor/UnrealEd/Private/Kismet2/DebuggerCommands.cpp b/Engine/Source/Editor/UnrealEd/Private/Kismet2/DebuggerCommands.cpp
--- a/Engine/Source/Editor/UnrealEd/Private/Kismet2/DebuggerCommands.cpp
+++ b/Engine/Source/Editor/UnrealEd/Private/Kismet2/DebuggerCommands.cpp
@@ -1196,7 +1196,7 @@ void FInternalPlayWorldCommandCallbacks::Simulate_Clicked()
 bool FInternalPlayWorldCommandCallbacks::Simulate_CanExecute()
 {
 	// Can't simulate while already simulating; PIE is fine as we toggle to simulate
-	return !(HasPlayWorld() && GUnrealEd->bIsSimulatingInEditor);
+	return !(HasPlayWorld() && GUnrealEd->bIsSimulatingInEditor) && !GEditor->IsLightingBuildCurrentlyRunning();
 }
 
 
@@ -1415,7 +1415,7 @@ void FInternalPlayWorldCommandCallbacks::PlayInViewport_Clicked( )
 bool FInternalPlayWorldCommandCallbacks::PlayInViewport_CanExecute()
 {
 	// Allow PIE if we don't already have a play session or the play session is simulate in editor (which we can toggle to PIE)
-	return (!GEditor->bIsPlayWorldQueued && !HasPlayWorld() ) || GUnrealEd->bIsSimulatingInEditor;
+	return (!GEditor->bIsPlayWorldQueued && !HasPlayWorld() && !GEditor->IsLightingBuildCurrentlyRunning() ) || GUnrealEd->bIsSimulatingInEditor;
 }
 
 
@@ -1465,7 +1465,7 @@ void FInternalPlayWorldCommandCallbacks::PlayInEditorFloating_Clicked( )
 
 bool FInternalPlayWorldCommandCallbacks::PlayInEditorFloating_CanExecute()
 {
-	return !HasPlayWorld() || !GUnrealEd->bIsSimulatingInEditor;
+	return (!HasPlayWorld() || !GUnrealEd->bIsSimulatingInEditor) && !GEditor->IsLightingBuildCurrentlyRunning();
 }
 
 void FInternalPlayWorldCommandCallbacks::PlayInVR_Clicked()
@@ -1515,7 +1515,7 @@ void FInternalPlayWorldCommandCallbacks::PlayInVR_Clicked()
 
 bool FInternalPlayWorldCommandCallbacks::PlayInVR_CanExecute()
 {
-	return (!HasPlayWorld() || !GUnrealEd->bIsSimulatingInEditor) && GEngine && GEngine->HMDDevice.IsValid();
+	return (!HasPlayWorld() || !GUnrealEd->bIsSimulatingInEditor) && !GEditor->IsLightingBuildCurrentlyRunning() && GEngine && GEngine->HMDDevice.IsValid();
 }
 
 
diff --git a/Engine/Source/Editor/UnrealEd/Private/StaticLightingSystem/StaticLightingPrivate.h b/Engine/Source/Editor/UnrealEd/Private/StaticLightingSystem/StaticLightingPrivate.h
--- a/Engine/Source/Editor/UnrealEd/Private/StaticLightingSystem/StaticLightingPrivate.h
+++ b/Engine/Source/Editor/UnrealEd/Private/StaticLightingSystem/StaticLightingPrivate.h
@@ -293,7 +293,7 @@ public:
 
 	bool ShouldOperateOnLevel(ULevel* InLevel) const
 	{
-		return !InLevel->bIsLightingScenario || InLevel == LightingScenario;
+		return InLevel && (!InLevel->bIsLightingScenario || InLevel == LightingScenario) && InLevel->bIsVisible;
 	}
 
 private:
diff --git a/Engine/Source/Editor/UnrealEd/Private/StaticLightingSystem/StaticLightingSystem.cpp b/Engine/Source/Editor/UnrealEd/Private/StaticLightingSystem/StaticLightingSystem.cpp
--- a/Engine/Source/Editor/UnrealEd/Private/StaticLightingSystem/StaticLightingSystem.cpp
+++ b/Engine/Source/Editor/UnrealEd/Private/StaticLightingSystem/StaticLightingSystem.cpp
@@ -609,19 +609,20 @@ bool FStaticLightingSystem::BeginLightmassProcess()
 
 					if (bLightIsInWorld && ShouldOperateOnLevel(Light->GetOwner()->GetLevel()))
 					{
-						if (Light->bAffectsWorld
-						&& (Light->HasStaticShadowing() || Light->HasStaticLighting()))
-					{
-						// Make sure the light GUIDs are up-to-date.
-						Light->ValidateLightGUIDs();
+						if (Light->bAffectsWorld 
+							&& Light->IsRegistered()
+							&& (Light->HasStaticShadowing() || Light->HasStaticLighting()))
+						{
+							// Make sure the light GUIDs are up-to-date.
+							Light->ValidateLightGUIDs();
 
-						// Add the light to the system's list of lights in the world.
-						Lights.Add(Light);
+							// Add the light to the system's list of lights in the world.
+							Lights.Add(Light);
+						}
 					}
 				}
 			}
 		}
-		}
 
 		{
 			FLightmassStatistics::FScopedGather GatherStatScope(LightmassStatistics.GatherLightingInfoTime);
diff --git a/Engine/Source/Editor/UnrealEd/Private/UnrealEdPrivatePCH.h b/Engine/Source/Editor/UnrealEd/Private/UnrealEdPrivatePCH.h
--- a/Engine/Source/Editor/UnrealEd/Private/UnrealEdPrivatePCH.h
+++ b/Engine/Source/Editor/UnrealEd/Private/UnrealEdPrivatePCH.h
@@ -34,7 +34,6 @@
 #include "Templates/EnableIf.h"
 #include "Templates/RemoveReference.h"
 #include "Templates/TypeCompatibleBytes.h"
-#include "Templates/AlignOf.h"
 #include "Templates/ChooseClass.h"
 #include "Templates/IntegralConstant.h"
 #include "Templates/IsClass.h"
diff --git a/Engine/Source/Editor/UnrealEd/Private/UnrealWidget.cpp b/Engine/Source/Editor/UnrealEd/Private/UnrealWidget.cpp
--- a/Engine/Source/Editor/UnrealEd/Private/UnrealWidget.cpp
+++ b/Engine/Source/Editor/UnrealEd/Private/UnrealWidget.cpp
@@ -630,7 +630,7 @@ void FWidget::Render_Translate( const FSceneView* View, FPrimitiveDrawInterface*
 		const FVector CameraZAxis = View->ViewMatrices.GetViewMatrix().GetColumn(2);
 
 		UMaterialInstanceDynamic* XYZMaterial = ( CurrentAxis&EAxisList::Screen) ? CurrentAxisMaterial : OpaquePlaneMaterialXY;
-		DrawSphere( PDI, InLocation, 4.0f * Space.Scale, 10, 5, XYZMaterial->GetRenderProxy(false), SDPG_Foreground );
+		DrawSphere( PDI, InLocation, FRotator::ZeroRotator, 4.0f * Space.Scale, 10, 5, XYZMaterial->GetRenderProxy(false), SDPG_Foreground );
 
 		PDI->SetHitProxy( NULL );
 	}
diff --git a/Engine/Source/Editor/UnrealEd/Public/UnrealEdSharedPCH.h b/Engine/Source/Editor/UnrealEd/Public/UnrealEdSharedPCH.h
--- a/Engine/Source/Editor/UnrealEd/Public/UnrealEdSharedPCH.h
+++ b/Engine/Source/Editor/UnrealEd/Public/UnrealEdSharedPCH.h
@@ -34,7 +34,6 @@
 #include "Templates/EnableIf.h"
 #include "Templates/RemoveReference.h"
 #include "Templates/TypeCompatibleBytes.h"
-#include "Templates/AlignOf.h"
 #include "Templates/ChooseClass.h"
 #include "Templates/IntegralConstant.h"
 #include "Templates/IsClass.h"
diff --git a/Engine/Source/Programs/ShaderCompileWorker/Private/ShaderCompileWorker.cpp b/Engine/Source/Programs/ShaderCompileWorker/Private/ShaderCompileWorker.cpp
--- a/Engine/Source/Programs/ShaderCompileWorker/Private/ShaderCompileWorker.cpp
+++ b/Engine/Source/Programs/ShaderCompileWorker/Private/ShaderCompileWorker.cpp
@@ -633,6 +633,7 @@ static void DirectCompile(const TArray<const class IShaderFormat*>& ShaderFormat
 	FName FormatName;
 	FString Entry = TEXT("Main");
 	bool bPipeline = false;
+	bool bUseMCPP = false;
 	EShaderFrequency Frequency = SF_Pixel;
 	TArray<FString> UsedOutputs;
 	bool bIncludeUsedOutputs = false;
@@ -681,6 +682,10 @@ static void DirectCompile(const TArray<const class IShaderFormat*>& ShaderFormat
 			{
 				bPipeline = true;
 			}
+			else if (!FCString::Strcmp(*Token, TEXT("mcpp")))
+			{
+				bUseMCPP = true;
+			}
 			else if (Token.StartsWith(TEXT("usedoutputs=")))
 			{
 				FString Outputs = Token.RightChop(12);
@@ -711,7 +716,7 @@ static void DirectCompile(const TArray<const class IShaderFormat*>& ShaderFormat
 	Input.SourceFilename = InputFile;
 	Input.Target.Platform =  FormatNameToEnum(FormatName);
 	Input.Target.Frequency = Frequency;
-	Input.bSkipPreprocessedCache = true;
+	Input.bSkipPreprocessedCache = !bUseMCPP;
 
 	auto AddResourceTableEntry = [](TMap<FString, FResourceTableEntry>& Map, const FString& Name, const FString& UBName, int32 Type, int32 ResourceIndex)
 	{
diff --git a/Engine/Source/Programs/UnrealBuildTool/Linux/LinuxToolChain.cs b/Engine/Source/Programs/UnrealBuildTool/Linux/LinuxToolChain.cs
--- a/Engine/Source/Programs/UnrealBuildTool/Linux/LinuxToolChain.cs
+++ b/Engine/Source/Programs/UnrealBuildTool/Linux/LinuxToolChain.cs
@@ -354,8 +354,8 @@ namespace UnrealBuildTool
 				// GCC only option
 				Result += " -fno-strict-aliasing";
 				Result += " -Wno-sign-compare"; // needed to suppress: comparison between signed and unsigned integer expressions
-				Result += " -Wno-enum-compare"; // Stats2.h triggers this (ALIGNOF(int64) <= DATA_ALIGN)
-				Result += " -Wno-return-type"; // Variant.h triggers this
+				Result += " -Wno-enum-compare"; // Stats2.h triggers this (alignof(int64) <= DATA_ALIGN)
+                Result += " -Wno-return-type"; // Variant.h triggers this
 				Result += " -Wno-unused-local-typedefs";
 				Result += " -Wno-multichar";
 				Result += " -Wno-unused-but-set-variable";
diff --git a/Engine/Source/Programs/UnrealHeaderTool/Private/CodeGenerator.cpp b/Engine/Source/Programs/UnrealHeaderTool/Private/CodeGenerator.cpp
--- a/Engine/Source/Programs/UnrealHeaderTool/Private/CodeGenerator.cpp
+++ b/Engine/Source/Programs/UnrealHeaderTool/Private/CodeGenerator.cpp
@@ -2752,7 +2752,7 @@ void FNativeClassHeaderGenerator::ExportGeneratedStructBodyMacros(FOutputDevice&
 	}
 	else
 	{
-		ExplicitSizeString = FString::Printf(TEXT(", sizeof(%s), ALIGNOF(%s)"), NameLookupCPP.GetNameCPP(Struct), NameLookupCPP.GetNameCPP(Struct));
+		ExplicitSizeString = FString::Printf(TEXT(", sizeof(%s), alignof(%s)"), NameLookupCPP.GetNameCPP(Struct), NameLookupCPP.GetNameCPP(Struct));
 	}
 
 	const TCHAR* UStructObjectFlags = bIsDynamic ? TEXT("RF_Public|RF_Transient") : TEXT("RF_Public|RF_Transient|RF_MarkAsNative");
diff --git a/Engine/Source/Runtime/Apple/MetalRHI/Private/MetalBlitCommandEncoder.cpp b/Engine/Source/Runtime/Apple/MetalRHI/Private/MetalBlitCommandEncoder.cpp
--- a/Engine/Source/Runtime/Apple/MetalRHI/Private/MetalBlitCommandEncoder.cpp
+++ b/Engine/Source/Runtime/Apple/MetalRHI/Private/MetalBlitCommandEncoder.cpp
@@ -14,6 +14,8 @@ extern int32 GMetalRuntimeDebugLevel;
 
 @implementation FMetalDebugBlitCommandEncoder
 
+APPLE_PLATFORM_OBJECT_ALLOC_OVERRIDES(FMetalDebugBlitCommandEncoder)
+
 @synthesize Inner;
 @synthesize Buffer;
 
diff --git a/Engine/Source/Runtime/Apple/MetalRHI/Private/MetalComputeCommandEncoder.cpp b/Engine/Source/Runtime/Apple/MetalRHI/Private/MetalComputeCommandEncoder.cpp
--- a/Engine/Source/Runtime/Apple/MetalRHI/Private/MetalComputeCommandEncoder.cpp
+++ b/Engine/Source/Runtime/Apple/MetalRHI/Private/MetalComputeCommandEncoder.cpp
@@ -15,6 +15,8 @@ NS_ASSUME_NONNULL_BEGIN
 
 @implementation FMetalDebugComputeCommandEncoder
 
+APPLE_PLATFORM_OBJECT_ALLOC_OVERRIDES(FMetalDebugComputeCommandEncoder)
+
 @synthesize Inner;
 @synthesize Buffer;
 @synthesize Pipeline;
diff --git a/Engine/Source/Runtime/Apple/MetalRHI/Private/MetalContext.cpp b/Engine/Source/Runtime/Apple/MetalRHI/Private/MetalContext.cpp
--- a/Engine/Source/Runtime/Apple/MetalRHI/Private/MetalContext.cpp
+++ b/Engine/Source/Runtime/Apple/MetalRHI/Private/MetalContext.cpp
@@ -862,7 +862,7 @@ FMetalRHICommandContext* FMetalDeviceContext::AcquireContext(int32 NewIndex, int
 	
 	if (NewIndex == 0)
 	{
-		if (FRHICommandListExecutor::GetImmediateCommandList().Bypass() || !GRHIThread)
+		if (FRHICommandListExecutor::GetImmediateCommandList().Bypass() || !IsRunningRHIInSeparateThread())
 		{
 			FMetalRHICommandUpdateFence UpdateCommand(this, StartFence, FMetalRHICommandUpdateFence::End);
 			UpdateCommand.Execute(FRHICommandListExecutor::GetImmediateCommandList());
@@ -873,7 +873,7 @@ FMetalRHICommandContext* FMetalDeviceContext::AcquireContext(int32 NewIndex, int
 		}
 	}
 	
-	if (FRHICommandListExecutor::GetImmediateCommandList().Bypass() || !GRHIThread)
+	if (FRHICommandListExecutor::GetImmediateCommandList().Bypass() || !IsRunningRHIInSeparateThread())
 	{
 		FMetalRHICommandUpdateFence UpdateCommand(this, EndFence, FMetalRHICommandUpdateFence::Start);
 		UpdateCommand.Execute(FRHICommandListExecutor::GetImmediateCommandList());
diff --git a/Engine/Source/Runtime/Apple/MetalRHI/Private/MetalProfiler.cpp b/Engine/Source/Runtime/Apple/MetalRHI/Private/MetalProfiler.cpp
--- a/Engine/Source/Runtime/Apple/MetalRHI/Private/MetalProfiler.cpp
+++ b/Engine/Source/Runtime/Apple/MetalRHI/Private/MetalProfiler.cpp
@@ -254,6 +254,11 @@ void FMetalEventNode::StopDraw(void)
 	}
 }
 #endif
+
+bool MetalGPUProfilerIsInSafeThread()
+{
+	return IsInRHIThread() || IsInActualRenderingThread();
+}
 	
 /** Start this frame of per tracking */
 void FMetalEventNodeFrame::StartFrame()
@@ -295,12 +300,18 @@ void FMetalGPUProfiler::Cleanup()
 
 void FMetalGPUProfiler::PushEvent(const TCHAR* Name, FColor Color)
 {
-	FGPUProfiler::PushEvent(Name, Color);
+	if(MetalGPUProfilerIsInSafeThread())
+	{
+		FGPUProfiler::PushEvent(Name, Color);
+	}
 }
 
 void FMetalGPUProfiler::PopEvent()
 {
-	FGPUProfiler::PopEvent();
+	if(MetalGPUProfilerIsInSafeThread())
+	{
+		FGPUProfiler::PopEvent();
+	}
 }
 
 //TGlobalResource<FVector4VertexDeclaration> GMetalVector4VertexDeclaration;
@@ -314,6 +325,11 @@ void FMetalGPUProfiler::BeginFrame()
 		CurrentEventNodeFrame = new FMetalEventNodeFrame(Context, GTriggerGPUProfile);
 		CurrentEventNodeFrame->StartFrame();
 		
+		if(GNumActiveGPUsForRendering > 1)
+		{
+			GTriggerGPUProfile = false;
+		}
+
 		if(GTriggerGPUProfile)
 		{
 			bTrackingEvents = true;
@@ -373,7 +389,10 @@ void FMetalGPUProfiler::BeginFrame()
 			}*/
 		}
 		
-		PushEvent(TEXT("FRAME"), FColor(0, 255, 0, 255));
+		if(GEmitDrawEvents)
+		{
+			PushEvent(TEXT("FRAME"), FColor(0, 255, 0, 255));
+		}
 	}
 	NumNestedFrames++;
 }
@@ -382,7 +401,10 @@ void FMetalGPUProfiler::EndFrame()
 {
 	if(--NumNestedFrames == 0)
 	{
-		PopEvent();
+		if(GEmitDrawEvents)
+		{
+			PopEvent();
+		}
 		
 #if PLATFORM_MAC
 		FPlatformMisc::UpdateDriverMonitorStatistics(GetMetalDeviceContext().GetDeviceIndex());
@@ -444,7 +466,7 @@ void FMetalGPUProfiler::EndFrame()
 
 void FMetalGPUProfiler::StartGPUWork(uint32 StartPoint, uint32 EndPoint, uint32 NumPrimitives, uint32 NumVertices)
 {
-	if(CurrentEventNode)
+	if(CurrentEventNode && MetalGPUProfilerIsInSafeThread())
 	{
 		RegisterGPUWork(NumPrimitives, NumVertices);
 #if METAL_STATISTICS
@@ -457,7 +479,7 @@ void FMetalGPUProfiler::StartGPUWork(uint32 StartPoint, uint32 EndPoint, uint32
 void FMetalGPUProfiler::FinishGPUWork(void)
 {
 #if METAL_STATISTICS
-	if(CurrentEventNode)
+	if(CurrentEventNode && MetalGPUProfilerIsInSafeThread())
 	{
 		FMetalEventNode* EventNode = (FMetalEventNode*)CurrentEventNode;
 		EventNode->StopDraw();
diff --git a/Engine/Source/Runtime/Apple/MetalRHI/Private/MetalRenderCommandEncoder.cpp b/Engine/Source/Runtime/Apple/MetalRHI/Private/MetalRenderCommandEncoder.cpp
--- a/Engine/Source/Runtime/Apple/MetalRHI/Private/MetalRenderCommandEncoder.cpp
+++ b/Engine/Source/Runtime/Apple/MetalRHI/Private/MetalRenderCommandEncoder.cpp
@@ -15,6 +15,8 @@ NS_ASSUME_NONNULL_BEGIN
 
 @implementation FMetalDebugRenderCommandEncoder
 
+APPLE_PLATFORM_OBJECT_ALLOC_OVERRIDES(FMetalDebugRenderCommandEncoder)
+
 @synthesize Inner;
 @synthesize Buffer;
 @synthesize Pipeline;
diff --git a/Engine/Source/Runtime/Apple/MetalRHI/Private/MetalRenderPipelineDesc.cpp b/Engine/Source/Runtime/Apple/MetalRHI/Private/MetalRenderPipelineDesc.cpp
--- a/Engine/Source/Runtime/Apple/MetalRHI/Private/MetalRenderPipelineDesc.cpp
+++ b/Engine/Source/Runtime/Apple/MetalRHI/Private/MetalRenderPipelineDesc.cpp
@@ -206,7 +206,7 @@ FMetalShaderPipeline* FMetalRenderPipelineDesc::CreatePipelineStateForBoundShade
 	
 	NSError* Error = nil;
 	
-	if(GUseRHIThread)
+	if (IsRunningRHIInSeparateThread())
 	{
 		SCOPE_CYCLE_COUNTER(STAT_MetalPipelineLockTime);
 		int Err = pthread_rwlock_rdlock(&MetalPipelineMutex.Mutex);
@@ -224,7 +224,7 @@ FMetalShaderPipeline* FMetalRenderPipelineDesc::CreatePipelineStateForBoundShade
 	FMetalShaderPipeline* statePack = MetalPipelineCache.FindRef(ComparableDesc);
 	if(statePack == nil)
 	{
-		if(GUseRHIThread)
+		if(IsRunningRHIInSeparateThread())
 		{
 			SCOPE_CYCLE_COUNTER(STAT_MetalPipelineLockTime);
 			int Err = pthread_rwlock_unlock(&MetalPipelineMutex.Mutex);
@@ -540,7 +540,7 @@ FMetalShaderPipeline* FMetalRenderPipelineDesc::CreatePipelineStateForBoundShade
 			statePack.FragmentSource = BSS->PixelShader ? BSS->PixelShader->GlslCodeNSString : nil;
 #endif
 			
-			if(GUseRHIThread)
+			if(IsRunningRHIInSeparateThread())
 			{
 				SCOPE_CYCLE_COUNTER(STAT_MetalPipelineLockTime);
 				int Err = pthread_rwlock_wrlock(&MetalPipelineMutex.Mutex);
@@ -558,7 +558,7 @@ FMetalShaderPipeline* FMetalRenderPipelineDesc::CreatePipelineStateForBoundShade
 				statePack = ExistingPipeline;
 			}
 			
-			if(GUseRHIThread)
+			if(IsRunningRHIInSeparateThread())
 			{
 				SCOPE_CYCLE_COUNTER(STAT_MetalPipelineLockTime);
 				int Err = pthread_rwlock_unlock(&MetalPipelineMutex.Mutex);
@@ -585,7 +585,7 @@ FMetalShaderPipeline* FMetalRenderPipelineDesc::CreatePipelineStateForBoundShade
 			return nil;
 		}
 	}
-	else if(GUseRHIThread)
+	else if(IsRunningRHIInSeparateThread())
 	{
 		SCOPE_CYCLE_COUNTER(STAT_MetalPipelineLockTime);
 		int Err = pthread_rwlock_unlock(&MetalPipelineMutex.Mutex);
diff --git a/Engine/Source/Runtime/Apple/MetalRHI/Private/MetalShaders.cpp b/Engine/Source/Runtime/Apple/MetalRHI/Private/MetalShaders.cpp
--- a/Engine/Source/Runtime/Apple/MetalRHI/Private/MetalShaders.cpp
+++ b/Engine/Source/Runtime/Apple/MetalRHI/Private/MetalShaders.cpp
@@ -14,6 +14,7 @@
 #include "MetalCommandBuffer.h"
 #include "Serialization/MemoryReader.h"
 #include "Misc/FileHelper.h"
+#include "ScopeRWLock.h"
 
 #define SHADERCOMPILERCOMMON_API
 #	include "Developer/ShaderCompilerCommon/Public/ShaderCompilerCommon.h"
@@ -56,14 +57,10 @@ struct FMetalCompiledShaderCache
 public:
 	FMetalCompiledShaderCache()
 	{
-		int Err = pthread_rwlock_init(&Lock, nullptr);
-		checkf(Err == 0, TEXT("pthread_rwlock_init failed with error: %d"), Err);
 	}
 	
 	~FMetalCompiledShaderCache()
 	{
-		int Err = pthread_rwlock_destroy(&Lock);
-		checkf(Err == 0, TEXT("pthread_rwlock_destroy failed with error: %d"), Err);
 		for (TPair<FMetalCompiledShaderKey, id<MTLFunction>> Pair : Cache)
 		{
 			[Pair.Value release];
@@ -72,25 +69,19 @@ public:
 	
 	id<MTLFunction> FindRef(FMetalCompiledShaderKey Key)
 	{
-		int Err = pthread_rwlock_rdlock(&Lock);
-		checkf(Err == 0, TEXT("pthread_rwlock_rdlock failed with error: %d"), Err);
+		FRWScopeLock(Lock, SLT_ReadOnly);
 		id<MTLFunction> Func = Cache.FindRef(Key);
-		Err = pthread_rwlock_unlock(&Lock);
-		checkf(Err == 0, TEXT("pthread_rwlock_unlock failed with error: %d"), Err);
 		return Func;
 	}
 	
 	void Add(FMetalCompiledShaderKey Key, id<MTLFunction> Function)
 	{
-		int Err = pthread_rwlock_wrlock(&Lock);
-		checkf(Err == 0, TEXT("pthread_rwlock_wrlock failed with error: %d"), Err);
+		FRWScopeLock(Lock, SLT_Write);
 		Cache.Add(Key, Function);
-		Err = pthread_rwlock_unlock(&Lock);
-		checkf(Err == 0, TEXT("pthread_rwlock_unlock failed with error: %d"), Err);
 	}
 	
 private:
-	pthread_rwlock_t Lock;
+	FRWLock Lock;
 	TMap<FMetalCompiledShaderKey, id<MTLFunction>> Cache;
 };
 
@@ -1053,7 +1044,7 @@ FMetalShaderPipeline* FMetalBoundShaderState::PrepareToDraw(FMetalHashedVertexDe
 	// generate a key for the current statez
 	FMetalRenderPipelineHash PipelineHash = RenderPipelineDesc.GetHash();
 	
-	if(GUseRHIThread)
+	if(IsRunningRHIInSeparateThread())
 	{
 		SCOPE_CYCLE_COUNTER(STAT_MetalBoundShaderLockTime);
 		int Err = pthread_rwlock_rdlock(&PipelineMutex);
@@ -1068,7 +1059,7 @@ FMetalShaderPipeline* FMetalBoundShaderState::PrepareToDraw(FMetalHashedVertexDe
 		PipelineStatePack = Dict->FindRef(VertexDesc);
 	}
 	
-	if(GUseRHIThread)
+	if(IsRunningRHIInSeparateThread())
 	{
 		SCOPE_CYCLE_COUNTER(STAT_MetalBoundShaderLockTime);
 		int Err = pthread_rwlock_unlock(&PipelineMutex);
@@ -1081,7 +1072,7 @@ FMetalShaderPipeline* FMetalBoundShaderState::PrepareToDraw(FMetalHashedVertexDe
 		PipelineStatePack = RenderPipelineDesc.CreatePipelineStateForBoundShaderState(this, VertexDesc);
 		check(PipelineStatePack);
 		
-		if(GUseRHIThread)
+		if(IsRunningRHIInSeparateThread())
 		{
 			SCOPE_CYCLE_COUNTER(STAT_MetalBoundShaderLockTime);
 			int Err = pthread_rwlock_wrlock(&PipelineMutex);
@@ -1107,7 +1098,7 @@ FMetalShaderPipeline* FMetalBoundShaderState::PrepareToDraw(FMetalHashedVertexDe
 			PipelineStatePack = ExistingPipeline;
 		}
 		
-		if(GUseRHIThread)
+		if(IsRunningRHIInSeparateThread())
 		{
 			SCOPE_CYCLE_COUNTER(STAT_MetalBoundShaderLockTime);
 			int Err = pthread_rwlock_unlock(&PipelineMutex);
diff --git a/Engine/Source/Runtime/Apple/MetalRHI/Private/MetalState.cpp b/Engine/Source/Runtime/Apple/MetalRHI/Private/MetalState.cpp
--- a/Engine/Source/Runtime/Apple/MetalRHI/Private/MetalState.cpp
+++ b/Engine/Source/Runtime/Apple/MetalRHI/Private/MetalState.cpp
@@ -306,7 +306,7 @@ FMetalBlendState::FMetalBlendState(const FBlendStateInitializerRHI& Initializer)
 			(BlendState.writeMask << 22);
 		
 		
-		if(GUseRHIThread)
+		if(IsRunningRHIInSeparateThread())
 		{
 			Mutex.Lock();
 		}
@@ -320,7 +320,7 @@ FMetalBlendState::FMetalBlendState(const FBlendStateInitializerRHI& Initializer)
 		}
 		// set the key
 		RenderTargetStates[RenderTargetIndex].BlendStateKey = *Key;
-		if(GUseRHIThread)
+		if(IsRunningRHIInSeparateThread())
 		{
 			Mutex.Unlock();
 		}
diff --git a/Engine/Source/Runtime/Apple/MetalRHI/Private/MetalStateCache.cpp b/Engine/Source/Runtime/Apple/MetalRHI/Private/MetalStateCache.cpp
--- a/Engine/Source/Runtime/Apple/MetalRHI/Private/MetalStateCache.cpp
+++ b/Engine/Source/Runtime/Apple/MetalRHI/Private/MetalStateCache.cpp
@@ -411,16 +411,16 @@ bool FMetalStateCache::SetRenderTargetsInfo(FRHISetRenderTargetsInfo const& InRe
 		// Deferred store actions make life a bit easier...
 		static bool bSupportsDeferredStore = GetMetalDeviceContext().GetCommandQueue().SupportsFeature(EMetalFeaturesDeferredStoreActions);
 		
-		if (bSupportsDeferredStore)
+		//Create local store action states if we support deferred store 
+		MTLStoreAction NewColorStore[MaxMetalRenderTargets];
+		for (uint32 i = 0; i < MaxMetalRenderTargets; ++i)
 		{
-			for (uint32 i = 0; i < MaxMetalRenderTargets; i++)
-			{
-				ColorStore[i] = MTLStoreActionUnknown;
-			}
-			DepthStore = MTLStoreActionUnknown;
-			StencilStore = MTLStoreActionUnknown;
+			NewColorStore[i] = MTLStoreActionUnknown;
 		}
 		
+		MTLStoreAction NewDepthStore = MTLStoreActionUnknown;
+		MTLStoreAction NewStencilStore = MTLStoreActionUnknown;
+		
 		// back this up for next frame
 		RenderTargetsInfo = InRenderTargets;
 		
@@ -533,8 +533,8 @@ bool FMetalStateCache::SetRenderTargetsInfo(FRHISetRenderTargetsInfo const& InRe
 				{
 					// set up an MSAA attachment
 					ColorAttachment.texture = Surface.MSAATexture;
-					ColorStore[RenderTargetIndex] = MTLStoreActionMultisampleResolve;
-					ColorAttachment.storeAction = bSupportsDeferredStore ? MTLStoreActionUnknown : ColorStore[RenderTargetIndex];
+					NewColorStore[RenderTargetIndex] = MTLStoreActionMultisampleResolve;
+					ColorAttachment.storeAction = bSupportsDeferredStore ? MTLStoreActionUnknown : NewColorStore[RenderTargetIndex];
 					ColorAttachment.resolveTexture = Surface.Texture;
 					PipelineDesc.SampleCount = Surface.MSAATexture.sampleCount;
 	
@@ -545,8 +545,8 @@ bool FMetalStateCache::SetRenderTargetsInfo(FRHISetRenderTargetsInfo const& InRe
 				{
 					// set up non-MSAA attachment
 					ColorAttachment.texture = Surface.Texture;
-					ColorStore[RenderTargetIndex] = GetMetalRTStoreAction(RenderTargetView.StoreAction);
-					ColorAttachment.storeAction = bSupportsDeferredStore ? MTLStoreActionUnknown : ColorStore[RenderTargetIndex];
+					NewColorStore[RenderTargetIndex] = GetMetalRTStoreAction(RenderTargetView.StoreAction);
+					ColorAttachment.storeAction = bSupportsDeferredStore ? MTLStoreActionUnknown : NewColorStore[RenderTargetIndex];
 					PipelineDesc.SampleCount = 1;
 				}
 				
@@ -744,8 +744,8 @@ bool FMetalStateCache::SetRenderTargetsInfo(FRHISetRenderTargetsInfo const& InRe
 				
 				bNeedsClear |= (DepthAttachment.loadAction == MTLLoadActionClear);
 				
-				DepthStore = GetMetalRTStoreAction(RenderTargetsInfo.DepthStencilRenderTarget.DepthStoreAction);
-				DepthAttachment.storeAction = bSupportsDeferredStore ? MTLStoreActionUnknown : DepthStore;
+				NewDepthStore = GetMetalRTStoreAction(RenderTargetsInfo.DepthStencilRenderTarget.DepthStoreAction);
+				DepthAttachment.storeAction = bSupportsDeferredStore ? MTLStoreActionUnknown : NewDepthStore;
 				DepthAttachment.clearDepth = DepthClearValue;
 
 				PipelineDesc.PipelineDescriptor.depthAttachmentPixelFormat = DepthAttachment.texture.pixelFormat;
@@ -756,7 +756,7 @@ bool FMetalStateCache::SetRenderTargetsInfo(FRHISetRenderTargetsInfo const& InRe
 				
 				bHasValidRenderTarget = true;
 
-				bCanRestartRenderPass &= (PipelineDesc.SampleCount <= 1) && ((RenderTargetsInfo.DepthStencilRenderTarget.Texture == FallbackDepthStencilSurface) || ((DepthAttachment.loadAction == MTLLoadActionLoad) && (!RenderTargetsInfo.DepthStencilRenderTarget.GetDepthStencilAccess().IsDepthWrite() || (RenderTargetsInfo.DepthStencilRenderTarget.DepthStoreAction == ERenderTargetStoreAction::EStore))));
+				bCanRestartRenderPass &= (PipelineDesc.SampleCount <= 1) && ((RenderTargetsInfo.DepthStencilRenderTarget.Texture == FallbackDepthStencilSurface) || ((DepthAttachment.loadAction == MTLLoadActionLoad) && (!RenderTargetsInfo.DepthStencilRenderTarget.GetDepthStencilAccess().IsDepthWrite() || (RenderTargetsInfo.DepthStencilRenderTarget.DepthStoreAction == ERenderTargetStoreAction::ENoAction) || (RenderTargetsInfo.DepthStencilRenderTarget.DepthStoreAction == ERenderTargetStoreAction::EStore))));
 				
 				// and assign it
 				RenderPass.depthAttachment = DepthAttachment;
@@ -778,8 +778,8 @@ bool FMetalStateCache::SetRenderTargetsInfo(FRHISetRenderTargetsInfo const& InRe
 				
 				bNeedsClear |= (StencilAttachment.loadAction == MTLLoadActionClear);
 				
-				StencilStore = GetMetalRTStoreAction(RenderTargetsInfo.DepthStencilRenderTarget.GetStencilStoreAction());
-				StencilAttachment.storeAction = bSupportsDeferredStore ? MTLStoreActionUnknown : StencilStore;
+				NewStencilStore = GetMetalRTStoreAction(RenderTargetsInfo.DepthStencilRenderTarget.GetStencilStoreAction());
+				StencilAttachment.storeAction = bSupportsDeferredStore ? MTLStoreActionUnknown : NewStencilStore;
 				StencilAttachment.clearStencil = StencilClearValue;
 
 				PipelineDesc.PipelineDescriptor.stencilAttachmentPixelFormat = StencilAttachment.texture.pixelFormat;
@@ -802,6 +802,17 @@ bool FMetalStateCache::SetRenderTargetsInfo(FRHISetRenderTargetsInfo const& InRe
 			}
 		}
 		
+		//Update deferred store states if required otherwise they're already set directly on the Metal Attachement Descriptors
+		if (bSupportsDeferredStore)
+		{
+			for (uint32 i = 0; i < MaxMetalRenderTargets; ++i)
+			{
+				ColorStore[i] = NewColorStore[i];
+			}
+			DepthStore = NewDepthStore;
+			StencilStore = NewStencilStore;
+		}
+		
 		bHasValidRenderTarget |= (InRenderTargets.NumUAVs > 0);
 		if (PipelineDesc.SampleCount == 0)
 		{
@@ -1556,16 +1567,6 @@ void FMetalStateCache::SetRenderStoreActions(FMetalCommandEncoder& CommandEncode
 
 void FMetalStateCache::SetRenderState(FMetalCommandEncoder& CommandEncoder, FMetalCommandEncoder* PrologueEncoder)
 {
-    if (RasterBits & EMetalRenderFlagPipelineState)
-    {
-    	check(PipelineState);
-        CommandEncoder.SetRenderPipelineState(PipelineState);
-		if (PipelineState.ComputePipelineState)
-		{
-			check(PrologueEncoder);
-    		PrologueEncoder->SetComputePipelineState(PipelineState);
-    	}
-    }
     if (RasterBits & EMetalRenderFlagViewport)
     {
         CommandEncoder.SetViewport(Viewport);
@@ -1610,6 +1611,17 @@ void FMetalStateCache::SetRenderState(FMetalCommandEncoder& CommandEncoder, FMet
     {
         CommandEncoder.SetVisibilityResultMode(VisibilityMode, VisibilityOffset);
     }
+	// Some Intel drivers need RenderPipeline state to be set after DepthStencil state to work properly
+	if (RasterBits & EMetalRenderFlagPipelineState)
+	{
+		check(PipelineState);
+		CommandEncoder.SetRenderPipelineState(PipelineState);
+		if (PipelineState.ComputePipelineState)
+		{
+			check(PrologueEncoder);
+			PrologueEncoder->SetComputePipelineState(PipelineState);
+		}
+	}
    	RasterBits = 0;
 }
 
diff --git a/Engine/Source/Runtime/Apple/MetalRHI/Private/MetalTexture.cpp b/Engine/Source/Runtime/Apple/MetalRHI/Private/MetalTexture.cpp
--- a/Engine/Source/Runtime/Apple/MetalRHI/Private/MetalTexture.cpp
+++ b/Engine/Source/Runtime/Apple/MetalRHI/Private/MetalTexture.cpp
@@ -1481,7 +1481,7 @@ void FMetalSurface::AsyncUnlock(class FRHICommandListImmediate& RHICmdList, uint
 	bool bDoDirectUnlock = Params.bDirectLock;
 	const bool bUnlockForCreate = Params.bCreateLock;
 			
-	if (RHICmdList.Bypass() || !GRHIThread || bDoDirectUnlock)
+	if (RHICmdList.Bypass() || !IsRunningRHIInSeparateThread() || bDoDirectUnlock)
 	{
 		if (bDoDirectUnlock)
 			{
@@ -1810,7 +1810,7 @@ FTexture2DRHIRef FMetalDynamicRHI::AsyncReallocateTexture2D_RenderThread(class F
 	@autoreleasepool {
 	FTexture2DRHIRef Result;
 	
-	if (RHICmdList.Bypass() || !GRHIThread)
+	if (RHICmdList.Bypass() || !IsRunningRHIInSeparateThread())
 	{
 		Result = GDynamicRHI->RHIAsyncReallocateTexture2D(Texture2D, NewMipCount, NewSizeX, NewSizeY, RequestStatus);
 	}
@@ -1978,7 +1978,7 @@ struct FMetalRHICommandUpdateTexture2D : public FRHICommand<FMetalRHICommandUpda
 void FMetalDynamicRHI::UpdateTexture2D_RenderThread(class FRHICommandListImmediate& RHICmdList, FTexture2DRHIParamRef Texture, uint32 MipIndex, const struct FUpdateTextureRegion2D& UpdateRegion, uint32 SourcePitch, const uint8* SourceData)
 {
 	@autoreleasepool {
-	if (RHICmdList.Bypass() || !GRHIThread)
+	if (RHICmdList.Bypass() || !IsRunningRHIInSeparateThread())
 	{
 	this->RHIUpdateTexture2D(Texture, MipIndex, UpdateRegion, SourcePitch, SourceData);
 }
@@ -2363,7 +2363,7 @@ void FMetalDynamicRHI::RHISetResourceAliasability_RenderThread(class FRHICommand
 			}
 			case EResourceAliasability::EUnaliasable:
 			{
-				if (RHICmdList.Bypass() || !GRHIThread)
+				if (RHICmdList.Bypass() || !IsRunningRHIInSeparateThread())
 				{
 					for (int32 i = 0; i < NumTextures; ++i)
 					{
diff --git a/Engine/Source/Runtime/Apple/MetalRHI/Private/MetalUniformBuffer.cpp b/Engine/Source/Runtime/Apple/MetalRHI/Private/MetalUniformBuffer.cpp
--- a/Engine/Source/Runtime/Apple/MetalRHI/Private/MetalUniformBuffer.cpp
+++ b/Engine/Source/Runtime/Apple/MetalRHI/Private/MetalUniformBuffer.cpp
@@ -100,7 +100,7 @@ void InitFrame_UniformBufferPoolCleanup()
 
 	SCOPE_CYCLE_COUNTER(STAT_MetalUniformBufferCleanupTime);
 	
-	if(GUseRHIThread)
+	if(IsRunningRHIInSeparateThread())
 	{
 		GMutex.Lock();
 	}
@@ -125,7 +125,7 @@ void InitFrame_UniformBufferPoolCleanup()
 		SafeUniformBufferPools[SafeFrameIndex][BucketIndex].Reset();
 	}
 	
-	if(GUseRHIThread)
+	if(IsRunningRHIInSeparateThread())
 	{
 		GMutex.Unlock();
 	}
@@ -135,7 +135,7 @@ void AddNewlyFreedBufferToUniformBufferPool(id<MTLBuffer> Buffer, uint32 Offset,
 {
 	check(Buffer);
 
-	if(GUseRHIThread)
+	if(IsRunningRHIInSeparateThread())
 	{
 		GMutex.Lock();
 	}
@@ -154,7 +154,7 @@ void AddNewlyFreedBufferToUniformBufferPool(id<MTLBuffer> Buffer, uint32 Offset,
 	INC_DWORD_STAT(STAT_MetalNumFreeUniformBuffers);
 	INC_MEMORY_STAT_BY(STAT_MetalFreeUniformBufferMemory, Buffer.length);
 	
-	if(GUseRHIThread)
+	if(IsRunningRHIInSeparateThread())
 	{
 		GMutex.Unlock();
 	}
@@ -196,7 +196,7 @@ FMetalUniformBuffer::FMetalUniformBuffer(const void* Contents, const FRHIUniform
 			else
 			{
 				// for single use buffers, allocate from the ring buffer to avoid thrashing memory
-				if (Usage == UniformBuffer_SingleDraw && !GUseRHIThread) // @todo Make this properly RHIThread safe.
+				if (Usage == UniformBuffer_SingleDraw && !IsRunningRHIInSeparateThread()) // @todo Make this properly RHIThread safe.
 				{
 					// use a bit of the ring buffer
 					Offset = GetMetalDeviceContext().AllocateFromRingBuffer(Layout.ConstantBufferSize);
@@ -205,7 +205,7 @@ FMetalUniformBuffer::FMetalUniformBuffer(const void* Contents, const FRHIUniform
 				else
 				{
 					// Find the appropriate bucket based on size
-					if(GUseRHIThread)
+					if(IsRunningRHIInSeparateThread())
 					{
 						GMutex.Lock();
 					}
@@ -237,7 +237,7 @@ FMetalUniformBuffer::FMetalUniformBuffer(const void* Contents, const FRHIUniform
 					}
 #endif
 					
-					if(GUseRHIThread)
+					if(IsRunningRHIInSeparateThread())
 					{
 						GMutex.Unlock();
 					}
@@ -279,7 +279,7 @@ FMetalUniformBuffer::~FMetalUniformBuffer()
 	INC_DWORD_STAT_BY(STAT_MetalUniformMemFreed, Size);
 	
 	// don't need to free the ring buffer!
-	if (GIsRHIInitialized && Buffer != nil && !(Usage == UniformBuffer_SingleDraw && !GUseRHIThread))
+	if (GIsRHIInitialized && Buffer != nil && !(Usage == UniformBuffer_SingleDraw && !IsRunningRHIInSeparateThread()))
 	{
 		check(Size <= 65536);
 		AddNewlyFreedBufferToUniformBufferPool(Buffer, Offset, Size);
diff --git a/Engine/Source/Runtime/Core/Private/Apple/ApplePlatformDebugEvents.cpp b/Engine/Source/Runtime/Core/Private/Apple/ApplePlatformDebugEvents.cpp
--- a/Engine/Source/Runtime/Core/Private/Apple/ApplePlatformDebugEvents.cpp
+++ b/Engine/Source/Runtime/Core/Private/Apple/ApplePlatformDebugEvents.cpp
@@ -9,6 +9,7 @@
 #include "UnrealString.h"
 #include "IConsoleManager.h"
 #include "Color.h"
+#include "ScopeRWLock.h"
 #include <sys/kdebug_signpost.h>
 #include <sys/syscall.h>
 
@@ -170,23 +171,6 @@ TArray<FApplePlatformDebugEvents::FEvent>* FApplePlatformDebugEvents::GetEventSt
 	return Current;
 }
 
-struct FAppleRWMutex
-{
-	FAppleRWMutex()
-	{
-		int Err = pthread_rwlock_init(&Mutex, nullptr);
-		checkf(Err == 0, TEXT("pthread_rwlock_init failed with error: %d"), errno);
-	}
-	
-	~FAppleRWMutex()
-	{
-		int Err = pthread_rwlock_destroy(&Mutex);
-		checkf(Err == 0, TEXT("pthread_rwlock_destroy failed with error: %d"), errno);
-	}
-	
-	pthread_rwlock_t Mutex;
-};
-
 uint16 FApplePlatformDebugEvents::GetEventCode(FString String)
 {
 	// Never emit 0 as we use that for the frame marker...
@@ -207,20 +191,19 @@ uint16 FApplePlatformDebugEvents::GetEventCode(FString String)
 	
 	uint32 Hash = GetTypeHash(String);
 	
-	static FAppleRWMutex RWMutex;
+	static FRWLock RWMutex;
 	static TMap<uint32, uint16> Names;
 	{
-		pthread_rwlock_rdlock(&RWMutex.Mutex);
+		FRWScopeLock Lock(RWMutex, SLT_ReadOnly);
 		uint16* CodePtr = Names.Find(Hash);
 		if(CodePtr)
 		{
 			Code = *CodePtr;
 		}
-		pthread_rwlock_unlock(&RWMutex.Mutex);
 		
 		if(!CodePtr)
 		{
-			pthread_rwlock_wrlock(&RWMutex.Mutex);
+			Lock.RaiseLockToWrite();
 			CodePtr = Names.Find(Hash);
 			if(CodePtr)
 			{
@@ -233,7 +216,6 @@ uint16 FApplePlatformDebugEvents::GetEventCode(FString String)
 				Names.Add(Hash, Code);
 				UE_LOG(LogInstruments, Display, TEXT("New Event Code: %u : %s"), (uint32)Code, *String);
 			}
-			pthread_rwlock_unlock(&RWMutex.Mutex);
 		}
 	}
 	
diff --git a/Engine/Source/Runtime/Core/Private/Async/TaskGraph.cpp b/Engine/Source/Runtime/Core/Private/Async/TaskGraph.cpp
--- a/Engine/Source/Runtime/Core/Private/Async/TaskGraph.cpp
+++ b/Engine/Source/Runtime/Core/Private/Async/TaskGraph.cpp
@@ -34,6 +34,8 @@ DEFINE_STAT(STAT_FTriggerEventGraphTask);
 DEFINE_STAT(STAT_ParallelFor);
 DEFINE_STAT(STAT_ParallelForTask);
 
+static int32 GNumWorkerThreadsToIgnore = 0;
+
 #if (PLATFORM_XBOXONE || PLATFORM_PS4 || PLATFORM_WINDOWS || PLATFORM_MAC || PLATFORM_LINUX) && !IS_PROGRAM && WITH_ENGINE && !UE_SERVER
 	#define CREATE_HIPRI_TASK_THREADS (1)
 	#define CREATE_BACKGROUND_TASK_THREADS (1)
@@ -50,14 +52,6 @@ namespace ENamedThreads
 	CORE_API int32 bHasHighPriorityThreads = CREATE_HIPRI_TASK_THREADS;
 }
 
-static int32 GNumWorkerThreadsToIgnore = 0;
-static FAutoConsoleVariableRef CVarNumWorkerThreadsToIgnore(
-	TEXT("TaskGraph.NumWorkerThreadsToIgnore"),
-	GNumWorkerThreadsToIgnore,
-	TEXT("Used to tune the number of task threads. Generally once you have found the right value, PlatformMisc::NumberOfWorkerThreadsToSpawn() should be hardcoded."),
-	ECVF_Cheat
-	);
-
 #if CREATE_HIPRI_TASK_THREADS || CREATE_BACKGROUND_TASK_THREADS
 	static void ThreadSwitchForABTest(const TArray<FString>& Args)
 	{
@@ -815,6 +809,19 @@ public:
 		Queue.StallRestartEvent->Trigger();
 	}
 
+	void StallForTuning(bool Stall)
+	{
+		if (Stall)
+		{
+			Queue.StallForTuning.Lock();
+			Queue.bStallForTuning = true;
+		}
+		else
+		{
+			Queue.bStallForTuning = false;
+			Queue.StallForTuning.Unlock();
+		}
+	}
 	/**
 	*Return true if this thread is processing tasks. This is only a "guess" if you ask for a thread other than yourself because that can change before the function returns.
 	*@param QueueIndex, Queue to request quit from
@@ -886,6 +893,26 @@ private:
 			TestRandomizedThreads();
 			Task->Execute(NewTasks, ENamedThreads::Type(ThreadId));
 			TestRandomizedThreads();
+			if (Queue.bStallForTuning)
+			{
+#if STATS
+				if (bTasksOpen)
+				{
+					ProcessingTasks.Stop();
+					bTasksOpen = false;
+				}
+#endif
+				{
+					FScopeLock Lock(&Queue.StallForTuning);
+				}
+#if STATS
+				if (FThreadStats::IsCollectingData(StatName))
+				{
+					bTasksOpen = true;
+					ProcessingTasks.Start(StatName);
+				}
+#endif
+			}
 		}
 		verify(!--Queue.RecursionGuard);
 	}
@@ -893,17 +920,21 @@ private:
 	/** Grouping of the data for an individual queue. **/
 	struct FThreadTaskQueue
 	{
+		/** Event that this thread blocks on when it runs out of work. **/
+		FEvent* StallRestartEvent;
 		/** We need to disallow reentry of the processing loop **/
 		uint32 RecursionGuard;
 		/** Indicates we executed a return task, so break out of the processing loop. **/
 		bool QuitForShutdown;
-		/** Event that this thread blocks on when it runs out of work. **/
-		FEvent*												StallRestartEvent;
+		/** Should we stall for tuning? **/
+		bool bStallForTuning;
+		FCriticalSection StallForTuning;
 
 		FThreadTaskQueue()
-			: RecursionGuard(0)
+			: StallRestartEvent(FPlatformProcess::GetSynchEventFromPool(false))
+			, RecursionGuard(0)
 			, QuitForShutdown(false)
-			, StallRestartEvent(FPlatformProcess::GetSynchEventFromPool(false))
+			, bStallForTuning(false)
 		{
 
 		}
@@ -1325,7 +1356,7 @@ public:
 	void StartTaskThread(int32 Priority, int32 IndexToStart)
 	{
 		ENamedThreads::Type ThreadToWake = ENamedThreads::Type(IndexToStart + Priority * NumTaskThreadsPerSet + NumNamedThreads);
-		((FTaskThreadAnyThread&)Thread(ThreadToWake)).FTaskThreadAnyThread::WakeUp();
+		((FTaskThreadAnyThread&)Thread(ThreadToWake)).WakeUp();
 	}
 	void StartAllTaskThreads(bool bDoBackgroundThreads)
 	{
@@ -1346,7 +1377,7 @@ public:
 
 	FBaseGraphTask* FindWork(ENamedThreads::Type ThreadInNeed)
 	{
-		int32 LocalNumWorkingThread = GetNumWorkerThreads();
+		int32 LocalNumWorkingThread = GetNumWorkerThreads() + GNumWorkerThreadsToIgnore;
 		int32 MyIndex = int32((uint32(ThreadInNeed) - NumNamedThreads) % NumTaskThreadsPerSet);
 		int32 Priority = int32((uint32(ThreadInNeed) - NumNamedThreads) / NumTaskThreadsPerSet);
 		check(MyIndex >= 0 && MyIndex < LocalNumWorkingThread &&
@@ -1356,6 +1387,14 @@ public:
 		return IncomingAnyThreadTasks[Priority].Pop(MyIndex, true);
 	}
 
+	void StallForTuning(int32 Index, bool Stall)
+	{
+		for (int32 Priority = 0; Priority < ENamedThreads::NumThreadPriorities; Priority++)
+		{
+			ENamedThreads::Type ThreadToWake = ENamedThreads::Type(Index + Priority * NumTaskThreadsPerSet + NumNamedThreads);
+			((FTaskThreadAnyThread&)Thread(ThreadToWake)).StallForTuning(Stall);
+		}
+	}
 	void SetTaskThreadPriorities(EThreadPriority Pri)
 	{
 		check(NumTaskThreadSets == 1); // otherwise tuning this doesn't make a lot of sense
@@ -1660,7 +1699,7 @@ void FTaskGraphInterface::BroadcastSlow_OnlyUseForSpecialPurposes(bool bDoTaskTh
 
 	FGraphEventArray Tasks;
 	STAT(Tasks.Add(TGraphTask<FBroadcastTask>::CreateTask().ConstructAndDispatchWhenReady(Callback, ENamedThreads::SetTaskPriority(ENamedThreads::StatsThread, ENamedThreads::HighTaskPriority), nullptr, nullptr, nullptr)););
-	if (GRHIThread)
+	if (GRHIThread_InternalUseOnly)
 	{
 		Tasks.Add(TGraphTask<FBroadcastTask>::CreateTask().ConstructAndDispatchWhenReady(Callback, ENamedThreads::SetTaskPriority(ENamedThreads::RHIThread, ENamedThreads::HighTaskPriority), nullptr, nullptr, nullptr));
 	}
@@ -1693,6 +1732,39 @@ void FTaskGraphInterface::BroadcastSlow_OnlyUseForSpecialPurposes(bool bDoTaskTh
 	}
 }
 
+static void HandleNumWorkerThreadsToIgnore(const TArray<FString>& Args)
+{
+	if (Args.Num() > 0)
+	{
+		int32 Arg = FCString::Atoi(*Args[0]);
+		int32 MaxNumPerBank = FTaskGraphInterface::Get().GetNumWorkerThreads() + GNumWorkerThreadsToIgnore;
+		if (Arg < MaxNumPerBank && Arg >= 0 && Arg != GNumWorkerThreadsToIgnore)
+		{
+			if (Arg > GNumWorkerThreadsToIgnore)
+			{
+				for (int32 Index = MaxNumPerBank - GNumWorkerThreadsToIgnore - 1; Index >= MaxNumPerBank - Arg; Index--)
+				{
+					FTaskGraphImplementation::Get().StallForTuning(Index, true);
+				}
+			}
+			else
+			{
+				for (int32 Index = MaxNumPerBank - Arg - 1; Index >= MaxNumPerBank - GNumWorkerThreadsToIgnore; Index--)
+				{
+					FTaskGraphImplementation::Get().StallForTuning(Index, false);
+				}
+			}
+			GNumWorkerThreadsToIgnore = Arg;
+		}
+	}
+	UE_LOG(LogConsoleResponse, Display, TEXT("Currently ignoring %d threads per priority bank"), GNumWorkerThreadsToIgnore);
+}
+
+static FAutoConsoleCommand CVarNumWorkerThreadsToIgnore(
+	TEXT("TaskGraph.NumWorkerThreadsToIgnore"),
+	TEXT("Used to tune the number of task threads. Generally once you have found the right value, PlatformMisc::NumberOfWorkerThreadsToSpawn() should be hardcoded."),
+	FConsoleCommandWithArgsDelegate::CreateStatic(&HandleNumWorkerThreadsToIgnore)
+	);
 
 // Benchmark
 
diff --git a/Engine/Source/Runtime/Core/Private/GenericPlatform/GenericPlatformAtomics.cpp b/Engine/Source/Runtime/Core/Private/GenericPlatform/GenericPlatformAtomics.cpp
--- a/Engine/Source/Runtime/Core/Private/GenericPlatform/GenericPlatformAtomics.cpp
+++ b/Engine/Source/Runtime/Core/Private/GenericPlatform/GenericPlatformAtomics.cpp
@@ -1,8 +1,7 @@
 // Copyright 1998-2017 Epic Games, Inc. All Rights Reserved.
 
 #include "GenericPlatform/GenericPlatformAtomics.h"
-#include "Templates/AlignOf.h"
 
 
 static_assert(sizeof(FInt128) == 16, "FInt128 size must be 16 bytes.");
-static_assert(ALIGNOF(FInt128) == 16, "FInt128 alignment must equals 16.");
+static_assert(alignof(FInt128) == 16, "FInt128 alignment must equals 16.");
diff --git a/Engine/Source/Runtime/Core/Private/GenericPlatform/GenericPlatformProcess.cpp b/Engine/Source/Runtime/Core/Private/GenericPlatform/GenericPlatformProcess.cpp
--- a/Engine/Source/Runtime/Core/Private/GenericPlatform/GenericPlatformProcess.cpp
+++ b/Engine/Source/Runtime/Core/Private/GenericPlatform/GenericPlatformProcess.cpp
@@ -25,6 +25,7 @@
 DEFINE_STAT(STAT_Sleep);
 DEFINE_STAT(STAT_EventWait);
 
+TArray<FString> FGenericPlatformProcess::ShaderDirs;
 
 void* FGenericPlatformProcess::GetDllHandle( const TCHAR* Filename )
 {
@@ -135,6 +136,25 @@ void FGenericPlatformProcess::SetShaderDir(const TCHAR*Where)
 	}
 }
 
+
+const TArray<FString>& FGenericPlatformProcess::AllShaderDirs()
+{
+	return ShaderDirs;
+}
+
+void FGenericPlatformProcess::AddShaderDir(const FString& InShaderDir)
+{
+	check(IsInGameThread());
+
+	FString ShaderDir = InShaderDir;
+	// make sure we store only relative paths
+	if (!FPaths::IsRelative(ShaderDir))
+	{
+		FPaths::MakePathRelativeTo(ShaderDir, FPlatformProcess::BaseDir());
+	}
+	ShaderDirs.Add(ShaderDir);
+}
+
 /**
  *	Get the shader working directory
  */
diff --git a/Engine/Source/Runtime/Core/Private/HAL/ThreadingBase.cpp b/Engine/Source/Runtime/Core/Private/HAL/ThreadingBase.cpp
--- a/Engine/Source/Runtime/Core/Private/HAL/ThreadingBase.cpp
+++ b/Engine/Source/Runtime/Core/Private/HAL/ThreadingBase.cpp
@@ -32,7 +32,7 @@ CORE_API FRunnableThread* GAudioThread = nullptr;
 CORE_API bool IsInAudioThread()
 {
 	// True if this is the audio thread or if there is no audio thread, then if it is the game thread
-	return (GAudioThreadId != 0 && FPlatformTLS::GetCurrentThreadId() == GAudioThreadId) || (GAudioThreadId == 0 && FPlatformTLS::GetCurrentThreadId() == GGameThreadId);
+	return FPlatformTLS::GetCurrentThreadId() == (GAudioThreadId ? GAudioThreadId : GGameThreadId);
 }
 
 CORE_API int32 GIsRenderingThreadSuspended = 0;
@@ -54,11 +54,14 @@ CORE_API bool IsInParallelRenderingThread()
 	return !GRenderingThread || GIsRenderingThreadSuspended || (FPlatformTLS::GetCurrentThreadId() != GGameThreadId);
 }
 
+CORE_API uint32 GRHIThreadId = 0;
+CORE_API FRunnableThread* GRHIThread_InternalUseOnly = nullptr;
+
 CORE_API bool IsInRHIThread()
 {
-	return GRHIThread && FPlatformTLS::GetCurrentThreadId() == GRHIThread->GetThreadID();
+	return GRHIThreadId && FPlatformTLS::GetCurrentThreadId() == GRHIThreadId;
 }
-CORE_API FRunnableThread* GRHIThread = nullptr;
+
 // Fake threads
 
 // Core version of IsInAsyncLoadingThread
diff --git a/Engine/Source/Runtime/Core/Private/Linux/LinuxPlatformProcess.cpp b/Engine/Source/Runtime/Core/Private/Linux/LinuxPlatformProcess.cpp
--- a/Engine/Source/Runtime/Core/Private/Linux/LinuxPlatformProcess.cpp
+++ b/Engine/Source/Runtime/Core/Private/Linux/LinuxPlatformProcess.cpp
@@ -151,14 +151,17 @@ void FLinuxPlatformProcess::CleanFileCache()
 	if (bShouldCleanShaderWorkingDirectory && !FParse::Param( FCommandLine::Get(), TEXT("Multiprocess")))
 	{
 		// get shader path, and convert it to the userdirectory
-		FString ShaderDir = FString(FPlatformProcess::BaseDir()) / FPlatformProcess::ShaderDir();
-		FString UserShaderDir = IFileManager::Get().ConvertToAbsolutePathForExternalAppForWrite(*ShaderDir);
-		FPaths::CollapseRelativeDirectories(ShaderDir);
-
-		// make sure we don't delete from the source directory
-		if (ShaderDir != UserShaderDir)
+		for (FString Dir : FPlatformProcess::AllShaderDirs())
 		{
-			IFileManager::Get().DeleteDirectory(*UserShaderDir, false, true);
+			FString ShaderDir = FString(FPlatformProcess::BaseDir()) / Dir;
+			FString UserShaderDir = IFileManager::Get().ConvertToAbsolutePathForExternalAppForWrite(*ShaderDir);
+			FPaths::CollapseRelativeDirectories(ShaderDir);
+
+			// make sure we don't delete from the source directory
+			if (ShaderDir != UserShaderDir)
+			{
+				IFileManager::Get().DeleteDirectory(*UserShaderDir, false, true);
+			}
 		}
 
 		FPlatformProcess::CleanShaderWorkingDir();
diff --git a/Engine/Source/Runtime/Core/Private/Mac/MacPlatformProcess.cpp b/Engine/Source/Runtime/Core/Private/Mac/MacPlatformProcess.cpp
--- a/Engine/Source/Runtime/Core/Private/Mac/MacPlatformProcess.cpp
+++ b/Engine/Source/Runtime/Core/Private/Mac/MacPlatformProcess.cpp
@@ -9,6 +9,7 @@
 #include "Misc/App.h"
 #include "Misc/Paths.h"
 #include "MacApplication.h"
+#include "HAL/FileManager.h"
 #include <mach-o/dyld.h>
 #include <libproc.h>
 
@@ -755,10 +756,24 @@ void FMacPlatformProcess::CleanFileCache()
 	bShouldCleanShaderWorkingDirectory = GIsFirstInstance;
 #endif
 
-	if (bShouldCleanShaderWorkingDirectory && !FParse::Param( FCommandLine::Get(), TEXT("Multiprocess")))
-	{
-		FPlatformProcess::CleanShaderWorkingDir();
-	}
+    if (bShouldCleanShaderWorkingDirectory && !FParse::Param( FCommandLine::Get(), TEXT("Multiprocess")))
+    {
+        // get shader path, and convert it to the userdirectory
+        for (FString Dir : FPlatformProcess::AllShaderDirs())
+        {
+            FString ShaderDir = FString(FPlatformProcess::BaseDir()) / Dir;
+            FString UserShaderDir = IFileManager::Get().ConvertToAbsolutePathForExternalAppForWrite(*ShaderDir);
+            FPaths::CollapseRelativeDirectories(ShaderDir);
+            
+            // make sure we don't delete from the source directory
+            if (ShaderDir != UserShaderDir)
+            {
+                IFileManager::Get().DeleteDirectory(*UserShaderDir, false, true);
+            }
+        }
+        
+        FPlatformProcess::CleanShaderWorkingDir();
+    }
 }
 
 const TCHAR* FMacPlatformProcess::BaseDir()
