Copying //UE4/Orion-Staging to //UE4/Main (Source: //Orion/Dev-General @ 3483207)

diff --git a/Engine/Source/Editor/UnrealEd/Private/CookOnTheFlyServer.cpp b/Engine/Source/Editor/UnrealEd/Private/CookOnTheFlyServer.cpp
--- a/Engine/Source/Editor/UnrealEd/Private/CookOnTheFlyServer.cpp
+++ b/Engine/Source/Editor/UnrealEd/Private/CookOnTheFlyServer.cpp
@@ -660,6 +660,12 @@ const FString& GetAssetRegistryFilename()
 	return AssetRegistryFilename;
 }
 
+const FString& GetDevelopmentAssetRegistryFilename()
+{
+	static const FString DevelopmentAssetRegistryFilename = FString(TEXT("DevelopmentAssetRegistry.bin"));
+	return DevelopmentAssetRegistryFilename;
+}
+
 /**
  * Uses the FMessageLog to log a message
  * 
@@ -722,16 +728,20 @@ FName UCookOnTheFlyServer::GetCachedStandardPackageFileFName( const UPackage* Pa
 
 bool UCookOnTheFlyServer::ClearPackageFilenameCacheForPackage( const FName& PackageName ) const
 {
+	check(IsInGameThread());
 	return PackageFilenameCache.Remove( PackageName ) >= 1;
 }
 
 bool UCookOnTheFlyServer::ClearPackageFilenameCacheForPackage( const UPackage* Package ) const
 {
+	check(IsInGameThread());
 	return PackageFilenameCache.Remove( Package->GetFName() ) >= 1;
 }
 
 const FString& UCookOnTheFlyServer::GetCachedSandboxFilename( const UPackage* Package, TUniquePtr<class FSandboxPlatformFile>& InSandboxFile ) const 
 {
+	check(IsInGameThread());
+
 	FName PackageFName = Package->GetFName();
 	static TMap<FName, FString> CachedSandboxFilenames;
 	FString* CachedSandboxFilename = CachedSandboxFilenames.Find(PackageFName);
@@ -746,6 +756,7 @@ const FString& UCookOnTheFlyServer::GetCachedSandboxFilename( const UPackage* Pa
 
 const UCookOnTheFlyServer::FCachedPackageFilename& UCookOnTheFlyServer::Cache(const FName& PackageName) const 
 {
+	check( IsInGameThread() );
 	FCachedPackageFilename *Cached = PackageFilenameCache.Find( PackageName );
 	if ( Cached != NULL )
 	{
@@ -772,6 +783,7 @@ const UCookOnTheFlyServer::FCachedPackageFilename& UCookOnTheFlyServer::Cache(co
 
 const FName* UCookOnTheFlyServer::GetCachedPackageFilenameToPackageFName(const FName& StandardPackageFilename) const
 {
+	check(IsInGameThread());
 	const FName* Result = PackageFilenameToPackageFNameCache.Find(StandardPackageFilename);
 	if ( Result )
 	{
@@ -793,6 +805,7 @@ const FName* UCookOnTheFlyServer::GetCachedPackageFilenameToPackageFName(const F
 
 void UCookOnTheFlyServer::ClearPackageFilenameCache() const 
 {
+	check(IsInGameThread());
 	PackageFilenameCache.Empty();
 	PackageFilenameToPackageFNameCache.Empty();
 }
@@ -880,6 +893,56 @@ TStatId UCookOnTheFlyServer::GetStatId() const
 	RETURN_QUICK_DECLARE_CYCLE_STAT(UCookServer, STATGROUP_Tickables);
 }
 
+const TArray<ITargetPlatform*>& UCookOnTheFlyServer::GetCookingTargetPlatforms() const
+{
+	ITargetPlatformManagerModule& TPM = GetTargetPlatformManagerRef();
+
+	FString PlatformStr;
+	if ( FParse::Value(FCommandLine::Get(), TEXT("TARGETPLATFORM="), PlatformStr) == false )
+	{
+		FString ValueName = TEXT("DefaultTargetPlatform");
+
+		if ( IsCookingInEditor() )
+		{
+			ValueName += TEXT("Editor");
+		}
+		if ( IsCookOnTheFlyMode() )
+		{
+			ValueName += TEXT("OnTheFly");
+		}
+
+
+		TArray<FString> TargetPlatformNames;
+		// see if we have specified in an ini file which target platforms we should use
+		if ( GConfig->GetArray(TEXT("CookSettings"), *ValueName, TargetPlatformNames, GEditorIni) )
+		{
+			for ( const FString& TargetPlatformName : TargetPlatformNames )
+			{
+				ITargetPlatform* TargetPlatform = TPM.FindTargetPlatform(TargetPlatformName);
+
+				if ( TargetPlatform )
+				{
+					CookingTargetPlatforms.AddUnique(TargetPlatform);
+				}
+				else
+				{
+					UE_LOG(LogCook, Warning, TEXT("Unable to resolve targetplatform name %s"), *TargetPlatformName );
+				}
+			}
+		}
+	}
+
+
+	if ( CookingTargetPlatforms.Num() == 0 )
+	{
+		const TArray<ITargetPlatform*>& Platforms = TPM.GetCookingTargetPlatforms();
+
+		CookingTargetPlatforms = Platforms;
+	}
+
+	return CookingTargetPlatforms;
+}
+
 bool UCookOnTheFlyServer::StartNetworkFileServer( const bool BindAnyPort )
 {
 	check( IsCookOnTheFlyMode() );
@@ -890,13 +953,12 @@ bool UCookOnTheFlyServer::StartNetworkFileServer( const bool BindAnyPort )
 #endif
 	ValidateCookOnTheFlySettings();
 
-	ITargetPlatformManagerModule& TPM = GetTargetPlatformManagerRef();
-	const TArray<ITargetPlatform*>& Platforms = TPM.GetCookingTargetPlatforms();
-
 	GenerateAssetRegistry();
 
 	InitializeSandbox();
 
+	const TArray<ITargetPlatform*>& Platforms = GetCookingTargetPlatforms();
+
 	// When cooking on the fly the full registry is saved at the beginning
 	// in cook by the book asset registry is saved after the cook is finished
 	for (int32 Index = 0; Index < Platforms.Num(); Index++)
@@ -904,17 +966,30 @@ bool UCookOnTheFlyServer::StartNetworkFileServer( const bool BindAnyPort )
 		FAssetRegistryGenerator* Generator = RegistryGenerators.FindRef(FName(*Platforms[Index]->PlatformName()));
 		if (Generator)
 		{
-			Generator->SaveAssetRegistry(GetSandboxAssetRegistryFilename());
+			Generator->SaveAssetRegistry(GetSandboxAssetRegistryFilename(), false);
 		}
 	}
 
 	// start the listening thread
+	FNewConnectionDelegate NewConnectionDelegate(FNewConnectionDelegate::CreateUObject(this, &UCookOnTheFlyServer::HandleNetworkFileServerNewConnection));
 	FFileRequestDelegate FileRequestDelegate(FFileRequestDelegate::CreateUObject(this, &UCookOnTheFlyServer::HandleNetworkFileServerFileRequest));
 	FRecompileShadersDelegate RecompileShadersDelegate(FRecompileShadersDelegate::CreateUObject(this, &UCookOnTheFlyServer::HandleNetworkFileServerRecompileShaders));
 	FSandboxPathDelegate SandboxPathDelegate(FSandboxPathDelegate::CreateUObject(this, &UCookOnTheFlyServer::HandleNetworkGetSandboxPath));
+	FInitialPrecookedListDelegate InitialPrecookedListDelegate(FInitialPrecookedListDelegate::CreateUObject(this, &UCookOnTheFlyServer::HandleNetworkGetPrecookedList));
+
+
+	FNetworkFileDelegateContainer NetworkFileDelegateContainer;
+	NetworkFileDelegateContainer.NewConnectionDelegate = NewConnectionDelegate;
+	NetworkFileDelegateContainer.InitialPrecookedListDelegate = InitialPrecookedListDelegate;
+	NetworkFileDelegateContainer.FileRequestDelegate = FileRequestDelegate;
+	NetworkFileDelegateContainer.RecompileShadersDelegate = RecompileShadersDelegate;
+	NetworkFileDelegateContainer.SandboxPathOverrideDelegate = SandboxPathDelegate;
+	
+	NetworkFileDelegateContainer.OnFileModifiedCallback = &FileModifiedDelegate;
+
 
 	INetworkFileServer *TcpFileServer = FModuleManager::LoadModuleChecked<INetworkFileSystemModule>("NetworkFileSystem")
-		.CreateNetworkFileServer(true, BindAnyPort ? 0 : -1, &FileRequestDelegate, &RecompileShadersDelegate, &SandboxPathDelegate, &FileModifiedDelegate, ENetworkFileServerProtocol::NFSP_Tcp);
+		.CreateNetworkFileServer(true, BindAnyPort ? 0 : -1, NetworkFileDelegateContainer, ENetworkFileServerProtocol::NFSP_Tcp);
 	if ( TcpFileServer )
 	{
 		NetworkFileServers.Add(TcpFileServer);
@@ -923,7 +998,7 @@ bool UCookOnTheFlyServer::StartNetworkFileServer( const bool BindAnyPort )
 	// cookonthefly server for html5 -- NOTE: if this is crashing COTF servers, please ask for Nick.Shin (via Josh.Adams)
 #if 1
 	INetworkFileServer *HttpFileServer = FModuleManager::LoadModuleChecked<INetworkFileSystemModule>("NetworkFileSystem")
-		.CreateNetworkFileServer(true, BindAnyPort ? 0 : -1, &FileRequestDelegate, &RecompileShadersDelegate, &SandboxPathDelegate, &FileModifiedDelegate, ENetworkFileServerProtocol::NFSP_Http);
+		.CreateNetworkFileServer(true, BindAnyPort ? 0 : -1, NetworkFileDelegateContainer, ENetworkFileServerProtocol::NFSP_Http);
 	if ( HttpFileServer )
 	{
 		NetworkFileServers.Add( HttpFileServer );
@@ -1730,10 +1805,10 @@ uint32 UCookOnTheFlyServer::TickCookOnTheSide( const float TimeSlice, uint32 &Co
 
 		
 		// if we are cook on the fly (even in editor) and we have cook requests, try process them straight away
-		if ( IsCookOnTheFlyMode() && CookRequests.HasItems() )
+		/*if ( IsCookOnTheFlyMode() && CookRequests.HasItems() )
 		{
 			continue;
-		}
+		}*/
 		
 		if ( Timer.IsTimeUp() )
 		{
@@ -2077,7 +2152,7 @@ bool UCookOnTheFlyServer::SaveCookedPackages(TArray<UPackage*>& PackagesToSave,
 
 			bool bShouldFinishTick = false;
 
-			if (IsCookByTheBookMode() && Timer.IsTimeUp())
+			if (Timer.IsTimeUp() && IsCookByTheBookMode() )
 			{
 				bShouldFinishTick = true;
 				// our timeslice is up
@@ -2086,7 +2161,7 @@ bool UCookOnTheFlyServer::SaveCookedPackages(TArray<UPackage*>& PackagesToSave,
 			// if we are cook the fly then save the package which was requested as fast as we can because the client is waiting on it
 
 			bool ProcessingUnsolicitedPackages = (I >= FirstUnsolicitedPackage);
-			bool ForceSavePackage = false;
+			bool bForceSavePackage = false;
 
 			if (IsCookOnTheFlyMode())
 			{
@@ -2120,7 +2195,10 @@ bool UCookOnTheFlyServer::SaveCookedPackages(TArray<UPackage*>& PackagesToSave,
 				}
 				else
 				{
-					ForceSavePackage = true;
+					if (!IsRealtimeMode())
+					{
+						bForceSavePackage = true;
+					}
 				}
 			}
 
@@ -2129,7 +2207,14 @@ bool UCookOnTheFlyServer::SaveCookedPackages(TArray<UPackage*>& PackagesToSave,
 
 			MakePackageFullyLoaded(Package);
 
-			if (!bShouldFinishTick)
+			if ( IsCookOnTheFlyMode() )
+			{
+				// never want to requeue packages
+				HasCheckedAllPackagesAreCached = true;
+			}
+
+			// if we are forcing save the package then it doesn't matter if we call FinishPackageCacheForCookedPlatformData
+			if (!bShouldFinishTick && !bForceSavePackage)
 			{
 				AllObjectsCookedDataCached = FinishPackageCacheForCookedPlatformData(Package, TargetPlatformsToCache, Timer);
 				if (AllObjectsCookedDataCached == false)
@@ -2141,29 +2226,28 @@ bool UCookOnTheFlyServer::SaveCookedPackages(TArray<UPackage*>& PackagesToSave,
 
 			// if we are in realtime mode and this package isn't ready to be saved then we should exit the tick here so we don't save it while in launch on
 			if (IsRealtimeMode() &&
-				(!IsCookOnTheFlyMode() || ProcessingUnsolicitedPackages) && // if we are in cook ont eh fly mode, EVEN if we are realtime we should force save if this is not an unsolicited package
 				(AllObjectsCookedDataCached == false) &&
 				HasCheckedAllPackagesAreCached)
 			{
 				bShouldFinishTick = true;
 			}
 
-			if (bShouldFinishTick && (!ForceSavePackage))
+			if (bShouldFinishTick && (!bForceSavePackage))
 			{
 				SCOPE_TIMER(EnqueueUnsavedPackages);
 				// enqueue all the packages which we were about to save
 				Timer.SavedPackage();  // this is a special case to prevent infinite loop, if we only have one package we might fall through this and could loop forever.  
-				if (IsCookByTheBookMode())
+				int32 NumPackagesToRequeue = PackagesToSave.Num();
+
+				if (IsCookOnTheFlyMode())
 				{
-					for (int32 RemainingIndex = I; RemainingIndex < PackagesToSave.Num(); ++RemainingIndex)
-					{
-						FName StandardFilename = GetCachedStandardPackageFileFName(PackagesToSave[RemainingIndex]);
-						CookRequests.EnqueueUnique(FFilePlatformRequest(StandardFilename, AllTargetPlatformNames));
-					}
+					NumPackagesToRequeue = FirstUnsolicitedPackage;
 				}
-				else
+				
+				for (int32 RemainingIndex = I; RemainingIndex < NumPackagesToRequeue; ++RemainingIndex)
 				{
-					check(ProcessingUnsolicitedPackages == true);
+					FName StandardFilename = GetCachedStandardPackageFileFName(PackagesToSave[RemainingIndex]);
+					CookRequests.EnqueueUnique(FFilePlatformRequest(StandardFilename, AllTargetPlatformNames));
 				}
 				Result |= COSR_WaitingOnCache;
 
@@ -2196,16 +2280,17 @@ bool UCookOnTheFlyServer::SaveCookedPackages(TArray<UPackage*>& PackagesToSave,
 			// if we already went through the entire package list then don't keep requeuing requests
 			if ((HasCheckedAllPackagesAreCached == false) &&
 				(AllObjectsCookedDataCached == false) &&
-				(ForceSavePackage == false))
+				(bForceSavePackage == false) &&
+				IsCookByTheBookMode() )
 			{
-				check(IsCookByTheBookMode() || ProcessingUnsolicitedPackages == true);
+				// check(IsCookByTheBookMode() || ProcessingUnsolicitedPackages == true);
 				// add to back of queue
 				PackagesToSave.Add(Package);
 				// UE_LOG(LogCook, Display, TEXT("Delaying save for package %s"), *PackageFName.ToString());
 				continue;
 			}
 
-			if (HasCheckedAllPackagesAreCached)
+			if ( HasCheckedAllPackagesAreCached && (AllObjectsCookedDataCached == false) )
 			{
 				UE_LOG(LogCook, Display, TEXT("Forcing save package %s because was already requeued once"), *PackageFName.ToString());
 			}
@@ -2636,6 +2721,12 @@ void UCookOnTheFlyServer::OnObjectPropertyChanged(UObject* ObjectBeingModified,
 	{
 		return;
 	}
+	if ( PropertyChangedEvent.Property == nullptr && 
+		PropertyChangedEvent.MemberProperty == nullptr )
+	{
+		// probably nothing changed... 
+		return;
+	}
 
 	OnObjectUpdated( ObjectBeingModified );
 }
@@ -3054,7 +3145,6 @@ void UCookOnTheFlyServer::SaveCookedPackage(UPackage* Package, uint32 SaveFlags,
 		UWorld* World = nullptr;
 		EObjectFlags FlagsToCook = RF_Public;
 
-
 		ITargetPlatformManagerModule& TPM = GetTargetPlatformManagerRef();
 
 		static TArray<ITargetPlatform*> ActiveStartupPlatforms = TPM.GetCookingTargetPlatforms();
@@ -3350,7 +3440,8 @@ void UCookOnTheFlyServer::Initialize( ECookMode::Type DesiredCookMode, ECookInit
 			}
 		}
 	}
-	
+
+
 	UE_LOG(LogCook, Display, TEXT("Mobile HDR setting %d"), IsMobileHDR());
 
 	// See if there are any plugins that need to be remapped for the sandbox
@@ -3507,8 +3598,7 @@ FString UCookOnTheFlyServer::GetOutputDirectoryOverride() const
 		// Output directory needs to contain [Platform] token to be able to cook for multiple targets.
 		if ( IsCookByTheBookMode() )
 		{
-			ITargetPlatformManagerModule& TPM = GetTargetPlatformManagerRef();
-			const TArray<ITargetPlatform*>& TargetPlatforms = TPM.GetCookingTargetPlatforms();
+			const TArray<ITargetPlatform*>& TargetPlatforms = GetCookingTargetPlatforms();
 
 			// more then one target platform specified append "[platform]" to output directory so that multiple platforms can be cooked
 			check( TargetPlatforms.Num() == 1 );
@@ -3664,7 +3754,7 @@ void GetAdditionalCurrentIniVersionStrings( const ITargetPlatform* TargetPlatfor
 	}
 
 
-	// TODO: Add support for physx version tracking, currently this happens so infrequently that invalidating a cook based on it is not essentual
+	// TODO: Add support for physx version tracking, currently this happens so infrequently that invalidating a cook based on it is not essential
 	//GetVersionFormatNumbersForIniVersionStrings(IniVersionMap, TEXT("PhysXCooking"), TPM->GetPhysXCooking());
 
 
@@ -4818,7 +4908,7 @@ void UCookOnTheFlyServer::PopulateCookedPackagesFromDisk(const TArray<ITargetPla
 			FAssetRegistryGenerator* PlatformAssetRegistry = RegistryGenerators.FindRef(PlatformFName);
 
 			// Load the platform cooked asset registry file
-			const FString CookedAssetRegistry = FPaths::GameDir() / GetAssetRegistryFilename();
+			const FString CookedAssetRegistry = FPaths::GameDir() / GetDevelopmentAssetRegistryFilename();
 			const FString SandboxCookedAssetRegistryFilename = ConvertToFullSandboxPath(*CookedAssetRegistry, true, Target->PlatformName());
 
 			bool bIsIterateSharedBuild = IsCookFlagSet(ECookInitializationFlags::IterateSharedBuild);
@@ -4829,11 +4919,14 @@ void UCookOnTheFlyServer::PopulateCookedPackagesFromDisk(const TArray<ITargetPla
 				FDateTime CurrentLocalCookedBuild = IFileManager::Get().GetTimeStamp(*SandboxCookedAssetRegistryFilename);
 
 				// iterate on the shared build if the option is set
-				FString SharedCookedAssetRegistry = FPaths::Combine(*FPaths::GameSavedDir(), TEXT("SharedIterativeBuild"), *Target->PlatformName(), TEXT("Cooked"), TEXT("DevelopmentAssetRegistry.bin"));
+				FString SharedCookedAssetRegistry = FPaths::Combine(*FPaths::GameSavedDir(), TEXT("SharedIterativeBuild"), *Target->PlatformName(), TEXT("Cooked"), GetDevelopmentAssetRegistryFilename());
 
 				FDateTime CurrentIterativeCookedBuild = IFileManager::Get().GetTimeStamp(*SharedCookedAssetRegistry);
 
-				if (CurrentIterativeCookedBuild >= CurrentLocalCookedBuild)
+
+
+				if ( (CurrentIterativeCookedBuild >= CurrentLocalCookedBuild) && 
+					(CurrentIterativeCookedBuild != FDateTime::MinValue()) )
 				{
 					// clean the sandbox 
 					ClearPlatformCookedData(FName(*Target->PlatformName()));
@@ -4910,12 +5003,17 @@ void UCookOnTheFlyServer::PopulateCookedPackagesFromDisk(const TArray<ITargetPla
 			uint32 NumPackagesFileHashMismatch = 0;
 			uint32 NumPackagesKept = 0;
 			uint32 NumMarkedFailedSaveKept = 0;
+			uint32 NumPackagesRemoved = 0;
 
 			for (const auto& CookedPaths : UncookedPathToCookedPath)
 			{
 				const FName CookedFile = CookedPaths.Value;
 				const FName UncookedFilename = CookedPaths.Key;
 				const FName* FoundPackageName = GetCachedPackageFilenameToPackageFName(UncookedFilename);
+				if ( !FoundPackageName )
+				{
+					++NumPackagesRemoved;
+				}
 				const FName PackageName = *FoundPackageName;
 				const FString UncookedFilenameString = UncookedFilename.ToString();
 				bool bShouldKeep = true;
@@ -4932,8 +5030,8 @@ void UCookOnTheFlyServer::PopulateCookedPackagesFromDisk(const TArray<ITargetPla
 				}
 				else if (IdenticalUncookedPackages.Contains(PackageName))
 				{
-					// May have switched between editor only and not, recheck
-					++NumPackagesUnableToFindCookedPackageInfo;
+					// These are packages which failed to save the first time 
+					// most likely because they are editor only packages
 					bShouldKeep = false;
 				}
 				
@@ -4963,8 +5061,11 @@ void UCookOnTheFlyServer::PopulateCookedPackagesFromDisk(const TArray<ITargetPla
 				}
 				else
 				{
-					// Force cook the modified file
-					CookRequests.EnqueueUnique(MoveTemp(FFilePlatformRequest(UncookedFilename, PlatformNames)));
+					if ( IsCookByTheBookMode() ) // cook on the fly will requeue this package when it wants it 
+					{
+						// Force cook the modified file
+						CookRequests.EnqueueUnique(MoveTemp(FFilePlatformRequest(UncookedFilename, PlatformNames)));
+					}
 					if (CookedFile != NAME_DummyCookedFilename)
 					{
 						// delete the old package 
@@ -5678,8 +5779,7 @@ void UCookOnTheFlyServer::CleanSandbox(const bool bIterative)
 	//child cookers shouldn't clean sandbox we would be deleting the master cookers / other cookers data
 	check(IsChildCooker() == false);
 
-	ITargetPlatformManagerModule& TPM = GetTargetPlatformManagerRef();
-	const TArray<ITargetPlatform*>& Platforms = TPM.GetCookingTargetPlatforms();
+	const TArray<ITargetPlatform*>& Platforms = GetCookingTargetPlatforms();
 
 	// before we can delete any cooked files we need to make sure that we have finished writing them
 	UPackage::WaitForAsyncFileWrites();
@@ -5808,10 +5908,15 @@ void UCookOnTheFlyServer::GenerateAssetRegistry()
 		// Perform a synchronous search of any .ini based asset paths (note that the per-game delegate may
 		// have already scanned paths on its own)
 		// We want the registry to be fully initialized when generating streaming manifests too.
-		bool bEditor = IsRealtimeMode();
 
 		// editor will scan asset registry automagically 
-		if ( !bEditor )
+		bool bCanDelayAssetregistryProcessing = IsRealtimeMode(); 
+
+		// if we are running in the editor we need the asset registry to be finished loaded before we process any iterative cook requests
+		bCanDelayAssetregistryProcessing &= !IsCookFlagSet(ECookInitializationFlags::Iterative); // 
+
+		
+		if ( !bCanDelayAssetregistryProcessing)
 		{
 			TArray<FString> ScanPaths;
 			if (GConfig->GetArray(TEXT("AssetRegistry"), TEXT("PathsToScanForCook"), ScanPaths, GEngineIni) > 0)
@@ -5825,8 +5930,7 @@ void UCookOnTheFlyServer::GenerateAssetRegistry()
 		}
 	}
 
-	ITargetPlatformManagerModule& TPM = GetTargetPlatformManagerRef();
-	const TArray<ITargetPlatform*>& Platforms = TPM.GetCookingTargetPlatforms();
+	const TArray<ITargetPlatform*>& Platforms = GetCookingTargetPlatforms();
 
 	for (ITargetPlatform* TargetPlatform : Platforms)
 	{
@@ -6841,8 +6945,7 @@ void UCookOnTheFlyServer::InitializeSandbox()
 {
 	if ( SandboxFile == nullptr )
 	{
-		ITargetPlatformManagerModule& TPM = GetTargetPlatformManagerRef();
-		const TArray<ITargetPlatform*>& TargetPlatforms = TPM.GetCookingTargetPlatforms();
+		const TArray<ITargetPlatform*>& TargetPlatforms = GetCookingTargetPlatforms();
 
 		CreateSandboxFile();
 
@@ -6865,7 +6968,6 @@ void UCookOnTheFlyServer::TermSandbox()
 
 void UCookOnTheFlyServer::ValidateCookOnTheFlySettings() const
 {
-	check( IsCookFlagSet(ECookInitializationFlags::IterateSharedBuild) == false );
 }
 
 void UCookOnTheFlyServer::ValidateCookByTheBookSettings() const
@@ -7208,7 +7310,6 @@ void UCookOnTheFlyServer::StartCookByTheBook( const FCookByTheBookStartupOptions
 	}
 }
 
-
 bool UCookOnTheFlyServer::RecompileChangedShaders(const TArray<FName>& TargetPlatforms)
 {
 	bool bShadersRecompiled = false;
@@ -7219,9 +7320,6 @@ bool UCookOnTheFlyServer::RecompileChangedShaders(const TArray<FName>& TargetPla
 	return bShadersRecompiled;
 }
 
-
-
-
 class FChildCookerRunnable : public FRunnable
 {
 private:
@@ -7586,6 +7684,23 @@ void UCookOnTheFlyServer::MaybeMarkPackageAsAlreadyLoaded(UPackage *Package)
 }
 
 
+bool UCookOnTheFlyServer::HandleNetworkFileServerNewConnection(const FString& VersionInfo, const FString& PlatformName)
+{
+	const uint32 CL = FEngineVersion::CompatibleWith().GetChangelist();
+	const FString Branch = FEngineVersion::CompatibleWith().GetBranch();
+
+	const FString LocalVersionInfo = FString::Printf(TEXT("%s %d"), *Branch, CL);
+
+	UE_LOG(LogCook, Display, TEXT("Connection received of version %s local version %s"), *VersionInfo, *LocalVersionInfo);
+
+	if (LocalVersionInfo != VersionInfo)
+	{
+		UE_LOG(LogCook, Warning, TEXT("Connection tried to connect with incompatable version"));
+		// return false;
+	}
+	return true;
+}
+
 void UCookOnTheFlyServer::HandleNetworkFileServerFileRequest( const FString& Filename, const FString &PlatformName, TArray<FString>& UnsolicitedFiles )
 {
 	check( IsCookOnTheFlyMode() );	
@@ -7708,6 +7823,27 @@ FString UCookOnTheFlyServer::HandleNetworkGetSandboxPath()
 	return SandboxFile->GetSandboxDirectory();
 }
 
+void UCookOnTheFlyServer::HandleNetworkGetPrecookedList(const FString& PlatformName, TMap<FString, FDateTime>& PrecookedFileList)
+{
+	FName PlatformFName = FName(*PlatformName);
+
+	TArray<FName> CookedPlatformFiles;
+	CookedPackages.GetCookedFilesForPlatform(PlatformFName, CookedPlatformFiles);
+
+
+	for ( const FName& CookedFile : CookedPlatformFiles)
+	{
+		const FString SandboxFilename = ConvertToFullSandboxPath(CookedFile.ToString(), true, PlatformName);
+		if (IFileManager::Get().FileExists(*SandboxFilename))
+		{
+			continue;
+		}
+
+		PrecookedFileList.Add(CookedFile.ToString(),FDateTime::MinValue());
+	}
+}
+
+
 void UCookOnTheFlyServer::HandleNetworkFileServerRecompileShaders(const FShaderRecompileData& RecompileData)
 {
 	// shouldn't receive network requests unless we are in cook on the fly mode
diff --git a/Engine/Source/Editor/UnrealEd/Private/Factories/ApexDestructibleAssetImport.cpp b/Engine/Source/Editor/UnrealEd/Private/Factories/ApexDestructibleAssetImport.cpp
--- a/Engine/Source/Editor/UnrealEd/Private/Factories/ApexDestructibleAssetImport.cpp
+++ b/Engine/Source/Editor/UnrealEd/Private/Factories/ApexDestructibleAssetImport.cpp
@@ -43,7 +43,7 @@ DEFINE_LOG_CATEGORY_STATIC(LogApexDestructibleAssetImport, Log, All);
 // Forward declarations and external processing functions
 struct ExistingSkelMeshData;
 extern ExistingSkelMeshData* SaveExistingSkelMeshData(USkeletalMesh* ExistingSkelMesh, bool bSaveMaterials, int32 ReimportLODIndex);
-extern void RestoreExistingSkelMeshData(ExistingSkelMeshData* MeshData, USkeletalMesh* SkeletalMesh, int32 ReimportLODIndex, bool bResetMaterialSlots);
+extern void RestoreExistingSkelMeshData(ExistingSkelMeshData* MeshData, USkeletalMesh* SkeletalMesh, int32 ReimportLODIndex, bool bResetMaterialSlots, bool bIsReimportPreview);
 extern void ProcessImportMeshInfluences(FSkeletalMeshImportData& ImportData);
 extern void ProcessImportMeshMaterials(TArray<FSkeletalMaterial>& Materials, FSkeletalMeshImportData& ImportData);
 extern bool ProcessImportMeshSkeleton(const USkeleton* SkeletonAsset, FReferenceSkeleton& RefSkeleton, int32& SkeletalDepth, FSkeletalMeshImportData& ImportData);
@@ -256,7 +256,7 @@ static void RestoreExistingDestMeshData(ExistingDestMeshData* MeshData, UDestruc
 		// Restore old settings, but resize arrays to make sense with the new NxDestructibleAsset
 		if (MeshData->SkelMeshData != NULL)
 		{
-			RestoreExistingSkelMeshData(MeshData->SkelMeshData, DestructibleMesh, INDEX_NONE, false);
+			RestoreExistingSkelMeshData(MeshData->SkelMeshData, DestructibleMesh, INDEX_NONE, false, false);
 		}
 		DestructibleMesh->BodySetup =  MeshData->BodySetup;
 		DestructibleMesh->FractureEffects = MeshData->FractureEffects;
diff --git a/Engine/Source/Editor/UnrealEd/Private/Factories/EditorFactories.cpp b/Engine/Source/Editor/UnrealEd/Private/Factories/EditorFactories.cpp
--- a/Engine/Source/Editor/UnrealEd/Private/Factories/EditorFactories.cpp
+++ b/Engine/Source/Editor/UnrealEd/Private/Factories/EditorFactories.cpp
@@ -4998,7 +4998,7 @@ EReimportResult::Type UReimportFbxStaticMeshFactory::Reimport( UObject* Obj )
 		bool bOutImportAll = false;
 		bool bIsObjFormat = false;
 		bool bIsAutomated = false;
-		GetImportOptions( FFbxImporter, ReimportUI, bShowOptionDialog, bIsAutomated, Obj->GetPathName(), bOperationCanceled, bOutImportAll, bIsObjFormat, bForceImportType, FBXIT_StaticMesh );
+		GetImportOptions( FFbxImporter, ReimportUI, bShowOptionDialog, bIsAutomated, Obj->GetPathName(), bOperationCanceled, bOutImportAll, bIsObjFormat, bForceImportType, FBXIT_StaticMesh, Mesh);
 	}
 
 	//We do not touch bAutoComputeLodDistances when we re-import, setting it to true will make sure we do not change anything.
@@ -5112,6 +5112,7 @@ EReimportResult::Type UReimportFbxStaticMeshFactory::Reimport( UObject* Obj )
 	}
 	else
 	{
+		FFbxImporter->ReleaseScene();
 		return EReimportResult::Cancelled;
 	}
 }
@@ -5253,7 +5254,7 @@ EReimportResult::Type UReimportFbxSkeletalMeshFactory::Reimport( UObject* Obj )
 		ImportOptions->bCreatePhysicsAsset = false;
 		ImportOptions->PhysicsAsset = SkeletalMesh->PhysicsAsset;
 
-		ImportOptions = GetImportOptions( FFbxImporter, ReimportUI, bShowOptionDialog, bIsAutomated, Obj->GetPathName(), bOperationCanceled, bOutImportAll, bIsObjFormat, bForceImportType, FBXIT_SkeletalMesh );
+		ImportOptions = GetImportOptions( FFbxImporter, ReimportUI, bShowOptionDialog, bIsAutomated, Obj->GetPathName(), bOperationCanceled, bOutImportAll, bIsObjFormat, bForceImportType, FBXIT_SkeletalMesh, Obj );
 	}
 
 	if( !bOperationCanceled && ensure(ImportData) )
@@ -5312,6 +5313,7 @@ EReimportResult::Type UReimportFbxSkeletalMeshFactory::Reimport( UObject* Obj )
 	}
 	else
 	{
+		FFbxImporter->ReleaseScene();
 		return EReimportResult::Cancelled;
 	}
 }
diff --git a/Engine/Source/Editor/UnrealEd/Private/Factories/SkeletalMeshImport.cpp b/Engine/Source/Editor/UnrealEd/Private/Factories/SkeletalMeshImport.cpp
--- a/Engine/Source/Editor/UnrealEd/Private/Factories/SkeletalMeshImport.cpp
+++ b/Engine/Source/Editor/UnrealEd/Private/Factories/SkeletalMeshImport.cpp
@@ -814,7 +814,10 @@ ExistingSkelMeshData* SaveExistingSkelMeshData(USkeletalMesh* ExistingSkelMesh,
 				int32 SectionMaterialIndex = ImportedResource->LODModels[LodIndex].Sections[SectionIndex].MaterialIndex;
 				bool SectionCastShadow = ImportedResource->LODModels[LodIndex].Sections[SectionIndex].bCastShadow;
 				bool SectionRecomputeTangents = ImportedResource->LODModels[LodIndex].Sections[SectionIndex].bRecomputeTangent;
-				ExistingMeshDataPtr->ExistingImportMeshLodSectionMaterialData[LodIndex].Add(ExistingMeshLodSectionData(ExistingMeshDataPtr->ExistingImportMaterialOriginalNameData[SectionMaterialIndex], SectionCastShadow, SectionRecomputeTangents));
+				if (ExistingMeshDataPtr->ExistingImportMaterialOriginalNameData.IsValidIndex(SectionMaterialIndex))
+				{
+					ExistingMeshDataPtr->ExistingImportMeshLodSectionMaterialData[LodIndex].Add(ExistingMeshLodSectionData(ExistingMeshDataPtr->ExistingImportMaterialOriginalNameData[SectionMaterialIndex], SectionCastShadow, SectionRecomputeTangents));
+				}
 			}
 		}
 
@@ -968,7 +971,7 @@ void TryRegenerateLODs(ExistingSkelMeshData* MeshData, USkeletalMesh* SkeletalMe
 	}
 }
 
-void RestoreExistingSkelMeshData(ExistingSkelMeshData* MeshData, USkeletalMesh* SkeletalMesh, int32 ReimportLODIndex, bool bResetMaterialSlots)
+void RestoreExistingSkelMeshData(ExistingSkelMeshData* MeshData, USkeletalMesh* SkeletalMesh, int32 ReimportLODIndex, bool bResetMaterialSlots, bool bIsReimportPreview)
 {
 	if (!MeshData || !SkeletalMesh)
 	{
@@ -1204,7 +1207,7 @@ void RestoreExistingSkelMeshData(ExistingSkelMeshData* MeshData, USkeletalMesh*
 				}
 			}
 
-			if (bRegenLODs)
+			if (bRegenLODs && !bIsReimportPreview)
 			{
 				TryRegenerateLODs(MeshData, SkeletalMesh);
 			}
diff --git a/Engine/Source/Editor/UnrealEd/Private/Fbx/FbxCompareWindow.cpp b/Engine/Source/Editor/UnrealEd/Private/Fbx/FbxCompareWindow.cpp
new file mode 100644
--- /dev/null
+++ b/Engine/Source/Editor/UnrealEd/Private/Fbx/FbxCompareWindow.cpp
@@ -0,0 +1,1032 @@
+// Copyright 1998-2017 Epic Games, Inc. All Rights Reserved.
+#include "FbxCompareWindow.h"
+#include "Modules/ModuleManager.h"
+#include "Widgets/Layout/SBorder.h"
+#include "Widgets/Text/STextBlock.h"
+#include "Widgets/Layout/SBox.h"
+#include "Widgets/Layout/SUniformGridPanel.h"
+#include "Widgets/Layout/SScrollBox.h"
+#include "Widgets/Layout/SSeparator.h"
+#include "Widgets/Input/SButton.h"
+#include "Widgets/Images/SImage.h"
+#include "Widgets/Colors/SColorBlock.h"
+#include "Widgets/Input/SCheckBox.h"
+#include "Widgets/Views/SListView.h"
+#include "Widgets/Views/STreeView.h"
+#include "Widgets/Views/STableRow.h"
+#include "EditorStyleSet.h"
+#include "Factories/FbxAnimSequenceImportData.h"
+#include "IDocumentation.h"
+#include "PropertyEditorModule.h"
+#include "IDetailsView.h"
+#include "Framework/Commands/UIAction.h"
+#include "Framework/MultiBox/MultiBoxBuilder.h"
+#include "Factories/FbxSceneImportFactory.h"
+#include "Toolkits/AssetEditorManager.h"
+#include "Framework/Application/SlateApplication.h"
+#include "Engine/SkeletalMesh.h"
+
+
+#define LOCTEXT_NAMESPACE "FBXOption"
+
+void SFbxCompareWindow::Construct(const FArguments& InArgs)
+{
+	CurrentDisplayOption = FMaterialCompareData::All;
+	bShowSectionFlag[EFBXCompareSection_General] = true;
+	bShowSectionFlag[EFBXCompareSection_Materials] = true;
+	bShowSectionFlag[EFBXCompareSection_Skeleton] = true;
+
+	WidgetWindow = InArgs._WidgetWindow;
+	FullFbxPath = InArgs._FullFbxPath.ToString();
+	FbxSceneInfo = InArgs._FbxSceneInfo;
+	FbxGeneralInfo = InArgs._FbxGeneralInfo;
+	if (InArgs._AssetReferencingSkeleton != nullptr)
+	{
+		//Copy the aray
+		AssetReferencingSkeleton = *(InArgs._AssetReferencingSkeleton);
+	}
+	CurrentMeshData = InArgs._CurrentMeshData;
+	FbxMeshData = InArgs._FbxMeshData;
+	PreviewObject = InArgs._PreviewObject;
+
+	FillGeneralListItem();
+	FillMaterialListItem();
+	if (PreviewObject->IsA(USkeletalMesh::StaticClass()))
+	{
+		FilSkeletonTreeItem();
+	}
+
+	SetMatchJointInfo();
+
+	// Material comparison
+	TSharedPtr<SWidget> MaterialCompareSection = ConstructMaterialComparison();
+	// Skeleton comparison
+	TSharedPtr<SWidget> SkeletonCompareSection = ConstructSkeletonComparison();
+	// General section
+	TSharedPtr<SWidget> GeneralInfoSection = ConstructGeneralInfo();
+
+	TSharedRef<SScrollBar> VerticalScrollBar = SNew(SScrollBar)
+		.Orientation(EOrientation::Orient_Vertical)
+		.AlwaysShowScrollbar(false);
+
+	this->ChildSlot
+		[
+		SNew(SBox)
+		[
+			SNew(SVerticalBox)
+			+SVerticalBox::Slot()
+			.FillHeight(1.0f)
+			[
+				SNew(SScrollBox)
+				+ SScrollBox::Slot()
+				[
+					SNew(SBorder)
+					.BorderImage(FEditorStyle::GetBrush("ToolPanel.DarkGroupBorder"))
+					[
+						SNew(SVerticalBox)
+						+SVerticalBox::Slot()
+						.AutoHeight()
+						.Padding(2)
+						[
+							// Header with the file path
+							SNew(SBorder)
+							.Padding(FMargin(3))
+							.BorderImage(FEditorStyle::GetBrush("ToolPanel.GroupBorder"))
+							[
+								SNew(SHorizontalBox)
+								+SHorizontalBox::Slot()
+								.AutoWidth()
+								[
+									SNew(STextBlock)
+									.Font(FEditorStyle::GetFontStyle("CurveEd.LabelFont"))
+									.Text(LOCTEXT("Import_CurrentFileTitle", "Current File: "))
+								]
+								+SHorizontalBox::Slot()
+								.Padding(5, 0, 0, 0)
+								.AutoWidth()
+								.VAlign(VAlign_Center)
+								[
+									SNew(STextBlock)
+									.Font(FEditorStyle::GetFontStyle("CurveEd.InfoFont"))
+									.Text(InArgs._FullFbxPath)
+								]
+							]
+						]
+						+ SVerticalBox::Slot()
+						.FillHeight(1.0f)
+						.Padding(2)
+						[
+							SNew(SVerticalBox)
+							+ SVerticalBox::Slot()
+							.AutoHeight()
+							.Padding(2)
+							[
+								// Material Compare section
+								MaterialCompareSection.ToSharedRef()
+							]
+							+ SVerticalBox::Slot()
+							.AutoHeight()
+							.Padding(2)
+							[
+								// Skeleton Compare section
+								SkeletonCompareSection.ToSharedRef()
+							]
+							+ SVerticalBox::Slot()
+							.AutoHeight()
+							.Padding(2)
+							[
+								GeneralInfoSection.ToSharedRef()
+							]
+						]
+					]
+				]
+			]
+			+ SVerticalBox::Slot()
+			.AutoHeight()
+			.HAlign(HAlign_Right)
+			.Padding(2)
+			[
+				SNew(SButton)
+				.HAlign(HAlign_Center)
+				.Text(LOCTEXT("SFbxCompareWindow_Preview_Done", "Done"))
+				.OnClicked(this, &SFbxCompareWindow::OnDone)
+			]
+		]
+	];	
+}
+
+FReply SFbxCompareWindow::SetSectionVisible(EFBXCompareSection SectionIndex)
+{
+	bShowSectionFlag[SectionIndex] = !bShowSectionFlag[SectionIndex];
+	return FReply::Handled();
+}
+
+EVisibility SFbxCompareWindow::IsSectionVisible(EFBXCompareSection SectionIndex)
+{
+	return bShowSectionFlag[SectionIndex] ? EVisibility::All : EVisibility::Collapsed;
+}
+
+const FSlateBrush* SFbxCompareWindow::GetCollapsableArrow(EFBXCompareSection SectionIndex) const
+{
+	return bShowSectionFlag[SectionIndex] ? FEditorStyle::GetBrush("Symbols.DownArrow") : FEditorStyle::GetBrush("Symbols.RightArrow");
+}
+
+TSharedPtr<SWidget> SFbxCompareWindow::ConstructGeneralInfo()
+{
+	return SNew(SBox)
+	.MaxDesiredHeight(205)
+	[
+		SNew(SBorder)
+		.Padding(FMargin(3))
+		.BorderImage(FEditorStyle::GetBrush("ToolPanel.GroupBorder"))
+		[
+			SNew(SVerticalBox)
+			+SVerticalBox::Slot()
+			.AutoHeight()
+			.Padding(2)
+			[
+				SNew(SHorizontalBox)
+				+ SHorizontalBox::Slot()
+				.AutoWidth()
+				[
+					SNew(SButton)
+					.HAlign(HAlign_Center)
+					.VAlign(VAlign_Center)
+					.IsFocusable(false)
+					.ButtonStyle(FEditorStyle::Get(), "NoBorder")
+					.OnClicked(this, &SFbxCompareWindow::SetSectionVisible, EFBXCompareSection_General)
+					[
+						SNew(SImage).Image(this, &SFbxCompareWindow::GetCollapsableArrow, EFBXCompareSection_General)
+					]
+				]
+				+ SHorizontalBox::Slot()
+				.AutoWidth()
+				[
+					SNew(STextBlock)
+					.Font(FEditorStyle::GetFontStyle("DetailsView.CategoryFontStyle"))
+					.Text(LOCTEXT("SFbxCompareWindow_GeneralInfoHeader", "Fbx File Information"))
+				]
+			
+			]
+			+SVerticalBox::Slot()
+			.FillHeight(1.0f)
+			.Padding(2)
+			[
+				SNew(SBox)
+				.Visibility(TAttribute<EVisibility>::Create(TAttribute<EVisibility>::FGetter::CreateSP(this, &SFbxCompareWindow::IsSectionVisible, EFBXCompareSection_General)))
+				[
+					SNew(SBorder)
+					.BorderImage(FEditorStyle::GetBrush("ToolPanel.DarkGroupBorder"))
+					[
+						//Show the general fbx information
+						SNew(SListView<TSharedPtr<FString>>)
+						.ListItemsSource(&GeneralListItem)
+						.OnGenerateRow(this, &SFbxCompareWindow::OnGenerateRowGeneralFbxInfo)
+					]
+				]
+			]
+		]
+	];
+}
+
+TSharedPtr<SWidget> SFbxCompareWindow::ConstructMaterialComparison()
+{
+	return SNew(SBox)
+	.MaxDesiredHeight(500)
+	[
+		SNew(SBorder)
+		.Padding(FMargin(3))
+		.BorderImage(FEditorStyle::GetBrush("ToolPanel.GroupBorder"))
+		[
+			SNew(SVerticalBox)
+			+SVerticalBox::Slot()
+			.AutoHeight()
+			.Padding(2)
+			[
+				SNew(SHorizontalBox)
+				+ SHorizontalBox::Slot()
+				.AutoWidth()
+				[
+					SNew(SButton)
+					.HAlign(HAlign_Center)
+					.VAlign(VAlign_Center)
+					.IsFocusable(false)
+					.ButtonStyle(FEditorStyle::Get(), "NoBorder")
+					.OnClicked(this, &SFbxCompareWindow::SetSectionVisible, EFBXCompareSection_Materials)
+					[
+						SNew(SImage).Image(this, &SFbxCompareWindow::GetCollapsableArrow, EFBXCompareSection_Materials)
+					]
+				]
+				+ SHorizontalBox::Slot()
+				.AutoWidth()
+				[
+					SNew(STextBlock)
+					.Font(FEditorStyle::GetFontStyle("DetailsView.CategoryFontStyle"))
+					.Text(LOCTEXT("SFbxCompareWindow_MaterialCompareHeader", "Materials"))
+				]
+			]
+			+SVerticalBox::Slot()
+			.FillHeight(1.0f)
+			.Padding(2)
+			[
+				SNew(SBox)
+				.Visibility(TAttribute<EVisibility>::Create(TAttribute<EVisibility>::FGetter::CreateSP(this, &SFbxCompareWindow::IsSectionVisible, EFBXCompareSection_Materials)))
+				[
+					SNew(SVerticalBox)
+					+SVerticalBox::Slot()
+					.FillHeight(1.0f)
+					.Padding(2)
+					[
+						//Show the Comparison of the meshes
+						SNew(SListView< TSharedPtr<FMaterialCompareData> >)
+						.ItemHeight(24)
+						.ListItemsSource(&CompareMaterialListItem)
+						.OnGenerateRow(this, &SFbxCompareWindow::OnGenerateRowForCompareMaterialList)
+						.HeaderRow
+						(
+							SNew(SHeaderRow)
+							+ SHeaderRow::Column("RowIndex")
+								.DefaultLabel(LOCTEXT("SFbxCompareWindow_RowIndex_ColumnHeader", ""))
+								.FixedWidth(25)
+							+ SHeaderRow::Column("Current")
+								.DefaultLabel(LOCTEXT("SFbxCompareWindow_Current_ColumnHeader", "Current Asset"))
+								.FillWidth(0.5f)
+							+ SHeaderRow::Column("Fbx")
+								.DefaultLabel(LOCTEXT("SFbxCompareWindow_Fbx_ColumnHeader", "Reimport Asset (Preview)"))
+								.FillWidth(0.5f)
+						)
+					]
+					+SVerticalBox::Slot()
+					.AutoHeight()
+					.Padding(2)
+					[
+						//Show the toggle button to display different re-import problem
+						SNew(SHorizontalBox)
+						+SHorizontalBox::Slot()
+						.AutoWidth()
+						[
+							SNew(SCheckBox)
+							.OnCheckStateChanged(this, &SFbxCompareWindow::ToggleMaterialDisplay, FMaterialCompareData::EMaterialCompareDisplayOption::All)
+							.IsChecked(this, &SFbxCompareWindow::IsToggleMaterialDisplayChecked, FMaterialCompareData::EMaterialCompareDisplayOption::All)
+							[
+								SNew(STextBlock)
+								.Text(LOCTEXT("SFbxCompareWindow_Display_Option_All", "All"))
+							]
+							
+						]
+						+SHorizontalBox::Slot()
+						.Padding(5, 0, 0, 0)
+						.AutoWidth()
+						.VAlign(VAlign_Center)
+						[
+							SNew(SCheckBox)
+							.OnCheckStateChanged(this, &SFbxCompareWindow::ToggleMaterialDisplay, FMaterialCompareData::EMaterialCompareDisplayOption::NoMatch)
+							.IsChecked(this, &SFbxCompareWindow::IsToggleMaterialDisplayChecked, FMaterialCompareData::EMaterialCompareDisplayOption::NoMatch)
+							[
+								SNew(STextBlock)
+								.Text(LOCTEXT("SFbxCompareWindow_Display_Option_NoMatch", "No Match"))
+								.ColorAndOpacity(FSlateColor(FLinearColor(0.7f, 0.3f, 0.0f)))
+								.ToolTipText(LOCTEXT("SFbxCompareWindow_Display_Option_NoMatch_tooltip", "Can impact gameplay code using material slot name."))
+							]
+						]
+						+SHorizontalBox::Slot()
+						.Padding(5, 0, 0, 0)
+						.AutoWidth()
+						.VAlign(VAlign_Center)
+						[
+							SNew(SCheckBox)
+							.OnCheckStateChanged(this, &SFbxCompareWindow::ToggleMaterialDisplay, FMaterialCompareData::EMaterialCompareDisplayOption::IndexChanged)
+							.IsChecked(this, &SFbxCompareWindow::IsToggleMaterialDisplayChecked, FMaterialCompareData::EMaterialCompareDisplayOption::IndexChanged)
+							[
+								SNew(STextBlock)
+								.Text(LOCTEXT("SFbxCompareWindow_Display_Option_IndexChanged", "Index Changed"))
+								.ColorAndOpacity(FSlateColor(FLinearColor::Yellow))
+								.ToolTipText(LOCTEXT("SFbxCompareWindow_Display_Option_IndexChanged_tooltip", "Can impact gameplay code using index base material."))
+							]
+						]
+						+SHorizontalBox::Slot()
+						.Padding(5, 0, 0, 0)
+						.AutoWidth()
+						.VAlign(VAlign_Center)
+						[
+							SNew(SCheckBox)
+							.OnCheckStateChanged(this, &SFbxCompareWindow::ToggleMaterialDisplay, FMaterialCompareData::EMaterialCompareDisplayOption::SkinxxError)
+							.IsChecked(this, &SFbxCompareWindow::IsToggleMaterialDisplayChecked, FMaterialCompareData::EMaterialCompareDisplayOption::SkinxxError)
+							[
+								SNew(STextBlock)
+								.Text(LOCTEXT("SFbxCompareWindow_Display_Option_SkinxxError", "SkinXX Error"))
+								.ColorAndOpacity(FSlateColor(FLinearColor::Red))
+								.ToolTipText(LOCTEXT("SFbxCompareWindow_Display_Option_SkinxxError_tooltip", "The list of materials will not be re-order correctly."))
+							]
+						]
+					]
+				]
+			]
+		]
+	];
+}
+
+void SFbxCompareWindow::ToggleMaterialDisplay(ECheckBoxState InNewValue, FMaterialCompareData::EMaterialCompareDisplayOption InDisplayOption)
+{
+	//Cannot uncheck a radio button
+	if (InNewValue != ECheckBoxState::Checked)
+	{
+		return;
+	}
+	CurrentDisplayOption = InDisplayOption;
+	for (TSharedPtr<FMaterialCompareData> CompareMaterial : CompareMaterialListItem)
+	{
+		CompareMaterial->CompareDisplayOption = CurrentDisplayOption;
+	}
+}
+
+ECheckBoxState SFbxCompareWindow::IsToggleMaterialDisplayChecked(FMaterialCompareData::EMaterialCompareDisplayOption InDisplayOption) const
+{
+	return CurrentDisplayOption == InDisplayOption ? ECheckBoxState::Checked : ECheckBoxState::Unchecked;
+}
+
+TSharedRef<ITableRow> SFbxCompareWindow::OnGenerateRowGeneralFbxInfo(TSharedPtr<FString> InItem, const TSharedRef<STableViewBase>& OwnerTable)
+{
+	int32 GeneralListIndex = GeneralListItem.Find(InItem);
+	bool LightBackgroundColor = GeneralListIndex % 2 == 1;
+	return SNew(STableRow<TSharedPtr<FString> >, OwnerTable)
+		[
+			SNew(STextBlock).Text(FText::FromString(*(InItem.Get())))
+		];
+}
+
+TSharedRef<ITableRow> SFbxCompareWindow::OnGenerateRowForCompareMaterialList(TSharedPtr<FMaterialCompareData> RowData, const TSharedRef<STableViewBase>& Table)
+{
+	TSharedRef< SCompareRowDataTableListViewRow > ReturnRow = SNew(SCompareRowDataTableListViewRow, Table)
+		.CompareRowData(RowData);
+	return ReturnRow;
+}
+
+FSlateColor FMaterialCompareData::GetCellColor(FCompMesh *DataA, int32 MaterialIndexA, int32 MaterialMatchA, FCompMesh *DataB, int32 MaterialIndexB, bool bSkinxxError) const
+{
+	if (!DataA->CompMaterials.IsValidIndex(MaterialIndexA))
+	{
+		return FSlateColor::UseForeground();
+	}
+
+	bool bMatchIndexChanged = MaterialMatchA == INDEX_NONE || (DataA->CompMaterials.IsValidIndex(MaterialIndexA) && DataB->CompMaterials.IsValidIndex(MaterialIndexB) && MaterialMatchA == MaterialIndexB);
+
+	if ((CompareDisplayOption == NoMatch || CompareDisplayOption == All) && MaterialMatchA == INDEX_NONE)
+	{
+		//There is no match for this material, so it will be add to the material array
+		return FSlateColor(FLinearColor(0.7f, 0.3f, 0.0f));
+	}
+	if ((CompareDisplayOption == IndexChanged || CompareDisplayOption == All) && !bMatchIndexChanged)
+	{
+		//The match index has changed, so index base gameplay will be broken
+		return FSlateColor(FLinearColor::Yellow);
+	}
+	if ((CompareDisplayOption == SkinxxError || CompareDisplayOption == All) && bSkinxxError)
+	{
+		//Skinxx error
+		return FSlateColor(FLinearColor::Red);
+	}
+	return FSlateColor::UseForeground();
+}
+
+TSharedRef<SWidget> FMaterialCompareData::ConstructCell(FCompMesh *MeshData, int32 MeshMaterialIndex, bool bSkinxxDuplicate, bool bSkinxxMissing)
+{
+	if (!MeshData->CompMaterials.IsValidIndex(MeshMaterialIndex))
+	{
+		return SNew(SBox)
+			.Padding(FMargin(5.0f, 0.0f, 0.0f, 0.0f))
+			[
+				SNew(STextBlock)
+				.Text(LOCTEXT("FMaterialCompareData_EmptyCell", ""))
+			];
+	}
+	FString CellContent = MeshData->CompMaterials[MeshMaterialIndex].ImportedMaterialSlotName.ToString();
+	FString CellTooltip = TEXT("Material Slot Name: ") + MeshData->CompMaterials[MeshMaterialIndex].MaterialSlotName.ToString();
+	if (bSkinxxDuplicate)
+	{
+		CellTooltip += TEXT(" (skinxx duplicate)");
+	}
+	if (bSkinxxMissing)
+	{
+		CellTooltip += TEXT(" (skinxx missing)");
+	}
+
+	return SNew(SBox)
+		.Padding(FMargin(5.0f, 0.0f, 0.0f, 0.0f))
+		[
+			SNew(STextBlock)
+			.Text(FText::FromString(CellContent))
+		.ToolTipText(FText::FromString(CellTooltip))
+		.ColorAndOpacity(this, MeshData == CurrentData ? &FMaterialCompareData::GetCurrentCellColor : &FMaterialCompareData::GetFbxCellColor)
+		];
+}
+
+FSlateColor FMaterialCompareData::GetCurrentCellColor() const
+{
+	return GetCellColor(CurrentData, CurrentMaterialIndex, CurrentMaterialMatch, FbxData, FbxMaterialIndex, bCurrentSkinxxDuplicate || bCurrentSkinxxMissing);
+}
+
+TSharedRef<SWidget> FMaterialCompareData::ConstructCellCurrent()
+{
+	return ConstructCell(CurrentData, CurrentMaterialIndex, bCurrentSkinxxDuplicate, bCurrentSkinxxMissing);
+}
+
+FSlateColor FMaterialCompareData::GetFbxCellColor() const
+{
+	return GetCellColor(FbxData, FbxMaterialIndex, FbxMaterialMatch, CurrentData, CurrentMaterialIndex, bFbxSkinxxDuplicate || bFbxSkinxxMissing);
+}
+
+TSharedRef<SWidget> FMaterialCompareData::ConstructCellFbx()
+{
+	return ConstructCell(FbxData, FbxMaterialIndex, bFbxSkinxxDuplicate, bFbxSkinxxMissing);
+}
+
+void SFbxCompareWindow::FillGeneralListItem()
+{
+	GeneralListItem.Add(MakeShareable(new FString(FbxGeneralInfo.UE4SdkVersion)));
+	GeneralListItem.Add(MakeShareable(new FString(FbxGeneralInfo.ApplicationCreator)));
+	GeneralListItem.Add(MakeShareable(new FString(FbxGeneralInfo.CreationDate)));
+	GeneralListItem.Add(MakeShareable(new FString(FbxGeneralInfo.FileVersion)));
+	GeneralListItem.Add(MakeShareable(new FString(FbxGeneralInfo.AxisSystem)));
+	GeneralListItem.Add(MakeShareable(new FString(FbxGeneralInfo.UnitSystem)));
+	GeneralListItem.Add(MakeShareable(new FString(TEXT("Unskinned Mesh Count:    ") + FString::FromInt(FbxSceneInfo->NonSkinnedMeshNum))));
+	GeneralListItem.Add(MakeShareable(new FString(TEXT("Skinned Count:    ") + FString::FromInt(FbxSceneInfo->SkinnedMeshNum))));
+	GeneralListItem.Add(MakeShareable(new FString(TEXT("Material Count:    ") + FString::FromInt(FbxSceneInfo->TotalMaterialNum))));
+	FString HasAnimationStr = TEXT("Has Animation:    ");
+	HasAnimationStr += FbxSceneInfo->bHasAnimation ? TEXT("True") : TEXT("False");
+	GeneralListItem.Add(MakeShareable(new FString(HasAnimationStr)));
+	if (FbxSceneInfo->bHasAnimation)
+	{
+		TArray<FStringFormatArg> Args;
+		Args.Add(FbxSceneInfo->TotalTime);
+		FString AnimationTimeStr = FString::Format(TEXT("Animation Time:    {0}"), Args);
+		GeneralListItem.Add(MakeShareable(new FString(AnimationTimeStr)));
+
+		Args.Empty();
+		Args.Add(FbxSceneInfo->FrameRate);
+		FString AnimationRateStr = FString::Format(TEXT("Animation Rate:    {0}"), Args);
+		GeneralListItem.Add(MakeShareable(new FString(AnimationRateStr)));
+	}
+}
+
+/*
+* Return true if there is some skinxx error. Both array will be allocate to the size of the materials array of MeshData
+*/
+bool SFbxCompareWindow::FindSkinxxErrors(FCompMesh *MeshData, TArray<bool> &DuplicateSkinxxMaterialNames, TArray<bool> &MissingSkinxxSuffixeMaterialNames)
+{
+	MissingSkinxxSuffixeMaterialNames.Empty();
+	MissingSkinxxSuffixeMaterialNames.AddZeroed(MeshData->CompMaterials.Num());
+	DuplicateSkinxxMaterialNames.Empty();
+	DuplicateSkinxxMaterialNames.AddZeroed(MeshData->CompMaterials.Num());
+	TArray<int32> SkinxxErrorIndexes;
+	bool bContainSkinxxIndex = false;
+	for (FCompMaterial CompMaterial : MeshData->CompMaterials)
+	{
+		if (CompMaterial.ImportedMaterialSlotName == NAME_None)
+		{
+			continue;
+		}
+		FString ImportedMaterialName = CompMaterial.ImportedMaterialSlotName.ToString();
+		int32 Offset = ImportedMaterialName.Find(TEXT("_SKIN"), ESearchCase::IgnoreCase, ESearchDir::FromEnd);
+		if (Offset != INDEX_NONE)
+		{
+			FString SkinXXNumber = ImportedMaterialName.Right(ImportedMaterialName.Len() - (Offset + 1)).RightChop(4);
+			if (SkinXXNumber.IsNumeric())
+			{
+				bContainSkinxxIndex = true;
+				break;
+			}
+		}
+	}
+
+	//There is no skinxx suffixe, so no skinxx error
+	if (!bContainSkinxxIndex)
+	{
+		return false;
+	}
+
+	bool bContainSkinxxError = false;
+	for (int32 MaterialNamesIndex = 0; MaterialNamesIndex < MeshData->CompMaterials.Num(); ++MaterialNamesIndex)
+	{
+		FName MaterialName = MeshData->CompMaterials[MaterialNamesIndex].ImportedMaterialSlotName;
+		if (MaterialName == NAME_None)
+		{
+			MissingSkinxxSuffixeMaterialNames[MaterialNamesIndex] = true;
+			bContainSkinxxError = true;
+			continue;
+		}
+
+		FString ImportedMaterialName = MaterialName.ToString();
+		int32 Offset = ImportedMaterialName.Find(TEXT("_SKIN"), ESearchCase::IgnoreCase, ESearchDir::FromEnd);
+		if (Offset != INDEX_NONE)
+		{
+			FString SkinXXNumber = ImportedMaterialName.Right(ImportedMaterialName.Len() - (Offset + 1)).RightChop(4);
+
+			if (SkinXXNumber.IsNumeric())
+			{
+				int32 TmpIndex = FPlatformString::Atoi(*SkinXXNumber);
+				if (SkinxxErrorIndexes.Contains(TmpIndex))
+				{
+					DuplicateSkinxxMaterialNames[MaterialNamesIndex] = true;
+					bContainSkinxxError = true;
+				}
+				SkinxxErrorIndexes.Add(TmpIndex);
+			}
+			else
+			{
+				MissingSkinxxSuffixeMaterialNames[MaterialNamesIndex] = true;
+				bContainSkinxxError = true;
+			}
+		}
+		else
+		{
+			MissingSkinxxSuffixeMaterialNames[MaterialNamesIndex] = true;
+			bContainSkinxxError = true;
+		}
+	}
+	return bContainSkinxxError;
+}
+
+void SFbxCompareWindow::FillMaterialListItem()
+{
+	TArray<bool> CurrentDuplicateSkinxx;
+	TArray<bool> CurrentMissingSkinxxSuffixe;
+	FindSkinxxErrors(CurrentMeshData, CurrentDuplicateSkinxx, CurrentMissingSkinxxSuffixe);
+	
+	TArray<bool> FbxDuplicateSkinxx;
+	TArray<bool> FbxMissingSkinxxSuffixe;
+	FindSkinxxErrors(FbxMeshData, FbxDuplicateSkinxx, FbxMissingSkinxxSuffixe);
+
+	//Build the compare data to show in the UI
+	int32 MaterialCompareRowNumber = FMath::Max(CurrentMeshData->CompMaterials.Num(), FbxMeshData->CompMaterials.Num());
+	for (int32 RowIndex = 0; RowIndex < MaterialCompareRowNumber; ++RowIndex)
+	{
+		TSharedPtr<FMaterialCompareData> CompareRowData = MakeShareable(new FMaterialCompareData());
+		CompareRowData->RowIndex = CompareMaterialListItem.Add(CompareRowData);
+		CompareRowData->CurrentData = CurrentMeshData;
+		CompareRowData->FbxData = FbxMeshData;
+		
+		CompareRowData->bCurrentSkinxxDuplicate = CurrentDuplicateSkinxx.IsValidIndex(RowIndex) && CurrentDuplicateSkinxx[RowIndex];
+		CompareRowData->bCurrentSkinxxMissing = CurrentMissingSkinxxSuffixe.IsValidIndex(RowIndex) && CurrentMissingSkinxxSuffixe[RowIndex];
+		CompareRowData->bFbxSkinxxDuplicate = FbxDuplicateSkinxx.IsValidIndex(RowIndex) && FbxDuplicateSkinxx[RowIndex];
+		CompareRowData->bFbxSkinxxDuplicate = FbxMissingSkinxxSuffixe.IsValidIndex(RowIndex) && FbxMissingSkinxxSuffixe[RowIndex];
+
+		CompareRowData->CompareDisplayOption = FMaterialCompareData::All;
+		if (CurrentMeshData->CompMaterials.IsValidIndex(RowIndex))
+		{
+			CompareRowData->CurrentMaterialIndex = RowIndex;
+			for (int32 FbxMaterialIndex = 0; FbxMaterialIndex < FbxMeshData->CompMaterials.Num(); ++FbxMaterialIndex)
+			{
+				if (FbxMeshData->CompMaterials[FbxMaterialIndex].ImportedMaterialSlotName == CurrentMeshData->CompMaterials[RowIndex].ImportedMaterialSlotName)
+				{
+					CompareRowData->CurrentMaterialMatch = FbxMaterialIndex;
+					break;
+				}
+			}
+		}
+		if (FbxMeshData->CompMaterials.IsValidIndex(RowIndex))
+		{
+			CompareRowData->FbxMaterialIndex = RowIndex;
+			for (int32 CurrentMaterialIndex = 0; CurrentMaterialIndex < CurrentMeshData->CompMaterials.Num(); ++CurrentMaterialIndex)
+			{
+				if (CurrentMeshData->CompMaterials[CurrentMaterialIndex].ImportedMaterialSlotName == FbxMeshData->CompMaterials[RowIndex].ImportedMaterialSlotName)
+				{
+					CompareRowData->FbxMaterialMatch = CurrentMaterialIndex;
+					break;
+				}
+			}
+		}
+	}
+}
+
+TSharedPtr<SWidget> SFbxCompareWindow::ConstructSkeletonComparison()
+{
+	if (!PreviewObject->IsA(USkeletalMesh::StaticClass()))
+	{
+		//Return an empty widget, we do not show the skeleton when the mesh is not a skeletal mesh
+		return SNew(SBox);
+	}
+
+	FString SkeletonStatusTooltip;
+	if (AssetReferencingSkeleton.Num() > 0)
+	{
+		SkeletonStatusTooltip += TEXT("Skeleton is references by ") + FString::FromInt(AssetReferencingSkeleton.Num()) + TEXT(" assets.");
+	}
+	
+	FText SkeletonStatus = FText(FbxMeshData->CompSkeleton.bSkeletonFitMesh ? LOCTEXT("SFbxCompareWindow_ConstructSkeletonComparison_MatchAndMerge", "The skeleton can be merged") : LOCTEXT("SFbxCompareWindow_ConstructSkeletonComparison_CannotMatchAndMerge", "The skeleton must be regenerated, it cannot be merged"));
+	
+	CompareTree = SNew(STreeView< TSharedPtr<FSkeletonCompareData> >)
+		.ItemHeight(24)
+		.SelectionMode(ESelectionMode::None)
+		.TreeItemsSource(&DisplaySkeletonTreeItem)
+		.OnGenerateRow(this, &SFbxCompareWindow::OnGenerateRowCompareTreeView)
+		.OnGetChildren(this, &SFbxCompareWindow::OnGetChildrenRowCompareTreeView);
+
+
+	return SNew(SBox)
+	.MaxDesiredHeight(600)
+	[
+		SNew(SBorder)
+		.Padding(FMargin(3))
+		.BorderImage(FEditorStyle::GetBrush("ToolPanel.GroupBorder"))
+		[
+			SNew(SVerticalBox)
+			+SVerticalBox::Slot()
+			.AutoHeight()
+			.Padding(2)
+			[
+				SNew(SHorizontalBox)
+				+ SHorizontalBox::Slot()
+				.AutoWidth()
+				[
+					SNew(SButton)
+					.HAlign(HAlign_Center)
+					.VAlign(VAlign_Center)
+					.IsFocusable(false)
+					.ButtonStyle(FEditorStyle::Get(), "NoBorder")
+					.OnClicked(this, &SFbxCompareWindow::SetSectionVisible, EFBXCompareSection_Skeleton)
+					[
+						SNew(SImage).Image(this, &SFbxCompareWindow::GetCollapsableArrow, EFBXCompareSection_Skeleton)
+					]
+				]
+				+ SHorizontalBox::Slot()
+				.AutoWidth()
+				[
+					SNew(STextBlock)
+					.Font(FEditorStyle::GetFontStyle("DetailsView.CategoryFontStyle"))
+					.Text(LOCTEXT("SFbxCompareWindow_SkeletonCompareHeader", "Skeleton"))
+				]
+			]
+			+SVerticalBox::Slot()
+			.FillHeight(1.0f)
+			.Padding(2)
+			[
+				SNew(SBox)
+				.Visibility(TAttribute<EVisibility>::Create(TAttribute<EVisibility>::FGetter::CreateSP(this, &SFbxCompareWindow::IsSectionVisible, EFBXCompareSection_Skeleton)))
+				[
+					SNew(SBorder)
+					.Padding(FMargin(3))
+					.BorderImage(FEditorStyle::GetBrush("ToolPanel.DarkGroupBorder"))
+					[
+						SNew(SVerticalBox)
+						+SVerticalBox::Slot()
+						.AutoHeight()
+						.Padding(2)
+						[
+							SNew(STextBlock)
+							.Font(FEditorStyle::GetFontStyle("DetailsView.CategoryFontStyle"))
+							.Text(SkeletonStatus)
+							.ToolTipText(FText::FromString(SkeletonStatusTooltip))
+							.ColorAndOpacity(FbxMeshData->CompSkeleton.bSkeletonFitMesh ? FSlateColor::UseForeground() : FSlateColor(FLinearColor(0.7f, 0.3f, 0.0f)))
+						]
+						+SVerticalBox::Slot()
+						.AutoHeight()
+						.Padding(2)
+						[
+							SNew(SSeparator)
+							.Orientation(EOrientation::Orient_Horizontal)
+						]
+						+SVerticalBox::Slot()
+						.FillHeight(1.0f)
+						.Padding(2)
+						[
+							SNew(SVerticalBox)
+							+SVerticalBox::Slot()
+							.FillHeight(1.0f)
+							[
+								CompareTree.ToSharedRef()
+							]
+							+ SVerticalBox::Slot()
+							.AutoHeight()
+							.Padding(2)
+							[
+								SNew(SSeparator)
+								.Orientation(EOrientation::Orient_Horizontal)
+							]
+							+SVerticalBox::Slot()
+							.AutoHeight()
+							.MaxHeight(200.0f)
+							[
+								//Show the general fbx information
+								SNew(SListView<TSharedPtr<FString>>)
+								.ListItemsSource(&AssetReferencingSkeleton)
+								.OnGenerateRow(this, &SFbxCompareWindow::OnGenerateRowAssetReferencingSkeleton)
+							]
+						]
+					]
+				]
+			]
+		]
+	];
+}
+
+class SCompareSkeletonTreeViewItem : public STableRow< TSharedPtr<FSkeletonCompareData> >
+{
+public:
+
+	SLATE_BEGIN_ARGS(SCompareSkeletonTreeViewItem)
+		: _SkeletonCompareData(nullptr)
+		, _CurrentMeshData(nullptr)
+		, _FbxMeshData(nullptr)
+	{}
+
+	/** The item content. */
+	SLATE_ARGUMENT(TSharedPtr<FSkeletonCompareData>, SkeletonCompareData)
+	SLATE_ARGUMENT(FCompMesh*, CurrentMeshData)
+	SLATE_ARGUMENT(FCompMesh*, FbxMeshData)
+	SLATE_END_ARGS()
+
+	/**
+	* Construct the widget
+	*
+	* @param InArgs   A declaration from which to construct the widget
+	*/
+	void Construct(const FArguments& InArgs, const TSharedRef<STableViewBase>& InOwnerTableView)
+	{
+		SkeletonCompareData = InArgs._SkeletonCompareData;
+		CurrentMeshData = InArgs._CurrentMeshData;
+		FbxMeshData = InArgs._FbxMeshData;
+
+		//This is suppose to always be valid
+		check(SkeletonCompareData.IsValid());
+		check(CurrentMeshData != nullptr);
+		check(FbxMeshData != nullptr);
+
+		const FSlateBrush* JointIcon = SkeletonCompareData->bMatchJoint ? FEditorStyle::GetDefaultBrush() : SkeletonCompareData->FbxJointIndex != INDEX_NONE ? FEditorStyle::GetBrush("FBXIcon.ReimportCompareAdd") : FEditorStyle::GetBrush("FBXIcon.ReimportCompareRemoved");
+
+		//Prepare the tooltip
+		FString Tooltip = SkeletonCompareData->bMatchJoint ? TEXT("") : FText(SkeletonCompareData->FbxJointIndex != INDEX_NONE ? LOCTEXT("SCompareSkeletonTreeViewItem_AddJoint_tooltip", "Fbx reimport will add this joint") : LOCTEXT("SCompareSkeletonTreeViewItem_RemoveJoint_tooltip", "Fbx reimport will remove this joint")).ToString();
+
+		this->ChildSlot
+			[
+				SNew(SHorizontalBox)
+				+ SHorizontalBox::Slot()
+				.AutoWidth()
+				[
+					SNew(SExpanderArrow, SharedThis(this))
+				]
+
+				+ SHorizontalBox::Slot()
+				.AutoWidth()
+				.Padding(0.0f, 2.0f, 6.0f, 2.0f)
+				[
+					SNew(SImage)
+					.Image(JointIcon)
+					.Visibility(JointIcon != FEditorStyle::GetDefaultBrush() ? EVisibility::Visible : EVisibility::Collapsed)
+				]
+				+ SHorizontalBox::Slot()
+				.FillWidth(1.0f)
+				.Padding(0.0f, 3.0f, 6.0f, 3.0f)
+				.VAlign(VAlign_Center)
+				[
+					SNew(STextBlock)
+					.Text(FText::FromString(SkeletonCompareData->JointName.ToString()))
+					.ToolTipText(FText::FromString(Tooltip))
+					.ColorAndOpacity(SkeletonCompareData->bMatchJoint && !SkeletonCompareData->bChildConflict ? FSlateColor::UseForeground() : FSlateColor(FLinearColor(0.7f, 0.3f, 0.0f)))
+				]
+			];
+
+		STableRow< TSharedPtr<FSkeletonCompareData> >::ConstructInternal(
+			STableRow::FArguments()
+			.ShowSelection(true),
+			InOwnerTableView
+			);
+	}
+
+private:
+	/** The node info to build the tree view row from. */
+	TSharedPtr<FSkeletonCompareData> SkeletonCompareData;
+	FCompMesh *CurrentMeshData;
+	FCompMesh *FbxMeshData;
+};
+
+
+TSharedRef<ITableRow> SFbxCompareWindow::OnGenerateRowCompareTreeView(TSharedPtr<FSkeletonCompareData> RowData, const TSharedRef<STableViewBase>& Table)
+{
+	TSharedRef< SCompareSkeletonTreeViewItem > ReturnRow = SNew(SCompareSkeletonTreeViewItem, Table)
+		.SkeletonCompareData(RowData)
+		.CurrentMeshData(CurrentMeshData)
+		.FbxMeshData(FbxMeshData);
+	return ReturnRow;
+}
+
+void SFbxCompareWindow::OnGetChildrenRowCompareTreeView(TSharedPtr<FSkeletonCompareData> InParent, TArray< TSharedPtr<FSkeletonCompareData> >& OutChildren)
+{
+	for (int32 ChildIndex = 0; ChildIndex < InParent->ChildJoints.Num(); ++ChildIndex)
+	{
+		TSharedPtr<FSkeletonCompareData> ChildJoint = InParent->ChildJoints[ChildIndex];
+		if (ChildJoint.IsValid())
+		{
+			OutChildren.Add(ChildJoint);
+		}
+	}
+}
+
+TSharedRef<ITableRow> SFbxCompareWindow::OnGenerateRowAssetReferencingSkeleton(TSharedPtr<FString> InItem, const TSharedRef<STableViewBase>& OwnerTable)
+{
+	int32 AssetListIndex = AssetReferencingSkeleton.Find(InItem);
+	bool LightBackgroundColor = AssetListIndex % 2 == 1;
+	return SNew(STableRow<TSharedPtr<FString> >, OwnerTable)
+		[
+			SNew(SBorder)
+			.BorderImage(LightBackgroundColor ? FEditorStyle::GetBrush("ToolPanel.GroupBorder") : FEditorStyle::GetBrush("ToolPanel.DarkGroupBorder"))
+			[
+				SNew(STextBlock)
+				.Text(FText::FromString(*(InItem.Get())))
+				//.ColorAndOpacity(LightBackgroundColor ? FSlateColor::UseSubduedForeground() : FSlateColor::UseForeground())
+			]
+		];
+}
+
+void SFbxCompareWindow::FilSkeletonTreeItem()
+{
+	//Create all the entries
+	for (int32 RowIndex = 0; RowIndex < CurrentMeshData->CompSkeleton.Joints.Num(); ++RowIndex)
+	{
+		TSharedPtr<FSkeletonCompareData> CompareRowData = MakeShareable(new FSkeletonCompareData());
+		int32 AddedIndex = CurrentSkeletonTreeItem.Add(CompareRowData);
+		check(AddedIndex == RowIndex);
+		CompareRowData->CurrentJointIndex = RowIndex;
+		CompareRowData->JointName = CurrentMeshData->CompSkeleton.Joints[RowIndex].Name;
+		CompareRowData->ChildJointIndexes = CurrentMeshData->CompSkeleton.Joints[RowIndex].ChildIndexes;
+	}
+
+	//Set the childrens and parent pointer
+	for (int32 RowIndex = 0; RowIndex < CurrentMeshData->CompSkeleton.Joints.Num(); ++RowIndex)
+	{
+		check(CurrentSkeletonTreeItem.IsValidIndex(RowIndex));
+		TSharedPtr<FSkeletonCompareData> CompareRowData = CurrentSkeletonTreeItem[RowIndex];
+		if (CurrentSkeletonTreeItem.IsValidIndex(CurrentMeshData->CompSkeleton.Joints[RowIndex].ParentIndex))
+		{
+			CompareRowData->ParentJoint = CurrentSkeletonTreeItem[CurrentMeshData->CompSkeleton.Joints[RowIndex].ParentIndex];
+		}
+		
+		for (int32 ChildJointIndex = 0; ChildJointIndex < CompareRowData->ChildJointIndexes.Num(); ++ChildJointIndex)
+		{
+			if (CurrentSkeletonTreeItem.IsValidIndex(CompareRowData->ChildJointIndexes[ChildJointIndex]))
+			{
+				CompareRowData->ChildJoints.Add(CurrentSkeletonTreeItem[CompareRowData->ChildJointIndexes[ChildJointIndex]]);
+			}
+		}
+	}
+
+	for (int32 RowIndex = 0; RowIndex < FbxMeshData->CompSkeleton.Joints.Num(); ++RowIndex)
+	{
+		TSharedPtr<FSkeletonCompareData> CompareRowData = MakeShareable(new FSkeletonCompareData());
+		int32 AddedIndex = FbxSkeletonTreeItem.Add(CompareRowData);
+		check(AddedIndex == RowIndex);
+		CompareRowData->FbxJointIndex = RowIndex;
+		CompareRowData->JointName = FbxMeshData->CompSkeleton.Joints[RowIndex].Name;
+		CompareRowData->ChildJointIndexes = FbxMeshData->CompSkeleton.Joints[RowIndex].ChildIndexes;
+	}
+
+	//Set the childrens and parent pointer
+	for (int32 RowIndex = 0; RowIndex < FbxMeshData->CompSkeleton.Joints.Num(); ++RowIndex)
+	{
+		check(FbxSkeletonTreeItem.IsValidIndex(RowIndex));
+		TSharedPtr<FSkeletonCompareData> CompareRowData = FbxSkeletonTreeItem[RowIndex];
+		if (FbxSkeletonTreeItem.IsValidIndex(FbxMeshData->CompSkeleton.Joints[RowIndex].ParentIndex))
+		{
+			CompareRowData->ParentJoint = FbxSkeletonTreeItem[FbxMeshData->CompSkeleton.Joints[RowIndex].ParentIndex];
+		}
+
+		for (int32 ChildJointIndex = 0; ChildJointIndex < CompareRowData->ChildJointIndexes.Num(); ++ChildJointIndex)
+		{
+			if (FbxSkeletonTreeItem.IsValidIndex(CompareRowData->ChildJointIndexes[ChildJointIndex]))
+			{
+				CompareRowData->ChildJoints.Add(FbxSkeletonTreeItem[CompareRowData->ChildJointIndexes[ChildJointIndex]]);
+			}
+		}
+	}
+}
+
+void SFbxCompareWindow::RecursiveMatchJointInfo(TSharedPtr<FSkeletonCompareData> SkeletonItem)
+{
+	TArray<TSharedPtr<FSkeletonCompareData>> DisplayChilds;
+	//Find the display child
+	if (CurrentSkeletonTreeItem.IsValidIndex(SkeletonItem->CurrentJointIndex))
+	{
+		for (int32 ChildIndex = 0; ChildIndex < CurrentSkeletonTreeItem[SkeletonItem->CurrentJointIndex]->ChildJoints.Num(); ++ChildIndex)
+		{
+			DisplayChilds.Add(CurrentSkeletonTreeItem[SkeletonItem->CurrentJointIndex]->ChildJoints[ChildIndex]);
+		}
+	}
+	if (FbxSkeletonTreeItem.IsValidIndex(SkeletonItem->FbxJointIndex))
+	{
+		for (int32 ChildIndex = 0; ChildIndex < FbxSkeletonTreeItem[SkeletonItem->FbxJointIndex]->ChildJoints.Num(); ++ChildIndex)
+		{
+			TSharedPtr<FSkeletonCompareData> FbxSkeletonItem = FbxSkeletonTreeItem[SkeletonItem->FbxJointIndex]->ChildJoints[ChildIndex];
+			bool bFoundChildMatch = false;
+			for (TSharedPtr<FSkeletonCompareData> DisplayChildJoint : DisplayChilds)
+			{
+				if (DisplayChildJoint->JointName == FbxSkeletonItem->JointName)
+				{
+					bFoundChildMatch = true;
+					DisplayChildJoint->bMatchJoint = true;
+					DisplayChildJoint->FbxJointIndex = FbxSkeletonItem->FbxJointIndex;
+					break;
+				}
+			}
+			if (!bFoundChildMatch)
+			{
+				DisplayChilds.Add(FbxSkeletonItem);
+			}
+		}
+	}
+
+	if (!SkeletonItem->bMatchJoint)
+	{
+		TSharedPtr<FSkeletonCompareData> ParentSkeletonItem = SkeletonItem->ParentJoint;
+		while (ParentSkeletonItem.IsValid() && ParentSkeletonItem->bChildConflict == false)
+		{
+			ParentSkeletonItem->bChildConflict = true;
+			ParentSkeletonItem = ParentSkeletonItem->ParentJoint;
+		}
+	}
+	//Set the new child list to the display joint
+	SkeletonItem->ChildJoints = DisplayChilds;
+	SkeletonItem->ChildJointIndexes.Empty();
+	for (TSharedPtr<FSkeletonCompareData> ChildJoint : SkeletonItem->ChildJoints)
+	{
+		ChildJoint->ParentJoint = SkeletonItem;
+		RecursiveMatchJointInfo(ChildJoint);
+	}
+}
+
+void SFbxCompareWindow::SetMatchJointInfo()
+{
+	TArray<TSharedPtr<FSkeletonCompareData>> RootJoint;
+	for (TSharedPtr<FSkeletonCompareData> CurrentSkeletonItem : CurrentSkeletonTreeItem)
+	{
+		if (!CurrentSkeletonItem->ParentJoint.IsValid())
+		{
+			DisplaySkeletonTreeItem.Add(CurrentSkeletonItem);
+		}
+	}
+	for (TSharedPtr<FSkeletonCompareData> CurrentSkeletonItem : FbxSkeletonTreeItem)
+	{
+		if (!CurrentSkeletonItem->ParentJoint.IsValid())
+		{
+			bool bInsertJoint = true;
+			for (TSharedPtr<FSkeletonCompareData> DisplayTreeItem : DisplaySkeletonTreeItem)
+			{
+				if(DisplayTreeItem->JointName == CurrentSkeletonItem->JointName)
+				{
+					DisplayTreeItem->FbxJointIndex = CurrentSkeletonItem->FbxJointIndex;
+					DisplayTreeItem->bMatchJoint = true;
+					bInsertJoint = false;
+				}
+			}
+			if (bInsertJoint)
+			{
+				DisplaySkeletonTreeItem.Add(CurrentSkeletonItem);
+			}
+		}
+	}
+
+	for (int32 SkeletonTreeIndex = 0; SkeletonTreeIndex < DisplaySkeletonTreeItem.Num(); ++SkeletonTreeIndex)
+	{
+		RecursiveMatchJointInfo(DisplaySkeletonTreeItem[SkeletonTreeIndex]);
+	}
+}
+#undef LOCTEXT_NAMESPACE
diff --git a/Engine/Source/Editor/UnrealEd/Private/Fbx/FbxFactory.cpp b/Engine/Source/Editor/UnrealEd/Private/Fbx/FbxFactory.cpp
--- a/Engine/Source/Editor/UnrealEd/Private/Fbx/FbxFactory.cpp
+++ b/Engine/Source/Editor/UnrealEd/Private/Fbx/FbxFactory.cpp
@@ -243,7 +243,7 @@ UObject* UFbxFactory::FactoryCreateBinary
 	bool bIsAutomated = IsAutomatedImport();
 	bool bShowImportDialog = bShowOption && !bIsAutomated;
 	bool bImportAll = false;
-	ImportOptions = GetImportOptions(FbxImporter, ImportUI, bShowImportDialog, bIsAutomated, InParent->GetPathName(), bOperationCanceled, bImportAll, bIsObjFormat, bIsObjFormat, ForcedImportType);
+	ImportOptions = GetImportOptions(FbxImporter, ImportUI, bShowImportDialog, bIsAutomated, InParent->GetPathName(), bOperationCanceled, bImportAll, bIsObjFormat, bIsObjFormat, ForcedImportType, ExistingObject);
 	bOutOperationCanceled = bOperationCanceled;
 	
 	if( bImportAll )
@@ -890,5 +890,13 @@ void UFbxImportUI::ParseFromJson(TSharedRef<class FJsonObject> ImportSettingsJso
 	}
 }
 
+void UFbxImportUI::ResetToDefault()
+{
+	ReloadConfig();
+	AnimSequenceImportData->ReloadConfig();
+	StaticMeshImportData->ReloadConfig();
+	SkeletalMeshImportData->ReloadConfig();
+	TextureImportData->ReloadConfig();
+}
 
 #undef LOCTEXT_NAMESPACE
diff --git a/Engine/Source/Editor/UnrealEd/Private/Fbx/FbxMainImport.cpp b/Engine/Source/Editor/UnrealEd/Private/Fbx/FbxMainImport.cpp
--- a/Engine/Source/Editor/UnrealEd/Private/Fbx/FbxMainImport.cpp
+++ b/Engine/Source/Editor/UnrealEd/Private/Fbx/FbxMainImport.cpp
@@ -37,9 +37,11 @@ namespace UnFbx
 
 TSharedPtr<FFbxImporter> FFbxImporter::StaticInstance;
 
+TSharedPtr<FFbxImporter> FFbxImporter::StaticPreviewInstance;
 
 
-FBXImportOptions* GetImportOptions( UnFbx::FFbxImporter* FbxImporter, UFbxImportUI* ImportUI, bool bShowOptionDialog, bool bIsAutomated, const FString& FullPath, bool& OutOperationCanceled, bool& bOutImportAll, bool bIsObjFormat, bool bForceImportType, EFBXImportType ImportType )
+
+FBXImportOptions* GetImportOptions( UnFbx::FFbxImporter* FbxImporter, UFbxImportUI* ImportUI, bool bShowOptionDialog, bool bIsAutomated, const FString& FullPath, bool& OutOperationCanceled, bool& bOutImportAll, bool bIsObjFormat, bool bForceImportType, EFBXImportType ImportType, UObject* ReimportObject)
 {
 	OutOperationCanceled = false;
 
@@ -100,10 +102,17 @@ FBXImportOptions* GetImportOptions( UnFbx::FFbxImporter* FbxImporter, UFbxImport
 
 		TSharedRef<SWindow> Window = SNew(SWindow)
 			.Title(NSLOCTEXT("UnrealEd", "FBXImportOpionsTitle", "FBX Import Options"))
-			.SizingRule( ESizingRule::Autosized )
+			.SizingRule(ESizingRule::Autosized)
 			.AutoCenter(EAutoCenter::None)
 			.ScreenPosition(WindowPosition);
 		
+		auto OnPreviewFbxImportLambda = FOnPreviewFbxImport::CreateLambda([=]
+		{
+			UnFbx::FFbxImporter* PreviewFbxImporter = UnFbx::FFbxImporter::GetPreviewInstance();
+			PreviewFbxImporter->ShowFbxReimportPreview(ReimportObject, ImportUI, FullPath);
+			UnFbx::FFbxImporter::DeletePreviewInstance();
+		});
+
 		TSharedPtr<SFbxOptionWindow> FbxOptionWindow;
 		Window->SetContent
 		(
@@ -115,6 +124,7 @@ FBXImportOptions* GetImportOptions( UnFbx::FFbxImporter* FbxImporter, UFbxImport
 			.IsObjFormat( bIsObjFormat )
 			.MaxWindowHeight(FbxImportWindowHeight)
 			.MaxWindowWidth(FbxImportWindowWidth)
+			.OnPreviewFbxImport(OnPreviewFbxImportLambda)
 		);
 
 		// @todo: we can make this slow as showing progress bar later
@@ -395,6 +405,20 @@ void FFbxImporter::DeleteInstance()
 	StaticInstance.Reset();
 }
 
+FFbxImporter* FFbxImporter::GetPreviewInstance()
+{
+	if (!StaticPreviewInstance.IsValid())
+	{
+		StaticPreviewInstance = MakeShareable(new FFbxImporter());
+	}
+	return StaticPreviewInstance.Get();
+}
+
+void FFbxImporter::DeletePreviewInstance()
+{
+	StaticPreviewInstance.Reset();
+}
+
 //-------------------------------------------------------------------------
 //
 //-------------------------------------------------------------------------
@@ -991,10 +1015,6 @@ bool FFbxImporter::ImportFile(FString Filename, bool bPreventMaterialNameClash /
 	{
 		UE_LOG(LogFbx, Log, TEXT("FBX Scene Loaded Succesfully"));
 		CurPhase = IMPORTED;
-		
-		// Release importer now as it is unneeded
-		Importer->Destroy();
-		Importer = NULL;
 	}
 	else
 	{
@@ -1011,6 +1031,10 @@ bool FFbxImporter::ImportFile(FString Filename, bool bPreventMaterialNameClash /
 
 void FFbxImporter::ConvertScene()
 {
+	//Set the original file information
+	FileAxisSystem = Scene->GetGlobalSettings().GetAxisSystem();
+	FileUnitSystem = Scene->GetGlobalSettings().GetSystemUnit();
+
 	if (GetImportOptions()->bConvertScene)
 	{
 		// we use -Y as forward axis here when we import. This is odd considering our forward axis is technically +X
diff --git a/Engine/Source/Editor/UnrealEd/Private/Fbx/FbxOptionWindow.cpp b/Engine/Source/Editor/UnrealEd/Private/Fbx/FbxOptionWindow.cpp
--- a/Engine/Source/Editor/UnrealEd/Private/Fbx/FbxOptionWindow.cpp
+++ b/Engine/Source/Editor/UnrealEd/Private/Fbx/FbxOptionWindow.cpp
@@ -19,10 +19,12 @@ void SFbxOptionWindow::Construct(const FArguments& InArgs)
 	ImportUI = InArgs._ImportUI;
 	WidgetWindow = InArgs._WidgetWindow;
 	bIsObjFormat = InArgs._IsObjFormat;
+	OnPreviewFbxImport = InArgs._OnPreviewFbxImport;
 
 	check (ImportUI);
 	
 	TSharedPtr<SBox> ImportTypeDisplay;
+	TSharedPtr<SHorizontalBox> FbxHeaderButtons;
 	TSharedPtr<SBox> InspectorBox;
 	this->ChildSlot
 	[
@@ -116,41 +118,82 @@ void SFbxOptionWindow::Construct(const FArguments& InArgs)
 	FDetailsViewArgs DetailsViewArgs;
 	DetailsViewArgs.bAllowSearch = false;
 	DetailsViewArgs.NameAreaSettings = FDetailsViewArgs::HideNameArea;
-	TSharedPtr<IDetailsView> DetailsView = PropertyEditorModule.CreateDetailView(DetailsViewArgs);
+	DetailsView = PropertyEditorModule.CreateDetailView(DetailsViewArgs);
 
 	InspectorBox->SetContent(DetailsView->AsShared());
 
-	if (ImportUI->bIsReimport)
-	{
-		ImportTypeDisplay->SetContent(
-			SNew(SBorder)
-			.Padding(FMargin(3))
-			.BorderImage(FEditorStyle::GetBrush("ToolPanel.GroupBorder"))
+	ImportTypeDisplay->SetContent(
+		SNew(SBorder)
+		.Padding(FMargin(3))
+		.BorderImage(FEditorStyle::GetBrush("ToolPanel.GroupBorder"))
+		[
+			SNew(SHorizontalBox)
+			+ SHorizontalBox::Slot()
+			.VAlign(VAlign_Center)
 			[
-				SNew(SHorizontalBox)
-				+ SHorizontalBox::Slot()
-				.VAlign(VAlign_Center)
+				SNew(STextBlock)
+				.Text(this, &SFbxOptionWindow::GetImportTypeDisplayText)
+			]
+			+ SHorizontalBox::Slot()
+			[
+				SNew(SBox)
+				.HAlign(HAlign_Right)
 				[
-					SNew(STextBlock)
-					.Text(this, &SFbxOptionWindow::GetImportTypeDisplayText)
+					SAssignNew(FbxHeaderButtons, SHorizontalBox)
+					+ SHorizontalBox::Slot()
+					.AutoWidth()
+					.Padding(FMargin(2.0f, 0.0f))
+					[
+						SNew(SButton)
+						.Text(LOCTEXT("FbxOptionWindow_ResetOptions", "Reset to Default"))
+						.OnClicked(this, &SFbxOptionWindow::OnResetToDefaultClick)
+					]
 				]
 			]
-		);
+		]
+	);
+
+	if (ImportUI->bIsReimport && OnPreviewFbxImport.IsBound())
+	{
+		FbxHeaderButtons->AddSlot()
+		.AutoWidth()
+		.Padding(FMargin(2.0f, 0.0f))
+		[
+			//Create the fbx import preview button
+			SNew(SButton)
+			.Text(LOCTEXT("FbxOptionWindow_Preview", "Preview..."))
+			.OnClicked(this, &SFbxOptionWindow::OnPreviewClick)
+		];
 	}
 
 	DetailsView->SetObject(ImportUI);
 }
 
+FReply SFbxOptionWindow::OnPreviewClick() const
+{
+	//Pop a preview window to let the user see the content of the fbx file
+	OnPreviewFbxImport.ExecuteIfBound();
+	return FReply::Handled();
+}
+
+FReply SFbxOptionWindow::OnResetToDefaultClick() const
+{
+	ImportUI->ResetToDefault();
+	//Refresh the view to make sure the custom UI are updating correctly
+	DetailsView->SetObject(ImportUI, true);
+	return FReply::Handled();
+}
+
 FText SFbxOptionWindow::GetImportTypeDisplayText() const
 {
 	switch (ImportUI->MeshTypeToImport)
 	{
 	case EFBXImportType::FBXIT_Animation :
-		return LOCTEXT("FbxOptionWindow_ImportTypeAnim", "Reimport Animation");
+		return ImportUI->bIsReimport ? LOCTEXT("FbxOptionWindow_ReImportTypeAnim", "Reimport Animation") : LOCTEXT("FbxOptionWindow_ImportTypeAnim", "Import Animation");
 	case EFBXImportType::FBXIT_SkeletalMesh:
-		return LOCTEXT("FbxOptionWindow_ImportTypeSK", "Reimport Skeletal Mesh");
+		return ImportUI->bIsReimport ? LOCTEXT("FbxOptionWindow_ReImportTypeSK", "Reimport Skeletal Mesh") : LOCTEXT("FbxOptionWindow_ImportTypeSK", "Import Skeletal Mesh");
 	case EFBXImportType::FBXIT_StaticMesh:
-		return LOCTEXT("FbxOptionWindow_ImportTypeSM", "Reimport Static Mesh");
+		return ImportUI->bIsReimport ? LOCTEXT("FbxOptionWindow_ReImportTypeSM", "Reimport Static Mesh") : LOCTEXT("FbxOptionWindow_ImportTypeSM", "Import Static Mesh");
 	}
 	return FText::GetEmpty();
 }
diff --git a/Engine/Source/Editor/UnrealEd/Private/Fbx/FbxPreviewReimport.cpp b/Engine/Source/Editor/UnrealEd/Private/Fbx/FbxPreviewReimport.cpp
--- a/Engine/Source/Editor/UnrealEd/Private/Fbx/FbxPreviewReimport.cpp
+++ b/Engine/Source/Editor/UnrealEd/Private/Fbx/FbxPreviewReimport.cpp
@@ -16,6 +16,17 @@
 #include "PackageTools.h"
 #include "FbxImporter.h"
 #include "Misc/FbxErrors.h"
+#include "HAL/FileManager.h"
+#include "Factories/FbxSceneImportFactory.h"
+#include "Toolkits/AssetEditorManager.h"
+#include "AssetRegistryModule.h"
+
+//Windows dialog popup
+#include "Widgets/DeclarativeSyntaxSupport.h"
+#include "Widgets/SWindow.h"
+#include "Framework/Application/SlateApplication.h"
+#include "Interfaces/IMainFrameModule.h"
+#include "FbxCompareWindow.h"
 
 //Meshes includes
 #include "MeshUtilities.h"
@@ -31,6 +42,7 @@
 #include "Engine/StaticMesh.h"
 #include "Engine/StaticMeshSocket.h"
 #include "StaticMeshResources.h"
+#include "Factories/FbxStaticMeshImportData.h"
 
 //Skeletal mesh includes
 #include "SkelImport.h"
@@ -43,6 +55,8 @@
 #include "ApexClothingUtils.h"
 #include "Engine/SkeletalMeshSocket.h"
 #include "Assets/ClothingAsset.h"
+#include "Factories/FbxSkeletalMeshImportData.h"
+#include "FbxCompareWindow.h"
 
 
 
@@ -50,103 +64,443 @@
 
 using namespace UnFbx;
 
-class FCompMaterial
+struct FCreateCompFromFbxArg
 {
-public:
-	FCompMaterial(FName InMaterialSlotName, FName InImportedMaterialSlotName)
-		: MaterialSlotName(InMaterialSlotName)
-		, ImportedMaterialSlotName(InImportedMaterialSlotName)
-	{}
-	FName MaterialSlotName;
-	FName ImportedMaterialSlotName;
+	FString MeshName;
+	bool IsStaticMesh;
+	bool IsStaticHasLodGroup;
 };
 
-class FCompSection
+void CreateCompFromSkeletalMesh(USkeletalMesh* SkeletalMesh, FCompMesh &CurrentData)
 {
-public:
-	int32 MaterialIndex;
-};
-
-class FCompLOD
-{
-public:
-	TArray<FCompSection> Sections;
-};
+	//Fill the material array
+	for (const FSkeletalMaterial &Material : SkeletalMesh->Materials)
+	{
+		FCompMaterial CompMaterial(Material.MaterialSlotName, Material.ImportedMaterialSlotName);
+		CurrentData.CompMaterials.Add(CompMaterial);
+	}
+		
+	//Fill the section topology
+	if (SkeletalMesh->GetResourceForRendering())
+	{
+		CurrentData.CompLods.AddZeroed(SkeletalMesh->GetResourceForRendering()->LODModels.Num());
+		//Fill sections data
+		for (int32 LodIndex = 0; LodIndex < SkeletalMesh->GetResourceForRendering()->LODModels.Num(); ++LodIndex)
+		{
+			//Find the LodMaterialMap, which must be use for all LOD except the base
+			TArray<int32> LODMaterialMap;
+			if(LodIndex > 0 && SkeletalMesh->LODInfo.IsValidIndex(LodIndex))
+			{
+				const FSkeletalMeshLODInfo &SkeletalMeshLODInfo = SkeletalMesh->LODInfo[LodIndex];
+				LODMaterialMap = SkeletalMeshLODInfo.LODMaterialMap;
+			}
 
-class FCompMesh
-{
-public:
-	TArray<FCompMaterial> Materials;
-	TArray<FCompLOD> Lods;
-};
+			const FStaticLODModel &StaticLodModel = SkeletalMesh->GetResourceForRendering()->LODModels[LodIndex];
+			CurrentData.CompLods[LodIndex].Sections.AddZeroed(StaticLodModel.Sections.Num());
+			for (int32 SectionIndex = 0; SectionIndex < StaticLodModel.Sections.Num(); ++SectionIndex)
+			{
+				const FSkelMeshSection &SkelMeshSection = StaticLodModel.Sections[SectionIndex];
+				int32 MaterialIndex = SkelMeshSection.MaterialIndex;
+				if (LodIndex > 0 && LODMaterialMap.IsValidIndex(MaterialIndex))
+				{
+					MaterialIndex = LODMaterialMap[MaterialIndex];
+				}
+				CurrentData.CompLods[LodIndex].Sections[SectionIndex].MaterialIndex = MaterialIndex;
+			}
+		}
+	}
 
-class FCompSkeletalMesh : public FCompMesh
-{
-public:
-	FCompSkeletalMesh(USkeletalMesh*InSkeletalMesh)
-		: SkeletalMesh(InSkeletalMesh)
+	//Fill the skeleton joint
+	CurrentData.CompSkeleton.Joints.AddZeroed(SkeletalMesh->RefSkeleton.GetNum());
+	for (int JointIndex = 0; JointIndex < CurrentData.CompSkeleton.Joints.Num(); ++JointIndex)
 	{
-		//Fill the material array
-		for (const FSkeletalMaterial &Material : SkeletalMesh->Materials)
+		CurrentData.CompSkeleton.Joints[JointIndex].Name = SkeletalMesh->RefSkeleton.GetBoneName(JointIndex);
+		CurrentData.CompSkeleton.Joints[JointIndex].ParentIndex = SkeletalMesh->RefSkeleton.GetParentIndex(JointIndex);
+		int32 ParentIndex = CurrentData.CompSkeleton.Joints[JointIndex].ParentIndex;
+		if (CurrentData.CompSkeleton.Joints.IsValidIndex(ParentIndex))
 		{
-			FCompMaterial CompMaterial(Material.MaterialSlotName, Material.ImportedMaterialSlotName);
-			Materials.Add(CompMaterial);
+			CurrentData.CompSkeleton.Joints[ParentIndex].ChildIndexes.Add(JointIndex);
 		}
-		if (SkeletalMesh->GetResourceForRendering())
+	}
+
+	USkeleton* Skeleton = SkeletalMesh->Skeleton;
+	if (Skeleton != nullptr && !Skeleton->MergeAllBonesToBoneTree(SkeletalMesh))
+	{
+		CurrentData.CompSkeleton.bSkeletonFitMesh = false;
+	}
+}
+
+void CreateCompFromStaticMesh(UStaticMesh* StaticMesh, FCompMesh &CurrentData)
+{
+	//Fill the material array
+	for (const FStaticMaterial &Material : StaticMesh->StaticMaterials)
+	{
+		FCompMaterial CompMaterial(Material.MaterialSlotName, Material.ImportedMaterialSlotName);
+		CurrentData.CompMaterials.Add(CompMaterial);
+	}
+
+	//Fill the section topology
+	if (StaticMesh->RenderData)
+	{
+		CurrentData.CompLods.AddZeroed(StaticMesh->RenderData->LODResources.Num());
+
+		//Fill sections data
+		for (int32 LodIndex = 0; LodIndex < StaticMesh->RenderData->LODResources.Num(); ++LodIndex)
 		{
-			//Fill sections data
-			for (int32 LodIndex = 0; LodIndex < SkeletalMesh->GetResourceForRendering()->LODModels.Num(); ++LodIndex)
-			{
-				//Find the LodMaterialMap, which must be use for all LOD except the base
-				TArray<int32> LODMaterialMap;
-				if(LodIndex > 0 && SkeletalMesh->LODInfo.IsValidIndex(LodIndex))
-				{
-					const FSkeletalMeshLODInfo &SkeletalMeshLODInfo = SkeletalMesh->LODInfo[LodIndex];
-					LODMaterialMap = SkeletalMeshLODInfo.LODMaterialMap;
-				}
+			//StaticMesh->SectionInfoMap.Get()
 
-				const FStaticLODModel &StaticLodModel = SkeletalMesh->GetResourceForRendering()->LODModels[LodIndex];
-				for (int32 SectionIndex = 0; SectionIndex < StaticLodModel.Sections.Num(); ++SectionIndex)
+			const FStaticMeshLODResources &StaticLodRessources = StaticMesh->RenderData->LODResources[LodIndex];
+			CurrentData.CompLods[LodIndex].Sections.AddZeroed(StaticLodRessources.Sections.Num());
+			for (int32 SectionIndex = 0; SectionIndex < StaticLodRessources.Sections.Num(); ++SectionIndex)
+			{
+				const FStaticMeshSection &StaticMeshSection = StaticLodRessources.Sections[SectionIndex];
+				int32 MaterialIndex = StaticMeshSection.MaterialIndex;
+				if (StaticMesh->SectionInfoMap.IsValidSection(LodIndex, SectionIndex))
 				{
-					const FSkelMeshSection &SkelMeshSection = StaticLodModel.Sections[SectionIndex];
-					int32 MaterialIndex = SkelMeshSection.MaterialIndex;
-					if (LodIndex > 0 && LODMaterialMap.IsValidIndex(MaterialIndex))
-					{
-						MaterialIndex = LODMaterialMap[MaterialIndex];
-					}
+					FMeshSectionInfo MeshSectionInfo = StaticMesh->SectionInfoMap.Get(LodIndex, SectionIndex);
+					MaterialIndex = MeshSectionInfo.MaterialIndex;
 				}
+				CurrentData.CompLods[LodIndex].Sections[SectionIndex].MaterialIndex = MaterialIndex;
 			}
 		}
-		else
+	}
+}
+
+void GetSkeletalMeshCompData(FFbxImporter *FbxImporter, UFbxImportUI* ImportUI, const FCreateCompFromFbxArg &CreateCompFromFbxArg, FCompMesh &FbxData, const USkeletalMesh *SkeletalMeshRef, UObject **PreviewObject)
+{
+	USkeletalMesh *SkeletalMesh = Cast<USkeletalMesh>(StaticDuplicateObject(SkeletalMeshRef, GetTransientPackage()));
+
+	if (SkeletalMesh != nullptr && FbxImporter->ReimportSkeletalMesh(SkeletalMesh, ImportUI->SkeletalMeshImportData))
+	{
+		if (GEditor->IsObjectInTransactionBuffer(SkeletalMesh))
 		{
-			
+			GEditor->ResetTransaction(LOCTEXT("PreviewReimportSkeletalMeshTransactionReset", "Preview Reimporting a skeletal mesh which was in the undo buffer"));
 		}
+		CreateCompFromSkeletalMesh(SkeletalMesh, FbxData);
+	}
+	else
+	{
+		//Log an error here
 	}
 
-	USkeletalMesh *GetSkeletalMesh() const { return SkeletalMesh; }
-	const FCompMesh &GetCompMesh() const { return *((FCompMesh *)this); }
-private:
-	USkeletalMesh *SkeletalMesh;
-};
+	*PreviewObject = SkeletalMesh;
+}
+
+void GetStaticMeshCompData(FFbxImporter *FbxImporter, UFbxImportUI* ImportUI, const FCreateCompFromFbxArg &CreateCompFromFbxArg, FCompMesh &FbxData, const UStaticMesh* StaticMeshRef, UObject **PreviewObject)
+{
+	UStaticMesh *StaticMesh = Cast<UStaticMesh>(StaticDuplicateObject(StaticMeshRef, GetTransientPackage()));
+
+	if (StaticMesh != nullptr && FbxImporter->ReimportStaticMesh(StaticMesh, ImportUI->StaticMeshImportData))
+	{
+		FbxImporter->ImportStaticMeshGlobalSockets(StaticMesh);
+		CreateCompFromStaticMesh(StaticMesh, FbxData);
+	}
+	else
+	{
+		//Log an error here
+	}
+	*PreviewObject = StaticMesh;
+}
+
+void CreateCompFromFbxData(FFbxImporter *FbxImporter, UFbxImportUI* ImportUI, const FString& FullPath, FCompMesh &FbxData, const FCreateCompFromFbxArg &CreateCompFromFbxArg, UStaticMesh *StaticMesh, USkeletalMesh *SkeletalMesh, UObject **PreviewObject)
+{
+	FbxImporter->ImportOptions->bIsReimportPreview = true;
+	//////////////////////////////////////////////////////////////////////////
+	//Static mesh reimport
+	if (CreateCompFromFbxArg.IsStaticMesh)
+	{
+		GetStaticMeshCompData(FbxImporter, ImportUI, CreateCompFromFbxArg, FbxData, StaticMesh, PreviewObject);
+		
+	}
+	//////////////////////////////////////////////////////////////////////////
+	else //Skeletal mesh reimport
+	{
+		GetSkeletalMeshCompData(FbxImporter, ImportUI, CreateCompFromFbxArg, FbxData, SkeletalMesh, PreviewObject);
+	}
+	FbxImporter->ImportOptions->bIsReimportPreview = false;
+}
 
-class FCompStaticMesh : public FCompMesh
+void FFbxImporter::FillGeneralFbxFileInformation(void *GeneralInfoPtr)
 {
-public:
-	FCompStaticMesh(UStaticMesh* InStaticMesh)
-		: StaticMesh(InStaticMesh)
+	FGeneralFbxFileInfo &FbxGeneralInfo = *(FGeneralFbxFileInfo*)GeneralInfoPtr;
+	FbxAxisSystem SourceSetup = Scene->GetGlobalSettings().GetAxisSystem();
+
+	//Get the UE4 sdk version
+	int32 SDKMajor, SDKMinor, SDKRevision;
+	FbxManager::GetFileFormatVersion(SDKMajor, SDKMinor, SDKRevision);
+
+	int32 FileMajor, FileMinor, FileRevision;
+	Importer->GetFileVersion(FileMajor, FileMinor, FileRevision);
+
+	FString DateVersion = FString(FbxManager::GetVersion(false));
+	FbxGeneralInfo.UE4SdkVersion = TEXT("UE4 Sdk Version: ") + FString::FromInt(SDKMajor) + TEXT(".") + FString::FromInt(SDKMinor) + TEXT(".") + FString::FromInt(SDKRevision) + TEXT(" (") + DateVersion + TEXT(")");
+
+	FbxIOFileHeaderInfo *FileHeaderInfo = Importer->GetFileHeaderInfo();
+	if (FileHeaderInfo)
 	{
-		//Fill the material array
-		for (const FStaticMaterial &Material : StaticMesh->StaticMaterials)
+		FbxGeneralInfo.ApplicationCreator = TEXT("Creator:    ") + FString(FileHeaderInfo->mCreator.Buffer());
+		FbxGeneralInfo.FileVersion = TEXT("Fbx File Version:    ") + FString::FromInt(FileMajor) + TEXT(".") + FString::FromInt(FileMinor) + TEXT(".") + FString::FromInt(FileRevision) + TEXT(" (") + FString::FromInt(FileHeaderInfo->mFileVersion) + TEXT(")");
+		FbxGeneralInfo.CreationDate = TEXT("Created Time:    ") + FString::FromInt(FileHeaderInfo->mCreationTimeStamp.mYear) + TEXT("-") + FString::FromInt(FileHeaderInfo->mCreationTimeStamp.mMonth) + TEXT("-") + FString::FromInt(FileHeaderInfo->mCreationTimeStamp.mDay) + TEXT(" (Y-M-D)");
+	}
+	int32 UpVectorSign = 1;
+	FbxAxisSystem::EUpVector UpVector = FileAxisSystem.GetUpVector(UpVectorSign);
+
+	int32 FrontVectorSign = 1;
+	FbxAxisSystem::EFrontVector FrontVector = FileAxisSystem.GetFrontVector(FrontVectorSign);
+
+
+	FbxAxisSystem::ECoordSystem CoordSystem = FileAxisSystem.GetCoorSystem();
+
+	FbxGeneralInfo.AxisSystem = TEXT("File Axis System:    UP: ");
+	if (UpVectorSign == -1)
+	{
+		FbxGeneralInfo.AxisSystem = TEXT("-");
+	}
+	FbxGeneralInfo.AxisSystem += (UpVector == FbxAxisSystem::EUpVector::eXAxis) ? TEXT("X, Front: ") : (UpVector == FbxAxisSystem::EUpVector::eYAxis) ? TEXT("Y, Front: ") : TEXT("Z, Front: ");
+	if (FrontVectorSign == -1)
+	{
+		FbxGeneralInfo.AxisSystem = TEXT("-");
+	}
+
+	if (UpVector == FbxAxisSystem::EUpVector::eXAxis)
+	{
+		FbxGeneralInfo.AxisSystem += (FrontVector == FbxAxisSystem::EFrontVector::eParityEven) ? TEXT("Y") : TEXT("Z");
+	}
+	else if (UpVector == FbxAxisSystem::EUpVector::eYAxis)
+	{
+		FbxGeneralInfo.AxisSystem += (FrontVector == FbxAxisSystem::EFrontVector::eParityEven) ? TEXT("X") : TEXT("Z");
+	}
+	else if (UpVector == FbxAxisSystem::EUpVector::eZAxis)
+	{
+		FbxGeneralInfo.AxisSystem += (FrontVector == FbxAxisSystem::EFrontVector::eParityEven) ? TEXT("X") : TEXT("Y");
+	}
+
+	//Hand side
+	FbxGeneralInfo.AxisSystem += (CoordSystem == FbxAxisSystem::ECoordSystem::eLeftHanded) ? TEXT(" Left Handed") : TEXT(" Right Handed");
+
+
+	if (FileAxisSystem == FbxAxisSystem::MayaZUp)
+	{
+		FbxGeneralInfo.AxisSystem += TEXT(" (Maya ZUp)");
+	}
+	else if (FileAxisSystem == FbxAxisSystem::MayaYUp)
+	{
+		FbxGeneralInfo.AxisSystem += TEXT(" (Maya YUp)");
+	}
+	else if (FileAxisSystem == FbxAxisSystem::Max)
+	{
+		FbxGeneralInfo.AxisSystem += TEXT(" (Max)");
+	}
+	else if (FileAxisSystem == FbxAxisSystem::Motionbuilder)
+	{
+		FbxGeneralInfo.AxisSystem += TEXT(" (Motion Builder)");
+	}
+	else if (FileAxisSystem == FbxAxisSystem::OpenGL)
+	{
+		FbxGeneralInfo.AxisSystem += TEXT(" (OpenGL)");
+	}
+	else if (FileAxisSystem == FbxAxisSystem::DirectX)
+	{
+		FbxGeneralInfo.AxisSystem += TEXT(" (DirectX)");
+	}
+	else if (FileAxisSystem == FbxAxisSystem::Lightwave)
+	{
+		FbxGeneralInfo.AxisSystem += TEXT(" (Lightwave)");
+	}
+
+	FbxGeneralInfo.UnitSystem = TEXT("Units:    ");
+	if (FileUnitSystem == FbxSystemUnit::mm)
+	{
+		FbxGeneralInfo.UnitSystem += TEXT("mm (millimeter)");
+	}
+	else if (FileUnitSystem == FbxSystemUnit::cm)
+	{
+		FbxGeneralInfo.UnitSystem += TEXT("cm (centimeter)");
+	}
+	else if (FileUnitSystem == FbxSystemUnit::dm)
+	{
+		FbxGeneralInfo.UnitSystem += TEXT("dm (decimeter)");
+	}
+	else if (FileUnitSystem == FbxSystemUnit::m)
+	{
+		FbxGeneralInfo.UnitSystem += TEXT("m (meter)");
+	}
+	else if (FileUnitSystem == FbxSystemUnit::km)
+	{
+		FbxGeneralInfo.UnitSystem += TEXT("km (kilometer)");
+	}
+	else if (FileUnitSystem == FbxSystemUnit::Inch)
+	{
+		FbxGeneralInfo.UnitSystem += TEXT("Inch");
+	}
+	else if (FileUnitSystem == FbxSystemUnit::Foot)
+	{
+		FbxGeneralInfo.UnitSystem += TEXT("Foot");
+	}
+	else if (FileUnitSystem == FbxSystemUnit::Yard)
+	{
+		FbxGeneralInfo.UnitSystem += TEXT("Yard");
+	}
+	else if (FileUnitSystem == FbxSystemUnit::Mile)
+	{
+		FbxGeneralInfo.UnitSystem += TEXT("Mile");
+	}
+}
+
+void FFbxImporter::ShowFbxReimportPreview(UObject *ReimportObj, UFbxImportUI* ImportUI, const FString& FullPath)
+{
+	if (ReimportObj == nullptr || ImportUI == nullptr)
+	{
+		return;
+	}
+	//Show a dialog
+	UStaticMesh *StaticMesh = Cast<UStaticMesh>(ReimportObj);
+	USkeletalMesh *SkeletalMesh = Cast<USkeletalMesh>(ReimportObj);
+	UObject *PreviewObject = nullptr;
+	FCompMesh CurrentData;
+	FCompMesh FbxData;
+	FCreateCompFromFbxArg CreateCompFromFbxArg;
+	FString Filename;
+	//Create the current data to compare from
+	if (StaticMesh)
+	{
+		CreateCompFromStaticMesh(StaticMesh, CurrentData);
+		CreateCompFromFbxArg.MeshName = StaticMesh->GetName();
+		CreateCompFromFbxArg.IsStaticMesh = true;
+		CreateCompFromFbxArg.IsStaticHasLodGroup = StaticMesh->LODGroup != NAME_None;
+		Filename = ImportUI->StaticMeshImportData->GetFirstFilename();
+	}
+	else if (SkeletalMesh)
+	{
+		CreateCompFromSkeletalMesh(SkeletalMesh, CurrentData);
+		CreateCompFromFbxArg.MeshName = SkeletalMesh->GetName();
+		CreateCompFromFbxArg.IsStaticMesh = false;
+		Filename = ImportUI->SkeletalMeshImportData->GetFirstFilename();
+	}
+	
+	//////////////////////////////////////////////////////////////////////////
+	// Set the import options
+	ApplyImportUIToImportOptions(ImportUI, *ImportOptions);
+
+	// Set all reimport options
+	ImportOptions->bAutoComputeLodDistances = true;
+	ImportOptions->LodNumber = 0;
+	ImportOptions->MinimumLodNumber = 0;
+	ImportOptions->bImportRigidMesh = true;
+	ImportOptions->bImportMaterials = false;
+	ImportOptions->bImportTextures = false;
+	ImportOptions->bImportAnimations = false;
+
+	//////////////////////////////////////////////////////////////////////////
+	// Open the fbx file
+
+	const FString FileExtension = FPaths::GetExtension(Filename);
+	const bool bIsValidFile = FileExtension.Equals(TEXT("fbx"), ESearchCase::IgnoreCase) || FileExtension.Equals("obj", ESearchCase::IgnoreCase);
+	if (!bIsValidFile || !(Filename.Len()) || (IFileManager::Get().FileSize(*Filename) == INDEX_NONE) ||
+		!ImportFromFile(*Filename, FPaths::GetExtension(Filename), true))
+	{
+		return;
+	}
+
+	//Query general information
+	FGeneralFbxFileInfo FbxGeneralInfo;
+	FillGeneralFbxFileInformation(&FbxGeneralInfo);
+
+	//Apply the Transform to the scene
+	UFbxAssetImportData* ImportAssetData = nullptr;
+	if (CreateCompFromFbxArg.IsStaticMesh)
+	{
+		ImportAssetData = ImportUI->StaticMeshImportData;
+	}
+	else
+	{
+		ImportAssetData = ImportUI->SkeletalMeshImportData;
+	}
+	ApplyTransformSettingsToFbxNode(Scene->GetRootNode(), ImportAssetData);
+
+	UnFbx::FbxSceneInfo SceneInfo;
+	//Read the scene and found all instance with their scene information.
+	GetSceneInfo(Filename, SceneInfo, true);
+	//Convert old structure to the new scene export structure
+	TSharedPtr<FFbxSceneInfo> SceneInfoPtr = UFbxSceneImportFactory::ConvertSceneInfo(this, &SceneInfo);
+	//Get import material info
+	UFbxSceneImportFactory::ExtractMaterialInfo(this, SceneInfoPtr);
+
+	CreateCompFromFbxData(this, ImportUI, FullPath, FbxData, CreateCompFromFbxArg, StaticMesh, SkeletalMesh, &PreviewObject);
+	TArray<TSharedPtr<FString>> AssetReferencingSkeleton;
+	if (SkeletalMesh != nullptr && SkeletalMesh->Skeleton != nullptr && !FbxData.CompSkeleton.bSkeletonFitMesh)
+	{
+		UObject* SelectedObject = SkeletalMesh->Skeleton;
+		if (SelectedObject)
 		{
-			FCompMaterial CompMaterial(Material.MaterialSlotName, Material.ImportedMaterialSlotName);
-			Materials.Add(CompMaterial);
+			FAssetRegistryModule& AssetRegistryModule = FModuleManager::LoadModuleChecked<FAssetRegistryModule>(TEXT("AssetRegistry"));
+			const FName SelectedPackageName = SelectedObject->GetOutermost()->GetFName();
+			//Get the Hard dependencies
+			TArray<FName> HardDependencies;
+			AssetRegistryModule.Get().GetReferencers(SelectedPackageName, HardDependencies, EAssetRegistryDependencyType::Hard);
+			//Get the Soft dependencies
+			TArray<FName> SoftDependencies;
+			AssetRegistryModule.Get().GetReferencers(SelectedPackageName, SoftDependencies, EAssetRegistryDependencyType::Soft);
+			//Compose the All dependencies array
+			TArray<FName> AllDependencies = HardDependencies;
+			AllDependencies += SoftDependencies;
+			if (AllDependencies.Num() > 0)
+			{
+				for (const FName AssetDependencyName : AllDependencies)
+				{
+					const FString PackageString = AssetDependencyName.ToString();
+					const FString FullAssetPathName = FString::Printf(TEXT("%s.%s"), *PackageString, *FPackageName::GetLongPackageAssetName(PackageString));
+					
+					FAssetData AssetData = AssetRegistryModule.Get().GetAssetByObjectPath(FName(*FullAssetPathName));
+					if (AssetData.GetClass() != nullptr)
+					{
+						TSharedPtr<FString> AssetReferencing = MakeShareable(new FString(AssetData.AssetClass.ToString() + TEXT(" ") + FullAssetPathName));
+						AssetReferencingSkeleton.Add(AssetReferencing);
+					}
+				}
+			}
 		}
 	}
 
-	UStaticMesh *GetStaticMesh() const { return StaticMesh; }
-	const FCompMesh &GetCompMesh() const { return *((FCompMesh *)this); }
-private:
-	UStaticMesh *StaticMesh;
-};
+	//Create the modal dialog window to let the user see the result of the compare
+	TSharedPtr<SWindow> ParentWindow;
+
+	if (FModuleManager::Get().IsModuleLoaded("MainFrame"))
+	{
+		IMainFrameModule& MainFrame = FModuleManager::LoadModuleChecked<IMainFrameModule>("MainFrame");
+		ParentWindow = MainFrame.GetParentWindow();
+	}
+
+	TSharedRef<SWindow> Window = SNew(SWindow)
+		.Title(NSLOCTEXT("UnrealEd", "FbxMaterialConflictOpionsTitle", "FBX Import Conflict"))
+		.AutoCenter(EAutoCenter::PreferredWorkArea)
+		.SizingRule(ESizingRule::UserSized)
+		.ClientSize(FVector2D(700, 650))
+		.MinWidth(700)
+		.MinHeight(650);
 
+	TSharedPtr<SFbxCompareWindow> FbxMaterialConflictWindow;
+	Window->SetContent
+		(
+			SNew(SFbxCompareWindow)
+			.WidgetWindow(Window)
+			.FullFbxPath(FText::FromString(Filename))
+			.FbxSceneInfo(SceneInfoPtr)
+			.FbxGeneralInfo(FbxGeneralInfo)
+			.AssetReferencingSkeleton(&AssetReferencingSkeleton)
+			.CurrentMeshData(&CurrentData)
+			.FbxMeshData(&FbxData)
+			.PreviewObject(PreviewObject)
+			);
+
+	// @todo: we can make this slow as showing progress bar later
+	FSlateApplication::Get().AddModalWindow(Window, ParentWindow, false);
+
+	if (PreviewObject != nullptr)
+	{
+		FAssetEditorManager::Get().CloseAllEditorsForAsset(PreviewObject);
+		PreviewObject->MarkPendingKill();
+		PreviewObject = nullptr;
+	}
+}
 #undef LOCTEXT_NAMESPACE
