Copying //UE4/Dev-Framework to //UE4/Dev-Main (Source: //UE4/Dev-Framework @ 3510040)

diff --git a/Engine/Source/Editor/UnrealEd/Private/Commandlets/AssetRegistryGenerator.cpp b/Engine/Source/Editor/UnrealEd/Private/Commandlets/AssetRegistryGenerator.cpp
--- a/Engine/Source/Editor/UnrealEd/Private/Commandlets/AssetRegistryGenerator.cpp
+++ b/Engine/Source/Editor/UnrealEd/Private/Commandlets/AssetRegistryGenerator.cpp
@@ -233,6 +233,12 @@ bool FAssetRegistryGenerator::GenerateStreamingInstallManifest()
 			int64 CurrentPakSize = 0;
 			bool bFinishedAllFiles = true;
 
+			if (bUseAssetManager)
+			{
+				// Sort so the order is consistent. If load order is important then it should be specified as a load order file to UnrealPak
+				ChunkFilenames.Sort();
+			}
+
 			for (; FilenameIndex < ChunkFilenames.Num(); ++FilenameIndex)
 			{
 				FString Filename = ChunkFilenames[FilenameIndex];
@@ -433,78 +439,17 @@ void FAssetRegistryGenerator::Initialize(const TArray<FName> &InStartupPackages)
 
 	FAssetRegistrySerializationOptions SaveOptions;
 
+	ensureMsgf(!AssetRegistry.IsLoadingAssets(), TEXT("Cannot initialize asset registry generator while asset registry is still scanning source assets "));
+
 	AssetRegistry.InitializeSerializationOptions(SaveOptions, TargetPlatform->IniPlatformName());
 
 	AssetRegistry.InitializeTemporaryAssetRegistryState(State, SaveOptions);
-
-	UpdatePackageSourceHashes();
 }
 
-void FAssetRegistryGenerator::UpdatePackageSourceHashes()
+void FAssetRegistryGenerator::ComputePackageDifferences(TSet<FName>& ModifiedPackages, TSet<FName>& NewPackages, TSet<FName>& RemovedPackages, TSet<FName>& IdenticalCookedPackages, TSet<FName>& IdenticalUncookedPackages, bool bRecurseModifications, bool bRecurseScriptModifications)
 {
-	UAssetManager* AssetManager = bUseAssetManager ? &UAssetManager::Get() : nullptr;
-
-	for (const TPair<FName, const FAssetPackageData*>& PackagePair : State.GetAssetPackageDataMap())
-	{
-		FName PackageName = PackagePair.Key;
-		FAssetPackageData* PackageData = const_cast<FAssetPackageData*>(PackagePair.Value);
-
-		FMD5 PackageSourceHash;
-
-		// Include package guid
-		PackageSourceHash.Update((uint8*)&PackageData->PackageGuid, sizeof(FGuid));
-
-		if (AssetManager)
-		{
-			AssetManager->UpdatePackageSourceHash(PackageName, PackageSourceHash);
-		}
-
-		// Find direct hard dependencies. We do not recurse the hashes, because dirty state is instead recursed
-		TArray<FAssetIdentifier> Dependencies;
-
-		State.GetDependencies(PackageName, Dependencies, EAssetRegistryDependencyType::Hard);
-
-		Dependencies.Sort([&](const FAssetIdentifier& A, const FAssetIdentifier& B) { return A.PackageName > B.PackageName; } );
-
-		for (FAssetIdentifier& Dependency : Dependencies)
-		{
-			const FAssetPackageData* DependencyData = State.GetAssetPackageData(Dependency.PackageName);
-
-			if (DependencyData)
-			{
-				// Hash it's guid
-				PackageSourceHash.Update((uint8*)&DependencyData->PackageGuid, sizeof(FGuid));
-
-				if (AssetManager)
-				{
-					AssetManager->UpdatePackageSourceHash(PackageName, PackageSourceHash);
-				}
-			}
-			else
-			{
-				FString PackageString = Dependency.PackageName.ToString();
-
-				if (FPackageName::IsScriptPackage(PackageString))
-				{
-					// Get the guid off the script package, this is updated when script is changed
-					UPackage* Package = FindPackage(nullptr, *PackageString);
+	TArray<FName> ModifiedScriptPackages;
 
-					if (Package)
-					{
-						FGuid PackageGuid = Package->GetGuid();
-						PackageSourceHash.Update((uint8*)&PackageGuid, sizeof(FGuid));
-					}
-				}
-			}
-		}
-
-		// Source Hash is safe to modify inline so const cast it
-		PackageData->PackageSourceHash.Set(PackageSourceHash);
-	}
-}
-
-void FAssetRegistryGenerator::ComputePackageDifferences(TSet<FName>& ModifiedPackages, TSet<FName>& NewPackages, TSet<FName>& RemovedPackages, TSet<FName>& IdenticalCookedPackages, TSet<FName>& IdenticalUncookedPackages, bool bRecurseModifications)
-{
 	for (const TPair<FName, const FAssetPackageData*>& PackagePair : State.GetAssetPackageDataMap())
 	{
 		FName PackageName = PackagePair.Key;
@@ -516,7 +461,7 @@ void FAssetRegistryGenerator::ComputePackageDifferences(TSet<FName>& ModifiedPac
 		{
 			NewPackages.Add(PackageName);
 		}
-		else if (CurrentPackageData->PackageSourceHash == PreviousPackageData->PackageSourceHash)
+		else if (CurrentPackageData->PackageGuid == PreviousPackageData->PackageGuid)
 		{
 			if (PreviousPackageData->DiskSize < 0)
 			{
@@ -529,7 +474,14 @@ void FAssetRegistryGenerator::ComputePackageDifferences(TSet<FName>& ModifiedPac
 		}
 		else
 		{
-			ModifiedPackages.Add(PackageName);
+			if (FPackageName::IsScriptPackage(PackageName.ToString()))
+			{
+				ModifiedScriptPackages.Add(PackageName);
+			}
+			else
+			{
+				ModifiedPackages.Add(PackageName);
+			}
 		}
 	}
 
@@ -548,9 +500,14 @@ void FAssetRegistryGenerator::ComputePackageDifferences(TSet<FName>& ModifiedPac
 
 	if (bRecurseModifications)
 	{
-		// Recurse modified packages to their dependencies. In theory a small number of packages are modified at once so this is faster than computing the hashes recursively
+		// Recurse modified packages to their dependencies. This is needed because we only compare package guids
 		TArray<FName> ModifiedPackagesToRecurse = ModifiedPackages.Array();
 
+		if (bRecurseScriptModifications)
+		{
+			ModifiedPackagesToRecurse.Append(ModifiedScriptPackages);
+		}
+
 		for (int32 RecurseIndex = 0; RecurseIndex < ModifiedPackagesToRecurse.Num(); RecurseIndex++)
 		{
 			FName ModifiedPackage = ModifiedPackagesToRecurse[RecurseIndex];
@@ -626,10 +583,11 @@ void FAssetRegistryGenerator::BuildChunkManifest(const TSet<FName>& InCookedPack
 				FoundIDList = &PackageChunkIDMap.Add(AssetData.PackageName);
 			}
 			FoundIDList->AddUnique(ChunkID);
-
-			// Now clear the original chunk id list. We will fill it with real IDs when cooking.
-			AssetData.ChunkIDs.Empty();
 		}
+
+		// Now clear the original chunk id list. We will fill it with real IDs when cooking.
+		AssetData.ChunkIDs.Empty();
+
 		// Update whether the owner package contains a map
 		if (AssetData.GetClass()->IsChildOf(UWorld::StaticClass()) || AssetData.GetClass()->IsChildOf(ULevel::StaticClass()))
 		{
@@ -1248,8 +1206,8 @@ void FAssetRegistryGenerator::FixupPackageDependenciesForChunks(FSandboxPlatform
 		}
 	}
 
-	auto* ChunkDepGraph = DependencyInfo->GetChunkDependencyGraph(ChunkManifests.Num());
-	//Once complete, Add any remaining assets (that are not assigned to a chunk) to the first chunk.
+	const FChunkDependencyTreeNode* ChunkDepGraph = DependencyInfo->GetOrBuildChunkDependencyGraph(ChunkManifests.Num() - 1);
+	// Once complete, Add any remaining assets (that are not assigned to a chunk) to the first chunk.
 	if (FinalChunkManifests.Num() == 0)
 	{
 		FinalChunkManifests.Add(nullptr);
@@ -1288,11 +1246,6 @@ void FAssetRegistryGenerator::FixupPackageDependenciesForChunks(FSandboxPlatform
 		}
 	}
 
-/*	if (!CheckChunkAssetsAreNotInChild(*ChunkDepGraph))
-	{
-		UE_LOG(LogAssetRegistryGenerator, Log, TEXT("Second Scan of chunks found duplicate asset entries in children."));
-	}*/
-
 	for (int32 ChunkID = 0, MaxChunk = ChunkManifests.Num(); ChunkID < MaxChunk; ++ChunkID)
 	{
 		const int32 ChunkManifestNum = ChunkManifests[ChunkID] ? ChunkManifests[ChunkID]->Num() : 0;
diff --git a/Engine/Source/Editor/UnrealEd/Private/Commandlets/ChunkDependencyInfo.cpp b/Engine/Source/Editor/UnrealEd/Private/Commandlets/ChunkDependencyInfo.cpp
--- a/Engine/Source/Editor/UnrealEd/Private/Commandlets/ChunkDependencyInfo.cpp
+++ b/Engine/Source/Editor/UnrealEd/Private/Commandlets/ChunkDependencyInfo.cpp
@@ -5,19 +5,47 @@
 UChunkDependencyInfo::UChunkDependencyInfo(const FObjectInitializer& ObjectInitializer)
 	: Super(ObjectInitializer)
 {
+	CachedHighestChunk = -1;
 }
 
-const FChunkDependencyTreeNode* UChunkDependencyInfo::GetChunkDependencyGraph(uint32 TotalChunks)
+const FChunkDependencyTreeNode* UChunkDependencyInfo::GetOrBuildChunkDependencyGraph(int32 HighestChunk)
+{
+	if (HighestChunk > CachedHighestChunk)
+	{
+		return BuildChunkDependencyGraph(HighestChunk);
+	}
+	return &RootTreeNode;
+}
+
+const FChunkDependencyTreeNode* UChunkDependencyInfo::BuildChunkDependencyGraph(int32 HighestChunk)
 {
 	// Reset any current tree
 	RootTreeNode.ChunkID = 0;
 	RootTreeNode.ChildNodes.Reset(0);
 
+	ChildToParentMap.Reset();
+	CachedHighestChunk = HighestChunk;
+
 	// Ensure the DependencyArray is OK to work with.
-	// Remove cycles
-	DependencyArray.RemoveAllSwap([](const FChunkDependency& A){ return A.ChunkID == A.ParentChunkID; });
+	for (int32 DepIndex = DependencyArray.Num() - 1; DepIndex >= 0; DepIndex --)
+	{
+		const FChunkDependency& Dep = DependencyArray[DepIndex];
+		if (Dep.ChunkID > HighestChunk)
+		{
+			HighestChunk = Dep.ChunkID;
+		}
+		if (Dep.ParentChunkID > HighestChunk)
+		{
+			HighestChunk = Dep.ParentChunkID;
+		}
+		if (Dep.ChunkID == Dep.ParentChunkID)
+		{
+			// Remove cycles
+			DependencyArray.RemoveAtSwap(DepIndex);
+		}
+	}
 	// Add missing links (assumes they parent to chunk zero)
-	for (uint32 i = 1; i < TotalChunks; ++i)
+	for (int32 i = 1; i <= HighestChunk; ++i)
 	{
 		if (!DependencyArray.FindByPredicate([=](const FChunkDependency& RHS){ return i == RHS.ChunkID; }))
 		{
@@ -41,13 +69,18 @@ const FChunkDependencyTreeNode* UChunkDependencyInfo::GetChunkDependencyGraph(ui
 		}
 	}
 
-	//
-	AddChildrenRecursive(RootTreeNode, DependencyArray);
+	AddChildrenRecursive(RootTreeNode, DependencyArray, TSet<int32>());
 	return &RootTreeNode;
 }
 
-void UChunkDependencyInfo::AddChildrenRecursive(FChunkDependencyTreeNode& Node, TArray<FChunkDependency>& DepInfo)
+void UChunkDependencyInfo::AddChildrenRecursive(FChunkDependencyTreeNode& Node, TArray<FChunkDependency>& DepInfo, TSet<int32> Parents)
 {
+	if (Parents.Num() > 0)
+	{
+		ChildToParentMap.FindOrAdd(Node.ChunkID).Append(Parents);
+	}
+
+	Parents.Add(Node.ChunkID);
 	auto ChildNodeIndices = DepInfo.FilterByPredicate(
 		[&](const FChunkDependency& RHS) 
 		{
@@ -59,6 +92,26 @@ void UChunkDependencyInfo::AddChildrenRecursive(FChunkDependencyTreeNode& Node,
 	}
 	for (auto& Child : Node.ChildNodes) 
 	{
-		AddChildrenRecursive(Child, DepInfo);
+		AddChildrenRecursive(Child, DepInfo, Parents);
 	}
 }
+
+void UChunkDependencyInfo::RemoveRedundantChunks(TArray<int32>& ChunkIDs) const
+{
+	for (int32 ChunkIndex = ChunkIDs.Num() - 1; ChunkIndex >= 0; ChunkIndex--)
+	{
+		const TSet<int32>* FoundParents = ChildToParentMap.Find(ChunkIDs[ChunkIndex]);
+
+		if (FoundParents)
+		{
+			for (int32 ParentChunk : *FoundParents)
+			{
+				if (ChunkIDs.Contains(ParentChunk))
+				{
+					ChunkIDs.RemoveAt(ChunkIndex);
+					break;
+				}
+			}
+		}
+	}
+}
\ No newline at end of file
diff --git a/Engine/Source/Editor/UnrealEd/Private/Commandlets/CookCommandlet.cpp b/Engine/Source/Editor/UnrealEd/Private/Commandlets/CookCommandlet.cpp
--- a/Engine/Source/Editor/UnrealEd/Private/Commandlets/CookCommandlet.cpp
+++ b/Engine/Source/Editor/UnrealEd/Private/Commandlets/CookCommandlet.cpp
@@ -347,16 +347,6 @@ bool UCookCommandlet::CookOnTheFly( FGuid InstanceId, int32 Timeout, bool bForce
 	UCookerSettings const* CookerSettings = GetDefault<UCookerSettings>();
 	ECookInitializationFlags IterateFlags = ECookInitializationFlags::Iterative;
 
-	if (CookerSettings->bUseAssetRegistryForIteration || Switches.Contains(TEXT("iterateregistry")))
-	{
-		// Asset registry overwrites other options
-		IterateFlags = ECookInitializationFlags::Iterative | ECookInitializationFlags::IterateOnAssetRegistry;
-	}
-	else if (Switches.Contains(TEXT("iteratehash")))
-	{
-		IterateFlags = ECookInitializationFlags::Iterative | ECookInitializationFlags::IterateOnHash;
-	}
-
 	ECookInitializationFlags CookFlags = ECookInitializationFlags::None;
 	CookFlags |= bIterativeCooking ? IterateFlags : ECookInitializationFlags::None;
 	CookFlags |= bSkipEditorContent ? ECookInitializationFlags::SkipEditorContent : ECookInitializationFlags::None;
@@ -626,25 +616,11 @@ bool UCookCommandlet::CookByTheBook( const TArray<ITargetPlatform*>& Platforms,
 	UCookerSettings const* CookerSettings = GetDefault<UCookerSettings>();
 	ECookInitializationFlags IterateFlags = ECookInitializationFlags::Iterative;
 
-	if (CookerSettings->bUseAssetRegistryForIteration || Switches.Contains(TEXT("iterateregistry")))
-	{
-		// Asset registry overwrites other options
-		IterateFlags = ECookInitializationFlags::Iterative | ECookInitializationFlags::IterateOnAssetRegistry;
-	}
-	else if (Switches.Contains(TEXT("iteratehash")))
-	{
-		IterateFlags = ECookInitializationFlags::Iterative | ECookInitializationFlags::IterateOnHash;
-	}
-
-	if (Switches.Contains(TEXT("iteratesharedcookedbuild")))
+	if (Switches.Contains(TEXT("IterateSharedCookedbuild")))
 	{
 		// Add shared build flag to method flag, and enable iterative
 		IterateFlags |= ECookInitializationFlags::IterateSharedBuild;
-		if (!(IterateFlags | ECookInitializationFlags::IterateOnAssetRegistry))
-		{
-			IterateFlags |= ECookInitializationFlags::IterateOnHash;
-		}
-
+		
 		bIterativeCooking = true;
 	}
 	
@@ -655,10 +631,11 @@ bool UCookCommandlet::CookByTheBook( const TArray<ITargetPlatform*>& Platforms,
 	CookFlags |= bUnversioned ? ECookInitializationFlags::Unversioned : ECookInitializationFlags::None;
 	CookFlags |= bVerboseCookerWarnings ? ECookInitializationFlags::OutputVerboseCookerWarnings : ECookInitializationFlags::None;
 	CookFlags |= bPartialGC ? ECookInitializationFlags::EnablePartialGC : ECookInitializationFlags::None;
-	bool bTestCook = Switches.Contains(TEXT("Testcook"));
+	bool bTestCook = Switches.Contains(TEXT("TestCook"));
 	CookFlags |= bTestCook ? ECookInitializationFlags::TestCook : ECookInitializationFlags::None;
-	CookFlags |= Switches.Contains(TEXT("logdebuginfo")) ? ECookInitializationFlags::LogDebugInfo : ECookInitializationFlags::None;
-	CookFlags |= Switches.Contains(TEXT("Ignoreinisettingsoutofdate")) ? ECookInitializationFlags::IgnoreIniSettingsOutOfDate : ECookInitializationFlags::None;
+	CookFlags |= Switches.Contains(TEXT("LogDebugInfo")) ? ECookInitializationFlags::LogDebugInfo : ECookInitializationFlags::None;
+	CookFlags |= Switches.Contains(TEXT("IgnoreIniSettingsOutOfDate")) || CookerSettings->bIgnoreIniSettingsOutOfDateForIteration ? ECookInitializationFlags::IgnoreIniSettingsOutOfDate : ECookInitializationFlags::None;
+	CookFlags |= Switches.Contains(TEXT("IgnoreScriptPackagesOutOfDate")) || CookerSettings->bIgnoreScriptPackagesOutOfDateForIteration ? ECookInitializationFlags::IgnoreScriptPackagesOutOfDate : ECookInitializationFlags::None;
 
 	TArray<UClass*> FullGCAssetClasses;
 	if (FullGCAssetClassNames.Num())
diff --git a/Engine/Source/Editor/UnrealEd/Private/CookOnTheFlyServer.cpp b/Engine/Source/Editor/UnrealEd/Private/CookOnTheFlyServer.cpp
--- a/Engine/Source/Editor/UnrealEd/Private/CookOnTheFlyServer.cpp
+++ b/Engine/Source/Editor/UnrealEd/Private/CookOnTheFlyServer.cpp
@@ -52,15 +52,11 @@
 #include "Interfaces/IAudioFormat.h"
 #include "Interfaces/IShaderFormat.h"
 #include "Interfaces/ITextureFormat.h"
-#include "Developer/PackageDependencyInfo/Public/PackageDependencyInfo.h"
 #include "IMessageContext.h"
 #include "Helpers/MessageEndpoint.h"
 #include "Helpers/MessageEndpointBuilder.h"
 #include "Interfaces/INetworkFileServer.h"
 #include "Interfaces/INetworkFileSystemModule.h"
-#include "Dom/JsonObject.h"
-#include "Serialization/JsonReader.h"
-#include "Serialization/JsonSerializer.h"
 #include "PlatformInfo.h"
 
 #include "AssetRegistryModule.h"
@@ -2781,9 +2777,6 @@ void UCookOnTheFlyServer::MarkPackageDirtyForCooker( UPackage *Package )
 			return;
 		}
 
-		// find all packages which depend on this one
-		MarkDependentPackagesDirtyForCooker( Package->GetFName() );
-
 		UE_LOG(LogCook, Verbose, TEXT("Modification detected to package %s"), *PackageFFileName.ToString());
 
 		if ( IsCookingInEditor() )
@@ -2830,14 +2823,6 @@ void UCookOnTheFlyServer::MarkPackageDirtyForCooker( UPackage *Package )
 	}
 }
 
-
-void UCookOnTheFlyServer::MarkDependentPackagesDirtyForCooker( const FName& PackageName)
-{
-	// mark all packages which are dependent on this package as need recooking
-	//  currently this is handled by the VerifyCookedPackagesAreUptoDate
-	//  when a new cook by the book request starts it will call VerifyCookedPackagesAreUptodate.
-}
-
 void UCookOnTheFlyServer::EndNetworkFileServer()
 {
 	for ( int i = 0; i < NetworkFileServers.Num(); ++i )
@@ -3440,8 +3425,7 @@ void UCookOnTheFlyServer::Initialize( ECookMode::Type DesiredCookMode, ECookInit
 			}
 		}
 	}
-
-
+	
 	UE_LOG(LogCook, Display, TEXT("Mobile HDR setting %d"), IsMobileHDR());
 
 	// See if there are any plugins that need to be remapped for the sandbox
@@ -4335,1499 +4319,341 @@ FString UCookOnTheFlyServer::ConvertCookedPathToUncookedPath(const FString& Cook
 
 }
 
-struct FCookedAssetPackageInfo
-{
-	FCookedAssetPackageInfo() : IsEditorPackage(false), SucceededSave(true) { }
-	FString RelativeFilename;
-	FString SourcePackageName;
-	FMD5Hash DependentPackageHash;
-	FMD5Hash FullPackageHash;
-	bool IsEditorPackage;
-	bool SucceededSave;
-	TArray<FString> DependentPackages;
-};
-
-bool ReadCookedAssetRegistry(const FString& Filename, TArray<FCookedAssetPackageInfo>& CookedPackageInfo )
+void UCookOnTheFlyServer::GetAllCookedFiles(TMap<FName, FName>& UncookedPathToCookedPath, const FString& SandboxPath)
 {
-	FString JsonInString;
-	bool Succeeded = FFileHelper::LoadFileToString(JsonInString, *Filename);
-	if (Succeeded)
+	IPlatformFile& PlatformFile = FPlatformFileManager::Get().GetPlatformFile();
+	TArray<FString> CookedFiles;
+	FPackageSearchVisitor PackageSearch(CookedFiles);
+	PlatformFile.IterateDirectoryRecursively(*SandboxPath, PackageSearch);
+	for (const FString& CookedFile : CookedFiles)
 	{
-		// load up previous package asset registry and fill in any packages which weren't recooked on this run
-		TSharedRef<TJsonReader<TCHAR>> Reader = TJsonReaderFactory<TCHAR>::Create(JsonInString);
-		TSharedPtr<FJsonObject> JsonObject;
-		Succeeded = FJsonSerializer::Deserialize<TCHAR>(Reader, JsonObject); 
-		// JsonObject.IsValid() && JsonObject->HasTypedField<EJson::Array>(TEXT("Packages"));
-		if (Succeeded)
-		{
-			TArray<TSharedPtr<FJsonValue>> PackageList = JsonObject->GetArrayField(TEXT("Packages"));
-			for (auto PackageListIt = PackageList.CreateConstIterator(); PackageListIt && Succeeded; ++PackageListIt)
-{
-				const TSharedPtr<FJsonValue>& JsonValue = *PackageListIt;
-				Succeeded = JsonValue->Type == EJson::Object;
-				if (Succeeded)
-				{
-					const TSharedPtr<FJsonObject>& JsonPackage = JsonValue->AsObject();
-
-					// get the package name and see if we have already written it out this run
-					int32 Index = CookedPackageInfo.AddDefaulted(1);
-					FCookedAssetPackageInfo& PackageInfo = CookedPackageInfo[Index];
-
-					check( PackageInfo.IsEditorPackage == false );
-					FString SourcePackageName;
-					verify(JsonPackage->TryGetStringField(TEXT("SourcePackageName"), SourcePackageName));
-					PackageInfo.SourcePackageName = SourcePackageName;
-
-
-					FString CookedPackageName;
-					verify(JsonPackage->TryGetStringField(TEXT("CookedPackageName"), CookedPackageName));
-
-					JsonPackage->TryGetStringField(TEXT("ResolvedRelativePath"), PackageInfo.RelativeFilename);
-
-					/*FPackageName::TryConvertLongPackageNameToFilename(SourcePackageName, PackageInfo.RelativeFilename);
-					PackageInfo.RelativeFilename += FPaths::GetExtension(CookedPackageName, true);*/
-
-					FString DependentPackageHash;
-					if (JsonPackage->TryGetStringField(TEXT("DependentPackageHash"), DependentPackageHash))
-					{
-						Lex::FromString(PackageInfo.DependentPackageHash, *DependentPackageHash);
-					}
-
-					FString FullPackageHash;
-					if (JsonPackage->TryGetStringField(TEXT("FullPackageHash"), FullPackageHash))
-					{
-						Lex::FromString(PackageInfo.FullPackageHash, *FullPackageHash);
-					}
-
-					FString IsEditorPackage;
-					if (JsonPackage->TryGetStringField(TEXT("IsEditorOnly"), IsEditorPackage))
-					{
-						Lex::FromString(PackageInfo.IsEditorPackage, *IsEditorPackage);
-					}
-
-					FString SucceededSave;
-					if (JsonPackage->TryGetStringField(TEXT("SucceededSave"), SucceededSave))
-					{
-						Lex::FromString(PackageInfo.SucceededSave, *SucceededSave);
-					}
-
-
-					const TArray<TSharedPtr<FJsonValue>>* DependencyArray = nullptr;
-					if (JsonPackage->TryGetArrayField(TEXT("DependencyData"), DependencyArray) )
-					{
-						check( DependencyArray != nullptr );
-						for ( const TSharedPtr<FJsonValue>& Dependency : *DependencyArray )
-						{
-							// if ( Dependency != nullptr )
-							{
-								PackageInfo.DependentPackages.Add( Dependency->AsString() );
-							}
-						}
-					}
+		const FName CookedFName(*CookedFile);
 
+		const FString CookedFullPath = FPaths::ConvertRelativePathToFull(CookedFile);
+		const FString UncookedFilename = ConvertCookedPathToUncookedPath(CookedFullPath);
 
-				}
-			}
-		}
+		const FName UncookedFName(*UncookedFilename);
+		UncookedPathToCookedPath.Add(UncookedFName, CookedFName);
 	}
-	return Succeeded;
 }
 
-
-
-void InternalCopyJsonValueToWriter(TSharedRef<TJsonWriter<TCHAR>> &Json, const FString& ValueName, const TSharedPtr<FJsonValue>& JsonValue)
+void UCookOnTheFlyServer::PopulateCookedPackagesFromDisk(const TArray<ITargetPlatform*>& Platforms)
 {
-	if (JsonValue->Type == EJson::String)
-	{
-		Json->WriteValue(ValueName, JsonValue->AsString());
-	}
-	else if (JsonValue->Type == EJson::Array)
-	{
-		if (ValueName.IsEmpty())
-	{
-			Json->WriteArrayStart();
-		}
-		else
-		{
-			Json->WriteArrayStart(ValueName);
-	}
+	check(IsChildCooker() == false);
 
-		const TArray<TSharedPtr<FJsonValue>>& Array = JsonValue->AsArray();
-		for (const auto& ArrayValue : Array)
+	// See what files are out of date in the sandbox folder
+	for (int32 Index = 0; Index < Platforms.Num(); Index++)
 	{
-			InternalCopyJsonValueToWriter(Json, FString(), ArrayValue);
-		}
+		TArray<FString> CookedPackagesToDelete;
 
-		Json->WriteArrayEnd();
-	}
-	else if (JsonValue->Type == EJson::Object)
-	{
-		if (ValueName.IsEmpty())
-		{
-			Json->WriteObjectStart();
-		}
-		else
-		{
-			Json->WriteObjectStart(ValueName);
-		}
+		ITargetPlatform* Target = Platforms[Index];
+		FString SandboxPath = GetSandboxDirectory(Target->PlatformName());
+		FName PlatformFName(*Target->PlatformName());
 
-		const TSharedPtr<FJsonObject>& Object = JsonValue->AsObject();
-		for (const auto& ObjectProperty : Object->Values)
+		FString EngineSandboxPath = SandboxFile->ConvertToSandboxPath(*FPaths::EngineDir()) + TEXT("/");
+		EngineSandboxPath.ReplaceInline(TEXT("[Platform]"), *Target->PlatformName());
+
+		FString GameSandboxPath = SandboxFile->ConvertToSandboxPath(*(FPaths::GameDir() + TEXT("a.txt")));
+		GameSandboxPath.ReplaceInline(TEXT("a.txt"), TEXT(""));
+		GameSandboxPath.ReplaceInline(TEXT("[Platform]"), *Target->PlatformName());
+
+		FString LocalGamePath = FPaths::GameDir();
+		if (FPaths::IsProjectFilePathSet())
 		{
-			InternalCopyJsonValueToWriter(Json, ObjectProperty.Key, ObjectProperty.Value);
+			LocalGamePath = FPaths::GetPath(FPaths::GetProjectFilePath()) + TEXT("/");
 		}
 
-		Json->WriteObjectEnd();
-	}
-	else
-		{
+		FString LocalEnginePath = FPaths::EngineDir();
 
-		UE_LOG(LogCook, Warning, TEXT("Unrecognized json value type %d in object %s"), *UEnum::GetValueAsString(TEXT("Json.EJson"), JsonValue->Type), *ValueName)
-	}
-		}
+		static bool bFindCulprit = false; //debugging setting if we want to find culprit for iterative cook issues
 
+		// Registry generator already exists
+		FAssetRegistryGenerator* PlatformAssetRegistry = RegistryGenerators.FindRef(PlatformFName);
 
-bool UCookOnTheFlyServer::SaveCookedAssetRegistry(const FString& InCookedAssetRegistryFilename, const TArray<FName>& InAllCookedPackages, const TArray<FName>& UncookedEditorPackages, const TArray<FName>& FailedToSavePackages, const FString& PlatformName, const bool Append) const
-{
-	bool bSuccess = true;
+		// Load the platform cooked asset registry file
+		const FString CookedAssetRegistry = FPaths::GameDir() / GetDevelopmentAssetRegistryFilename();
+		const FString SandboxCookedAssetRegistryFilename = ConvertToFullSandboxPath(*CookedAssetRegistry, true, Target->PlatformName());
 
-	const FString CookedAssetRegistryFilename = InCookedAssetRegistryFilename.Replace(TEXT("[Platform]"), *PlatformName);
+		bool bIsIterateSharedBuild = IsCookFlagSet(ECookInitializationFlags::IterateSharedBuild);
 
+		if (bIsIterateSharedBuild)
+		{
+			// see if the shared build is newer then the current cooked content in the local directory
+			FDateTime CurrentLocalCookedBuild = IFileManager::Get().GetTimeStamp(*SandboxCookedAssetRegistryFilename);
 
-	FString JsonOutString;
-	TSharedRef<TJsonWriter<TCHAR>> Json = TJsonWriterFactory<TCHAR, TPrettyJsonPrintPolicy<TCHAR> >::Create(&JsonOutString);
+			// iterate on the shared build if the option is set
+			FString SharedCookedAssetRegistry = FPaths::Combine(*FPaths::GameSavedDir(), TEXT("SharedIterativeBuild"), *Target->PlatformName(), TEXT("Cooked"), GetDevelopmentAssetRegistryFilename());
 
-	Json->WriteObjectStart();
-	Json->WriteArrayStart(TEXT("Packages"));
+			FDateTime CurrentIterativeCookedBuild = IFileManager::Get().GetTimeStamp(*SharedCookedAssetRegistry);
 
-	FPackageDependencyInfoModule& PDInfoModule = FModuleManager::LoadModuleChecked<FPackageDependencyInfoModule>("PackageDependencyInfo");
 
-	TArray<FName> AllCookedPackages = InAllCookedPackages;
-	TSet<FName> FailedToSavePackagesSet;
-	for ( const FName& FailedToSavePackage : FailedToSavePackages)
-	{
-		const FName& FailedStandardName = GetCachedStandardPackageFileFName(FailedToSavePackage);
-		AllCookedPackages.Add(FailedStandardName );
-		FailedToSavePackagesSet.Add(FailedStandardName);
-	}
 
-	{
-		SCOPE_TIMER(DetermineDependentTimeStamps);
-		
-		TArray<FString> RelativePaths;
-		for ( const auto& FNamePath : AllCookedPackages )
-		{
-			RelativePaths.Add(FNamePath.ToString());
-		}
-		// resolving dependencies for packages
-		int32 NumThreadsToDeterminePackageInfo = FMath::Max(GLargeThreadPool->GetNumThreads() - 2, 1);
-		check(bIgnoreMarkupPackageAlreadyLoaded==false);
-		bIgnoreMarkupPackageAlreadyLoaded = true;
-		PDInfoModule.AsyncDetermineAllDependentPackageInfo(RelativePaths, NumThreadsToDeterminePackageInfo);
-		bIgnoreMarkupPackageAlreadyLoaded = false;
-	}
+			if ( (CurrentIterativeCookedBuild >= CurrentLocalCookedBuild) && 
+				(CurrentIterativeCookedBuild != FDateTime::MinValue()) )
+			{
+				// clean the sandbox 
+				ClearPlatformCookedData(FName(*Target->PlatformName()));
+				FString SandboxDirectory = GetSandboxDirectory(Target->PlatformName());
+				IFileManager::Get().DeleteDirectory(*SandboxDirectory, false, true);
 
-	AllCookedPackages.StableSort();
+				// SaveCurrentIniSettings(Target); // use this if we don't care about ini safty.
+				// copy the ini settings from the shared cooked build. 
+				const FString SharedCookedIniFile = FPaths::Combine(*FPaths::GameSavedDir(), TEXT("SharedIterativeBuild"), *Target->PlatformName(), TEXT("Cooked"), TEXT("CookedIniVersion.txt"));
+				const FString SandboxCookedIniFile = ConvertToFullSandboxPath(*(FPaths::GameDir() / TEXT("CookedIniVersion.txt")), true).Replace(TEXT("[Platform]"), *Target->PlatformName());
 
-	TSet<FName> ProcessedCookedPackages; // keep track of which packages are processed
-	for (const FName& Package : AllCookedPackages)
-	{
-		const FString& RelativePath = Package.ToString();
+				IFileManager::Get().Copy(*SandboxCookedIniFile, *SharedCookedIniFile);
 
-		FName PackageName;
-		FString PlatformSandboxPath;
-		//FString PlatformSandboxPath = SandboxPath.Replace(TEXT("[Platform]"), *Platform->PlatformName());
-		if (FPackageName::IsScriptPackage(RelativePath) || FPackageName::IsMemoryPackage(RelativePath))
-		{
-			PackageName = Package;
-		}
-		else
-		{
-			FString StringPackageName;
-			if (FPackageName::TryConvertFilenameToLongPackageName(RelativePath, StringPackageName) == false)
+				bool bIniSettingsOutOfDate = IniSettingsOutOfDate(Target);
+				if (bIniSettingsOutOfDate && !IsCookFlagSet(ECookInitializationFlags::IgnoreIniSettingsOutOfDate))
+				{
+					UE_LOG(LogCook, Display, TEXT("Shared iterative build ini settings out of date, not using shared cooked build"));
+				}
+				else
+				{
+					if (bIniSettingsOutOfDate)
+					{
+						UE_LOG(LogCook, Display, TEXT("Shared iterative build ini settings out of date, but we don't care"));
+					}
+
+					UE_LOG(LogCook, Display, TEXT("Shared iterative build is newer then local cooked build, iteratively cooking from shared build "));
+					PlatformAssetRegistry->LoadPreviousAssetRegistry(SharedCookedAssetRegistry);
+				}
+			}
+			else
 			{
-				UE_LOG(LogCook, Warning, TEXT("Unable to determine extension for file %s skipping asset registry entry"), *PackageName.ToString());
-				continue;
+				UE_LOG(LogCook, Display, TEXT("Local cook is newer then shared cooked build, iterativly cooking from local build"));
+				PlatformAssetRegistry->LoadPreviousAssetRegistry(SandboxCookedAssetRegistryFilename);
 			}
-			PackageName = FName(*StringPackageName);
-
-			PlatformSandboxPath = SandboxFile->ConvertToSandboxPath(*StringPackageName);
-			PlatformSandboxPath.ReplaceInline(TEXT("[Platform]"), *PlatformName);
 		}
-
-		if (UncookedEditorPackages.Contains(PackageName))
+		else
 		{
-			continue;
+			PlatformAssetRegistry->LoadPreviousAssetRegistry(SandboxCookedAssetRegistryFilename);
 		}
 
-		Json->WriteObjectStart(); // unnamed package start
-
-		ProcessedCookedPackages.Add(PackageName);
+		// Get list of changed packages
+		TSet<FName> ModifiedPackages, NewPackages, RemovedPackages, IdenticalCookedPackages, IdenticalUncookedPackages;
 
+		// We recurse modifications up the reference chain because it is safer, if this ends up being a significant issue in some games we can add a command line flag
+		bool bRecurseModifications = true;
+		bool bRecurseScriptModifications = !IsCookFlagSet(ECookInitializationFlags::IgnoreScriptPackagesOutOfDate);
+		PlatformAssetRegistry->ComputePackageDifferences(ModifiedPackages, NewPackages, RemovedPackages, IdenticalCookedPackages, IdenticalUncookedPackages, bRecurseModifications, bRecurseScriptModifications);
 
+		// check the files on disk 
+		TMap<FName, FName> UncookedPathToCookedPath;
+		// get all the on disk cooked files
+		GetAllCookedFiles(UncookedPathToCookedPath, SandboxPath);
 
-		FPackageName::FindPackageFileWithoutExtension(PlatformSandboxPath, PlatformSandboxPath);
+		const static FName NAME_DummyCookedFilename(TEXT("DummyCookedFilename")); // pls never name a package dummycookedfilename otherwise shit might go wonky
+		if (bIsIterateSharedBuild)
+		{
+			TSet<FName> ExistingPackages = ModifiedPackages;
+			ExistingPackages.Append(RemovedPackages);
+			ExistingPackages.Append(IdenticalCookedPackages);
+			ExistingPackages.Append(IdenticalUncookedPackages);
 
-		FDateTime TimeStamp = IFileManager::Get().GetTimeStamp(*PlatformSandboxPath);
-		int64 FileSize = IFileManager::Get().FileSize(*PlatformSandboxPath);
+			// if we are iterating of a shared build the cooked files might not exist in the cooked directory because we assume they are packaged in the pak file (which we don't want to extract)
+			for (FName PackageName : ExistingPackages)
+			{
+				FString Filename;
+				if (FPackageName::DoesPackageExist(PackageName.ToString(), nullptr, &Filename))
+				{
+					UncookedPathToCookedPath.Add(FName(*Filename), NAME_DummyCookedFilename);
+				}
+			}
+		}
 
-		Json->WriteValue(TEXT("SourcePackageName"), PackageName.ToString());
-		Json->WriteValue(TEXT("CookedPackageName"), PlatformSandboxPath);
-		Json->WriteValue(TEXT("CookedPackageTimeStamp"), TimeStamp.ToString());
-		Json->WriteValue(TEXT("FileSize"), FString::Printf(TEXT("%lld"), FileSize));
-		Json->WriteValue(TEXT("ResolvedRelativePath"), *RelativePath);
+		uint32 NumPackagesConsidered = UncookedPathToCookedPath.Num();
+		uint32 NumPackagesUnableToFindCookedPackageInfo = 0;
+		uint32 NumPackagesFileHashMismatch = 0;
+		uint32 NumPackagesKept = 0;
+		uint32 NumMarkedFailedSaveKept = 0;
+		uint32 NumPackagesRemoved = 0;
 
-		Json->WriteValue(TEXT("IsEditorOnly"), TEXT("false"));
-		if ( FailedToSavePackagesSet.Contains(Package) )
-		{
-			Json->WriteValue(TEXT("SucceededSave"), TEXT("false"));
-		}
-		else
+		for (const auto& CookedPaths : UncookedPathToCookedPath)
 		{
-			Json->WriteValue(TEXT("SucceededSave"), TEXT("true"));
-		}
+			const FName CookedFile = CookedPaths.Value;
+			const FName UncookedFilename = CookedPaths.Key;
+			const FName* FoundPackageName = GetCachedPackageFilenameToPackageFName(UncookedFilename);
+			if ( !FoundPackageName )
+			{
+				++NumPackagesRemoved;
+			}
+			const FName PackageName = *FoundPackageName;
+			const FString UncookedFilenameString = UncookedFilename.ToString();
+			bool bShouldKeep = true;
 
+			if (ModifiedPackages.Contains(PackageName))
+			{
+				++NumPackagesFileHashMismatch;
+				bShouldKeep = false;
+			}
+			else if (NewPackages.Contains(PackageName) || RemovedPackages.Contains(PackageName))
+			{
+				++NumPackagesUnableToFindCookedPackageInfo;
+				bShouldKeep = false;
+			}
+			else if (IdenticalUncookedPackages.Contains(PackageName))
+			{
+				// These are packages which failed to save the first time 
+				// most likely because they are editor only packages
+				bShouldKeep = false;
+			}
+				
+			if (CookedFile == NAME_DummyCookedFilename)
+			{
+				check(IFileManager::Get().FileExists(*CookedFile.ToString()) == false);
+			}
 
-		{
-			SCOPE_TIMER(DeterminePackageDependentHash);
-			FMD5Hash Hash;
-			PDInfoModule.DeterminePackageDependentHash(*RelativePath, Hash);
-			FString DependentPackageHashString = Lex::ToString(Hash);
-			Json->WriteValue(TEXT("DependentPackageHash"), *DependentPackageHashString);
-		}
-		{
-			SCOPE_TIMER(DetermineFullPackageHash);
-			FMD5Hash Hash;
-			PDInfoModule.DetermineFullPackageHash(*RelativePath, Hash);
-			FString FullPackageHashString = Lex::ToString(Hash);
-			Json->WriteValue(TEXT("FullPackageHash"), *FullPackageHashString);
-		}
+			TArray<FName> PlatformNames;
+			PlatformNames.Add(PlatformFName);
 
-		{
-			SCOPE_TIMER(GetPackageDependencies);
-			TArray<FString> DependencyNames;
-			if (PDInfoModule.GetPackageDependencies(*RelativePath, DependencyNames))
+			if (bShouldKeep)
 			{
-				Json->WriteArrayStart(TEXT("DependencyData"));
+				// Add this package to the CookedPackages list so that we don't try cook it again
+				if (CookedFile != NAME_DummyCookedFilename)
+				{
+					check(IFileManager::Get().FileExists(*CookedFile.ToString()));
+				}
+				TArray<bool> Succeeded;
+				Succeeded.Add(true);
 
-				for (const FString& DependencyName : DependencyNames)
+				if (IdenticalCookedPackages.Contains(PackageName))
 				{
-					Json->WriteValue(DependencyName);
+					CookedPackages.Add(FFilePlatformCookedPackage(UncookedFilename, MoveTemp(PlatformNames), MoveTemp(Succeeded)));
+					++NumPackagesKept;
 				}
-				Json->WriteArrayEnd();
 			}
-		}
-
-		Json->WriteValue(TEXT("FileSize"), FString::Printf(TEXT("%lld"), FileSize));
-		Json->WriteArrayStart("AssetData");
-
-
-		TArray<FAssetData> AssetRegistryData;
-		AssetRegistry->GetAssetsByPackageName(PackageName, AssetRegistryData, true);
-
-		for (const FAssetData& AssetData : AssetRegistryData)
-		{	// Add only assets that have actually been cooked and belong to any chunk
-			if (AssetData.ChunkIDs.Num() > 0 && (AssetData.PackageName == PackageName))
+			else
 			{
-				Json->WriteObjectStart();
-				// save all their infos 
-				Json->WriteValue(TEXT("ObjectPath"), AssetData.ObjectPath.ToString());
-				Json->WriteValue(TEXT("PackageName"), AssetData.PackageName.ToString());
-				Json->WriteValue(TEXT("PackagePath"), AssetData.PackagePath.ToString());
-				Json->WriteValue(TEXT("AssetName"), AssetData.AssetName.ToString());
-				Json->WriteValue(TEXT("AssetClass"), AssetData.AssetClass.ToString());
-				Json->WriteObjectStart("TagsAndValues");
-				for (const auto& Tag : AssetData.TagsAndValues)
+				if ( IsCookByTheBookMode() ) // cook on the fly will requeue this package when it wants it 
+				{
+					// Force cook the modified file
+					CookRequests.EnqueueUnique(MoveTemp(FFilePlatformRequest(UncookedFilename, PlatformNames)));
+				}
+				if (CookedFile != NAME_DummyCookedFilename)
 				{
-					Json->WriteValue(Tag.Key.ToString(), Tag.Value);
+					// delete the old package 
+					const FString CookedFullPath = FPaths::ConvertRelativePathToFull(CookedFile.ToString());
+					UE_LOG(LogCook, Verbose, TEXT("Deleting cooked package %s failed filehash test"), *CookedFullPath);
+					CookedPackagesToDelete.Add(CookedFullPath);
+				}
+				else
+				{
+					// the cooker should rebuild this package because it's not in the cooked package list
+					// the new package will have higher priority then the package in the original shared cooked build
+					UE_LOG(LogCook, Verbose, TEXT("Shared cooked build: Detected package is out of date %s"), *UncookedFilenameString);
 				}
-				Json->WriteObjectEnd(); // end tags and values object
-				Json->WriteObjectEnd(); // end unnamed array object
 			}
 		}
-		Json->WriteArrayEnd();
-		Json->WriteObjectEnd(); // unnamed package
-	}
-	/*
-	for (const auto& FailedPackage : FailedToSavePackages)
-	{
-		const FString& RelativePath = FailedPackage.ToString();
-		
-		FName PackageName;
-		FString PlatformSandboxPath;
-		//FString PlatformSandboxPath = SandboxPath.Replace(TEXT("[Platform]"), *Platform->PlatformName());
-		if (FPackageName::IsScriptPackage(RelativePath))
-		{
-			PackageName = FailedPackage;
-		}
-		else
+
+		// Register identical uncooked packages from previous run
+		for (FName UncookedPackage : IdenticalUncookedPackages)
 		{
-			FString StringPackageName;
-			if (FPackageName::TryConvertFilenameToLongPackageName(RelativePath, StringPackageName) == false)
-			{
-				UE_LOG(LogCook, Warning, TEXT("Unable to determine extension for file %s skipping asset registry entry"), *PackageName.ToString());
-				continue;
-			}
-			PackageName = FName(*StringPackageName);
+			const FName UncookedFilename = GetCachedStandardPackageFileFName(UncookedPackage);
 
-			PlatformSandboxPath = SandboxFile->ConvertToSandboxPath(*StringPackageName);
-			PlatformSandboxPath.ReplaceInline(TEXT("[Platform]"), *PlatformName);
-		}
+			TArray<FName> PlatformNames;
+			PlatformNames.Add(PlatformFName);
 
-		Json->WriteObjectStart(); // unnamed package start
+			ensure(CookedPackages.Exists(UncookedFilename, PlatformNames, false) == false);
 
-		ProcessedCookedPackages.Add(PackageName);
+			CookedPackages.Add(MoveTemp(FFilePlatformCookedPackage(UncookedFilename, MoveTemp(PlatformNames))));
+			++NumMarkedFailedSaveKept;
+		}
 
+		UE_LOG(LogCook, Display, TEXT("Iterative cooking summary for %s, \nConsidered: %d, \nFile Hash missmatch: %d, \nPackages Kept: %d, \nPackages failed save kept: %d, \nMissing Cooked Info(expected 0): %d"),
+			*Target->PlatformName(),
+			NumPackagesConsidered, NumPackagesFileHashMismatch,
+			NumPackagesKept, NumMarkedFailedSaveKept,
+			NumPackagesUnableToFindCookedPackageInfo);
 
+		auto DeletePackageLambda = [&CookedPackagesToDelete](int32 PackageIndex)
+		{
+			const FString& CookedFullPath = CookedPackagesToDelete[PackageIndex];
+			IFileManager::Get().Delete(*CookedFullPath, true, true, true);
+		};
+		ParallelFor(CookedPackagesToDelete.Num(), DeletePackageLambda);
+	}
+}
 
-		FPackageName::FindPackageFileWithoutExtension(PlatformSandboxPath, PlatformSandboxPath);
+const FString ExtractPackageNameFromObjectPath( const FString ObjectPath )
+{
+	// get the path 
+	int32 Beginning = ObjectPath.Find(TEXT("'"), ESearchCase::CaseSensitive);
+	if ( Beginning == INDEX_NONE )
+	{
+		return ObjectPath;
+	}
+	int32 End = ObjectPath.Find(TEXT("."), ESearchCase::CaseSensitive, ESearchDir::FromStart, Beginning + 1);
+	if (End == INDEX_NONE )
+	{
+		End = ObjectPath.Find(TEXT("'"), ESearchCase::CaseSensitive, ESearchDir::FromStart, Beginning + 1);
+	}
+	if ( End == INDEX_NONE )
+	{
+		// one more use case is that the path is "Class'Path" example "OrionBoostItemDefinition'/Game/Misc/Boosts/XP_1Win" dunno why but this is actually dumb
+		if ( ObjectPath[Beginning+1] == '/' )
+		{
+			return ObjectPath.Mid(Beginning+1);
+		}
+		return ObjectPath;
+	}
+	return ObjectPath.Mid(Beginning + 1, End - Beginning - 1);
+}
 
-		FDateTime TimeStamp = IFileManager::Get().GetTimeStamp(*PlatformSandboxPath);
-		int64 FileSize = IFileManager::Get().FileSize(*PlatformSandboxPath);
+void UCookOnTheFlyServer::CleanSandbox(const bool bIterative)
+{
+	//child cookers shouldn't clean sandbox we would be deleting the master cookers / other cookers data
+	check(IsChildCooker() == false);
 
-		Json->WriteValue(TEXT("SourcePackageName"), PackageName.ToString());
-		Json->WriteValue(TEXT("CookedPackageName"), PlatformSandboxPath);
-		Json->WriteValue(TEXT("ResolvedRelativePath"), *RelativePath);
+	const TArray<ITargetPlatform*>& Platforms = GetCookingTargetPlatforms();
 
+	// before we can delete any cooked files we need to make sure that we have finished writing them
+	UPackage::WaitForAsyncFileWrites();
 
-		Json->WriteValue(TEXT("IsEditorOnly"), TEXT("false"));
-		Json->WriteValue(TEXT("SucceededSave"), TEXT("false"));
-				
+#if OUTPUT_TIMING
+	double SandboxCleanTime = 0.0;
+#endif
+	{
+#if OUTPUT_TIMING
+		SCOPE_SECONDS_COUNTER(SandboxCleanTime);
+		SCOPE_TIMER(CleanSandboxTime);
+#endif
+		if (bIterative == false)
+		{
+			// for now we are going to wipe the cooked directory
+			for (int32 Index = 0; Index < Platforms.Num(); Index++)
+			{
+				ITargetPlatform* Target = Platforms[Index];
+				UE_LOG(LogCook, Display, TEXT("Cooked content cleared for platform %s"), *Target->PlatformName());
 
-		Json->WriteObjectEnd();
-			}
-	*/
+				FString SandboxDirectory = GetSandboxDirectory(Target->PlatformName()); // GetOutputDirectory(Target->PlatformName());
+				IFileManager::Get().DeleteDirectory(*SandboxDirectory, false, true);
 
-	TArray<FName> SortedUncookedEditorPackages = UncookedEditorPackages;
-	SortedUncookedEditorPackages.StableSort();
+				ClearPlatformCookedData(FName(*Target->PlatformName()));
 
-	for (const auto& EditorOnlyPackage : SortedUncookedEditorPackages)
-	{
-		const FString& RelativePath = EditorOnlyPackage.ToString();
+				IniSettingsOutOfDate(Target);
+				SaveCurrentIniSettings(Target);
+			}
 
-		FName PackageName;
-		FString PlatformSandboxPath;
-		//FString PlatformSandboxPath = SandboxPath.Replace(TEXT("[Platform]"), *Platform->PlatformName());
-		if (FPackageName::IsScriptPackage(RelativePath) || FPackageName::IsMemoryPackage(RelativePath))
-		{
-			PackageName = EditorOnlyPackage;
 		}
 		else
 		{
-			FString StringPackageName;
-			if (FPackageName::TryConvertFilenameToLongPackageName(RelativePath, StringPackageName) == false)
+			for (const ITargetPlatform* Target : Platforms)
 			{
-				UE_LOG(LogCook, Warning, TEXT("Unable to determine extension for file %s skipping asset registry entry"), *PackageName.ToString());
-				continue;
-			}
-			PackageName = FName(*StringPackageName);
+				bool bIniSettingsOutOfDate = IniSettingsOutOfDate(Target);
+				if (bIniSettingsOutOfDate)
+				{
+					if (!IsCookFlagSet(ECookInitializationFlags::IgnoreIniSettingsOutOfDate))
+					{
+						UE_LOG(LogCook, Display, TEXT("Cook invalidated for platform %s ini settings don't match from last cook, clearing all cooked content"), *Target->PlatformName());
 
-			PlatformSandboxPath = SandboxFile->ConvertToSandboxPath(*StringPackageName);
-			PlatformSandboxPath.ReplaceInline(TEXT("[Platform]"), *PlatformName);
-		}
-
-		Json->WriteObjectStart(); // unnamed package start
-
-		ProcessedCookedPackages.Add(PackageName);
-
-
-
-		FPackageName::FindPackageFileWithoutExtension(PlatformSandboxPath, PlatformSandboxPath);
-
-		FDateTime TimeStamp = IFileManager::Get().GetTimeStamp(*PlatformSandboxPath);
-		int64 FileSize = IFileManager::Get().FileSize(*PlatformSandboxPath);
-
-		Json->WriteValue(TEXT("SourcePackageName"), PackageName.ToString());
-		Json->WriteValue(TEXT("CookedPackageName"), PlatformSandboxPath);
-		Json->WriteValue(TEXT("ResolvedRelativePath"), *RelativePath);
-
-		Json->WriteValue(TEXT("IsEditorOnly"), TEXT("true"));
-		Json->WriteValue(TEXT("SucceededSave"), TEXT("false"));
-
-		Json->WriteObjectEnd();
-	}
-
-
-	if (Append)
-	{
-		FString JsonInString;
-		if (FFileHelper::LoadFileToString(JsonInString, *CookedAssetRegistryFilename))
-		{
-			// load up previous package asset registry and fill in any packages which weren't recooked on this run
-			TSharedRef<TJsonReader<TCHAR>> Reader = TJsonReaderFactory<TCHAR>::Create(JsonInString);
-			TSharedPtr<FJsonObject> JsonObject;
-			bool shouldRead = FJsonSerializer::Deserialize(Reader, JsonObject) && JsonObject.IsValid() && JsonObject->HasTypedField<EJson::Array>(TEXT("Packages"));
-			if (shouldRead)
-				{
-				TArray<TSharedPtr<FJsonValue>> PackageList = JsonObject->GetArrayField(TEXT("Packages"));
-				for (auto PackageListIt = PackageList.CreateConstIterator(); PackageListIt && shouldRead; ++PackageListIt)
-				{
-					const TSharedPtr<FJsonValue>& JsonValue = *PackageListIt;
-					shouldRead = JsonValue->Type == EJson::Object;
-					if (shouldRead)
-					{
-						const TSharedPtr<FJsonObject>& JsonPackage = JsonValue->AsObject();
-
-						// get the package name and see if we have already written it out this run
-
-						FString CookedPackageName;
-						verify(JsonPackage->TryGetStringField(TEXT("SourcePackageName"), CookedPackageName));
-
-						const FName CookedPackageFName(*CookedPackageName);
-						if (ProcessedCookedPackages.Contains(CookedPackageFName))
-						{
-							// don't need to process this package
-							continue;
-						}
-
-
-						// check that the on disk version is still valid
-						FString SourcePackageName;
-						check(JsonPackage->TryGetStringField(TEXT("SourcePackageName"), SourcePackageName));
-
-						// if our timestamp is different then don't copy the information over
-						FDateTime CurrentTimeStamp = IFileManager::Get().GetTimeStamp(*CookedPackageName);
-
-						FString SavedTimeString;
-						check(JsonPackage->TryGetStringField(TEXT("CookedPackageTimeStamp"), SavedTimeString));
-						FDateTime SavedTimeStamp;
-						FDateTime::Parse(SavedTimeString, SavedTimeStamp);
-
-						if (SavedTimeStamp != CurrentTimeStamp)
-						{
-							continue;
-						}
-
-						InternalCopyJsonValueToWriter(Json, FString(), JsonValue);
-					}
-
-				}
-			}
-			else
-			{
-				UE_LOG(LogCook, Warning, TEXT("Unable to read or json is invalid format %s"), *CookedAssetRegistryFilename);
-			}
-		}
-	}
-
-
-	Json->WriteArrayEnd();
-	Json->WriteObjectEnd();
-
-	if (Json->Close())
-	{
-		FArchive* ItemTemplatesFile = IFileManager::Get().CreateFileWriter(*CookedAssetRegistryFilename);
-		if (ItemTemplatesFile)
-		{
-			// serialize the file contents
-			TStringConversion<FTCHARToUTF8_Convert> Convert(*JsonOutString);
-			ItemTemplatesFile->Serialize(const_cast<ANSICHAR*>(Convert.Get()), Convert.Length());
-			ItemTemplatesFile->Close();
-			if (!ItemTemplatesFile->IsError())
-			{
-				bSuccess = true;
-			}
-			else
-			{
-				UE_LOG(LogCook, Error, TEXT("Unable to write to %s"), *CookedAssetRegistryFilename);
-			}
-			delete ItemTemplatesFile;
-		}
-		else
-		{
-			UE_LOG(LogCook, Error, TEXT("Unable to open %s for writing."), *CookedAssetRegistryFilename);
-		}
-	}
-	else
-	{
-		UE_LOG(LogCook, Error, TEXT("Error closing Json Writer"));
-	}
-	return bSuccess;
-}
-
-
-void UCookOnTheFlyServer::GetAllCookedFiles(TMap<FName, FName>& UncookedPathToCookedPath, const FString& SandboxPath)
-{
-	IPlatformFile& PlatformFile = FPlatformFileManager::Get().GetPlatformFile();
-	TArray<FString> CookedFiles;
-	FPackageSearchVisitor PackageSearch(CookedFiles);
-	PlatformFile.IterateDirectoryRecursively(*SandboxPath, PackageSearch);
-	for (const FString& CookedFile : CookedFiles)
-	{
-		const FName CookedFName(*CookedFile);
-
-		const FString CookedFullPath = FPaths::ConvertRelativePathToFull(CookedFile);
-		const FString UncookedFilename = ConvertCookedPathToUncookedPath(CookedFullPath);
-
-		const FName UncookedFName(*UncookedFilename);
-		UncookedPathToCookedPath.Add(UncookedFName, CookedFName);
-	}
-}
-
-void UCookOnTheFlyServer::PopulateCookedPackagesFromDisk(const TArray<ITargetPlatform*>& Platforms)
-{
-	check(IsChildCooker() == false);
-
-	// check each package to see if it's out of date
-	FPackageDependencyInfoModule* PDInfoModule = nullptr;
-
-	if (IsCookFlagSet(ECookInitializationFlags::Iterative) && !IsCookFlagSet(ECookInitializationFlags::IterateOnAssetRegistry))
-	{
-		PDInfoModule = &FModuleManager::LoadModuleChecked<FPackageDependencyInfoModule>("PackageDependencyInfo");
-	}
-
-	// See what files are out of date in the sandbox folder
-	for (int32 Index = 0; Index < Platforms.Num(); Index++)
-	{
-		TArray<FString> CookedPackagesToDelete;
-
-		ITargetPlatform* Target = Platforms[Index];
-		FString SandboxPath = GetSandboxDirectory(Target->PlatformName());
-		FName PlatformFName(*Target->PlatformName());
-
-		FString EngineSandboxPath = SandboxFile->ConvertToSandboxPath(*FPaths::EngineDir()) + TEXT("/");
-		EngineSandboxPath.ReplaceInline(TEXT("[Platform]"), *Target->PlatformName());
-
-		FString GameSandboxPath = SandboxFile->ConvertToSandboxPath(*(FPaths::GameDir() + TEXT("a.txt")));
-		GameSandboxPath.ReplaceInline(TEXT("a.txt"), TEXT(""));
-		GameSandboxPath.ReplaceInline(TEXT("[Platform]"), *Target->PlatformName());
-
-		FString LocalGamePath = FPaths::GameDir();
-		if (FPaths::IsProjectFilePathSet())
-		{
-			LocalGamePath = FPaths::GetPath(FPaths::GetProjectFilePath()) + TEXT("/");
-		}
-
-		FString LocalEnginePath = FPaths::EngineDir();
-
-		static bool bFindCulprit = false; //debugging setting if we want to find culprit for iterative cook issues
-
-		if (IsCookFlagSet(ECookInitializationFlags::IterateOnAssetRegistry))
-		{
-			// Registry generator already exists
-			FAssetRegistryGenerator* PlatformAssetRegistry = RegistryGenerators.FindRef(PlatformFName);
-
-			// Load the platform cooked asset registry file
-			const FString CookedAssetRegistry = FPaths::GameDir() / GetDevelopmentAssetRegistryFilename();
-			const FString SandboxCookedAssetRegistryFilename = ConvertToFullSandboxPath(*CookedAssetRegistry, true, Target->PlatformName());
-
-			bool bIsIterateSharedBuild = IsCookFlagSet(ECookInitializationFlags::IterateSharedBuild);
-
-			if (bIsIterateSharedBuild)
-			{
-				// see if the shared build is newer then the current cooked content in the local directory
-				FDateTime CurrentLocalCookedBuild = IFileManager::Get().GetTimeStamp(*SandboxCookedAssetRegistryFilename);
-
-				// iterate on the shared build if the option is set
-				FString SharedCookedAssetRegistry = FPaths::Combine(*FPaths::GameSavedDir(), TEXT("SharedIterativeBuild"), *Target->PlatformName(), TEXT("Cooked"), GetDevelopmentAssetRegistryFilename());
-
-				FDateTime CurrentIterativeCookedBuild = IFileManager::Get().GetTimeStamp(*SharedCookedAssetRegistry);
-
-
-
-				if ( (CurrentIterativeCookedBuild >= CurrentLocalCookedBuild) && 
-					(CurrentIterativeCookedBuild != FDateTime::MinValue()) )
-				{
-					// clean the sandbox 
-					ClearPlatformCookedData(FName(*Target->PlatformName()));
-					FString SandboxDirectory = GetSandboxDirectory(Target->PlatformName());
-					IFileManager::Get().DeleteDirectory(*SandboxDirectory, false, true);
-
-					// SaveCurrentIniSettings(Target); // use this if we don't care about ini safty.
-					// copy the ini settings from the shared cooked build. 
-					const FString SharedCookedIniFile = FPaths::Combine(*FPaths::GameSavedDir(), TEXT("SharedIterativeBuild"), *Target->PlatformName(), TEXT("Cooked"), TEXT("CookedIniVersion.txt"));
-					const FString SandboxCookedIniFile = ConvertToFullSandboxPath(*(FPaths::GameDir() / TEXT("CookedIniVersion.txt")), true).Replace(TEXT("[Platform]"), *Target->PlatformName());
-
-					IFileManager::Get().Copy(*SandboxCookedIniFile, *SharedCookedIniFile);
-
-					bool bIniSettingsOutOfDate = IniSettingsOutOfDate(Target);
-					if (bIniSettingsOutOfDate && !IsCookFlagSet(ECookInitializationFlags::IgnoreIniSettingsOutOfDate))
-					{
-						UE_LOG(LogCook, Display, TEXT("Shared iterative build ini settings out of date, not using shared cooked build"));
-					}
-					else
-					{
-						if (bIniSettingsOutOfDate)
-						{
-							UE_LOG(LogCook, Display, TEXT("Shared iterative build ini settings out of date, but we don't care"));
-						}
-
-						UE_LOG(LogCook, Display, TEXT("Shared iterative build is newer then local cooked build, iteratively cooking from shared build "));
-						PlatformAssetRegistry->LoadPreviousAssetRegistry(SharedCookedAssetRegistry);
-					}
-				}
-				else
-				{
-					UE_LOG(LogCook, Display, TEXT("Local cook is newer then shared cooked build, iterativly cooking from local build"));
-					PlatformAssetRegistry->LoadPreviousAssetRegistry(SandboxCookedAssetRegistryFilename);
-				}
-			}
-			else
-			{
-				PlatformAssetRegistry->LoadPreviousAssetRegistry(SandboxCookedAssetRegistryFilename);
-			}
-
-			// Get list of changed packages
-			TSet<FName> ModifiedPackages, NewPackages, RemovedPackages, IdenticalCookedPackages, IdenticalUncookedPackages;
-
-			// We recurse modifications up the reference chain because it is safer, if this ends up being a significant issue in some games we can add a command line flag
-			bool bRecurseModifications = true;
-			PlatformAssetRegistry->ComputePackageDifferences(ModifiedPackages, NewPackages, RemovedPackages, IdenticalCookedPackages, IdenticalUncookedPackages, bRecurseModifications);
-
-			// check the files on disk 
-			TMap<FName, FName> UncookedPathToCookedPath;
-			// get all the on disk cooked files
-			GetAllCookedFiles(UncookedPathToCookedPath, SandboxPath);
-
-			const static FName NAME_DummyCookedFilename(TEXT("DummyCookedFilename")); // pls never name a package dummycookedfilename otherwise shit might go wonky
-			if (bIsIterateSharedBuild)
-			{
-				TSet<FName> ExistingPackages = ModifiedPackages;
-				ExistingPackages.Append(RemovedPackages);
-				ExistingPackages.Append(IdenticalCookedPackages);
-				ExistingPackages.Append(IdenticalUncookedPackages);
-
-				// if we are iterating of a shared build the cooked files might not exist in the cooked directory because we assume they are packaged in the pak file (which we don't want to extract)
-				for (FName PackageName : ExistingPackages)
-				{
-					FString Filename;
-					if (FPackageName::DoesPackageExist(PackageName.ToString(), nullptr, &Filename))
-					{
-						UncookedPathToCookedPath.Add(FName(*Filename), NAME_DummyCookedFilename);
-					}
-				}
-			}
-
-			uint32 NumPackagesConsidered = UncookedPathToCookedPath.Num();
-			uint32 NumPackagesUnableToFindCookedPackageInfo = 0;
-			uint32 NumPackagesFileHashMismatch = 0;
-			uint32 NumPackagesKept = 0;
-			uint32 NumMarkedFailedSaveKept = 0;
-			uint32 NumPackagesRemoved = 0;
-
-			for (const auto& CookedPaths : UncookedPathToCookedPath)
-			{
-				const FName CookedFile = CookedPaths.Value;
-				const FName UncookedFilename = CookedPaths.Key;
-				const FName* FoundPackageName = GetCachedPackageFilenameToPackageFName(UncookedFilename);
-				if ( !FoundPackageName )
-				{
-					++NumPackagesRemoved;
-				}
-				const FName PackageName = *FoundPackageName;
-				const FString UncookedFilenameString = UncookedFilename.ToString();
-				bool bShouldKeep = true;
-
-				if (ModifiedPackages.Contains(PackageName))
-				{
-					++NumPackagesFileHashMismatch;
-					bShouldKeep = false;
-				}
-				else if (NewPackages.Contains(PackageName) || RemovedPackages.Contains(PackageName))
-				{
-					++NumPackagesUnableToFindCookedPackageInfo;
-					bShouldKeep = false;
-				}
-				else if (IdenticalUncookedPackages.Contains(PackageName))
-				{
-					// These are packages which failed to save the first time 
-					// most likely because they are editor only packages
-					bShouldKeep = false;
-				}
-				
-				if (CookedFile == NAME_DummyCookedFilename)
-				{
-					check(IFileManager::Get().FileExists(*CookedFile.ToString()) == false);
-				}
-
-				TArray<FName> PlatformNames;
-				PlatformNames.Add(PlatformFName);
-
-				if (bShouldKeep)
-				{
-					// Add this package to the CookedPackages list so that we don't try cook it again
-					if (CookedFile != NAME_DummyCookedFilename)
-					{
-						check(IFileManager::Get().FileExists(*CookedFile.ToString()));
-					}
-					TArray<bool> Succeeded;
-					Succeeded.Add(true);
-
-					if (IdenticalCookedPackages.Contains(PackageName))
-					{
-						CookedPackages.Add(FFilePlatformCookedPackage(UncookedFilename, MoveTemp(PlatformNames), MoveTemp(Succeeded)));
-						++NumPackagesKept;
-					}
-				}
-				else
-				{
-					if ( IsCookByTheBookMode() ) // cook on the fly will requeue this package when it wants it 
-					{
-						// Force cook the modified file
-						CookRequests.EnqueueUnique(MoveTemp(FFilePlatformRequest(UncookedFilename, PlatformNames)));
-					}
-					if (CookedFile != NAME_DummyCookedFilename)
-					{
-						// delete the old package 
-						const FString CookedFullPath = FPaths::ConvertRelativePathToFull(CookedFile.ToString());
-						UE_LOG(LogCook, Verbose, TEXT("Deleting cooked package %s failed filehash test"), *CookedFullPath);
-						CookedPackagesToDelete.Add(CookedFullPath);
-					}
-					else
-					{
-						// the cooker should rebuild this package because it's not in the cooked package list
-						// the new package will have higher priority then the package in the original shared cooked build
-						UE_LOG(LogCook, Verbose, TEXT("Shared cooked build: Detected package is out of date %s"), *UncookedFilenameString);
-					}
-				}
-			}
-
-			// Register identical uncooked packages from previous run
-			for (FName UncookedPackage : IdenticalUncookedPackages)
-			{
-				const FName UncookedFilename = GetCachedStandardPackageFileFName(UncookedPackage);
-
-				TArray<FName> PlatformNames;
-				PlatformNames.Add(PlatformFName);
-
-				ensure(CookedPackages.Exists(UncookedFilename, PlatformNames, false) == false);
-
-				CookedPackages.Add(MoveTemp(FFilePlatformCookedPackage(UncookedFilename, MoveTemp(PlatformNames))));
-				++NumMarkedFailedSaveKept;
-			}
-
-			UE_LOG(LogCook, Display, TEXT("Iterative cooking summary for %s, \nConsidered: %d, \nFile Hash missmatch: %d, \nPackages Kept: %d, \nPackages failed save kept: %d, \nMissing Cooked Info(expected 0): %d"),
-				*Target->PlatformName(),
-				NumPackagesConsidered, NumPackagesFileHashMismatch,
-				NumPackagesKept, NumMarkedFailedSaveKept,
-				NumPackagesUnableToFindCookedPackageInfo);
-
-			auto DeletePackageLambda = [&CookedPackagesToDelete](int32 PackageIndex)
-			{
-				const FString& CookedFullPath = CookedPackagesToDelete[PackageIndex];
-				IFileManager::Get().Delete(*CookedFullPath, true, true, true);
-			};
-			ParallelFor(CookedPackagesToDelete.Num(), DeletePackageLambda);
-
-		}
-		else if (IsCookFlagSet(ECookInitializationFlags::IterateOnHash))
-		{
-			// load the platform cooked asset registry json
-			const FString CookedAssetRegistry = FPaths::GameDir() / TEXT("CookedAssetRegistry.json");
-			const FString SandboxCookedAssetRegistryFilename = ConvertToFullSandboxPath(*CookedAssetRegistry, true, Target->PlatformName());
-
-			bool bIsIterateSharedBuild = IsCookFlagSet(ECookInitializationFlags::IterateSharedBuild);
-
-
-			TArray<FCookedAssetPackageInfo> CookedPackagesInfo;
-			if ( bIsIterateSharedBuild )
-			{
-				// see if the shared build is newer then the current cooked content in the local directory
-				FDateTime CurrentLocalCookedBuild = IFileManager::Get().GetTimeStamp(*SandboxCookedAssetRegistryFilename);
-
-				// iterate on the shared build if the option is set
-				FString SharedCookedAssetRegistry = FPaths::Combine(*FPaths::GameSavedDir(), TEXT("SharedIterativeBuild"), *Target->PlatformName(), TEXT("Cooked"), TEXT("CookedAssetRegistry.json"));
-
-				FDateTime CurrentIterativeCookedBuild = IFileManager::Get().GetTimeStamp(*SharedCookedAssetRegistry);
-
-				if ( CurrentIterativeCookedBuild > CurrentLocalCookedBuild )
-				{
-					// clean the sandbox 
-					ClearPlatformCookedData(FName(*Target->PlatformName()));
-					FString SandboxDirectory = GetSandboxDirectory(Target->PlatformName());
-					IFileManager::Get().DeleteDirectory(*SandboxDirectory, false, true);
-
-					// SaveCurrentIniSettings(Target); // use this if we don't care about ini safty.
-					// copy the ini settings from the shared cooked build. 
-					const FString SharedCookedIniFile = FPaths::Combine(*FPaths::GameSavedDir(), TEXT("SharedIterativeBuild"), *Target->PlatformName(), TEXT("Cooked"), TEXT("CookedIniVersion.txt"));
-					const FString SandboxCookedIniFile = ConvertToFullSandboxPath(*(FPaths::GameDir() / TEXT("CookedIniVersion.txt")), true).Replace(TEXT("[Platform]"), *Target->PlatformName());
-					
-					IFileManager::Get().Copy(*SandboxCookedIniFile, *SharedCookedIniFile );
-
-					bool bIniSettingsOutOfDate = IniSettingsOutOfDate(Target);
-					if ( bIniSettingsOutOfDate && !IsCookFlagSet(ECookInitializationFlags::IgnoreIniSettingsOutOfDate) )
-					{
-						UE_LOG(LogCook, Display, TEXT("Shared iterative build ini settings out of date, not using shared cooked build"));
-					}
-					else
-					{
-						if ( bIniSettingsOutOfDate )
-						{
-							UE_LOG(LogCook, Display, TEXT("Shared iterative build ini settings out of date, but we don't care"));
-						}
-
-						UE_LOG(LogCook, Display, TEXT("Shared iterative build is newer then local cooked build, iteratively cooking from shared build "));
-						ReadCookedAssetRegistry(SharedCookedAssetRegistry, CookedPackagesInfo);
-					}
-				}
-				else
-				{
-					UE_LOG(LogCook, Display, TEXT("Local cook is newer then shared cooked build, iterativly cooking from local build"));
-				}
-
-			}
-
-			// if shared builds are enabled this will silently fail and no bad things will come of it
-			ReadCookedAssetRegistry(SandboxCookedAssetRegistryFilename, CookedPackagesInfo);
-
-			TMap<FName, FCookedAssetPackageInfo*> CookedPackageInfoLookup;
-			for (FCookedAssetPackageInfo& PackageInfo : CookedPackagesInfo)
-			{
-				CookedPackageInfoLookup.Add(FName(*PackageInfo.RelativeFilename), &PackageInfo);
-			}
-
-
-			// check the files on disk 
-			TMap<FName, FName> UncookedPathToCookedPath;
-			// get all the on disk cooked files
-			GetAllCookedFiles(UncookedPathToCookedPath, SandboxPath);
-			
-			const static FName NAME_DummyCookedFilename(TEXT("DummyCookedFilename")); // pls never name a package dummycookedfilename otherwise shit might go wonky
-			if ( bIsIterateSharedBuild )
-			{
-				// if we are iterating of a shared build the cooked files might not exist in the cooked directory because we assume they are packaged in the pak file (which we don't want to extract)
-				for (const FCookedAssetPackageInfo& CookedPackageInfo : CookedPackagesInfo)
-				{
-					UncookedPathToCookedPath.Add(FName(*CookedPackageInfo.RelativeFilename), NAME_DummyCookedFilename);
-				}
-			}		
-
-			TArray<FName> AllUncookedFNames;
-			{
-				TArray<FString> AllUncookedFilenames;
-				for (const auto& PathPair : UncookedPathToCookedPath)
-				{
-					if (CookedPackageInfoLookup.Contains(PathPair.Key))
-					{
-						AllUncookedFNames.Add(PathPair.Key);
-						AllUncookedFilenames.Add(PathPair.Key.ToString());
-					}
-				}
-
-				if (AllUncookedFilenames.Num() > 0)
-				{
-#if OUTPUT_TIMING
-					SCOPE_TIMER(DetermineAllDependentPackageInfoTime);
-#endif
-					//AllUncookedFilenames.Empty();
-
-					int32 NumThreadsToDeterminePackageInfo = FMath::Max(GLargeThreadPool->GetNumThreads() - 2,1);
-					check(bIgnoreMarkupPackageAlreadyLoaded == false);
-					bIgnoreMarkupPackageAlreadyLoaded = true;
-					PDInfoModule->AsyncDetermineAllDependentPackageInfo(AllUncookedFilenames, NumThreadsToDeterminePackageInfo);
-					bIgnoreMarkupPackageAlreadyLoaded = false;
-				}
-			}
-#if OUTPUT_TIMING
-			SCOPE_TIMER(ProcessingKeptPackageHashes);
-#endif
-			uint32 NumPackagesConsidered = UncookedPathToCookedPath.Num();
-			uint32 NumPackagesUnableToDetermineCurrentDependencyInfo = 0;
-			uint32 NumPackagesUnableToFindCookedPackageInfo = 0;
-			uint32 NumPackagesFileHashMissmatch = 0;
-			uint32 NumPackagesDependentHashMissmatch = 0;
-			uint32 NumPackagesKept = 0;
-
-			for (const auto& CookedPaths : UncookedPathToCookedPath)
-			{
-				const FName CookedFile = CookedPaths.Value;
-				const FName UncookedFilename = CookedPaths.Key;
-				const FString UncookedFilenameString = UncookedFilename.ToString();
-				bool bShouldKeep = true;
-
-
-				FCookedAssetPackageInfo* CookedPackageInfo = CookedPackageInfoLookup.FindRef(UncookedFilename);
-
-				if (bShouldKeep && CookedPackageInfo == nullptr)
-				{
-					++NumPackagesUnableToFindCookedPackageInfo;
-					bShouldKeep = false;
-				}
-
-				if ( bShouldKeep && CookedPackageInfo->SucceededSave == false )
-				{
-					// get rid of this package
-					bShouldKeep = false;
-				}
-				
-				FMD5Hash FullPackageHash;
-				{
-#if OUTPUT_TIMING
-					SCOPE_TIMER(DetermineFullPackageHash);
-#endif
-					if (bShouldKeep && PDInfoModule->DetermineFullPackageHash(*UncookedFilenameString, FullPackageHash) == false)
-					{
-						++NumPackagesUnableToDetermineCurrentDependencyInfo;
-						bShouldKeep = false;
-					}
-				}
-
-				FMD5Hash PackageDependentHash;
-				if (bShouldKeep && PDInfoModule->DeterminePackageDependentHash(*UncookedFilenameString, PackageDependentHash) == false)
-				{
-					UE_LOG(LogCook, Warning, TEXT("Should never happen, above function should have failed and not gone in here, Package %s"), *UncookedFilenameString);
-					bShouldKeep = false;
-				}
-
-				// find the package dependent hash in the previously cooked asset registry
-				if (bShouldKeep && CookedPackageInfo->FullPackageHash != FullPackageHash)
-				{
-					++NumPackagesFileHashMissmatch;
-					bShouldKeep = false;
-				}
-
-				if (bShouldKeep && CookedPackageInfo->DependentPackageHash != PackageDependentHash)
-				{
-					++NumPackagesDependentHashMissmatch;
-
-					if (UE_LOG_ACTIVE(LogCook, Verbose) )
-					// if (IsCookFlagSet(ECookInitializationFlags::LogDebugInfo))
-					{
-						UE_LOG(LogCook, Display, TEXT("Tracing package dependency which caused iterative cook to fail for: %s could be slow Dependent package hashes %s != %s"), *UncookedFilenameString, *Lex::ToString(CookedPackageInfo->DependentPackageHash), *Lex::ToString(PackageDependentHash));
-						TMap<FString, FMD5Hash> CurrentPackageHashes;
-						PDInfoModule->RecursiveGetDependentPackageHashes(*UncookedFilenameString, CurrentPackageHashes);
-						for (const auto& CurrentPackageHash : CurrentPackageHashes)
-						{
-							const FName PackageName(*CurrentPackageHash.Key);
-							FCookedAssetPackageInfo* CurrentCookedPackageInfo = CookedPackageInfoLookup.FindRef(PackageName);
-							if (CurrentCookedPackageInfo )
-							{
-								if ( CurrentCookedPackageInfo->FullPackageHash != CurrentPackageHash.Value )
-								{
-									UE_LOG(LogCook, Display, TEXT("\tinvalidated because full package hash for package: %s doesn't match"), *CurrentPackageHash.Key);
-								}
-							}
-							else
-							{
-								UE_LOG(LogCook, Display, TEXT("\tUnable to find cooked package info for: %s"), *CurrentPackageHash.Key);
-							}
-						}
-
-						TArray<FString> CurrentPackageDependencies;
-						PDInfoModule->GetPackageDependencies(*UncookedFilenameString, CurrentPackageDependencies);
-
-						FString DependencyString = TEXT("CookedPackageInfo:");
-						for (const FString& Dependency : CookedPackageInfo->DependentPackages)
-						{
-							DependencyString += TEXT("\n");
-							DependencyString += Dependency;
-						}
-						DependencyString += TEXT("\nCurrentPackageInfo:");
-						for (const FString& Dependency : CurrentPackageDependencies)
-						{
-							DependencyString += TEXT("\n");
-							DependencyString += Dependency;
-						}
-
-						UE_LOG(LogCook, Display, TEXT("Verifing package dependency count and order %d != %d %s"), CookedPackageInfo->DependentPackages.Num(), CurrentPackageDependencies.Num(), *DependencyString);
-
-						
-						UE_LOG(LogCook, Display, TEXT("Finish trace of: %s"), *UncookedFilenameString);
-					}
-					bShouldKeep = false;
-				}
-
-				if (CookedFile == NAME_DummyCookedFilename)
-				{
-					check(IFileManager::Get().FileExists(*CookedFile.ToString()) == false);
-				}
-
-				if (bShouldKeep)
-				{
-					// Add this package to the CookedPackages list so that we don't try cook it again
-					if ( CookedFile != NAME_DummyCookedFilename )
-					{
-						check( IFileManager::Get().FileExists(*CookedFile.ToString()));
-					}
-					TArray<FName> PlatformNames;
-					PlatformNames.Add(PlatformFName);
-					TArray<bool> Succeeded;
-					Succeeded.Add(true);
-					CookedPackages.Add(FFilePlatformCookedPackage(UncookedFilename, MoveTemp(PlatformNames), MoveTemp(Succeeded)));
-					++NumPackagesKept;
-				}
-				else
-				{
-					if ( IsCookByTheBookMode() )
-					{
-						TArray<FName> PlatformNames;
-						PlatformNames.Add(PlatformFName);
-						CookRequests.EnqueueUnique(MoveTemp(FFilePlatformRequest(UncookedFilename, PlatformNames)));
-					}
-					if ( CookedFile != NAME_DummyCookedFilename )
-					{
-						// delete the old package 
-						const FString CookedFullPath = FPaths::ConvertRelativePathToFull(CookedFile.ToString());
-						UE_LOG(LogCook, Verbose, TEXT("Deleting cooked package %s failed filehash test"), *CookedFullPath);
-						CookedPackagesToDelete.Add(CookedFullPath); 
-					}
-					else
-					{
-						// the cooker should rebuild this package because it's not in the cooked package list
-						// the new package will have higher priority then the package in the original shared cooked build
-						UE_LOG(LogCook, Verbose, TEXT("Shared cooked build: Detected package is out of date %s"), *UncookedFilenameString);
-					}
-				}
-			}
-
-			uint32 NumMarkedEditorOnly = 0;
-			uint32 NumMarkedFailedSave = 0;
-			uint32 NumMarkedFailedSaveKept = 0;
-			uint32 NumMarkedFailedSaveMismatchHash = 0;
-
-			// additionally we also need to populate the editor packages which were not cooked
-			for (const FCookedAssetPackageInfo& CookedPackageInfo : CookedPackagesInfo)
-			{
-				if (CookedPackageInfo.SucceededSave == false)
-				{
-					// if we failed to save then we shouldn't be in the cooked package lists
-
-					const FName StandardName = GetCachedStandardPackageFileFName(FName(*CookedPackageInfo.SourcePackageName));
-
-					++NumMarkedFailedSave;
-
-					TArray<FName> PlatformNames;
-					PlatformNames.Add(PlatformFName);
-
-
-					ensure(CookedPackages.Exists(StandardName, PlatformNames, false) == false);
-
-					bool bShouldKeep = true;
-					FMD5Hash FullPackageHash;
-					if (bShouldKeep && PDInfoModule->DetermineFullPackageHash(*StandardName.ToString(), FullPackageHash) == false)
-					{
-						++NumPackagesUnableToDetermineCurrentDependencyInfo;
-						// UE_LOG(LogCook, Display, TEXT("Unable to find package dependency information for: %s, unable to iterative cook"), *UncookedFilenameString);
-						bShouldKeep = false;
-					}
-
-					// test against the saved hash for the source packages
-					// if a package changes so it can now be saved it may need to be recooked
-					// failed packages only need their package hash processed as their dependencies wouldn't have been saved the first time around
-					if (bShouldKeep && CookedPackageInfo.FullPackageHash != FullPackageHash)
-					{
-						++NumMarkedFailedSaveMismatchHash;
-						UE_LOG(LogCook, Display, TEXT("Full package hash for: %s doesn't match, force recook"), *StandardName.ToString());
-						bShouldKeep = false;
-					}
-
-
-					if ( bShouldKeep )
-					{
-						++NumMarkedFailedSaveKept;
-						CookedPackages.Add(MoveTemp(FFilePlatformCookedPackage(StandardName, MoveTemp(PlatformNames))));
-					}
-					
-				}
-			}
-
-			// additionally we also need to populate the editor packages which were not cooked
-			for (const FCookedAssetPackageInfo& CookedPackageInfo : CookedPackagesInfo)
-			{
-				if (CookedPackageInfo.IsEditorPackage)
-				{
-					++NumMarkedEditorOnly;
-					UncookedEditorOnlyPackages.Add(FName(*CookedPackageInfo.SourcePackageName));
-				}
-			}
-
-			UE_LOG(LogCook, Display, TEXT("Iterative cooking summary for %s, \nConsidered: %d, \nFile Hash missmatch: %d, \nDependency Hash Missmatch: %d, \nPackages marked editor only: %d, \nPackages failed save: %d, \nPackages failed save hash mismatch: %d, \nPackages Kept: %d, \nPackages failed save kept: %d, \nMissing Cooked Info(expected 0): %d, \nFailed Determining Package Info(expected 0): %d"),
-				*Target->PlatformName(),
-				NumPackagesConsidered, NumPackagesFileHashMissmatch, NumPackagesDependentHashMissmatch, 
-				NumMarkedEditorOnly, NumMarkedFailedSave, NumMarkedFailedSaveMismatchHash,
-				NumPackagesKept, NumMarkedFailedSaveKept,
-				NumPackagesUnableToFindCookedPackageInfo, NumPackagesUnableToDetermineCurrentDependencyInfo);
-
-			
-			
-			auto DeletePackageLambda = [&CookedPackagesToDelete](int32 PackageIndex)
-				{ 
-					const FString& CookedFullPath = CookedPackagesToDelete[PackageIndex];
-					IFileManager::Get().Delete(*CookedFullPath, true, true, true);
-				};
-			ParallelFor( CookedPackagesToDelete.Num(), DeletePackageLambda );
-
-			// also need to clean up any files which are not uassets but might be related to uassets
-			
-			IPlatformFile& PlatformFile = FPlatformFileManager::Get().GetPlatformFile();
-			TSet<FString> CookedFilesNoExt;
-			TArray<FString> OtherFiles;
-			FAdditionalPackageSearchVisitor AdditionalPackageSearch(CookedFilesNoExt, OtherFiles);
-			PlatformFile.IterateDirectoryRecursively(*SandboxPath, AdditionalPackageSearch);
-		
-
-			auto DeleteOtherPackagesLambda = [&CookedFilesNoExt, &OtherFiles]( int32 PackageIndex )
-				{
-					const FString& OtherFile = OtherFiles[PackageIndex];
-					const FString NoExtOtherFile = FPaths::SetExtension(OtherFile, "");
-					if ( CookedFilesNoExt.Contains(NoExtOtherFile) == false )
-					{
-						// if we didn't keep the original package then get rid of this one
-						IFileManager::Get().Delete(*OtherFile, true, true, true);
-					}
-				};
-
-			ParallelFor(OtherFiles.Num(), DeleteOtherPackagesLambda );
-
-
-			/*for ( const FString& CookedFullPath : CookedPackagesToDelete )
-			{
-				IFileManager::Get().Delete(*CookedFullPath, true, true, true);
-			}*/
-
-		}
-		else
-		{
-			UE_LOG(LogCook, Display, TEXT("Used time stamps to iteratively cook"));
-			// list of directories to skip
-			TArray<FString> DirectoriesToSkip;
-			TArray<FString> DirectoriesToNotRecurse;
+						ClearPlatformCookedData(FName(*Target->PlatformName()));
 
-				// use the timestamp grabbing visitor
-				IPlatformFile& PlatformFile = FPlatformFileManager::Get().GetPlatformFile();
-				FLocalTimestampDirectoryVisitor Visitor(PlatformFile, DirectoriesToSkip, DirectoriesToNotRecurse, false);
-			
-				PlatformFile.IterateDirectory(*SandboxPath, Visitor);
-
-				for (TMap<FString, FDateTime>::TIterator TimestampIt(Visitor.FileTimes); TimestampIt; ++TimestampIt)
-				{
-					FString CookedFilename = TimestampIt.Key();
-					if ((CookedFilename.EndsWith(TEXT(".uasset")) == false) &&
-						(CookedFilename.EndsWith(TEXT(".umap")) == false))
-					{
-						// don't care if it's not a map or an asset
-						continue;
-					}
-					FDateTime CookedTimestamp = TimestampIt.Value();
-				CookedFilename = FPaths::ConvertRelativePathToFull(CookedFilename);
-
-					FString UncookedFilename = ConvertCookedPathToUncookedPath(CookedFilename);
-					
-				FDateTime DependentTimestamp;
-				if (PDInfoModule->DeterminePackageDependentTimeStamp(*UncookedFilename, DependentTimestamp) == true)
-				{
-					double Diff = (CookedTimestamp - DependentTimestamp).GetTotalSeconds();
-
-					if (Diff >= 0.0)
-					{
-						TArray<FName> PlatformNames;
-						PlatformNames.Add(PlatformFName);
-						TArray<bool> Succeeded;
-						Succeeded.Add(true);
-						CookedPackages.Add(FFilePlatformCookedPackage(FName(*UncookedFilename), MoveTemp(PlatformNames), MoveTemp(Succeeded)));
-					}
-				}
-			}
-		}
-	}
-}
-
-const FString ExtractPackageNameFromObjectPath( const FString ObjectPath )
-{
-	// get the path 
-	int32 Beginning = ObjectPath.Find(TEXT("'"), ESearchCase::CaseSensitive);
-	if ( Beginning == INDEX_NONE )
-	{
-		return ObjectPath;
-	}
-	int32 End = ObjectPath.Find(TEXT("."), ESearchCase::CaseSensitive, ESearchDir::FromStart, Beginning + 1);
-	if (End == INDEX_NONE )
-	{
-		End = ObjectPath.Find(TEXT("'"), ESearchCase::CaseSensitive, ESearchDir::FromStart, Beginning + 1);
-	}
-	if ( End == INDEX_NONE )
-	{
-		// one more use case is that the path is "Class'Path" example "OrionBoostItemDefinition'/Game/Misc/Boosts/XP_1Win" dunno why but this is actually dumb
-		if ( ObjectPath[Beginning+1] == '/' )
-		{
-			return ObjectPath.Mid(Beginning+1);
-		}
-		return ObjectPath;
-	}
-	return ObjectPath.Mid(Beginning + 1, End - Beginning - 1);
-}
-
-void UCookOnTheFlyServer::VerifyCookedPackagesAreUptodate(const TArray<ITargetPlatform*>& Platforms)
-{
-	// we remove cooked packages which are updated using the "MarkObjectDirty" callbacks from the editor
-	// here we process the dependencies by removing any packages which awere dependent on a package which was removed
-
-	if (IsCookFlagSet(ECookInitializationFlags::IterateSharedBuild) || IsCookFlagSet(ECookInitializationFlags::IterateOnAssetRegistry))
-	{
-		// can't verify cooked packages in shared builds, for asset registry it got verified during the gather step
-		check( IsCookOnTheFlyMode() == false );
-		return;
-	}
-
-	uint32 CookedPackagesMissingOnDisk = 0;
-
-	// loop through all packages on disk and remove any packages from the cooked packages array which don't exist
-	// this shouldn't ever happen because the cooker was not shutdown since it cooked these packages so unless someone manually went in and deleted the cooked packages this won't happen
-	for (const ITargetPlatform* TargetPlatform : Platforms)
-	{
-		FName PlatformFName(*TargetPlatform->PlatformName());
-		TArray<FName> PlatformNamesArray;
-					
-		TArray<FName> CookedFiles;
-		CookedPackages.GetCookedFilesForPlatform(PlatformFName, CookedFiles, false);
-
-		for (const FName& CookedFile : CookedFiles)
-		{
-			const FString SandboxPath = SandboxFile->ConvertToSandboxPath(*CookedFile.ToString()).Replace(TEXT("[Platform]"), *TargetPlatform->PlatformName());
-
-			if (!IFileManager::Get().FileExists(*SandboxPath))
-			{
-				++CookedPackagesMissingOnDisk;
-				// removing cooked file because it doesn't exist
-				if (IsCookFlagSet(ECookInitializationFlags::LogDebugInfo))
-				{
-					UE_LOG(LogCook, Display, TEXT("Cooked file invalidated because it doesn't exist %s"), *SandboxPath);
-				}
-				if ( CookedPackages.RemoveFileForPlatform(CookedFile, PlatformFName) )
-				{
-					if (IsCookByTheBookMode())
-					{
-						CookRequests.EnqueueUnique(FFilePlatformRequest(CookedFile, PlatformFName));
-					}
-				}
-			}
-		}
-	}
-
-	// loop through all cooked packages and remove any which are missing their dependencies
-	uint32 CookedPackagesConsidered = 0;
-	uint32 CookedPackagesFailedDependencyCheck = 0;
-	uint32 FoundRedirectors = 0;
-	uint32 CookedPackagesKept = 0;
-	
-	for (const ITargetPlatform* TargetPlatform : Platforms)
-	{
-		FName PlatformFName = FName(*TargetPlatform->PlatformName());
-
-
-		TArray<FName> PlatformNamesArray;
-		PlatformNamesArray.Add(PlatformFName);
-
-		TArray<FName> CookedFiles;
-		CookedPackages.GetCookedFilesForPlatform(PlatformFName, CookedFiles, false);
-		for (const FName& CookedFile : CookedFiles)
-		{
-			bool bShouldRemovePackage = false;
-
-			const FName* PackageNamePtr = GetCachedPackageFilenameToPackageFName(CookedFile);
-			check(PackageNamePtr != nullptr);
-			const FName& PackageName = *PackageNamePtr;
-
-			TArray<FName> Dependencies;
-			Dependencies.Add(PackageName);
-
-			++CookedPackagesConsidered;
-
-			for (int32 I = 0; I < Dependencies.Num(); ++I)
-			{
-				FName DependencyPackageName = Dependencies[I];
-
-				if (FPackageName::IsScriptPackage(DependencyPackageName.ToString()) || FPackageName::IsMemoryPackage(DependencyPackageName.ToString()))
-				{
-					// don't care about script packages we are assuming the scripts can't be invalidated after the cooker has started
-					continue;
-				}
-
-				if (!FPackageName::IsValidLongPackageName(*DependencyPackageName.ToString()))
-				{
-					// this is probably a soft reference to a UObject rather then a package anem.
-					const FString FinalPackageNameString = ExtractPackageNameFromObjectPath(DependencyPackageName.ToString());
-					if (FPackageName::IsValidLongPackageName(FinalPackageNameString))
-					{
-						DependencyPackageName = FName(*FinalPackageNameString);
-					}
-					else
-					{
-						UE_LOG(LogCook, Warning, TEXT("Unable resolve dependency package name %s, not able to iterative cook %s"), *DependencyPackageName.ToString(), *CookedFile.ToString());
-						bShouldRemovePackage = true;
-						break;
-					}
-				}
-
-				if (UncookedEditorOnlyPackages.Contains(DependencyPackageName))
-				{
-					// don't need to process this package, it was editor only so unless some other package comes along and needs it, it won't need to be resaved (and also shouldn't invalidate other packages)
-					continue;
-				}
-
-				TMap<FString, FString> RedirectedPaths;
-				if (ContainsRedirector(DependencyPackageName, RedirectedPaths))
-				{
-					FString DestinationObjectPath;
-					for (TPair<FString, FString>& RedirectedPath : RedirectedPaths)
-					{
-						DestinationObjectPath = RedirectedPath.Value;
-						break;
-					}
-
-					FString RedirectPackageName = ExtractPackageNameFromObjectPath(DestinationObjectPath);
-					if (!FPackageName::IsValidLongPackageName(RedirectPackageName))
-					{
-						UE_LOG(LogCook, Warning, TEXT("Unable to resolve redirector object path from string %s, potentially malformed"), *DestinationObjectPath);
-						continue;
-					}
-					Dependencies.AddUnique(FName(*RedirectPackageName));
-					++FoundRedirectors;
-					// don't process this dependency anymore 
-					continue;	
-				}
-
-				const FName StandardDependencyName = GetCachedStandardPackageFileFName(DependencyPackageName);
-				if (StandardDependencyName == NAME_None)
-				{
-					UE_LOG(LogCook, Warning, TEXT("Unable to resolve package name %s not able to iterative cook"), *DependencyPackageName.ToString());
-					bShouldRemovePackage = true;
-					break;
-				}
-
-				// check this package is cooked for the platform
-				if (!CookedPackages.Exists(StandardDependencyName, PlatformNamesArray))
-				{
-					//if (IsCookFlagSet(ECookInitializationFlags::LogDebugInfo))
-					{
-						UE_LOG(LogCook, Verbose, TEXT("Invalidating cooked package %s because %s was considered out of date"), *CookedFile.ToString(), *StandardDependencyName.ToString());
-					}
-					++CookedPackagesFailedDependencyCheck;
-					bShouldRemovePackage = true;
-					break;
-				}
-
-				if (CookedPackages.Exists(StandardDependencyName, PlatformNamesArray, false))
-				{
-					TArray<FName> NewDependencies;
-					AssetRegistry->GetDependencies(DependencyPackageName, NewDependencies, EAssetRegistryDependencyType::All);
-					for (const FName& NewDependency : NewDependencies)
-					{
-						Dependencies.AddUnique(NewDependency);
-					}
-				}
-			}
+						FString SandboxDirectory = GetSandboxDirectory(Target->PlatformName());
+						IFileManager::Get().DeleteDirectory(*SandboxDirectory, false, true);
 
-			if (bShouldRemovePackage)
-			{
-				UE_LOG(LogCook, Verbose, TEXT("Removing package from cooked package list %s"), *CookedFile.ToString());
-				if ( CookedPackages.RemoveFileForPlatform(CookedFile, PlatformFName) )
-				{
-					if ( IsCookByTheBookMode() )
-					{
-						CookRequests.EnqueueUnique(FFilePlatformRequest(CookedFile, PlatformFName));
+						SaveCurrentIniSettings(Target);
 					}
-				}
-			}
-			else
-			{
-				++CookedPackagesKept;
-			}
-		}
-	}
-
-
-	uint32 CookedPackagesDeleted = 0; // this should be the same as the cooked packages removed above but in cases where people tamper with the cooked directory this will be different
-
-	// delete any packages on disk which don't exist in the CookedPackages array 
-	for (const ITargetPlatform* TargetPlatform : Platforms)
-	{
-		FString SandboxPath = GetSandboxDirectory(TargetPlatform->PlatformName());
-		FName PlatformFName(*TargetPlatform->PlatformName());
-		TArray<FName> PlatformNamesArray;
-		PlatformNamesArray.Add(PlatformFName);
-
-
-		IPlatformFile& PlatformFile = FPlatformFileManager::Get().GetPlatformFile();
-		TArray<FString> CookedFiles;
-		FPackageSearchVisitor PackageSearch(CookedFiles);
-		PlatformFile.IterateDirectoryRecursively(*SandboxPath, PackageSearch);
-
-		for (const FString& CookedFile : CookedFiles)
-		{
-			const FString UncookedFilename = ConvertCookedPathToUncookedPath(CookedFile);
-			const FName UncookedFName(*UncookedFilename);
-			if (!CookedPackages.Exists(UncookedFName, PlatformNamesArray))
-					{
-				IFileManager::Get().Delete(*CookedFile, true, true, true);
-				if (IsCookFlagSet(ECookInitializationFlags::LogDebugInfo))
-				{
-					UE_LOG(LogCook, Display, TEXT("Deleted cooked package %s because it wasn't in the cooked package list"), *CookedFile);
-				}
-			}
-		}
-	}
-
-	// CookedPackagesConsidered are packages which passed the missing test and had their dependencies tested against
-	// FoundRedirectors are packages which were referencing redirectors
-	// CookedPackagesKept are packages which were kept after their dependencies were checked that they still exist
-	// CookedPackagesMissingOnDisk are packages which were in the CookedPackages list but weren't actually on disk in the "Cooked\" directory
-	// CookedPackagesDeleted are packages which were deleted as a result of them not being kept or being in the cooked directory and not on the CookedPackages list 
-	//		this can happen when something removed them from the cooked packages list either because it was modified in memory or package information couldn't be found about this package
-	UE_LOG(LogCook, Display, TEXT("Verify cooked packages: \nPackages considered for dependency checks: %d, \nPackages failed dependency checks %d, \nFound redirectors: %d, \nPackages Kept: %d, \nPackages Missing: %d, \nPackages Deleted: %d"),
-		CookedPackagesConsidered, 
-		CookedPackagesFailedDependencyCheck,
-		FoundRedirectors, 
-		CookedPackagesKept, 
-		CookedPackagesMissingOnDisk,
-		CookedPackagesDeleted );
-
-}
-
-void UCookOnTheFlyServer::CleanSandbox(const bool bIterative)
-{
-	//child cookers shouldn't clean sandbox we would be deleting the master cookers / other cookers data
-	check(IsChildCooker() == false);
-
-	const TArray<ITargetPlatform*>& Platforms = GetCookingTargetPlatforms();
-
-	// before we can delete any cooked files we need to make sure that we have finished writing them
-	UPackage::WaitForAsyncFileWrites();
-
-#if OUTPUT_TIMING
-	double SandboxCleanTime = 0.0;
-#endif
-	{
-#if OUTPUT_TIMING
-		SCOPE_SECONDS_COUNTER(SandboxCleanTime);
-		SCOPE_TIMER(CleanSandboxTime);
-#endif
-		if (bIterative == false)
-		{
-			// for now we are going to wipe the cooked directory
-			for (int32 Index = 0; Index < Platforms.Num(); Index++)
-			{
-				ITargetPlatform* Target = Platforms[Index];
-				UE_LOG(LogCook, Display, TEXT("Cooked content cleared for platform %s"), *Target->PlatformName());
-
-				FString SandboxDirectory = GetSandboxDirectory(Target->PlatformName()); // GetOutputDirectory(Target->PlatformName());
-				IFileManager::Get().DeleteDirectory(*SandboxDirectory, false, true);
-
-				ClearPlatformCookedData(FName(*Target->PlatformName()));
-
-				IniSettingsOutOfDate(Target);
-				SaveCurrentIniSettings(Target);
-			}
-
-		}
-		else
-		{
-			for (const ITargetPlatform* Target : Platforms)
-			{
-				bool bIniSettingsOutOfDate = IniSettingsOutOfDate(Target);
-				if (bIniSettingsOutOfDate)
-				{
-					if ( !IsCookFlagSet(ECookInitializationFlags::IgnoreIniSettingsOutOfDate) )
-				{
-					UE_LOG(LogCook, Display, TEXT("Cook invalidated for platform %s ini settings don't match from last cook, clearing all cooked content"), *Target->PlatformName());
-
-					ClearPlatformCookedData(FName(*Target->PlatformName()));
-
-					FString SandboxDirectory = GetSandboxDirectory(Target->PlatformName());
-					IFileManager::Get().DeleteDirectory(*SandboxDirectory, false, true);
-
-					SaveCurrentIniSettings(Target);
-				}
 					else
 					{
 						UE_LOG(LogCook, Display, TEXT("Inisettings were out of date for platform %s but we are going with it anyway because IgnoreIniSettingsOutOfDate is set"), *Target->PlatformName());
@@ -5835,33 +4661,9 @@ void UCookOnTheFlyServer::CleanSandbox(const bool bIterative)
 				}
 			}
 
-			if (IsCookFlagSet(ECookInitializationFlags::IterateOnAssetRegistry))
-			{
-				// This is fast in asset registry iterate, so just reconstruct
-				CookedPackages.Empty();
-				PopulateCookedPackagesFromDisk(Platforms);
-			}
-			else
-			{
-				if (bIsInitializingSandbox)
-				{
-					// if we are initializing sandbox that means that the sandbox has been recreated and could be pointing at something different
-					// need to hit the disk and populate our cooked packages form there
-					PopulateCookedPackagesFromDisk(Platforms);
-
-					VerifyCookedPackagesAreUptodate(Platforms);
-				}
-				else
-				{
-					// do a fast check and see if anything has changed since last cook
-					// we know that the sandbox is the same so the cookedpackagelist will be populated with what's on disk
-					// note this can only happen in editor
-					VerifyCookedPackagesAreUptodate(Platforms);
-				}
-
-				// Collect garbage to ensure we don't have any packages hanging around from dependent time stamp determination
-				CollectGarbage(GARBAGE_COLLECTION_KEEPFLAGS);
-			}
+			// This is fast in asset registry iterate, so just reconstruct
+			CookedPackages.Empty();
+			PopulateCookedPackagesFromDisk(Platforms);
 		}
 	}
 #if OUTPUT_TIMING
@@ -5918,12 +4720,13 @@ void UCookOnTheFlyServer::GenerateAssetRegistry()
 		if ( !bCanDelayAssetregistryProcessing)
 		{
 			TArray<FString> ScanPaths;
-			if (GConfig->GetArray(TEXT("AssetRegistry"), TEXT("PathsToScanForCook"), ScanPaths, GEngineIni) > 0)
+			if (GConfig->GetArray(TEXT("AssetRegistry"), TEXT("PathsToScanForCook"), ScanPaths, GEngineIni) > 0 && !AssetRegistry->IsLoadingAssets())
 			{
 				AssetRegistry->ScanPathsSynchronous(ScanPaths);
 			}
 			else
 			{
+				// This will flush the background gather if we're in the editor
 				AssetRegistry->SearchAllAssets(true);
 			}
 		}
@@ -6009,6 +4812,8 @@ void UCookOnTheFlyServer::CollectFilesToCook(TArray<FName>& FilesInPath, const T
 	bool bMapsOnly = (!!(FilesToCookFlags & ECookByTheBookOptions::MapsOnly)) || PackagingSettings->bCookMapsOnly;
 	bool bNoDev = !!(FilesToCookFlags & ECookByTheBookOptions::NoDevContent);
 
+	TArray<FName> InitialPackages = FilesInPath;
+
 	if (IsChildCooker())
 	{
 		const FString ChildCookFilename = CookByTheBookOptions->ChildCookFilename;
@@ -6200,6 +5005,43 @@ void UCookOnTheFlyServer::CollectFilesToCook(TArray<FName>& FilesInPath, const T
 			}
 		}
 
+		// If no packages were explicitly added by command line or game callback, add all maps
+		if (FilesInPath.Num() == InitialPackages.Num() || bCookAll)
+		{
+			TArray<FString> Tokens;
+			Tokens.Empty(2);
+			Tokens.Add(FString("*") + FPackageName::GetAssetPackageExtension());
+			Tokens.Add(FString("*") + FPackageName::GetMapPackageExtension());
+
+			uint8 PackageFilter = NORMALIZE_DefaultFlags | NORMALIZE_ExcludeEnginePackages;
+			if (bMapsOnly)
+			{
+				PackageFilter |= NORMALIZE_ExcludeContentPackages;
+			}
+
+			if (bNoDev)
+			{
+				PackageFilter |= NORMALIZE_ExcludeDeveloperPackages;
+			}
+
+			// assume the first token is the map wildcard/pathname
+			TArray<FString> Unused;
+			for (int32 TokenIndex = 0; TokenIndex < Tokens.Num(); TokenIndex++)
+			{
+				TArray<FString> TokenFiles;
+				if (!NormalizePackageNames(Unused, TokenFiles, Tokens[TokenIndex], PackageFilter))
+				{
+					UE_LOG(LogCook, Display, TEXT("No packages found for parameter %i: '%s'"), TokenIndex, *Tokens[TokenIndex]);
+					continue;
+				}
+
+				for (int32 TokenFileIndex = 0; TokenFileIndex < TokenFiles.Num(); ++TokenFileIndex)
+				{
+					AddFileToCook(FilesInPath, TokenFiles[TokenFileIndex]);
+				}
+			}
+		}
+
 		// Add any files of the desired cultures localized assets to cook.
 		for (const FString& CultureToCookName : CookCultures)
 		{
@@ -6303,43 +5145,6 @@ void UCookOnTheFlyServer::CollectFilesToCook(TArray<FName>& FilesInPath, const T
 		}
 	}
 
-	if ((FilesInPath.Num() == 0) || bCookAll)
-	{
-		TArray<FString> Tokens;
-		Tokens.Empty(2);
-		Tokens.Add(FString("*") + FPackageName::GetAssetPackageExtension());
-		Tokens.Add(FString("*") + FPackageName::GetMapPackageExtension());
-
-		uint8 PackageFilter = NORMALIZE_DefaultFlags | NORMALIZE_ExcludeEnginePackages;
-		if ( bMapsOnly )
-		{
-			PackageFilter |= NORMALIZE_ExcludeContentPackages;
-		}
-
-		if ( bNoDev )
-		{
-			PackageFilter |= NORMALIZE_ExcludeDeveloperPackages;
-		}
-
-		// assume the first token is the map wildcard/pathname
-		TArray<FString> Unused;
-		for ( int32 TokenIndex = 0; TokenIndex < Tokens.Num(); TokenIndex++ )
-		{
-			TArray<FString> TokenFiles;
-			if ( !NormalizePackageNames( Unused, TokenFiles, Tokens[TokenIndex], PackageFilter) )
-			{
-				UE_LOG(LogCook, Display, TEXT("No packages found for parameter %i: '%s'"), TokenIndex, *Tokens[TokenIndex]);
-				continue;
-			}
-
-			for (int32 TokenFileIndex = 0; TokenFileIndex < TokenFiles.Num(); ++TokenFileIndex)
-			{
-				AddFileToCook( FilesInPath, TokenFiles[TokenFileIndex]);
-			}
-		}
-	}
-
-
 	if (!(FilesToCookFlags & ECookByTheBookOptions::NoInputPackages))
 	{
 		// make sure we cook any extra assets for the default touch interface
@@ -6612,10 +5417,6 @@ void UCookOnTheFlyServer::CookByTheBookFinished()
 		
 		// Save modified asset registry with all streaming chunk info generated during cook
 		const FString& SandboxRegistryFilename = GetSandboxAssetRegistryFilename();
-		// the registry filename will be modified when we call save asset registry
-		
-		const FString CookedAssetRegistry = FPaths::GameDir() / TEXT("CookedAssetRegistry.json");
-		const FString SandboxCookedAssetRegistryFilename = ConvertToFullSandboxPath(*CookedAssetRegistry, true);
 
 		ITargetPlatformManagerModule& TPM = GetTargetPlatformManagerRef();
 		if (bCacheShaderLibraries && PackagingSettings->bShareMaterialShaderCode)
@@ -6718,27 +5519,6 @@ void UCookOnTheFlyServer::CookByTheBookFinished()
 					{
 						Generator.WriteCookerOpenOrder();
 					}
-
-				}
-				// If asset registry iteration is enabled, but this is a full build, don't save the old json as it is slow and won't be used
-				if (!IsCookFlagSet(ECookInitializationFlags::IterateOnAssetRegistry) && !CookerSettings->bUseAssetRegistryForIteration)
-				{
-					SCOPE_TIMER(SaveCookedPackageAssetRegistry);
-
-					TArray<FName> AllCookedPackages;
-					TArray<FName> UncookedEditorOnlyPackageNamesArray;
-					{
-						SCOPE_TIMER(SortingCookedPackages);
-						AllCookedPackages.Append(CookedPackagesFilenames);
-					
-						for (FName UncookedEditorOnlyPackage : UncookedEditorOnlyPackageNames)
-						{
-							UncookedEditorOnlyPackageNamesArray.Add(UncookedEditorOnlyPackage);
-						}
-						AllCookedPackages.Sort();
-						UncookedEditorOnlyPackageNamesArray.Sort();
-					}
-					SaveCookedAssetRegistry(SandboxCookedAssetRegistryFilename, AllCookedPackages, UncookedEditorOnlyPackageNamesArray, IgnorePackageNames.Array(), PlatformName.ToString(), false);
 				}
 				if (IsCreatingReleaseVersion())
 				{
diff --git a/Engine/Source/Editor/UnrealEd/Private/CookerSettings.cpp b/Engine/Source/Editor/UnrealEd/Private/CookerSettings.cpp
--- a/Engine/Source/Editor/UnrealEd/Private/CookerSettings.cpp
+++ b/Engine/Source/Editor/UnrealEd/Private/CookerSettings.cpp
@@ -9,7 +9,8 @@ UCookerSettings::UCookerSettings(const FObjectInitializer& ObjectInitializer)
 	, bEnableCookOnTheSide(false)
 	, bEnableBuildDDCInBackground(false)
 	, bIterativeCookingForLaunchOn(false)
-	, bUseAssetRegistryForIteration(false)
+	, bIgnoreIniSettingsOutOfDateForIteration(false)
+	, bIgnoreScriptPackagesOutOfDateForIteration(false)
 	, bCompileBlueprintsInDevelopmentMode(true)
 	, bCookBlueprintComponentTemplateData(false)
 {
diff --git a/Engine/Source/Editor/UnrealEd/Private/EdGraphToken.cpp b/Engine/Source/Editor/UnrealEd/Private/EdGraphToken.cpp
--- a/Engine/Source/Editor/UnrealEd/Private/EdGraphToken.cpp
+++ b/Engine/Source/Editor/UnrealEd/Private/EdGraphToken.cpp
@@ -40,6 +40,10 @@ FEdGraphToken::FEdGraphToken(const UObject* InObject, const UEdGraphPin* InPin)
 	if (InPin)
 	{
 		CachedText = InPin->GetDisplayName();
+		if (CachedText.IsEmpty())
+		{
+			CachedText = NSLOCTEXT("MessageLog", "UnnamedPin", "<Unnamed>");
+		}
 	}
 	else if (InObject)
 	{
