Copying //UE4/Dev-Framework to //UE4/Dev-Main (Source: //UE4/Dev-Framework @ 3510040)

diff --git a/Engine/Source/Editor/UnrealEd/Private/EditorCategoryUtils.cpp b/Engine/Source/Editor/UnrealEd/Private/EditorCategoryUtils.cpp
--- a/Engine/Source/Editor/UnrealEd/Private/EditorCategoryUtils.cpp
+++ b/Engine/Source/Editor/UnrealEd/Private/EditorCategoryUtils.cpp
@@ -285,7 +285,7 @@ FString FEditorCategoryUtils::GetCategoryDisplayString(const FString& Unsanitize
 }
 
 //------------------------------------------------------------------------------
-void FEditorCategoryUtils::GetClassHideCategories(const UClass* Class, TArray<FString>& CategoriesOut, bool bHomogenize)
+void FEditorCategoryUtils::GetClassHideCategories(const UStruct* Class, TArray<FString>& CategoriesOut, bool bHomogenize)
 {
 	CategoriesOut.Empty();
 
@@ -307,7 +307,7 @@ void FEditorCategoryUtils::GetClassHideCategories(const UClass* Class, TArray<FS
 }
 
 //------------------------------------------------------------------------------
-void  FEditorCategoryUtils::GetClassShowCategories(const UClass* Class, TArray<FString>& CategoriesOut)
+void  FEditorCategoryUtils::GetClassShowCategories(const UStruct* Class, TArray<FString>& CategoriesOut)
 {
 	CategoriesOut.Empty();
 
@@ -325,19 +325,19 @@ void  FEditorCategoryUtils::GetClassShowCategories(const UClass* Class, TArray<F
 }
 
 //------------------------------------------------------------------------------
-bool FEditorCategoryUtils::IsCategoryHiddenFromClass(const UClass* Class, FCommonEditorCategory::EValue CategoryId)
+bool FEditorCategoryUtils::IsCategoryHiddenFromClass(const UStruct* Class, FCommonEditorCategory::EValue CategoryId)
 {
 	return IsCategoryHiddenFromClass(Class, GetCommonCategory(CategoryId));
 }
 
 //------------------------------------------------------------------------------
-bool FEditorCategoryUtils::IsCategoryHiddenFromClass(const UClass* Class, const FText& Category)
+bool FEditorCategoryUtils::IsCategoryHiddenFromClass(const UStruct* Class, const FText& Category)
 {
 	return IsCategoryHiddenFromClass(Class, Category.ToString());
 }
 
 //------------------------------------------------------------------------------
-bool FEditorCategoryUtils::IsCategoryHiddenFromClass(const UClass* Class, const FString& Category)
+bool FEditorCategoryUtils::IsCategoryHiddenFromClass(const UStruct* Class, const FString& Category)
 {
 	TArray<FString> ClassHideCategories;
 	GetClassHideCategories(Class, ClassHideCategories);
@@ -345,7 +345,7 @@ bool FEditorCategoryUtils::IsCategoryHiddenFromClass(const UClass* Class, const
 }
 
 //------------------------------------------------------------------------------
-bool FEditorCategoryUtils::IsCategoryHiddenFromClass(const TArray<FString>& ClassHideCategories, const UClass* Class, const FString& Category)
+bool FEditorCategoryUtils::IsCategoryHiddenFromClass(const TArray<FString>& ClassHideCategories, const UStruct* Class, const FString& Category)
 {
 	bool bIsHidden = false;
 
@@ -410,7 +410,7 @@ void FEditorCategoryUtils::GetCategoryTooltipInfo(const FString& Category, FText
 }
 
 //------------------------------------------------------------------------------
-TSet<FString> FEditorCategoryUtils::GetHiddenCategories(const UClass* Class)
+TSet<FString> FEditorCategoryUtils::GetHiddenCategories(const UStruct* Class)
 {
 	TArray<FString> ClassHiddenCategories;
 	GetClassHideCategories(Class, ClassHiddenCategories);
diff --git a/Engine/Source/Editor/UnrealEd/Private/EditorEngine.cpp b/Engine/Source/Editor/UnrealEd/Private/EditorEngine.cpp
--- a/Engine/Source/Editor/UnrealEd/Private/EditorEngine.cpp
+++ b/Engine/Source/Editor/UnrealEd/Private/EditorEngine.cpp
@@ -6794,11 +6794,6 @@ FWorldContext* UEditorEngine::GetPIEWorldContext()
 	return nullptr;
 }
 
-bool UEditorEngine::IsUsingWorldAssets()
-{
-	return !FParse::Param(FCommandLine::Get(), TEXT("DisableWorldAssets"));
-}
-
 void UEditorEngine::OnAssetLoaded(UObject* Asset)
 {
 	UWorld* World = Cast<UWorld>(Asset);
diff --git a/Engine/Source/Editor/UnrealEd/Private/EditorLevelUtils.cpp b/Engine/Source/Editor/UnrealEd/Private/EditorLevelUtils.cpp
--- a/Engine/Source/Editor/UnrealEd/Private/EditorLevelUtils.cpp
+++ b/Engine/Source/Editor/UnrealEd/Private/EditorLevelUtils.cpp
@@ -1,7 +1,7 @@
 // Copyright 1998-2017 Epic Games, Inc. All Rights Reserved.
 
 /*=============================================================================
-	EditorLevelUtils.cpp: Editor-specific level management routines
+EditorLevelUtils.cpp: Editor-specific level management routines
 =============================================================================*/
 
 #include "EditorLevelUtils.h"
@@ -45,6 +45,7 @@
 #include "AssetRegistryModule.h"
 #include "Engine/LevelStreamingVolume.h"
 #include "Components/ModelComponent.h"
+#include "Misc/RuntimeErrors.h"
 
 DEFINE_LOG_CATEGORY(LogLevelTools);
 
@@ -53,7 +54,7 @@ DEFINE_LOG_CATEGORY(LogLevelTools);
 
 int32 UEditorLevelUtils::MoveActorsToLevel(const TArray<AActor*>& ActorsToMove, ULevelStreaming* DestStreamingLevel)
 {
-	return MoveActorsToLevel(ActorsToMove, DestStreamingLevel->GetLoadedLevel());
+	return MoveActorsToLevel(ActorsToMove, DestStreamingLevel ? DestStreamingLevel->GetLoadedLevel() : nullptr);
 }
 
 int32 UEditorLevelUtils::MoveActorsToLevel(const TArray<AActor*>& ActorsToMove, ULevel* DestLevel)
@@ -69,7 +70,6 @@ int32 UEditorLevelUtils::MoveActorsToLevel(const TArray<AActor*>& ActorsToMove,
 		FString OriginalClipboardContent;
 		FPlatformMisc::ClipboardPaste(OriginalClipboardContent);
 
-
 		// The final list of actors to move after invalid actors were removed
 		TArray<AActor*> FinalMoveList;
 		FinalMoveList.Reserve(ActorsToMove.Num());
@@ -85,9 +85,9 @@ int32 UEditorLevelUtils::MoveActorsToLevel(const TArray<AActor*>& ActorsToMove,
 				}
 
 				bool bIsSourceLevelLocked = FLevelUtils::IsLevelLocked(CurActor);
-				
+
 				if (!bIsSourceLevelLocked)
-		{
+				{
 					if (CurActor->GetLevel() != DestLevel)
 					{
 						FinalMoveList.Add(CurActor);
@@ -110,7 +110,7 @@ int32 UEditorLevelUtils::MoveActorsToLevel(const TArray<AActor*>& ActorsToMove,
 
 
 		if (FinalMoveList.Num() > 0)
-			{
+		{
 			GEditor->SelectNone(false, true, false);
 
 			USelection* ActorSelection = GEditor->GetSelectedActors();
@@ -136,16 +136,16 @@ int32 UEditorLevelUtils::MoveActorsToLevel(const TArray<AActor*>& ActorsToMove,
 				OwningWorld->SetCurrentLevel(DestLevel);
 				const bool bLevelVisible = DestLevel->bIsVisible;
 				if (!bLevelVisible)
-			{
+				{
 					UEditorLevelUtils::SetLevelVisibility(DestLevel, true, false);
-		}
-		
+				}
+
 				// Paste the actors into the new level
 				GEditor->edactPasteSelected(OwningWorld, false, false, false);
 
 				// Restore new level visibility to previous state
 				if (!bLevelVisible)
-		{
+				{
 					UEditorLevelUtils::SetLevelVisibility(DestLevel, false, false);
 				}
 
@@ -156,305 +156,318 @@ int32 UEditorLevelUtils::MoveActorsToLevel(const TArray<AActor*>& ActorsToMove,
 			// The moved (pasted) actors will now be selected
 			NumMovedActors += FinalMoveList.Num();
 		}
-		
+
 		// Restore the original clipboard contents
 		FPlatformMisc::ClipboardCopy(*OriginalClipboardContent);
 	}
 
-
 	return NumMovedActors;
 }
 
 int32 UEditorLevelUtils::MoveSelectedActorsToLevel(ULevelStreaming* DestStreamingLevel)
-	{
-	return MoveSelectedActorsToLevel(DestStreamingLevel->GetLoadedLevel());
+{
+	ensureAsRuntimeWarning(DestStreamingLevel != nullptr);
+	return DestStreamingLevel ? MoveSelectedActorsToLevel(DestStreamingLevel->GetLoadedLevel()) : 0;
 }
 
 int32 UEditorLevelUtils::MoveSelectedActorsToLevel(ULevel* DestLevel)
+{
+	if (ensureAsRuntimeWarning(DestLevel != nullptr))
+	{
+		TArray<AActor*> ActorsToMove;
+		for (FSelectionIterator It(GEditor->GetSelectedActorIterator()); It; ++It)
 		{
-	TArray<AActor*> ActorsToMove;
-	for (FSelectionIterator It(GEditor->GetSelectedActorIterator()); It; ++It)
+			if (AActor* Actor = Cast<AActor>(*It))
 			{
-		if (AActor* Actor = Cast<AActor>(*It))
-				{
-			ActorsToMove.Add(Actor);
+				ActorsToMove.Add(Actor);
 			}
 		}
 
-	return MoveActorsToLevel(ActorsToMove, DestLevel);
-	
+		return MoveActorsToLevel(ActorsToMove, DestLevel);
 	}
 
+	return 0;
+}
+
 ULevel* UEditorLevelUtils::AddLevelsToWorld(UWorld* InWorld, const TArray<FString>& LevelPackageNames, UClass* LevelStreamingClass)
+{
+	if (!ensure(InWorld))
 	{
-		if ( !ensure(InWorld) )
-		{
-			return nullptr;
-		}
+		return nullptr;
+	}
 
-		FScopedSlowTask SlowTask(LevelPackageNames.Num(), LOCTEXT("AddLevelsToWorldTask", "Adding Levels to World"));
-		SlowTask.MakeDialog();
+	FScopedSlowTask SlowTask(LevelPackageNames.Num(), LOCTEXT("AddLevelsToWorldTask", "Adding Levels to World"));
+	SlowTask.MakeDialog();
 
-		TArray<FString> PackageNames = LevelPackageNames;
+	TArray<FString> PackageNames = LevelPackageNames;
 
-		// Sort the level packages alphabetically by name.
-		PackageNames.Sort();
+	// Sort the level packages alphabetically by name.
+	PackageNames.Sort();
 
-		// Fire ULevel::LevelDirtiedEvent when falling out of scope.
-		FScopedLevelDirtied LevelDirtyCallback;
+	// Fire ULevel::LevelDirtiedEvent when falling out of scope.
+	FScopedLevelDirtied LevelDirtyCallback;
 
-		// Try to add the levels that were specified in the dialog.
-		ULevel* NewLevel = nullptr;
-		for (const auto& PackageName : PackageNames)
-		{
-			SlowTask.EnterProgressFrame();
+	// Try to add the levels that were specified in the dialog.
+	ULevel* NewLevel = nullptr;
+	for (const auto& PackageName : PackageNames)
+	{
+		SlowTask.EnterProgressFrame();
 
 		ULevelStreaming* NewStreamingLevel = AddLevelToWorld(InWorld, *PackageName, LevelStreamingClass);
 		NewLevel = NewStreamingLevel->GetLoadedLevel();
-			if (NewLevel)
-			{
-				LevelDirtyCallback.Request();
-			}
-
-		} // for each file
-
-		// Set the last loaded level to be the current level
 		if (NewLevel)
 		{
-			if (InWorld->SetCurrentLevel(NewLevel))
-			{
-				FEditorDelegates::NewCurrentLevel.Broadcast();
-			}
+			LevelDirtyCallback.Request();
 		}
 
-		// For safety
-		if (GLevelEditorModeTools().IsModeActive(FBuiltinEditorModes::EM_Landscape))
+	} // for each file
+
+	  // Set the last loaded level to be the current level
+	if (NewLevel)
+	{
+		if (InWorld->SetCurrentLevel(NewLevel))
 		{
-			GLevelEditorModeTools().ActivateDefaultMode();
+			FEditorDelegates::NewCurrentLevel.Broadcast();
 		}
+	}
 
-		// refresh editor windows
-		FEditorDelegates::RefreshAllBrowsers.Broadcast();
+	// For safety
+	if (GLevelEditorModeTools().IsModeActive(FBuiltinEditorModes::EM_Landscape))
+	{
+		GLevelEditorModeTools().ActivateDefaultMode();
+	}
 
-		// Update volume actor visibility for each viewport since we loaded a level which could potentially contain volumes
-		GUnrealEd->UpdateVolumeActorVisibility(NULL);
+	// refresh editor windows
+	FEditorDelegates::RefreshAllBrowsers.Broadcast();
 
-		return NewLevel;
-	}
+	// Update volume actor visibility for each viewport since we loaded a level which could potentially contain volumes
+	GUnrealEd->UpdateVolumeActorVisibility(NULL);
 
-ULevelStreaming* UEditorLevelUtils::AddLevelToWorld( UWorld* InWorld, const TCHAR* LevelPackageName, TSubclassOf<ULevelStreaming> LevelStreamingClass)
-	{
+	return NewLevel;
+}
+
+ULevelStreaming* UEditorLevelUtils::AddLevelToWorld(UWorld* InWorld, const TCHAR* LevelPackageName, TSubclassOf<ULevelStreaming> LevelStreamingClass)
+{
 	ULevel* NewLevel = nullptr;
 	ULevelStreaming* StreamingLevel = nullptr;
-		bool bIsPersistentLevel = (InWorld->PersistentLevel->GetOutermost()->GetName() == FString(LevelPackageName));
+	bool bIsPersistentLevel = (InWorld->PersistentLevel->GetOutermost()->GetName() == FString(LevelPackageName));
 
-		if ( bIsPersistentLevel || FLevelUtils::FindStreamingLevel( InWorld, LevelPackageName ) )
-		{
-			// Do nothing if the level already exists in the world.
-			FMessageDialog::Open( EAppMsgType::Ok, NSLOCTEXT("UnrealEd", "LevelAlreadyExistsInWorld", "A level with that name already exists in the world.") );
-		}
-		else
+	if (bIsPersistentLevel || FLevelUtils::FindStreamingLevel(InWorld, LevelPackageName))
+	{
+		// Do nothing if the level already exists in the world.
+		FMessageDialog::Open(EAppMsgType::Ok, NSLOCTEXT("UnrealEd", "LevelAlreadyExistsInWorld", "A level with that name already exists in the world."));
+	}
+	else
+	{
+		// If the selected class is still NULL, abort the operation.
+		if (LevelStreamingClass == nullptr)
 		{
-			// If the selected class is still NULL, abort the operation.
-		if ( LevelStreamingClass == nullptr)
-			{
 			return nullptr;
-			}	
+		}
 
-			const FScopedBusyCursor BusyCursor;
+		const FScopedBusyCursor BusyCursor;
 
 		StreamingLevel = NewObject<ULevelStreaming>(InWorld, LevelStreamingClass, NAME_None, RF_NoFlags, NULL);
 
-			// Associate a package name.
-			StreamingLevel->SetWorldAssetByPackageName(LevelPackageName);
+		// Associate a package name.
+		StreamingLevel->SetWorldAssetByPackageName(LevelPackageName);
 
-			// Seed the level's draw color.
-			StreamingLevel->LevelColor = FLinearColor::MakeRandomColor();
+		// Seed the level's draw color.
+		StreamingLevel->LevelColor = FLinearColor::MakeRandomColor();
 
-			// Add the new level to world.
-			InWorld->StreamingLevels.Add( StreamingLevel );
+		// Add the new level to world.
+		InWorld->StreamingLevels.Add(StreamingLevel);
 
-			// Refresh just the newly created level.
-			TArray<ULevelStreaming*> LevelsForRefresh;
-			LevelsForRefresh.Add( StreamingLevel );
-			InWorld->RefreshStreamingLevels( LevelsForRefresh );
-			InWorld->MarkPackageDirty();
+		// Refresh just the newly created level.
+		TArray<ULevelStreaming*> LevelsForRefresh;
+		LevelsForRefresh.Add(StreamingLevel);
+		InWorld->RefreshStreamingLevels(LevelsForRefresh);
+		InWorld->MarkPackageDirty();
 
-			NewLevel = StreamingLevel->GetLoadedLevel();
-			if ( NewLevel != nullptr)
-			{
-				EditorLevelUtils::SetLevelVisibility( NewLevel, true, true );
+		NewLevel = StreamingLevel->GetLoadedLevel();
+		if (NewLevel != nullptr)
+		{
+			EditorLevelUtils::SetLevelVisibility(NewLevel, true, true);
 
-				// Levels migrated from other projects may fail to load their world settings
-				// If so we create a new AWorldSettings actor here.
-				if ( NewLevel->GetWorldSettings(false) == nullptr )
-				{
-					UWorld* SubLevelWorld = CastChecked<UWorld>(NewLevel->GetOuter());
+			// Levels migrated from other projects may fail to load their world settings
+			// If so we create a new AWorldSettings actor here.
+			if (NewLevel->GetWorldSettings(false) == nullptr)
+			{
+				UWorld* SubLevelWorld = CastChecked<UWorld>(NewLevel->GetOuter());
 
-					FActorSpawnParameters SpawnInfo;
-					SpawnInfo.SpawnCollisionHandlingOverride = ESpawnActorCollisionHandlingMethod::AlwaysSpawn;
-					SpawnInfo.Name = GEngine->WorldSettingsClass->GetFName();
-					AWorldSettings* NewWorldSettings = SubLevelWorld->SpawnActor<AWorldSettings>( GEngine->WorldSettingsClass, SpawnInfo );
+				FActorSpawnParameters SpawnInfo;
+				SpawnInfo.SpawnCollisionHandlingOverride = ESpawnActorCollisionHandlingMethod::AlwaysSpawn;
+				SpawnInfo.Name = GEngine->WorldSettingsClass->GetFName();
+				AWorldSettings* NewWorldSettings = SubLevelWorld->SpawnActor<AWorldSettings>(GEngine->WorldSettingsClass, SpawnInfo);
 
-					NewLevel->SetWorldSettings(NewWorldSettings);
-				}
+				NewLevel->SetWorldSettings(NewWorldSettings);
 			}
 		}
+	}
 
-		if(NewLevel) // if the level was successfully added
-		{
-			FEditorDelegates::OnAddLevelToWorld.Broadcast(NewLevel);
-		}
+	if (NewLevel) // if the level was successfully added
+	{
+		FEditorDelegates::OnAddLevelToWorld.Broadcast(NewLevel);
+	}
 
 	return StreamingLevel;
-	}
+}
 
 ULevelStreaming* UEditorLevelUtils::SetStreamingClassForLevel(ULevelStreaming* InLevel, TSubclassOf<ULevelStreaming> LevelStreamingClass)
-	{
-		check(InLevel);
+{
+	check(InLevel);
 
-		const FScopedBusyCursor BusyCursor;
+	const FScopedBusyCursor BusyCursor;
 
-		// Cache off the package name, as it will be lost when unloading the level
-		const FName CachedPackageName = InLevel->GetWorldAssetPackageFName();
+	// Cache off the package name, as it will be lost when unloading the level
+	const FName CachedPackageName = InLevel->GetWorldAssetPackageFName();
 
-		// First hide and remove the level if it exists
-		ULevel* Level = InLevel->GetLoadedLevel();
-		check(Level);
-		SetLevelVisibility( Level, false, false );
-		check(Level->OwningWorld);
-		UWorld* World = Level->OwningWorld;
-		
-		World->StreamingLevels.Remove(InLevel);
-		
-		// re-add the level with the desired streaming class
-		AddLevelToWorld(World, *(CachedPackageName.ToString()), LevelStreamingClass);
+	// First hide and remove the level if it exists
+	ULevel* Level = InLevel->GetLoadedLevel();
+	check(Level);
+	SetLevelVisibility(Level, false, false);
+	check(Level->OwningWorld);
+	UWorld* World = Level->OwningWorld;
 
-		// Transfer level streaming settings
-		ULevelStreaming* NewStreamingLevel = FLevelUtils::FindStreamingLevel( Level );
-		if ( NewStreamingLevel )
-		{
-			NewStreamingLevel->LevelTransform = InLevel->LevelTransform;
-			NewStreamingLevel->EditorStreamingVolumes = InLevel->EditorStreamingVolumes;
-			NewStreamingLevel->MinTimeBetweenVolumeUnloadRequests = InLevel->MinTimeBetweenVolumeUnloadRequests;
-			NewStreamingLevel->LevelColor = InLevel->LevelColor;
-			NewStreamingLevel->Keywords = InLevel->Keywords;
-		}
+	World->StreamingLevels.Remove(InLevel);
+
+	// re-add the level with the desired streaming class
+	AddLevelToWorld(World, *(CachedPackageName.ToString()), LevelStreamingClass);
 
-		return NewStreamingLevel;
+	// Transfer level streaming settings
+	ULevelStreaming* NewStreamingLevel = FLevelUtils::FindStreamingLevel(Level);
+	if (NewStreamingLevel)
+	{
+		NewStreamingLevel->LevelTransform = InLevel->LevelTransform;
+		NewStreamingLevel->EditorStreamingVolumes = InLevel->EditorStreamingVolumes;
+		NewStreamingLevel->MinTimeBetweenVolumeUnloadRequests = InLevel->MinTimeBetweenVolumeUnloadRequests;
+		NewStreamingLevel->LevelColor = InLevel->LevelColor;
+		NewStreamingLevel->Keywords = InLevel->Keywords;
 	}
 
+	return NewStreamingLevel;
+}
+
 void UEditorLevelUtils::MakeLevelCurrent(ULevel* InLevel)
+{
+	if (ensureAsRuntimeWarning(InLevel != nullptr))
 	{
 		// Locked levels can't be made current.
-		if ( !FLevelUtils::IsLevelLocked( InLevel ) )
-		{ 
+		if (!FLevelUtils::IsLevelLocked(InLevel))
+		{
 			// Make current broadcast if it changed
-			if( InLevel->OwningWorld->SetCurrentLevel( InLevel ) )
+			if (InLevel->OwningWorld->SetCurrentLevel(InLevel))
 			{
 				FEditorDelegates::NewCurrentLevel.Broadcast();
 			}
-			
+
 			// Deselect all selected builder brushes.
 			bool bDeselectedSomething = false;
-			for ( FSelectionIterator It( GEditor->GetSelectedActorIterator() ) ; It ; ++It )
+			for (FSelectionIterator It(GEditor->GetSelectedActorIterator()); It; ++It)
 			{
-				AActor* Actor = static_cast<AActor*>( *It );
-				checkSlow( Actor->IsA(AActor::StaticClass()) );
-				ABrush* Brush = Cast< ABrush >( Actor );
-				if ( Brush && FActorEditorUtils::IsABuilderBrush(Actor) )
+				AActor* Actor = static_cast<AActor*>(*It);
+				checkSlow(Actor->IsA(AActor::StaticClass()));
+				ABrush* Brush = Cast< ABrush >(Actor);
+				if (Brush && FActorEditorUtils::IsABuilderBrush(Actor))
 				{
-					GEditor->SelectActor( Actor, /*bInSelected=*/ false, /*bNotify=*/ false );
+					GEditor->SelectActor(Actor, /*bInSelected=*/ false, /*bNotify=*/ false);
 					bDeselectedSomething = true;
 				}
 			}
 
 			// Send a selection change callback if necessary.
-			if ( bDeselectedSomething )
+			if (bDeselectedSomething)
 			{
 				GEditor->NoteSelectionChange();
 			}
 
 			// Force the current level to be visible.
-			SetLevelVisibility( InLevel, true, false );
+			SetLevelVisibility(InLevel, true, false);
 		}
 		else
 		{
-			FMessageDialog::Open( EAppMsgType::Ok, NSLOCTEXT("UnrealEd", "Error_OperationDisallowedOnLockedLevelMakeLevelCurrent", "MakeLevelCurrent: The requested operation could not be completed because the level is locked.") );
+			FMessageDialog::Open(EAppMsgType::Ok, NSLOCTEXT("UnrealEd", "Error_OperationDisallowedOnLockedLevelMakeLevelCurrent", "MakeLevelCurrent: The requested operation could not be completed because the level is locked."));
 		}
 	}
-	
+}
+
 void UEditorLevelUtils::MakeLevelCurrent(ULevelStreaming* InStreamingLevel)
 {
-	MakeLevelCurrent(InStreamingLevel->GetLoadedLevel());
+	if (ensureAsRuntimeWarning(InStreamingLevel != nullptr))
+	{
+		MakeLevelCurrent(InStreamingLevel->GetLoadedLevel());
+	}
 }
 
 bool UEditorLevelUtils::PrivateRemoveInvalidLevelFromWorld(ULevelStreaming* InLevelStreaming)
+{
+	bool bRemovedLevelStreaming = false;
+	if (InLevelStreaming != NULL)
 	{
-		bool bRemovedLevelStreaming = false;
-		if (InLevelStreaming != NULL)
-		{
-			check(InLevelStreaming->GetLoadedLevel() == NULL); // This method is designed to be used to remove left over references to null levels 
+		check(InLevelStreaming->GetLoadedLevel() == NULL); // This method is designed to be used to remove left over references to null levels 
 
-			InLevelStreaming->Modify();
+		InLevelStreaming->Modify();
 
-			// Disassociate the level from the volume.
-			for ( auto VolIter = InLevelStreaming->EditorStreamingVolumes.CreateIterator(); VolIter; VolIter++ )
+		// Disassociate the level from the volume.
+		for (auto VolIter = InLevelStreaming->EditorStreamingVolumes.CreateIterator(); VolIter; VolIter++)
+		{
+			ALevelStreamingVolume* LevelStreamingVolume = *VolIter;
+			if (LevelStreamingVolume)
 			{
-				ALevelStreamingVolume* LevelStreamingVolume = *VolIter;
-				if ( LevelStreamingVolume )
-				{
-					LevelStreamingVolume->Modify();
-					LevelStreamingVolume->StreamingLevelNames.Remove( InLevelStreaming->GetWorldAssetPackageFName() );
-				}
+				LevelStreamingVolume->Modify();
+				LevelStreamingVolume->StreamingLevelNames.Remove(InLevelStreaming->GetWorldAssetPackageFName());
 			}
+		}
 
-			// Disassociate the volumes from the level.
-			InLevelStreaming->EditorStreamingVolumes.Empty();
+		// Disassociate the volumes from the level.
+		InLevelStreaming->EditorStreamingVolumes.Empty();
 
-			UWorld* OwningWorld = Cast<UWorld>(InLevelStreaming->GetOuter());
-			if (OwningWorld != NULL)
-			{
-				OwningWorld->StreamingLevels.Remove(InLevelStreaming);
-				OwningWorld->RefreshStreamingLevels();
-				bRemovedLevelStreaming = true;
-			}
+		UWorld* OwningWorld = Cast<UWorld>(InLevelStreaming->GetOuter());
+		if (OwningWorld != NULL)
+		{
+			OwningWorld->StreamingLevels.Remove(InLevelStreaming);
+			OwningWorld->RefreshStreamingLevels();
+			bRemovedLevelStreaming = true;
 		}
-		return bRemovedLevelStreaming;
 	}
+	return bRemovedLevelStreaming;
+}
 
 bool UEditorLevelUtils::RemoveInvalidLevelFromWorld(ULevelStreaming* InLevelStreaming)
+{
+	bool bRemoveSuccessful = PrivateRemoveInvalidLevelFromWorld(InLevelStreaming);
+	if (bRemoveSuccessful)
 	{
-		bool bRemoveSuccessful = PrivateRemoveInvalidLevelFromWorld(InLevelStreaming);
-		if (bRemoveSuccessful)
-		{
-			// Redraw the main editor viewports.
-			FEditorSupportDelegates::RedrawAllViewports.Broadcast();
+		// Redraw the main editor viewports.
+		FEditorSupportDelegates::RedrawAllViewports.Broadcast();
 
-			// refresh editor windows
-			FEditorDelegates::RefreshAllBrowsers.Broadcast();
+		// refresh editor windows
+		FEditorDelegates::RefreshAllBrowsers.Broadcast();
 
-			// Update selection for any selected actors that were in the level and are no longer valid
-			GEditor->NoteSelectionChange();
+		// Update selection for any selected actors that were in the level and are no longer valid
+		GEditor->NoteSelectionChange();
 
-			// Collect garbage to clear out the destroyed level
-			CollectGarbage( GARBAGE_COLLECTION_KEEPFLAGS );
-		}
-		return bRemoveSuccessful;
+		// Collect garbage to clear out the destroyed level
+		CollectGarbage(GARBAGE_COLLECTION_KEEPFLAGS);
 	}
-	
+	return bRemoveSuccessful;
+}
+
 
 ULevelStreaming* UEditorLevelUtils::CreateNewStreamingLevel(TSubclassOf<ULevelStreaming> LevelStreamingClass, const FString& PackagePath /*= TEXT("")*/, bool bMoveSelectedActorsIntoNewLevel /*= false*/)
 {
 	FString Filename;
 	if (PackagePath.IsEmpty() || FPackageName::TryConvertLongPackageNameToFilename(PackagePath, Filename, FPackageName::GetMapPackageExtension()))
 	{
-		return CreateNewStreamingLevelForWorld(*GEditor->GetEditorWorldContext().World(), LevelStreamingClass, Filename, bMoveSelectedActorsIntoNewLevel);
+		if (ensureAsRuntimeWarning(LevelStreamingClass.Get() != nullptr))
+		{
+			return CreateNewStreamingLevelForWorld(*GEditor->GetEditorWorldContext().World(), LevelStreamingClass, Filename, bMoveSelectedActorsIntoNewLevel);
+		}
 	}
 
 	return nullptr;
 }
-	
+
 
 ULevelStreaming* UEditorLevelUtils::CreateNewStreamingLevelForWorld(UWorld& InWorld, TSubclassOf<ULevelStreaming> LevelStreamingClass, const FString& DefaultFilename /* = TEXT( "" ) */, bool bMoveSelectedActorsIntoNewLevel /* = false */)
 {
@@ -466,25 +479,17 @@ ULevelStreaming* UEditorLevelUtils::CreateNewStreamingLevelForWorld(UWorld& InWo
 
 	// This is the new streaming level's world not the persistent level world
 	UWorld* NewLevelWorld = nullptr;
-	if (UEditorEngine::IsUsingWorldAssets())
-	{
-		// Create a new world
-		UWorldFactory* Factory = NewObject<UWorldFactory>();
-		Factory->WorldType = EWorldType::Inactive;
-		UPackage* Pkg = CreatePackage(NULL, NULL);
-		FName WorldName(TEXT("Untitled"));
-		EObjectFlags Flags = RF_Public | RF_Standalone;
-		NewLevelWorld = CastChecked<UWorld>(Factory->FactoryCreateNew(UWorld::StaticClass(), Pkg, WorldName, Flags, NULL, GWarn));
-		if (NewLevelWorld)
-		{
-			FAssetRegistryModule::AssetCreated(NewLevelWorld);
-		}
-	}
-	else
+
+	// Create a new world
+	UWorldFactory* Factory = NewObject<UWorldFactory>();
+	Factory->WorldType = EWorldType::Inactive;
+	UPackage* Pkg = CreatePackage(NULL, NULL);
+	FName WorldName(TEXT("Untitled"));
+	EObjectFlags Flags = RF_Public | RF_Standalone;
+	NewLevelWorld = CastChecked<UWorld>(Factory->FactoryCreateNew(UWorld::StaticClass(), Pkg, WorldName, Flags, NULL, GWarn));
+	if (NewLevelWorld)
 	{
-		// Create a new world - so we can 'borrow' its level
-		NewLevelWorld = UWorld::CreateWorld(EWorldType::None, false);
-		check(NewLevelWorld);
+		FAssetRegistryModule::AssetCreated(NewLevelWorld);
 	}
 
 	// Save the new world to disk.
@@ -495,13 +500,6 @@ ULevelStreaming* UEditorLevelUtils::CreateNewStreamingLevelForWorld(UWorld& InWo
 		NewPackageName = NewLevelWorld->GetOutermost()->GetName();
 	}
 
-	if (!UEditorEngine::IsUsingWorldAssets())
-	{
-		// Destroy the new world we created and collect the garbage
-		NewLevelWorld->DestroyWorld(false);
-		CollectGarbage(GARBAGE_COLLECTION_KEEPFLAGS);
-	}
-
 	// If the new world was saved successfully, import it as a streaming level.
 	ULevelStreaming* NewStreamingLevel = nullptr;
 	ULevel* NewLevel = nullptr;
@@ -531,449 +529,449 @@ ULevelStreaming* UEditorLevelUtils::CreateNewStreamingLevelForWorld(UWorld& InWo
 
 
 bool UEditorLevelUtils::RemoveLevelFromWorld(ULevel* InLevel)
+{
+	if (GEditor->Layers.IsValid())
+	{
+		GEditor->Layers->RemoveLevelLayerInformation(InLevel);
+	}
+
+	GEditor->CloseEditedWorldAssets(CastChecked<UWorld>(InLevel->GetOuter()));
+
+	UWorld* OwningWorld = InLevel->OwningWorld;
+	const FName LevelPackageName = InLevel->GetOutermost()->GetFName();
+	const bool bRemovingCurrentLevel = InLevel->IsCurrentLevel();
+	const bool bRemoveSuccessful = PrivateRemoveLevelFromWorld(InLevel);
+	if (bRemoveSuccessful)
 	{
-		if (GEditor->Layers.IsValid())
+		if (bRemovingCurrentLevel)
 		{
-			GEditor->Layers->RemoveLevelLayerInformation( InLevel );
+			MakeLevelCurrent(OwningWorld->PersistentLevel);
 		}
-		
-		GEditor->CloseEditedWorldAssets(CastChecked<UWorld>(InLevel->GetOuter()));
 
-		UWorld* OwningWorld = InLevel->OwningWorld;
-		const FName LevelPackageName		= InLevel->GetOutermost()->GetFName();
-		const bool bRemovingCurrentLevel	= InLevel->IsCurrentLevel();
-		const bool bRemoveSuccessful		= PrivateRemoveLevelFromWorld( InLevel );
-		if ( bRemoveSuccessful )
-		{
-			if (bRemovingCurrentLevel)
-			{
-				MakeLevelCurrent(OwningWorld->PersistentLevel);
-			}
 
+		EditorDestroyLevel(InLevel);
 
-			EditorDestroyLevel(InLevel);
+		// Redraw the main editor viewports.
+		FEditorSupportDelegates::RedrawAllViewports.Broadcast();
 
-			// Redraw the main editor viewports.
-			FEditorSupportDelegates::RedrawAllViewports.Broadcast();
+		// refresh editor windows
+		FEditorDelegates::RefreshAllBrowsers.Broadcast();
 
-			// refresh editor windows
-			FEditorDelegates::RefreshAllBrowsers.Broadcast();
-			
-			// Reset transaction buffer and run GC to clear out the destroyed level
-			GEditor->Cleanse(true, false, LOCTEXT("RemoveLevelTransReset", "Removing Levels from World"));
+		// Reset transaction buffer and run GC to clear out the destroyed level
+		GEditor->Cleanse(true, false, LOCTEXT("RemoveLevelTransReset", "Removing Levels from World"));
 
-			// Ensure that world was removed
-			UPackage* LevelPackage = FindObjectFast<UPackage>(NULL, LevelPackageName);
-			if (LevelPackage != nullptr)
+		// Ensure that world was removed
+		UPackage* LevelPackage = FindObjectFast<UPackage>(NULL, LevelPackageName);
+		if (LevelPackage != nullptr)
+		{
+			UWorld* TheWorld = UWorld::FindWorldInPackage(LevelPackage->GetOutermost());
+			if (TheWorld != nullptr)
 			{
-				UWorld* TheWorld = UWorld::FindWorldInPackage(LevelPackage->GetOutermost());								
-				if (TheWorld != nullptr)
-				{	
-					StaticExec(NULL, *FString::Printf(TEXT("OBJ REFS CLASS=%s NAME=%s shortest"),*TheWorld->GetClass()->GetName(), *TheWorld->GetPathName()));
-					TMap<UObject*,UProperty*>	Route		= FArchiveTraceRoute::FindShortestRootPath( TheWorld, true, GARBAGE_COLLECTION_KEEPFLAGS );
-					FString						ErrorString	= FArchiveTraceRoute::PrintRootPath( Route, TheWorld );
-					UE_LOG(LogStreaming, Fatal, TEXT("%s didn't get garbage collected!") LINE_TERMINATOR TEXT("%s"), *TheWorld->GetFullName(), *ErrorString );
-				}
+				StaticExec(NULL, *FString::Printf(TEXT("OBJ REFS CLASS=%s NAME=%s shortest"), *TheWorld->GetClass()->GetName(), *TheWorld->GetPathName()));
+				TMap<UObject*, UProperty*>	Route = FArchiveTraceRoute::FindShortestRootPath(TheWorld, true, GARBAGE_COLLECTION_KEEPFLAGS);
+				FString						ErrorString = FArchiveTraceRoute::PrintRootPath(Route, TheWorld);
+				UE_LOG(LogStreaming, Fatal, TEXT("%s didn't get garbage collected!") LINE_TERMINATOR TEXT("%s"), *TheWorld->GetFullName(), *ErrorString);
 			}
 		}
-		return bRemoveSuccessful;
 	}
+	return bRemoveSuccessful;
+}
 
 
 bool UEditorLevelUtils::PrivateRemoveLevelFromWorld(ULevel* InLevel)
+{
+	if (!InLevel || InLevel->IsPersistentLevel())
 	{
-		if ( !InLevel || InLevel->IsPersistentLevel() )
-		{
-			return false;
-		}
+		return false;
+	}
 
-		if ( FLevelUtils::IsLevelLocked(InLevel) )
-		{
-			FMessageDialog::Open( EAppMsgType::Ok, NSLOCTEXT("UnrealEd", "Error_OperationDisallowedOnLockedLevelRemoveLevelFromWorld", "RemoveLevelFromWorld: The requested operation could not be completed because the level is locked.") );
-			return false;
-		}
+	if (FLevelUtils::IsLevelLocked(InLevel))
+	{
+		FMessageDialog::Open(EAppMsgType::Ok, NSLOCTEXT("UnrealEd", "Error_OperationDisallowedOnLockedLevelRemoveLevelFromWorld", "RemoveLevelFromWorld: The requested operation could not be completed because the level is locked."));
+		return false;
+	}
 
-		int32 StreamingLevelIndex = INDEX_NONE;
+	int32 StreamingLevelIndex = INDEX_NONE;
 
-		for( int32 LevelIndex = 0 ; LevelIndex < InLevel->OwningWorld->StreamingLevels.Num() ; ++LevelIndex )
+	for (int32 LevelIndex = 0; LevelIndex < InLevel->OwningWorld->StreamingLevels.Num(); ++LevelIndex)
+	{
+		ULevelStreaming* StreamingLevel = InLevel->OwningWorld->StreamingLevels[LevelIndex];
+		if (StreamingLevel && StreamingLevel->GetLoadedLevel() == InLevel)
 		{
-			ULevelStreaming* StreamingLevel = InLevel->OwningWorld->StreamingLevels[ LevelIndex ];
-			if( StreamingLevel && StreamingLevel->GetLoadedLevel() == InLevel )
-			{
-				StreamingLevelIndex = LevelIndex;
-				break;
-			}
+			StreamingLevelIndex = LevelIndex;
+			break;
 		}
+	}
 
-		if (StreamingLevelIndex != INDEX_NONE)
-		{
-			InLevel->OwningWorld->StreamingLevels[StreamingLevelIndex]->MarkPendingKill();
-			InLevel->OwningWorld->StreamingLevels.RemoveAt( StreamingLevelIndex );
-			InLevel->OwningWorld->RefreshStreamingLevels();
-		}
-		else if (InLevel->bIsVisible)
-		{
-			InLevel->OwningWorld->RemoveFromWorld(InLevel);
-			check(InLevel->bIsVisible == false);
-		}
+	if (StreamingLevelIndex != INDEX_NONE)
+	{
+		InLevel->OwningWorld->StreamingLevels[StreamingLevelIndex]->MarkPendingKill();
+		InLevel->OwningWorld->StreamingLevels.RemoveAt(StreamingLevelIndex);
+		InLevel->OwningWorld->RefreshStreamingLevels();
+	}
+	else if (InLevel->bIsVisible)
+	{
+		InLevel->OwningWorld->RemoveFromWorld(InLevel);
+		check(InLevel->bIsVisible == false);
+	}
 
-		InLevel->ReleaseRenderingResources();
+	InLevel->ReleaseRenderingResources();
 
-		IStreamingManager::Get().RemoveLevel( InLevel );
-		UWorld* World = InLevel->OwningWorld;
-		World->RemoveLevel(InLevel);
-		InLevel->ClearLevelComponents();
+	IStreamingManager::Get().RemoveLevel(InLevel);
+	UWorld* World = InLevel->OwningWorld;
+	World->RemoveLevel(InLevel);
+	InLevel->ClearLevelComponents();
 
-		// remove all group actors from the world in the level we are removing
-		// otherwise, this will cause group actors to not be garbage collected
-		for (int32 GroupIndex = World->ActiveGroupActors.Num()-1; GroupIndex >= 0; --GroupIndex)
+	// remove all group actors from the world in the level we are removing
+	// otherwise, this will cause group actors to not be garbage collected
+	for (int32 GroupIndex = World->ActiveGroupActors.Num() - 1; GroupIndex >= 0; --GroupIndex)
+	{
+		AGroupActor* GroupActor = Cast<AGroupActor>(World->ActiveGroupActors[GroupIndex]);
+		if (GroupActor && GroupActor->IsInLevel(InLevel))
 		{
-			AGroupActor* GroupActor = Cast<AGroupActor>(World->ActiveGroupActors[GroupIndex]);
-			if (GroupActor && GroupActor->IsInLevel(InLevel))
-			{
-				World->ActiveGroupActors.RemoveAt(GroupIndex);
-			}
+			World->ActiveGroupActors.RemoveAt(GroupIndex);
 		}
-		
-		// Mark all model components as pending kill so GC deletes references to them.
-		for (UModelComponent* ModelComponent : InLevel->ModelComponents)
+	}
+
+	// Mark all model components as pending kill so GC deletes references to them.
+	for (UModelComponent* ModelComponent : InLevel->ModelComponents)
+	{
+		if (ModelComponent != nullptr)
 		{
-			if (ModelComponent != nullptr)
-			{
-				ModelComponent->MarkPendingKill();
-			}
+			ModelComponent->MarkPendingKill();
 		}
+	}
 
-		// Mark all actors and their components as pending kill so GC will delete references to them.
-		for (AActor* Actor : InLevel->Actors)
+	// Mark all actors and their components as pending kill so GC will delete references to them.
+	for (AActor* Actor : InLevel->Actors)
+	{
+		if (Actor != nullptr)
 		{
-			if (Actor != nullptr)
-			{
-				Actor->MarkComponentsAsPendingKill();
-				Actor->MarkPendingKill();
-			}
+			Actor->MarkComponentsAsPendingKill();
+			Actor->MarkPendingKill();
 		}
+	}
 
-		World->MarkPackageDirty();
-		World->BroadcastLevelsChanged();
+	World->MarkPackageDirty();
+	World->BroadcastLevelsChanged();
 
-		return true;
-	}
+	return true;
+}
 
 bool UEditorLevelUtils::EditorDestroyLevel(ULevel* InLevel)
-	{
-		UWorld* World = InLevel->OwningWorld;
+{
+	UWorld* World = InLevel->OwningWorld;
 
-		InLevel->GetOuter()->MarkPendingKill();
-		InLevel->MarkPendingKill();
-		InLevel->GetOuter()->ClearFlags(RF_Public | RF_Standalone);
+	InLevel->GetOuter()->MarkPendingKill();
+	InLevel->MarkPendingKill();
+	InLevel->GetOuter()->ClearFlags(RF_Public | RF_Standalone);
 
-		UPackage* Package = InLevel->GetOutermost();
-		// We want to unconditionally destroy the level, so clear the dirty flag here so it can be unloaded successfully
-		Package->SetDirtyFlag(false);
+	UPackage* Package = InLevel->GetOutermost();
+	// We want to unconditionally destroy the level, so clear the dirty flag here so it can be unloaded successfully
+	Package->SetDirtyFlag(false);
 
-		TArray<UPackage*> Packages;
-		Packages.Add(Package);
-		if (!PackageTools::UnloadPackages(Packages))
-		{
-			FFormatNamedArguments Args;
-			Args.Add(TEXT("Package"), FText::FromString(Package->GetName()));
-			FMessageDialog::Open(EAppMsgType::Ok, FText::Format(LOCTEXT("UnloadPackagesFail", "Unable to unload package '{Package}'."), Args));
-			return false;
-		}
-
-		return true;
+	TArray<UPackage*> Packages;
+	Packages.Add(Package);
+	if (!PackageTools::UnloadPackages(Packages))
+	{
+		FFormatNamedArguments Args;
+		Args.Add(TEXT("Package"), FText::FromString(Package->GetName()));
+		FMessageDialog::Open(EAppMsgType::Ok, FText::Format(LOCTEXT("UnloadPackagesFail", "Unable to unload package '{Package}'."), Args));
+		return false;
 	}
 
-ULevel* UEditorLevelUtils::CreateNewLevel( UWorld* InWorld, bool bMoveSelectedActorsIntoNewLevel, TSubclassOf<ULevelStreaming> LevelStreamingClass, const FString& DefaultFilename )
-		{
+	return true;
+}
+
+ULevel* UEditorLevelUtils::CreateNewLevel(UWorld* InWorld, bool bMoveSelectedActorsIntoNewLevel, TSubclassOf<ULevelStreaming> LevelStreamingClass, const FString& DefaultFilename)
+{
 	ULevelStreaming* StreamingLevel = CreateNewStreamingLevelForWorld(*InWorld, LevelStreamingClass, DefaultFilename, bMoveSelectedActorsIntoNewLevel);
 	return StreamingLevel->GetLoadedLevel();
-		}
+}
 
 void UEditorLevelUtils::DeselectAllSurfacesInLevel(ULevel* InLevel)
+{
+	if (InLevel)
 	{
-		if(InLevel)
+		UModel* Model = InLevel->Model;
+		for (int32 SurfaceIndex = 0; SurfaceIndex < Model->Surfs.Num(); ++SurfaceIndex)
 		{
-			UModel* Model = InLevel->Model;
-			for( int32 SurfaceIndex = 0 ; SurfaceIndex < Model->Surfs.Num() ; ++SurfaceIndex )
+			FBspSurf& Surf = Model->Surfs[SurfaceIndex];
+			if (Surf.PolyFlags & PF_Selected)
 			{
-				FBspSurf& Surf = Model->Surfs[SurfaceIndex];
-				if( Surf.PolyFlags & PF_Selected )
-				{
-					Model->ModifySurf( SurfaceIndex, false );
-					Surf.PolyFlags &= ~PF_Selected;
-				}
+				Model->ModifySurf(SurfaceIndex, false);
+				Surf.PolyFlags &= ~PF_Selected;
 			}
 		}
 	}
+}
 
 void UEditorLevelUtils::SetLevelVisibility(ULevel* Level, bool bShouldBeVisible, bool bForceLayersVisible)
-	{	
-		// Nothing to do
-		if ( Level == NULL )
+{
+	// Nothing to do
+	if (Level == NULL)
+	{
+		return;
+	}
+
+	// Handle the case of the p-level
+	// The p-level can't be unloaded, so its actors/BSP should just be temporarily hidden/unhidden
+	// Also, intentionally do not force layers visible for the p-level
+	if (Level->IsPersistentLevel())
+	{
+		//create a transaction so we can undo the visibilty toggle
+		const FScopedTransaction Transaction(LOCTEXT("ToggleLevelVisibility", "Toggle Level Visibility"));
+		if (Level->bIsVisible != bShouldBeVisible)
 		{
-			return;
+			Level->Modify();
 		}
-
-		// Handle the case of the p-level
-		// The p-level can't be unloaded, so its actors/BSP should just be temporarily hidden/unhidden
-		// Also, intentionally do not force layers visible for the p-level
-		if ( Level->IsPersistentLevel() )
+		// Set the visibility of each actor in the p-level
+		for (TArray<AActor*>::TIterator PLevelActorIter(Level->Actors); PLevelActorIter; ++PLevelActorIter)
 		{
-			//create a transaction so we can undo the visibilty toggle
-			const FScopedTransaction Transaction( LOCTEXT( "ToggleLevelVisibility", "Toggle Level Visibility" ) );		
-			if ( Level->bIsVisible != bShouldBeVisible )
+			AActor* CurActor = *PLevelActorIter;
+			if (CurActor && !FActorEditorUtils::IsABuilderBrush(CurActor) && CurActor->bHiddenEdLevel == bShouldBeVisible)
 			{
-				Level->Modify();
+				CurActor->Modify();
+				CurActor->bHiddenEdLevel = !bShouldBeVisible;
+				CurActor->RegisterAllComponents();
+				CurActor->MarkComponentsRenderStateDirty();
 			}
-			// Set the visibility of each actor in the p-level
-			for ( TArray<AActor*>::TIterator PLevelActorIter( Level->Actors ); PLevelActorIter; ++PLevelActorIter )
+		}
+
+		// Set the visibility of each BSP surface in the p-level
+		UModel* CurLevelModel = Level->Model;
+		if (CurLevelModel)
+		{
+			CurLevelModel->Modify();
+			for (TArray<FBspSurf>::TIterator SurfaceIterator(CurLevelModel->Surfs); SurfaceIterator; ++SurfaceIterator)
 			{
-				AActor* CurActor = *PLevelActorIter;
-				if ( CurActor && !FActorEditorUtils::IsABuilderBrush(CurActor) && CurActor->bHiddenEdLevel == bShouldBeVisible )
-				{
-					CurActor->Modify();
-					CurActor->bHiddenEdLevel = !bShouldBeVisible;
-					CurActor->RegisterAllComponents();
-					CurActor->MarkComponentsRenderStateDirty();
-				}
+				FBspSurf& CurSurf = *SurfaceIterator;
+				CurSurf.bHiddenEdLevel = !bShouldBeVisible;
 			}
+		}
 
-			// Set the visibility of each BSP surface in the p-level
-			UModel* CurLevelModel = Level->Model;
-			if ( CurLevelModel )
+		// Add/remove model components from the scene
+		for (int32 ComponentIndex = 0; ComponentIndex < Level->ModelComponents.Num(); ComponentIndex++)
+		{
+			UModelComponent* CurLevelModelCmp = Level->ModelComponents[ComponentIndex];
+			if (CurLevelModelCmp)
 			{
-				CurLevelModel->Modify();
-				for ( TArray<FBspSurf>::TIterator SurfaceIterator( CurLevelModel->Surfs ); SurfaceIterator; ++SurfaceIterator )
+				if (bShouldBeVisible)
 				{
-					FBspSurf& CurSurf = *SurfaceIterator;
-					CurSurf.bHiddenEdLevel = !bShouldBeVisible;
+					CurLevelModelCmp->RegisterComponentWithWorld(Level->OwningWorld);
 				}
-			}
-
-			// Add/remove model components from the scene
-			for(int32 ComponentIndex = 0;ComponentIndex < Level->ModelComponents.Num();ComponentIndex++)
-			{
-				UModelComponent* CurLevelModelCmp = Level->ModelComponents[ComponentIndex];
-				if(CurLevelModelCmp)
+				else if (!bShouldBeVisible && CurLevelModelCmp->IsRegistered())
 				{
-					if (bShouldBeVisible)
-					{
-						CurLevelModelCmp->RegisterComponentWithWorld(Level->OwningWorld);
-					}
-					else if (!bShouldBeVisible && CurLevelModelCmp->IsRegistered())
-					{
-						CurLevelModelCmp->UnregisterComponent();
-					}
+					CurLevelModelCmp->UnregisterComponent();
 				}
 			}
+		}
 
-			Level->GetWorld()->OnLevelsChanged().Broadcast();
-			FEditorSupportDelegates::RedrawAllViewports.Broadcast();
+		Level->GetWorld()->OnLevelsChanged().Broadcast();
+		FEditorSupportDelegates::RedrawAllViewports.Broadcast();
+	}
+	else
+	{
+		ULevelStreaming* StreamingLevel = NULL;
+		if (Level->OwningWorld == NULL || Level->OwningWorld->PersistentLevel != Level)
+		{
+			StreamingLevel = FLevelUtils::FindStreamingLevel(Level);
 		}
-		else
+
+		// Create a transaction so we can undo the visibility toggle
+		const FScopedTransaction Transaction(LOCTEXT("ToggleLevelVisibility", "Toggle Level Visibility"));
+
+		// Handle the case of a streaming level
+		if (StreamingLevel)
 		{
-			ULevelStreaming* StreamingLevel = NULL;
-			if (Level->OwningWorld == NULL || Level->OwningWorld->PersistentLevel != Level )
-			{
-				StreamingLevel = FLevelUtils::FindStreamingLevel( Level );
-			}
+			// We need to set the RF_Transactional to make a streaming level serialize itself. so store the original ones, set the flag, and put the original flags back when done
+			EObjectFlags cachedFlags = StreamingLevel->GetFlags();
+			StreamingLevel->SetFlags(RF_Transactional);
+			StreamingLevel->Modify();
+			StreamingLevel->SetFlags(cachedFlags);
 
-			// Create a transaction so we can undo the visibility toggle
-			const FScopedTransaction Transaction( LOCTEXT( "ToggleLevelVisibility", "Toggle Level Visibility" ) );		
+			// Set the visibility state for this streaming level.  
+			StreamingLevel->bShouldBeVisibleInEditor = bShouldBeVisible;
+		}
 
-			// Handle the case of a streaming level
-			if ( StreamingLevel )
-			{
-				// We need to set the RF_Transactional to make a streaming level serialize itself. so store the original ones, set the flag, and put the original flags back when done
-				EObjectFlags cachedFlags = StreamingLevel->GetFlags();
-				StreamingLevel->SetFlags( RF_Transactional );
-				StreamingLevel->Modify();			
-				StreamingLevel->SetFlags( cachedFlags );
-
-				// Set the visibility state for this streaming level.  
-				StreamingLevel->bShouldBeVisibleInEditor = bShouldBeVisible;
-			}
+		if (!bShouldBeVisible && GEditor->Layers.IsValid())
+		{
+			GEditor->Layers->RemoveLevelLayerInformation(Level);
+		}
+
+		// UpdateLevelStreaming sets Level->bIsVisible directly, so we need to make sure it gets saved to the transaction buffer.
+		if (Level->bIsVisible != bShouldBeVisible)
+		{
+			Level->Modify();
+		}
 
-			if (!bShouldBeVisible && GEditor->Layers.IsValid())
+		if (StreamingLevel)
+		{
+			Level->OwningWorld->FlushLevelStreaming();
+
+			// In the Editor we expect this operation will complete in a single call
+			check(Level->bIsVisible == bShouldBeVisible);
+		}
+		else if (Level->OwningWorld != NULL)
+		{
+			// In case we level has no associated StreamingLevel, remove or add to world directly
+			if (bShouldBeVisible)
 			{
-				GEditor->Layers->RemoveLevelLayerInformation( Level );
+				if (!Level->bIsVisible)
+				{
+					Level->OwningWorld->AddToWorld(Level);
+				}
 			}
-
-			// UpdateLevelStreaming sets Level->bIsVisible directly, so we need to make sure it gets saved to the transaction buffer.
-			if ( Level->bIsVisible != bShouldBeVisible )
+			else
 			{
-				Level->Modify();
+				Level->OwningWorld->RemoveFromWorld(Level);
 			}
 
-			if ( StreamingLevel )
-			{
-				Level->OwningWorld->FlushLevelStreaming();
+			// In the Editor we expect this operation will complete in a single call
+			check(Level->bIsVisible == bShouldBeVisible);
+		}
 
-				// In the Editor we expect this operation will complete in a single call
-				check(Level->bIsVisible == bShouldBeVisible);
-			}
-			else if (Level->OwningWorld != NULL)
+		if (bShouldBeVisible && GEditor->Layers.IsValid())
+		{
+			GEditor->Layers->AddLevelLayerInformation(Level);
+		}
+
+		// Force the level's layers to be visible, if desired
+		FEditorSupportDelegates::RedrawAllViewports.Broadcast();
+
+		// Iterate over the level's actors, making a list of their layers and unhiding the layers.
+		TArray<AActor*>& Actors = Level->Actors;
+		for (int32 ActorIndex = 0; ActorIndex < Actors.Num(); ++ActorIndex)
+		{
+			AActor* Actor = Actors[ActorIndex];
+			if (Actor)
 			{
-				// In case we level has no associated StreamingLevel, remove or add to world directly
-				if (bShouldBeVisible)
+				bool bModified = false;
+				if (bShouldBeVisible && bForceLayersVisible &&
+					GEditor->Layers.IsValid() &&
+					GEditor->Layers->IsActorValidForLayer(Actor))
 				{
-					if (!Level->bIsVisible)
+					// Make the actor layer visible, if it's not already.
+					if (Actor->bHiddenEdLayer)
 					{
-						Level->OwningWorld->AddToWorld(Level);
+						bModified = Actor->Modify();
+						Actor->bHiddenEdLayer = false;
 					}
-				}
-				else
-				{
-					Level->OwningWorld->RemoveFromWorld(Level);
-				}
-				
-				// In the Editor we expect this operation will complete in a single call
-				check(Level->bIsVisible == bShouldBeVisible);
-			}
-
-			if (bShouldBeVisible && GEditor->Layers.IsValid())
-			{
-				GEditor->Layers->AddLevelLayerInformation( Level );
-			}
 
-			// Force the level's layers to be visible, if desired
-			FEditorSupportDelegates::RedrawAllViewports.Broadcast();
+					const bool bIsVisible = true;
+					GEditor->Layers->SetLayersVisibility(Actor->Layers, bIsVisible);
+				}
 
-			// Iterate over the level's actors, making a list of their layers and unhiding the layers.
-			TArray<AActor*>& Actors = Level->Actors;
-			for ( int32 ActorIndex = 0 ; ActorIndex < Actors.Num() ; ++ActorIndex )
-			{
-				AActor* Actor = Actors[ ActorIndex ];
-				if ( Actor )
+				// Set the visibility of each actor in the streaming level
+				if (!FActorEditorUtils::IsABuilderBrush(Actor) && Actor->bHiddenEdLevel == bShouldBeVisible)
 				{
-					bool bModified = false;
-					if (bShouldBeVisible && bForceLayersVisible && 
-						GEditor->Layers.IsValid() && 
-						GEditor->Layers->IsActorValidForLayer( Actor ))
+					if (!bModified)
 					{
-						// Make the actor layer visible, if it's not already.
-						if ( Actor->bHiddenEdLayer )
-						{
-							bModified = Actor->Modify();
-							Actor->bHiddenEdLayer = false;
-						}
-
-						const bool bIsVisible = true;
-						GEditor->Layers->SetLayersVisibility( Actor->Layers, bIsVisible );
+						bModified = Actor->Modify();
 					}
+					Actor->bHiddenEdLevel = !bShouldBeVisible;
 
-					// Set the visibility of each actor in the streaming level
-					if ( !FActorEditorUtils::IsABuilderBrush(Actor) && Actor->bHiddenEdLevel == bShouldBeVisible )
+					if (bShouldBeVisible)
 					{
-						if ( !bModified )
-						{
-							bModified = Actor->Modify();
-						}
-						Actor->bHiddenEdLevel = !bShouldBeVisible;
-
-						if (bShouldBeVisible)
-						{
-							Actor->ReregisterAllComponents();
-						}
-						else
-						{
-							Actor->UnregisterAllComponents();
-						}
+						Actor->ReregisterAllComponents();
+					}
+					else
+					{
+						Actor->UnregisterAllComponents();
 					}
 				}
 			}
 		}
+	}
 
-		FEditorDelegates::RefreshLayerBrowser.Broadcast();
+	FEditorDelegates::RefreshLayerBrowser.Broadcast();
 
-		// Notify the Scene Outliner, as new Actors may be present in the world.
-		GEngine->BroadcastLevelActorListChanged();
+	// Notify the Scene Outliner, as new Actors may be present in the world.
+	GEngine->BroadcastLevelActorListChanged();
 
-		// If the level is being hidden, deselect actors and surfaces that belong to this level.
-		if ( !bShouldBeVisible )
+	// If the level is being hidden, deselect actors and surfaces that belong to this level.
+	if (!bShouldBeVisible)
+	{
+		USelection* SelectedActors = GEditor->GetSelectedActors();
+		SelectedActors->Modify();
+		TArray<AActor*>& Actors = Level->Actors;
+		for (int32 ActorIndex = 0; ActorIndex < Actors.Num(); ++ActorIndex)
 		{
-			USelection* SelectedActors = GEditor->GetSelectedActors();
-			SelectedActors->Modify();
-			TArray<AActor*>& Actors = Level->Actors;
-			for ( int32 ActorIndex = 0 ; ActorIndex < Actors.Num() ; ++ActorIndex )
+			AActor* Actor = Actors[ActorIndex];
+			if (Actor)
 			{
-				AActor* Actor = Actors[ ActorIndex ];
-				if ( Actor )
-				{
-					SelectedActors->Deselect( Actor );
-				}
+				SelectedActors->Deselect(Actor);
 			}
+		}
 
-			DeselectAllSurfacesInLevel(Level);
+		DeselectAllSurfacesInLevel(Level);
 
-			// Tell the editor selection status was changed.
-			GEditor->NoteSelectionChange();
-		}
+		// Tell the editor selection status was changed.
+		GEditor->NoteSelectionChange();
+	}
 
-		Level->bIsVisible = bShouldBeVisible;
+	Level->bIsVisible = bShouldBeVisible;
 
-		if (Level->bIsLightingScenario)
-		{
-			Level->OwningWorld->PropagateLightingScenarioChange(bShouldBeVisible);
-		}
+	if (Level->bIsLightingScenario)
+	{
+		Level->OwningWorld->PropagateLightingScenarioChange(bShouldBeVisible);
 	}
+}
 
 void UEditorLevelUtils::GetWorlds(UWorld* InWorld, TArray<UWorld*>& OutWorlds, bool bIncludeInWorld, bool bOnlyEditorVisible)
-	{
-		OutWorlds.Empty();
+{
+	OutWorlds.Empty();
 
-		if (!InWorld)
-		{
-			return;
-		}
+	if (!InWorld)
+	{
+		return;
+	}
 
-		if ( bIncludeInWorld )
-		{
-			OutWorlds.AddUnique( InWorld );
-		}
+	if (bIncludeInWorld)
+	{
+		OutWorlds.AddUnique(InWorld);
+	}
 
-		// Iterate over the world's level array to find referenced levels ("worlds"). We don't 
-		for ( int32 LevelIndex = 0 ; LevelIndex < InWorld->StreamingLevels.Num() ; ++LevelIndex )
+	// Iterate over the world's level array to find referenced levels ("worlds"). We don't 
+	for (int32 LevelIndex = 0; LevelIndex < InWorld->StreamingLevels.Num(); ++LevelIndex)
+	{
+		ULevelStreaming* StreamingLevel = InWorld->StreamingLevels[LevelIndex];
+		if (StreamingLevel)
 		{
-			ULevelStreaming* StreamingLevel = InWorld->StreamingLevels[LevelIndex];
-			if ( StreamingLevel )
+			// If we asked for only sub-levels that are editor-visible, then limit our results appropriately
+			bool bShouldAlwaysBeLoaded = false; // Cast< ULevelStreamingAlwaysLoaded >( StreamingLevel ) != NULL;
+			if (!bOnlyEditorVisible || bShouldAlwaysBeLoaded || StreamingLevel->bShouldBeVisibleInEditor)
 			{
-				// If we asked for only sub-levels that are editor-visible, then limit our results appropriately
-				bool bShouldAlwaysBeLoaded = false; // Cast< ULevelStreamingAlwaysLoaded >( StreamingLevel ) != NULL;
-				if( !bOnlyEditorVisible || bShouldAlwaysBeLoaded || StreamingLevel->bShouldBeVisibleInEditor )
+				const ULevel* Level = StreamingLevel->GetLoadedLevel();
+
+				// This should always be the case for valid level names as the Editor preloads all packages.
+				if (Level != NULL)
 				{
-					const ULevel* Level = StreamingLevel->GetLoadedLevel();
-					
-					// This should always be the case for valid level names as the Editor preloads all packages.
-					if ( Level != NULL )
+					// Newer levels have their packages' world as the outer.
+					UWorld* World = Cast<UWorld>(Level->GetOuter());
+					if (World != NULL)
 					{
-						// Newer levels have their packages' world as the outer.
-						UWorld* World = Cast<UWorld>( Level->GetOuter() );
-						if ( World != NULL )
-						{
-							OutWorlds.AddUnique( World );
-						}
+						OutWorlds.AddUnique(World);
 					}
 				}
 			}
 		}
-		
-		// Levels can be loaded directly without StreamingLevel facilities
-		for ( int32 LevelIndex = 0 ; LevelIndex < InWorld->GetLevels().Num() ; ++LevelIndex )
+	}
+
+	// Levels can be loaded directly without StreamingLevel facilities
+	for (int32 LevelIndex = 0; LevelIndex < InWorld->GetLevels().Num(); ++LevelIndex)
+	{
+		ULevel* Level = InWorld->GetLevel(LevelIndex);
+		if (Level)
 		{
-			ULevel* Level = InWorld->GetLevel(LevelIndex);
-			if ( Level )
+			// Newer levels have their packages' world as the outer.
+			UWorld* World = Cast<UWorld>(Level->GetOuter());
+			if (World != NULL)
 			{
-				// Newer levels have their packages' world as the outer.
-				UWorld* World = Cast<UWorld>( Level->GetOuter() );
-				if ( World != NULL )
-				{
-					OutWorlds.AddUnique( World );
-				}
+				OutWorlds.AddUnique(World);
 			}
 		}
 	}
+}
 
 #undef LOCTEXT_NAMESPACE
diff --git a/Engine/Source/Editor/UnrealEd/Private/Factories/WorldFactory.cpp b/Engine/Source/Editor/UnrealEd/Private/Factories/WorldFactory.cpp
--- a/Engine/Source/Editor/UnrealEd/Private/Factories/WorldFactory.cpp
+++ b/Engine/Source/Editor/UnrealEd/Private/Factories/WorldFactory.cpp
@@ -12,7 +12,7 @@
 UWorldFactory::UWorldFactory(const FObjectInitializer& ObjectInitializer)
 	: Super(ObjectInitializer)
 {
-	bCreateNew = UEditorEngine::IsUsingWorldAssets();
+	bCreateNew = true;
 	SupportedClass = UWorld::StaticClass();
 	WorldType = EWorldType::Inactive;
 	bInformEngineOfWorld = false;
diff --git a/Engine/Source/Editor/UnrealEd/Private/FileHelpers.cpp b/Engine/Source/Editor/UnrealEd/Private/FileHelpers.cpp
--- a/Engine/Source/Editor/UnrealEd/Private/FileHelpers.cpp
+++ b/Engine/Source/Editor/UnrealEd/Private/FileHelpers.cpp
@@ -753,40 +753,20 @@ static bool SaveAsImplementation( UWorld* InWorld, const FString& DefaultFilenam
 	{
 		SaveFilename = FString();
 		bool bSaveFileLocationSelected = false;
-		if (UEditorEngine::IsUsingWorldAssets())
-		{
-			FString DefaultPackagePath;
-			FPackageName::TryConvertFilenameToLongPackageName(DefaultDirectory / Filename, DefaultPackagePath);
 
-			FString PackageName;
-			bSaveFileLocationSelected = OpenSaveAsDialog(
-				UWorld::StaticClass(),
-				FPackageName::GetLongPackagePath(DefaultPackagePath),
-				FPaths::GetBaseFilename(Filename),
-				PackageName);
+		FString DefaultPackagePath;
+		FPackageName::TryConvertFilenameToLongPackageName(DefaultDirectory / Filename, DefaultPackagePath);
+
+		FString PackageName;
+		bSaveFileLocationSelected = OpenSaveAsDialog(
+			UWorld::StaticClass(),
+			FPackageName::GetLongPackagePath(DefaultPackagePath),
+			FPaths::GetBaseFilename(Filename),
+			PackageName);
 
-			if ( bSaveFileLocationSelected )
-			{
-				SaveFilename = FPackageName::LongPackageNameToFilename(PackageName, FPackageName::GetMapPackageExtension());
-			}
-		}
-		else
-		{
-			bSaveFileLocationSelected = FileDialogHelpers::SaveFile(
-				NSLOCTEXT("UnrealEd", "SaveAs", "Save As").ToString(),
-				FEditorFileUtils::GetFilterString(FI_Save),
-				DefaultDirectory,
-				FPaths::GetCleanFilename(Filename),
-				SaveFilename
-				);
-		}
 		if( bSaveFileLocationSelected )
 		{
-			// Add a map file extension if none was supplied
-			if( FPaths::GetExtension(SaveFilename).IsEmpty() )
-			{
-				SaveFilename = SaveFilename + FPackageName::GetMapPackageExtension();
-			}
+			SaveFilename = FPackageName::LongPackageNameToFilename(PackageName, FPackageName::GetMapPackageExtension());
 
 			FText ErrorMessage;
 			bFilenameIsValid = FEditorFileUtils::IsValidMapFilename(SaveFilename, ErrorMessage);
@@ -2015,104 +1995,51 @@ bool FEditorFileUtils::LoadMap()
 	}
 
 	bool bResult = false;
-	if (UEditorEngine::IsUsingWorldAssets())
-	{
-		static bool bIsDialogOpen = false;
+	static bool bIsDialogOpen = false;
 
-		struct FLocal
-		{
-			static void HandleLevelsChosen(const TArray<FAssetData>& SelectedAssets, bool* OutResult)
-			{
-				bIsDialogOpen = false;
-
-				if ( SelectedAssets.Num() > 0 )
-				{
-					const FAssetData& AssetData = SelectedAssets[0];
-
-					if (!GIsDemoMode)
-					{
-						// If there are any unsaved changes to the current level, see if the user wants to save those first.
-						bool bPromptUserToSave = true;
-						bool bSaveMapPackages = true;
-						bool bSaveContentPackages = true;
-						if (FEditorFileUtils::SaveDirtyPackages(bPromptUserToSave, bSaveMapPackages, bSaveContentPackages) == false)
-						{
-							*OutResult = false;
-							return;
-						}
-					}
-
-					const FString FileToOpen = FPackageName::LongPackageNameToFilename(AssetData.PackageName.ToString(), FPackageName::GetMapPackageExtension());
-					const bool bLoadAsTemplate = false;
-					const bool bShowProgress = true;
-					*OutResult = FEditorFileUtils::LoadMap(FileToOpen, bLoadAsTemplate, bShowProgress);
-				}
-			}
-
-			static void HandleDialogCancelled()
-			{
-				bIsDialogOpen = false;
-			}
-		};
-		
-		if (!bIsDialogOpen)
-		{
-			bIsDialogOpen = true;
-			const bool bAllowMultipleSelection = false;
-			OpenLevelPickingDialog(FOnLevelsChosen::CreateStatic(&FLocal::HandleLevelsChosen, &bResult),
-								   FOnLevelPickingCancelled::CreateStatic(&FLocal::HandleDialogCancelled),
-								   bAllowMultipleSelection);
-		}
-	}
-	else
+	struct FLocal
 	{
-		bool bFilenameIsValid = false;
-		FString DefaultDirectory = FEditorDirectories::Get().GetLastDirectory(ELastDirectory::LEVEL);
-
-		while( !bFilenameIsValid )
+		static void HandleLevelsChosen(const TArray<FAssetData>& SelectedAssets, bool* OutResult)
 		{
-			TArray<FString> OutFiles;
-			if ( FileDialogHelpers::OpenFiles( NSLOCTEXT("UnrealEd", "Open", "Open").ToString(), GetFilterString(FI_Load), DefaultDirectory, EFileDialogFlags::None, OutFiles) )
-			{
-				const FString& FileToOpen = OutFiles[0];
+			bIsDialogOpen = false;
 
-				FText ErrorMessage;
-				bFilenameIsValid = FEditorFileUtils::IsValidMapFilename(FileToOpen, ErrorMessage);
-				if ( !bFilenameIsValid )
-				{
-					// Start the loop over, prompting for load again
-					const FText DisplayFilename = FText::FromString( IFileManager::Get().ConvertToAbsolutePathForExternalAppForRead(*FileToOpen) );
-					FFormatNamedArguments Arguments;
-					Arguments.Add(TEXT("Filename"), DisplayFilename);
-					Arguments.Add(TEXT("LineTerminators"), FText::FromString(LINE_TERMINATOR LINE_TERMINATOR));
-					Arguments.Add(TEXT("ErrorMessage"), ErrorMessage);
-					const FText DisplayMessage = FText::Format( NSLOCTEXT("LoadMap", "InvalidMapName", "Failed to load map {Filename}{LineTerminators}{ErrorMessage}"), Arguments );
-					FMessageDialog::Open( EAppMsgType::Ok, DisplayMessage );
-					continue;
-				}
+			if ( SelectedAssets.Num() > 0 )
+			{
+				const FAssetData& AssetData = SelectedAssets[0];
 
-				if( !GIsDemoMode )
+				if (!GIsDemoMode)
 				{
 					// If there are any unsaved changes to the current level, see if the user wants to save those first.
 					bool bPromptUserToSave = true;
 					bool bSaveMapPackages = true;
 					bool bSaveContentPackages = true;
-					if( FEditorFileUtils::SaveDirtyPackages(bPromptUserToSave, bSaveMapPackages, bSaveContentPackages) == false )
+					if (FEditorFileUtils::SaveDirtyPackages(bPromptUserToSave, bSaveMapPackages, bSaveContentPackages) == false)
 					{
-						// something went wrong or the user pressed cancel.  Return to the editor so the user doesn't lose their changes		
-						return false;
+						*OutResult = false;
+						return;
 					}
 				}
 
-				FEditorDirectories::Get().SetLastDirectory(ELastDirectory::LEVEL, FPaths::GetPath(FileToOpen));
-				bResult = LoadMap( FileToOpen, false, true );
-			}
-			else
-			{
-				// User canceled the open dialog, do not prompt again.
-				break;
+				const FString FileToOpen = FPackageName::LongPackageNameToFilename(AssetData.PackageName.ToString(), FPackageName::GetMapPackageExtension());
+				const bool bLoadAsTemplate = false;
+				const bool bShowProgress = true;
+				*OutResult = FEditorFileUtils::LoadMap(FileToOpen, bLoadAsTemplate, bShowProgress);
 			}
 		}
+
+		static void HandleDialogCancelled()
+		{
+			bIsDialogOpen = false;
+		}
+	};
+		
+	if (!bIsDialogOpen)
+	{
+		bIsDialogOpen = true;
+		const bool bAllowMultipleSelection = false;
+		OpenLevelPickingDialog(FOnLevelsChosen::CreateStatic(&FLocal::HandleLevelsChosen, &bResult),
+								FOnLevelPickingCancelled::CreateStatic(&FLocal::HandleDialogCancelled),
+								bAllowMultipleSelection);
 	}
 
 	return bResult;
@@ -2648,34 +2575,25 @@ static int32 InternalSavePackage( UPackage* PackageToSave, bool& bOutPackageLoca
 		while( NumSkips < NumSkipsBeforeAbort )
 		{
 			FString DefaultLocation = Directory;
-
-			bool bSaveFile = false;
-			if( UEditorEngine::IsUsingWorldAssets() )
+			FString DefaultPackagePath;
+			if (!FPackageName::TryConvertFilenameToLongPackageName(DefaultLocation / FinalPackageFilename, DefaultPackagePath))
 			{
-				FString DefaultPackagePath;
-				if (!FPackageName::TryConvertFilenameToLongPackageName(DefaultLocation / FinalPackageFilename, DefaultPackagePath))
-				{
-					// Original location is invalid; set default location to /Game/Maps
-					DefaultLocation = FPaths::GameContentDir() / TEXT("Maps");
-					ensure(FPackageName::TryConvertFilenameToLongPackageName(DefaultLocation / FinalPackageFilename, DefaultPackagePath));
-				}
+				// Original location is invalid; set default location to /Game/Maps
+				DefaultLocation = FPaths::GameContentDir() / TEXT("Maps");
+				ensure(FPackageName::TryConvertFilenameToLongPackageName(DefaultLocation / FinalPackageFilename, DefaultPackagePath));
+			}
 
-				FString SaveAsPackageName;
-				bSaveFile = OpenSaveAsDialog(
-					UWorld::StaticClass(),
-					FPackageName::GetLongPackagePath(DefaultPackagePath),
-					FPaths::GetBaseFilename(FinalPackageFilename),
-					SaveAsPackageName);
+			FString SaveAsPackageName;
+			bool bSaveFile = OpenSaveAsDialog(
+				UWorld::StaticClass(),
+				FPackageName::GetLongPackagePath(DefaultPackagePath),
+				FPaths::GetBaseFilename(FinalPackageFilename),
+				SaveAsPackageName);
 
-				if (bSaveFile)
-				{
-					// Leave out the extension. It will be added below.
-					FinalPackageFilename = FPackageName::LongPackageNameToFilename(SaveAsPackageName);
-				}
-			}
-			else
+			if (bSaveFile)
 			{
-				bSaveFile = FileDialogHelpers::SaveFile( SavePackageText.ToString(), FileTypes, DefaultLocation, FinalPackageFilename, FinalPackageFilename );
+				// Leave out the extension. It will be added below.
+				FinalPackageFilename = FPackageName::LongPackageNameToFilename(SaveAsPackageName);
 			}
 
 			if( bSaveFile )
diff --git a/Engine/Source/Editor/UnrealEd/Private/Kismet2/BlueprintEditorUtils.cpp b/Engine/Source/Editor/UnrealEd/Private/Kismet2/BlueprintEditorUtils.cpp
--- a/Engine/Source/Editor/UnrealEd/Private/Kismet2/BlueprintEditorUtils.cpp
+++ b/Engine/Source/Editor/UnrealEd/Private/Kismet2/BlueprintEditorUtils.cpp
@@ -1225,6 +1225,17 @@ void FBlueprintEditorUtils::RefreshVariables(UBlueprint* Blueprint)
 	Compiler.RefreshVariables(Blueprint);
 }
 
+void FBlueprintEditorUtils::PreloadBlueprintSpecificData(UBlueprint* Blueprint)
+{
+	TArray<UK2Node*> AllNodes;
+	FBlueprintEditorUtils::GetAllNodesOfClass(Blueprint, AllNodes);
+
+	for( UK2Node* K2Node : AllNodes )
+	{
+		K2Node->PreloadRequiredAssets();
+	}
+}
+
 UClass* FBlueprintEditorUtils::RegenerateBlueprintClass(UBlueprint* Blueprint, UClass* ClassToRegenerate, UObject* PreviousCDO, TArray<UObject*>& ObjLoaded)
 {
 	bool bRegenerated = false;
@@ -2019,6 +2030,7 @@ void FBlueprintEditorUtils::MarkBlueprintAsStructurallyModified(UBlueprint* Blue
 						FKismetCompilerOptions CompileOptions;
 						CompileOptions.CompileType = EKismetCompileType::SkeletonOnly;
 						Compiler.CompileBlueprint(SkelBlueprint, CompileOptions, Results);
+						SkelBlueprint->Status = BS_Dirty;
 
 						SkelBlueprint->BroadcastCompiled();
 
@@ -2043,24 +2055,35 @@ void FBlueprintEditorUtils::MarkBlueprintAsStructurallyModified(UBlueprint* Blue
 
 		BP_SCOPED_COMPILER_EVENT_STAT(EKismetCompilerStats_MarkBlueprintasStructurallyModified);
 
-		TArray<UEdGraph*> AllGraphs;
-		Blueprint->GetAllGraphs(AllGraphs);
-		for (int32 i = 0; i < AllGraphs.Num(); i++)
-		{
-			UK2Node_EditablePinBase* EntryNode = GetEntryNode(AllGraphs[i]);
-			if(UK2Node_Tunnel* TunnelNode = ExactCast<UK2Node_Tunnel>(EntryNode))
-			{
-				// Remove data marking graphs as latent, this will be re-cache'd as needed
-				TunnelNode->MetaData.HasLatentFunctions = INDEX_NONE;
-			}
-		}
-
 		TArray<UClass*> ChildrenOfClass;
 		if (UClass* SkelClass = Blueprint->SkeletonGeneratedClass)
 		{
 			if (!Blueprint->bIsRegeneratingOnLoad)
 			{
-				GetDerivedClasses(SkelClass, ChildrenOfClass, false);
+				if (IsInterfaceBlueprint(Blueprint))
+				{
+					// Find all dependent Blueprints that implement the interface. Note: Using
+					// GetDependentBlueprints() here as the result is cached and thus it should
+					// generally be a faster path than iterating through all loaded Blueprints.
+					TArray<UBlueprint*> DependentBlueprints;
+					GetDependentBlueprints(Blueprint, DependentBlueprints, true);
+					for (UBlueprint* DependentBlueprint : DependentBlueprints)
+					{
+						const bool bDependentBPImplementsInterface = DependentBlueprint->ImplementedInterfaces.ContainsByPredicate([&Blueprint](const FBPInterfaceDescription& InterfaceDesc)
+						{
+							return InterfaceDesc.Interface == Blueprint->GeneratedClass;
+						});
+
+						if (bDependentBPImplementsInterface)
+						{
+							ChildrenOfClass.Add(DependentBlueprint->SkeletonGeneratedClass);
+						}
+					}
+				}
+				else
+				{
+					GetDerivedClasses(SkelClass, ChildrenOfClass, false);
+				}
 			}
 		}
 
@@ -2071,20 +2094,20 @@ void FBlueprintEditorUtils::MarkBlueprintAsStructurallyModified(UBlueprint* Blue
 			FKismetCompilerOptions CompileOptions;
 			CompileOptions.CompileType = EKismetCompileType::SkeletonOnly;
 			Compiler.CompileBlueprint(Blueprint, CompileOptions, Results);
-			Blueprint->Status = BS_Dirty;
 		}
 		UpdateDelegatesInBlueprint(Blueprint);
 
 		FRefreshHelper::SkeletalRecompileChildren(ChildrenOfClass, Blueprint->bIsRegeneratingOnLoad);
 
+		// Call general modification callback as well
+		MarkBlueprintAsModified(Blueprint);
+
 		{
 			BP_SCOPED_COMPILER_EVENT_STAT(EKismetCompilerStats_NotifyBlueprintChanged);
 
 			// Notify any interested parties that the blueprint has changed
 			Blueprint->BroadcastChanged();
 		}
-
-		Blueprint->MarkPackageDirty();
 	}
 }
 
@@ -2183,6 +2206,48 @@ UBlueprint* FBlueprintEditorUtils::FindBlueprintForGraphChecked(const UEdGraph*
 	return Result;
 }
 
+UClass* FBlueprintEditorUtils::GetSkeletonClass(UClass* FromClass)
+{
+	if(UBlueprint* Generator = Cast<UBlueprint>(FromClass->ClassGeneratedBy))
+	{
+		return Generator->SkeletonGeneratedClass;
+	}
+	return nullptr;
+}
+
+UClass* FBlueprintEditorUtils::GetMostUpToDateClass(UClass* FromClass)
+{
+	if(!FromClass || FromClass->HasAnyClassFlags(CLASS_Native))
+	{
+		return FromClass;
+	}
+
+	if(GBlueprintUseCompilationManager)
+	{
+		// It's really not safe/coherent to try and dig out the 'right' class. Things that need the 'most up to date'
+		// version of a class should always be looking at the skeleton:
+		return GetSkeletonClass(FromClass);
+	}
+	else
+	{
+		return FromClass;
+	}
+}
+
+bool FBlueprintEditorUtils::PropertyStillExists(UProperty* Property)
+{
+	if (GBlueprintUseCompilationManager)
+	{
+		if (UClass* UpToDateClass = FBlueprintEditorUtils::GetMostUpToDateClass(Property->GetTypedOuter<UClass>()))
+		{
+			return (UpToDateClass->FindPropertyByName(Property->GetFName()) != nullptr);
+		}
+	}
+
+	// We can't reliably know if the property still exists, but assume that it does
+	return true;
+}
+
 bool FBlueprintEditorUtils::IsGraphNameUnique(UBlueprint* Blueprint, const FName& InName)
 {
 	// Check for any object directly created in the blueprint
@@ -4320,7 +4385,7 @@ void FBlueprintEditorUtils::RenameMemberVariable(UBlueprint* Blueprint, const FN
 
 			{
 				const bool bIsDelegateVar = (Variable.VarType.PinCategory == UEdGraphSchema_K2::PC_MCDelegate);
-				if (UEdGraph* DelegateSignatureGraph = bIsDelegateVar ? FindObject<UEdGraph>(Blueprint, *OldName.ToString()) : nullptr)
+				if (UEdGraph* DelegateSignatureGraph = bIsDelegateVar ? GetDelegateSignatureGraphByName(Blueprint, OldName) : nullptr)
 				{
 					FBlueprintEditorUtils::RenameGraph(DelegateSignatureGraph, NewName.ToString());
 
@@ -4618,7 +4683,7 @@ FName FBlueprintEditorUtils::DuplicateVariable(UBlueprint* InBlueprint, const US
 				{
 					// if there is a property for variable, it means the original default value was already copied, so it can be safely overridden
 					Variable.DefaultValue.Empty();
-					TargetProperty->ExportTextItem(NewVar.DefaultValue, OldPropertyAddr, OldPropertyAddr, nullptr, PPF_None);
+					TargetProperty->ExportTextItem(NewVar.DefaultValue, OldPropertyAddr, OldPropertyAddr, nullptr, PPF_SerializedAsImportText);
 				}
 			}
 
@@ -4636,8 +4701,6 @@ FName FBlueprintEditorUtils::DuplicateVariable(UBlueprint* InBlueprint, const US
 			{
 				FunctionEntry->Modify();
 
-				FBPVariableDescription& Variable = *LocalVariable;
-
 				NewVar = DuplicateVariableDescription(InBlueprint, *LocalVariable);
 
 				// Add the new variable
@@ -4694,6 +4757,7 @@ bool FBlueprintEditorUtils::AddLocalVariable(UBlueprint* Blueprint, UEdGraph* In
 		NewVar.VarName = InNewVarName;
 		NewVar.VarGuid = FGuid::NewGuid();
 		NewVar.VarType = InNewVarType;
+		NewVar.PropertyFlags |= CPF_BlueprintVisible;
 		NewVar.FriendlyName = FName::NameToDisplayString( NewVar.VarName.ToString(), (NewVar.VarType.PinCategory == K2Schema->PC_Boolean) ? true : false );
 		NewVar.Category = K2Schema->VR_DefaultCategory;
 		NewVar.DefaultValue = DefaultValue;
@@ -6037,43 +6101,77 @@ static void ConformInterfaceByName(UBlueprint* Blueprint, FBPInterfaceDescriptio
 		}
 
 		// Cache off the graph names for this interface, for easier searching
-		TArray<FName> GraphNames;
+		TMap<FName, const UEdGraph*> InterfaceFunctionGraphs;
 		for (int32 GraphIndex = 0; GraphIndex < CurrentInterfaceDesc.Graphs.Num(); GraphIndex++)
 		{
 			const UEdGraph* CurrentGraph = CurrentInterfaceDesc.Graphs[GraphIndex];
 			if( CurrentGraph )
 			{
-				GraphNames.AddUnique(CurrentGraph->GetFName());
+				InterfaceFunctionGraphs.Add(CurrentGraph->GetFName()) = CurrentGraph;
 			}
 		}
 
+		// If this is a Blueprint interface, redirect to the skeleton class for function iteration
+		const UClass* InterfaceClass = CurrentInterfaceDesc.Interface;
+		if (InterfaceClass && InterfaceClass->ClassGeneratedBy)
+		{
+			InterfaceClass = CastChecked<UBlueprint>(InterfaceClass->ClassGeneratedBy)->SkeletonGeneratedClass;
+		}
+
 		// Iterate over all the functions in the interface, and create graphs that are in the interface, but missing in the blueprint
-		for (TFieldIterator<UFunction> FunctionIter(CurrentInterfaceDesc.Interface, EFieldIteratorFlags::IncludeSuper); FunctionIter; ++FunctionIter)
+		for (TFieldIterator<UFunction> FunctionIter(InterfaceClass, EFieldIteratorFlags::IncludeSuper); FunctionIter; ++FunctionIter)
 		{
 			UFunction* Function = *FunctionIter;
 			const FName FunctionName = Function->GetFName();
 			if(!VariableNamesUsedInBlueprint.Contains(FunctionName))
 			{
-				if( UEdGraphSchema_K2::CanKismetOverrideFunction(Function) && !UEdGraphSchema_K2::FunctionCanBePlacedAsEvent(Function) && !GraphNames.Contains(FunctionName) )
+				if( UEdGraphSchema_K2::CanKismetOverrideFunction(Function) && !UEdGraphSchema_K2::FunctionCanBePlacedAsEvent(Function) )
 				{
-					// interface methods initially create EventGraph stubs, so we need
-					// to make sure we remove that entry so the new graph doesn't conflict (don't
-					// worry, these are regenerated towards the end of a compile)
-					for (UEdGraph* GraphStub : Blueprint->EventGraphs)
+					if (const UEdGraph** FunctionGraphPtr = InterfaceFunctionGraphs.Find(FunctionName))
 					{
-						if (GraphStub->GetFName() == FunctionName)
+						const bool bIsConstInterfaceFunction = (Function->FunctionFlags & FUNC_Const) != 0;
+
+						// Sync the 'const' attribute of the implementation with the interface function, in case it has been changed
+						TArray<UK2Node_FunctionEntry*> FunctionEntryNodes;
+						(*FunctionGraphPtr)->GetNodesOfClass<UK2Node_FunctionEntry>(FunctionEntryNodes);
+						for (UK2Node_FunctionEntry* FunctionEntryNode : FunctionEntryNodes)
 						{
-							FBlueprintEditorUtils::RemoveGraph(Blueprint, GraphStub, EGraphRemoveFlags::MarkTransient);
+							const bool bIsImplementedAsConstFunction = (FunctionEntryNode->GetExtraFlags() & FUNC_Const) != 0;
+							if (bIsImplementedAsConstFunction != bIsConstInterfaceFunction)
+							{
+								FunctionEntryNode->Modify();
+								if (bIsConstInterfaceFunction)
+								{
+									FunctionEntryNode->AddExtraFlags(FUNC_Const);
+								}
+								else
+								{
+									FunctionEntryNode->ClearExtraFlags(FUNC_Const);
+								}
+							}
 						}
 					}
+					else
+					{
+						// interface methods initially create EventGraph stubs, so we need
+						// to make sure we remove that entry so the new graph doesn't conflict (don't
+						// worry, these are regenerated towards the end of a compile)
+						for (UEdGraph* GraphStub : Blueprint->EventGraphs)
+						{
+							if (GraphStub->GetFName() == FunctionName)
+							{
+								FBlueprintEditorUtils::RemoveGraph(Blueprint, GraphStub, EGraphRemoveFlags::MarkTransient);
+							}
+						}
 
-					// Check to see if we already have implemented 
-					UEdGraph* NewGraph = FBlueprintEditorUtils::CreateNewGraph(Blueprint, FunctionName, UEdGraph::StaticClass(), UEdGraphSchema_K2::StaticClass());
-					NewGraph->bAllowDeletion = false;
-					NewGraph->InterfaceGuid = FBlueprintEditorUtils::FindInterfaceFunctionGuid(Function, CurrentInterfaceDesc.Interface);
-					CurrentInterfaceDesc.Graphs.Add(NewGraph);
+						// Check to see if we already have implemented 
+						UEdGraph* NewGraph = FBlueprintEditorUtils::CreateNewGraph(Blueprint, FunctionName, UEdGraph::StaticClass(), UEdGraphSchema_K2::StaticClass());
+						NewGraph->bAllowDeletion = false;
+						NewGraph->InterfaceGuid = FBlueprintEditorUtils::FindInterfaceFunctionGuid(Function, CurrentInterfaceDesc.Interface);
+						CurrentInterfaceDesc.Graphs.Add(NewGraph);
 
-					FBlueprintEditorUtils::AddInterfaceGraph(Blueprint, NewGraph, CurrentInterfaceDesc.Interface);
+						FBlueprintEditorUtils::AddInterfaceGraph(Blueprint, NewGraph, CurrentInterfaceDesc.Interface);
+					}
 				}
 			}
 			else
@@ -6548,8 +6646,24 @@ FName FBlueprintEditorUtils::FindUniqueKismetName(const UBlueprint* InBlueprint,
 	int32 Count = 0;
 	FString KismetName;
 	FString BaseName = InBaseName;
-
 	TSharedPtr<FKismetNameValidator> NameValidator = MakeShareable(new FKismetNameValidator(InBlueprint, NAME_None, InScope));
+
+	// Clean up BaseName to not contain any invalid characters, which will mean we can never find a legal name no matter how many numbers we add
+	if (NameValidator->IsValid(BaseName) == EValidatorResult::ContainsInvalidCharacters)
+	{
+		for (TCHAR& TestChar : BaseName)
+		{
+			for (TCHAR BadChar : UE_BLUEPRINT_INVALID_NAME_CHARACTERS)
+			{
+				if (TestChar == BadChar)
+				{
+					TestChar = TEXT('_');
+					break;
+				}
+			}
+		}
+	}
+
 	while(NameValidator->IsValid(KismetName) != EValidatorResult::Ok)
 	{
 		// Calculate the number of digits in the number, adding 2 (1 extra to correctly count digits, another to account for the '_' that will be added to the name
@@ -7538,21 +7652,53 @@ void FBlueprintEditorUtils::PostEditChangeBlueprintActors(UBlueprint* Blueprint,
 	}
 }
 
-bool FBlueprintEditorUtils::IsPropertyReadOnlyInCurrentBlueprint(const UBlueprint* Blueprint, const UProperty* Property)
+FBlueprintEditorUtils::EPropertyWritableState FBlueprintEditorUtils::IsPropertyWritableInBlueprint(const UBlueprint* Blueprint, const UProperty* Property)
 {
-	if (Property && Property->HasAnyPropertyFlags(CPF_BlueprintReadOnly))
+	if (Property)
 	{
-		return true;
+		if (!Property->HasAnyPropertyFlags(CPF_BlueprintVisible))
+		{
+			return EPropertyWritableState::NotBlueprintVisible;
+		}
+		if (Property->HasAnyPropertyFlags(CPF_BlueprintReadOnly))
+		{
+			return EPropertyWritableState::BlueprintReadOnly;
+		}
+		if (Property->GetBoolMetaData(FBlueprintMetadata::MD_Private))
+		{
+			const UClass* OwningClass = CastChecked<UClass>(Property->GetOuter());
+			if (OwningClass->ClassGeneratedBy != Blueprint)
+			{
+				return EPropertyWritableState::Private;
+			}
+		}
 	}
-	if (Property && Property->GetBoolMetaData(FBlueprintMetadata::MD_Private))
+	return EPropertyWritableState::Writable;
+}
+
+FBlueprintEditorUtils::EPropertyReadableState FBlueprintEditorUtils::IsPropertyReadableInBlueprint(const UBlueprint* Blueprint, const UProperty* Property)
+{
+	if (Property)
 	{
-		const UClass* OwningClass = CastChecked<UClass>(Property->GetOuter());
-		if(OwningClass->ClassGeneratedBy != Blueprint)
+		if (!Property->HasAnyPropertyFlags(CPF_BlueprintVisible))
 		{
-			return true;
+			return EPropertyReadableState::NotBlueprintVisible;
+		}
+		if (Property->GetBoolMetaData(FBlueprintMetadata::MD_Private))
+		{
+			const UClass* OwningClass = CastChecked<UClass>(Property->GetOuter());
+			if (OwningClass->ClassGeneratedBy != Blueprint)
+			{
+				return EPropertyReadableState::Private;
+			}
 		}
 	}
-	return false;
+	return EPropertyReadableState::Readable;
+}
+
+bool FBlueprintEditorUtils::IsPropertyReadOnlyInCurrentBlueprint(const UBlueprint* Blueprint, const UProperty* Property)
+{
+	return (IsPropertyWritableInBlueprint(Blueprint, Property) != EPropertyWritableState::Writable);
 }
 
 void FBlueprintEditorUtils::FindAndSetDebuggableBlueprintInstances()
@@ -7770,7 +7916,7 @@ bool FBlueprintEditorUtils::PropertyValueFromString_Direct(const UProperty* Prop
 		else if (Property->IsA(UTextProperty::StaticClass()))
 		{
 			FStringOutputDevice ImportError;
-			const TCHAR* EndOfParsedBuff = Property->ImportText(*StrValue, DirectValue, 0, nullptr, &ImportError);
+			const TCHAR* EndOfParsedBuff = Property->ImportText(*StrValue, DirectValue, PPF_SerializedAsImportText, nullptr, &ImportError);
 			bParseSucceeded = EndOfParsedBuff && ImportError.IsEmpty();
 		}
 		else
@@ -7781,7 +7927,7 @@ bool FBlueprintEditorUtils::PropertyValueFromString_Direct(const UProperty* Prop
 				: *StrValue;
 
 			FStringOutputDevice ImportError;
-			const TCHAR* EndOfParsedBuff = Property->ImportText(*StrValue, DirectValue, 0, nullptr, &ImportError);
+			const TCHAR* EndOfParsedBuff = Property->ImportText(*StrValue, DirectValue, PPF_SerializedAsImportText, nullptr, &ImportError);
 			bParseSucceeded = EndOfParsedBuff && ImportError.IsEmpty();
 		}
 	}
@@ -7829,7 +7975,7 @@ bool FBlueprintEditorUtils::PropertyValueFromString_Direct(const UProperty* Prop
 			bParseSucceeded = FStructureEditorUtils::Fill_MakeStructureDefaultValue(Cast<const UUserDefinedStruct>(Struct), DirectValue);
 
 			FStringOutputDevice ImportError;
-			const TCHAR* EndOfParsedBuff = StructProperty->ImportText(StrValue.IsEmpty() ? TEXT("()") : *StrValue, DirectValue, 0, nullptr, &ImportError);
+			const TCHAR* EndOfParsedBuff = StructProperty->ImportText(StrValue.IsEmpty() ? TEXT("()") : *StrValue, DirectValue, PPF_SerializedAsImportText, nullptr, &ImportError);
 			bParseSucceeded &= EndOfParsedBuff && ImportError.IsEmpty();
 		}
 	}
@@ -7895,7 +8041,7 @@ bool FBlueprintEditorUtils::PropertyValueToString_Direct(const UProperty* Proper
 			DefaultValue = StructOnScope.GetStructMemory();
 		}
 		
-		bSuccedded = Property->ExportText_Direct(OutForm, DirectValue, DefaultValue, nullptr, 0);
+		bSuccedded = Property->ExportText_Direct(OutForm, DirectValue, DefaultValue, nullptr, PPF_SerializedAsImportText);
 	}
 	return bSuccedded;
 }
diff --git a/Engine/Source/Editor/UnrealEd/Private/Kismet2/EnumEditorUtils.cpp b/Engine/Source/Editor/UnrealEd/Private/Kismet2/EnumEditorUtils.cpp
--- a/Engine/Source/Editor/UnrealEd/Private/Kismet2/EnumEditorUtils.cpp
+++ b/Engine/Source/Editor/UnrealEd/Private/Kismet2/EnumEditorUtils.cpp
@@ -15,6 +15,7 @@
 #include "Kismet2/BlueprintEditorUtils.h"
 #include "Internationalization/TextNamespaceUtil.h"
 #include "Internationalization/TextPackageNamespaceUtil.h"
+#include "ScopedTransaction.h"
 
 #define LOCTEXT_NAMESPACE "Enum"
 
@@ -97,6 +98,8 @@ void FEnumEditorUtils::AddNewEnumeratorForUserDefinedEnum(UUserDefinedEnum* Enum
 		return;
 	}
 
+	const FScopedTransaction Transaction(NSLOCTEXT("EnumEditor", "AddNewEnumerator", "Add Enumerator"));
+
 	PrepareForChange(Enum);
 
 	TArray<TPair<FName, int64>> OldNames, Names;
@@ -128,6 +131,8 @@ void FEnumEditorUtils::RemoveEnumeratorFromUserDefinedEnum(UUserDefinedEnum* Enu
 		return;
 	}
 
+	const FScopedTransaction Transaction(NSLOCTEXT("EnumEditor", "RemoveEnumerator", "Remove Enumerator"));
+
 	PrepareForChange(Enum);
 
 	TArray<TPair<FName, int64>> OldNames, Names;
@@ -158,6 +163,8 @@ void FEnumEditorUtils::SetEnumeratorBitflagsTypeState(UUserDefinedEnum* Enum, bo
 {
 	if (Enum)
 	{
+		const FScopedTransaction Transaction(NSLOCTEXT("EnumEditor", "SetEnumeratorBitflagsTypeState", "Set Bitflag Type State"));
+
 		PrepareForChange(Enum);
 
 		if (bBitflagsType)
@@ -185,6 +192,8 @@ void FEnumEditorUtils::MoveEnumeratorInUserDefinedEnum(UUserDefinedEnum* Enum, i
 		return;
 	}
 
+	const FScopedTransaction Transaction(NSLOCTEXT("EnumEditor", "MoveEnumeratorInUserDefinedEnum", "Reorder Enumerator"));
+
 	PrepareForChange(Enum);
 
 	TArray<TPair<FName, int64>> OldNames, Names;
@@ -252,9 +261,10 @@ public:
 	}
 };
 
-void FEnumEditorUtils::PrepareForChange(const UUserDefinedEnum* Enum)
+void FEnumEditorUtils::PrepareForChange(UUserDefinedEnum* Enum)
 {
 	FEnumEditorManager::Get().PreChange(Enum, EEnumEditorChangeInfo::Changed);
+	Enum->Modify();
 }
 
 void FEnumEditorUtils::BroadcastChanges(const UUserDefinedEnum* Enum, const TArray<TPair<FName, int64>>& OldNames, bool bResolveData)
@@ -446,6 +456,8 @@ bool FEnumEditorUtils::SetEnumeratorDisplayName(UUserDefinedEnum* Enum, int32 En
 	{
 		if (IsEnumeratorDisplayNameValid(Enum, EnumeratorIndex, NewDisplayName))
 		{
+			const FScopedTransaction Transaction(NSLOCTEXT("EnumEditor", "SetEnumeratorDisplayName", "Set Display Name"));
+
 			PrepareForChange(Enum);
 			
 			const FName EnumEntryName = *Enum->GetNameStringByIndex(EnumeratorIndex);
diff --git a/Engine/Source/Editor/UnrealEd/Private/Kismet2/Kismet2.cpp b/Engine/Source/Editor/UnrealEd/Private/Kismet2/Kismet2.cpp
--- a/Engine/Source/Editor/UnrealEd/Private/Kismet2/Kismet2.cpp
+++ b/Engine/Source/Editor/UnrealEd/Private/Kismet2/Kismet2.cpp
@@ -925,12 +925,6 @@ void FKismetEditorUtilities::CompileBlueprint(UBlueprint* BlueprintObj, EBluepri
 		}
 	}
 
-	// Default Values are now set in CDO. And these copies could be soon obsolete, so better to reset them.
-	for(int VarIndex = 0; VarIndex < BlueprintObj->NewVariables.Num(); ++VarIndex)
-	{
-		BlueprintObj->NewVariables[VarIndex].DefaultValue.Empty();
-	}
-
 	if (!bLetReinstancerRefreshDependBP && (bIsInterface || !BlueprintObj->bIsRegeneratingOnLoad) && !bSkipReinstancing)
 	{
 		BP_SCOPED_COMPILER_EVENT_STAT(EKismetCompilerStats_RefreshDependentBlueprints);
diff --git a/Engine/Source/Editor/UnrealEd/Private/Kismet2/Kismet2NameValidators.cpp b/Engine/Source/Editor/UnrealEd/Private/Kismet2/Kismet2NameValidators.cpp
--- a/Engine/Source/Editor/UnrealEd/Private/Kismet2/Kismet2NameValidators.cpp
+++ b/Engine/Source/Editor/UnrealEd/Private/Kismet2/Kismet2NameValidators.cpp
@@ -25,28 +25,38 @@ TSharedPtr<INameValidatorInterface> FNameValidatorFactory::MakeValidator(UEdGrap
 }
 
 
-FString INameValidatorInterface::GetErrorString(const FString& Name, EValidatorResult ErrorCode)
+FText INameValidatorInterface::GetErrorText(const FString& Name, EValidatorResult ErrorCode)
 {
-	FString ErrorText;
-
+	FText ErrorText;
 	switch (ErrorCode)
 	{
-		case EmptyName:
-			ErrorText = LOCTEXT("EmptyName_Error", "Name cannot be empty.").ToString();
-			break;
-
-		case AlreadyInUse:
-			ErrorText = FString::Printf( *LOCTEXT("AlreadyInUse_Error", "\"%s\" is already in use.").ToString(), *Name);
-			break;
-
-		case ExistingName:
-			ErrorText = LOCTEXT("ExistingName_Error", "Name cannot be the same as the existing name.").ToString();
-			break;
-		case LocallyInUse:
-		case TooLong:
-		case Ok:
-			break;
+	case EValidatorResult::EmptyName:
+		ErrorText = LOCTEXT("EmptyName_Error", "Name cannot be empty.");
+		break;
+	case EValidatorResult::AlreadyInUse:
+		ErrorText = LOCTEXT("AlreadyInUse_Error", "Name is already in use.");
+		break;
+	case EValidatorResult::ExistingName:
+		ErrorText = LOCTEXT("ExistingName_Error", "Name cannot be the same as the existing name.");
+		break;
+	case EValidatorResult::ContainsInvalidCharacters:
+		ErrorText = LOCTEXT("ContainsInvalidCharacters_Error", "Name cannot contain invalid characters.");
+		FName::IsValidXName(Name, UE_BLUEPRINT_INVALID_NAME_CHARACTERS, /*out*/ &ErrorText);
+		break;
+	case EValidatorResult::TooLong:
+		ErrorText = LOCTEXT("NameTooLong_Error", "Names must have fewer than 100 characters!");
+		break;
+	case EValidatorResult::LocallyInUse:
+		ErrorText  = LOCTEXT("LocallyInUse_Error", "Conflicts with another another object in the same scope!");
+		break;
+	case EValidatorResult::Ok:
+		break;
+	default:
+		ErrorText = LOCTEXT("UnknownError", "Unknown error");
+		check(false);
+		break;
 	}
+
 	return ErrorText;
 }
 
@@ -112,6 +122,10 @@ EValidatorResult FKismetNameValidator::IsValid(const FString& Name, bool /*bOrig
 	{
 		return EValidatorResult::TooLong;
 	}
+	else if (!FName::IsValidXName(Name, UE_BLUEPRINT_INVALID_NAME_CHARACTERS))
+	{
+		return EValidatorResult::ContainsInvalidCharacters;
+	}
 
 	// If not defined in name table, not current graph name
 	return IsValid( FName(*Name) );
diff --git a/Engine/Source/Editor/UnrealEd/Private/Kismet2/StructureEditorUtils.cpp b/Engine/Source/Editor/UnrealEd/Private/Kismet2/StructureEditorUtils.cpp
--- a/Engine/Source/Editor/UnrealEd/Private/Kismet2/StructureEditorUtils.cpp
+++ b/Engine/Source/Editor/UnrealEd/Private/Kismet2/StructureEditorUtils.cpp
@@ -12,6 +12,7 @@
 #include "EdMode.h"
 #include "ScopedTransaction.h"
 #include "EdGraphSchema_K2.h"
+#include "UserDefinedStructure/UserDefinedStructEditorData.h"
 #include "Kismet2/BlueprintEditorUtils.h"
 #include "Editor/UnrealEd/Public/Kismet2/CompilerResultsLog.h"
 #include "Editor/KismetCompiler/Public/KismetCompilerModule.h"
@@ -658,6 +659,26 @@ const TArray<FStructVariableDescription>* FStructureEditorUtils::GetVarDescPtr(c
 	return Struct->EditorData ? &CastChecked<const UUserDefinedStructEditorData>(Struct->EditorData)->VariablesDescriptions : nullptr;
 }
 
+FStructVariableDescription* FStructureEditorUtils::GetVarDescByGuid(UUserDefinedStruct* Struct, FGuid VarGuid)
+{
+	if (Struct)
+	{
+		TArray<FStructVariableDescription>* VarDescArray = GetVarDescPtr(Struct);
+		return VarDescArray ? VarDescArray->FindByPredicate(FFindByGuidHelper<FStructVariableDescription>(VarGuid)) : nullptr;
+	}
+	return nullptr;
+}
+
+const FStructVariableDescription* FStructureEditorUtils::GetVarDescByGuid(const UUserDefinedStruct* Struct, FGuid VarGuid)
+{
+	if (Struct)
+	{
+		const TArray<FStructVariableDescription>* VarDescArray = GetVarDescPtr(Struct);
+		return VarDescArray ? VarDescArray->FindByPredicate(FFindByGuidHelper<FStructVariableDescription>(VarGuid)) : nullptr;
+	}
+	return nullptr;
+}
+
 FString FStructureEditorUtils::GetTooltip(const UUserDefinedStruct* Struct)
 {
 	const UUserDefinedStructEditorData* StructEditorData = Struct ? Cast<const UUserDefinedStructEditorData>(Struct->EditorData) : NULL;
diff --git a/Engine/Source/Editor/UnrealEd/Private/UnrealEdEngine.cpp b/Engine/Source/Editor/UnrealEd/Private/UnrealEdEngine.cpp
--- a/Engine/Source/Editor/UnrealEd/Private/UnrealEdEngine.cpp
+++ b/Engine/Source/Editor/UnrealEd/Private/UnrealEdEngine.cpp
@@ -156,11 +156,15 @@ void UUnrealEdEngine::Init(IEngineLoop* InEngineLoop)
 		UEditorExperimentalSettings const* ExperimentalSettings = GetDefault<UEditorExperimentalSettings>();
 		UCookerSettings const* CookerSettings = GetDefault<UCookerSettings>();
 		ECookInitializationFlags BaseCookingFlags = ECookInitializationFlags::AutoTick | ECookInitializationFlags::AsyncSave;
-		const ECookInitializationFlags IterativeFlags = ECookInitializationFlags::Iterative | (CookerSettings->bUseAssetRegistryForIteration ? ECookInitializationFlags::IterateOnAssetRegistry : ECookInitializationFlags::IterateOnHash);
-		BaseCookingFlags |= CookerSettings->bIterativeCookingForLaunchOn ? IterativeFlags : ECookInitializationFlags::None;
 		BaseCookingFlags |= CookerSettings->bEnableBuildDDCInBackground ? ECookInitializationFlags::BuildDDCInBackground : ECookInitializationFlags::None;
 
-
+		if (CookerSettings->bIterativeCookingForLaunchOn)
+		{
+			BaseCookingFlags |= ECookInitializationFlags::Iterative;
+			BaseCookingFlags |= CookerSettings->bIgnoreIniSettingsOutOfDateForIteration ? ECookInitializationFlags::IgnoreIniSettingsOutOfDate : ECookInitializationFlags::None;
+			BaseCookingFlags |= CookerSettings->bIgnoreScriptPackagesOutOfDateForIteration ? ECookInitializationFlags::IgnoreScriptPackagesOutOfDate : ECookInitializationFlags::None;
+		}
+		
 		if (CookerSettings->bEnableCookOnTheSide)
 		{
 			if ( ExperimentalSettings->bSharedCookedBuilds )
diff --git a/Engine/Source/Editor/UnrealEd/Private/UserDefinedStructEditorData.cpp b/Engine/Source/Editor/UnrealEd/Private/UserDefinedStructEditorData.cpp
--- a/Engine/Source/Editor/UnrealEd/Private/UserDefinedStructEditorData.cpp
+++ b/Engine/Source/Editor/UnrealEd/Private/UserDefinedStructEditorData.cpp
@@ -7,6 +7,7 @@
 #include "Kismet2/StructureEditorUtils.h"
 #include "Kismet2/BlueprintEditorUtils.h"
 #include "Blueprint/BlueprintSupport.h"
+#include "Editor.h"
 
 #define LOCTEXT_NAMESPACE "UserDefinedStructEditorData"
 
@@ -56,10 +57,25 @@ UUserDefinedStruct* UUserDefinedStructEditorData::GetOwnerStruct() const
 	return Cast<UUserDefinedStruct>(GetOuter());
 }
 
+void UUserDefinedStructEditorData::PostUndo(bool bSuccess)
+{
+	GEditor->UnregisterForUndo(this);
+	// TODO: In the undo case we might want to flip the change type since an add is now a remove and vice versa
+	FStructureEditorUtils::OnStructureChanged(GetOwnerStruct(), CachedStructureChange);
+	CachedStructureChange = FStructureEditorUtils::Unknown;
+}
+
+void UUserDefinedStructEditorData::ConsolidatedPostEditUndo(const FStructureEditorUtils::EStructureEditorChangeInfo TransactedStructureChange)
+{
+	ensure(CachedStructureChange == FStructureEditorUtils::Unknown);
+	CachedStructureChange = TransactedStructureChange;
+	GEditor->RegisterForUndo(this);
+}
+
 void UUserDefinedStructEditorData::PostEditUndo()
 {
 	Super::PostEditUndo();
-	FStructureEditorUtils::OnStructureChanged(GetOwnerStruct());
+	ConsolidatedPostEditUndo(FStructureEditorUtils::Unknown);
 }
 
 class FStructureTransactionAnnotation : public ITransactionObjectAnnotation
@@ -89,24 +105,24 @@ TSharedPtr<ITransactionObjectAnnotation> UUserDefinedStructEditorData::GetTransa
 void UUserDefinedStructEditorData::PostEditUndo(TSharedPtr<ITransactionObjectAnnotation> TransactionAnnotation)
 {
 	Super::PostEditUndo();
-	FStructureEditorUtils::EStructureEditorChangeInfo ActiveChange = FStructureEditorUtils::Unknown;
+	FStructureEditorUtils::EStructureEditorChangeInfo TransactedStructureChange = FStructureEditorUtils::Unknown;
 
 	if (TransactionAnnotation.IsValid())
 	{
 		TSharedPtr<FStructureTransactionAnnotation> StructAnnotation = StaticCastSharedPtr<FStructureTransactionAnnotation>(TransactionAnnotation);
 		if (StructAnnotation.IsValid())
 		{
-			ActiveChange = StructAnnotation->GetActiveChange();
+			TransactedStructureChange = StructAnnotation->GetActiveChange();
 		}
 	}
-	FStructureEditorUtils::OnStructureChanged(GetOwnerStruct(), ActiveChange);
+	ConsolidatedPostEditUndo(TransactedStructureChange);
 }
 
 void UUserDefinedStructEditorData::PostLoadSubobjects(FObjectInstancingGraph* OuterInstanceGraph)
 {
 	Super::PostLoadSubobjects(OuterInstanceGraph);
 
-	for (auto& VarDesc : VariablesDescriptions)
+	for (FStructVariableDescription& VarDesc : VariablesDescriptions)
 	{
 		VarDesc.bInvalidMember = !FStructureEditorUtils::CanHaveAMemberVariableOfType(GetOwnerStruct(), VarDesc.ToPinType());
 	}
@@ -138,7 +154,7 @@ void UUserDefinedStructEditorData::RecreateDefaultInstance(FString* OutLog)
 			UProperty* Property = *It;
 			if (Property)
 			{
-				auto VarDesc = VariablesDescriptions.FindByPredicate(FStructureEditorUtils::FFindByNameHelper<FStructVariableDescription>(Property->GetFName()));
+				FStructVariableDescription* VarDesc = VariablesDescriptions.FindByPredicate(FStructureEditorUtils::FFindByNameHelper<FStructVariableDescription>(Property->GetFName()));
 				if (VarDesc && !VarDesc->CurrentDefaultValue.IsEmpty())
 				{
 					if (!FBlueprintEditorUtils::PropertyValueFromString(Property, VarDesc->CurrentDefaultValue, StructData))
diff --git a/Engine/Source/Editor/UnrealEd/Public/Commandlets/AssetRegistryGenerator.h b/Engine/Source/Editor/UnrealEd/Public/Commandlets/AssetRegistryGenerator.h
--- a/Engine/Source/Editor/UnrealEd/Public/Commandlets/AssetRegistryGenerator.h
+++ b/Engine/Source/Editor/UnrealEd/Public/Commandlets/AssetRegistryGenerator.h
@@ -52,8 +52,9 @@ public:
 	 * @param IdenticalCookedPackages list of cooked packages that have not changed
 	 * @param IdenticalUncookedPackages list of uncooked packages that have not changed. These were filtered out by platform or editor only
 	 * @param bRecurseModifications if true, modified packages are recursed to X in X->Y->Z chains. Otherwise, only Y and Z are seen as modified
+	 * @param bRecurseModifications if true, modified script/c++ packages are recursed, if false only asset references are recursed
 	 */
-	void ComputePackageDifferences(TSet<FName>& ModifiedPackages, TSet<FName>& NewPackages, TSet<FName>& RemovedPackages, TSet<FName>& IdenticalCookedPackages, TSet<FName>& IdenticalUncookedPackages, bool bRecurseModifications);
+	void ComputePackageDifferences(TSet<FName>& ModifiedPackages, TSet<FName>& NewPackages, TSet<FName>& RemovedPackages, TSet<FName>& IdenticalCookedPackages, TSet<FName>& IdenticalUncookedPackages, bool bRecurseModifications, bool bRecurseScriptModifications);
 
 	/**
 	 * GenerateChunkManifest 
@@ -279,9 +280,6 @@ private:
 		return RegistryChunkIDs;
 	}
 
-	/** Updates the Package Source hash based on dependency information */
-	void UpdatePackageSourceHashes();
-
 	/** Generate manifest for a single package */
 	void GenerateChunkManifestForPackage(const FName& PackageFName, const FString& PackagePathName, const FString& SandboxFilename, const FString& LastLoadedMapName, FSandboxPlatformFile* InSandboxFile);
 
diff --git a/Engine/Source/Editor/UnrealEd/Public/CookerSettings.h b/Engine/Source/Editor/UnrealEd/Public/CookerSettings.h
--- a/Engine/Source/Editor/UnrealEd/Public/CookerSettings.h
+++ b/Engine/Source/Editor/UnrealEd/Public/CookerSettings.h
@@ -36,9 +36,13 @@ public:
 	UPROPERTY(GlobalConfig, EditAnywhere, Category = Cooker, meta = (DisplayName = "Iterative cooking for builds launched from the editor (launch on)"))
 	bool bIterativeCookingForLaunchOn;
 
-	/** Use asset registry for iteration instead of older system. Currently experimental, will be default when finished */
-	UPROPERTY(GlobalConfig)
-	bool bUseAssetRegistryForIteration;
+	/** Ignore ini changes when doing iterative cooking, either in editor or out of editor */
+	UPROPERTY(GlobalConfig, EditAnywhere, Category = Cooker, AdvancedDisplay)
+	bool bIgnoreIniSettingsOutOfDateForIteration;
+
+	/** Ignore native header file changes when doing iterative cooking, either in editor or out of editor */
+	UPROPERTY(GlobalConfig, EditAnywhere, Category = Cooker, AdvancedDisplay)
+	bool bIgnoreScriptPackagesOutOfDateForIteration;
 
 	/** Whether or not to compile Blueprints in development mode when cooking. */
 	UPROPERTY(GlobalConfig, EditAnywhere, Category = Cooker, AdvancedDisplay)
