Copying //UE4/Orion-Staging to //UE4/Main (Source: //Orion/Dev-General @ 3564337)

diff --git a/Engine/Source/Editor/UnrealEd/Private/Commandlets/CompileAllBlueprintsCommandlet.cpp b/Engine/Source/Editor/UnrealEd/Private/Commandlets/CompileAllBlueprintsCommandlet.cpp
new file mode 100644
--- /dev/null
+++ b/Engine/Source/Editor/UnrealEd/Private/Commandlets/CompileAllBlueprintsCommandlet.cpp
@@ -0,0 +1,350 @@
+// Copyright 1998-2017 Epic Games, Inc. All Rights Reserved.
+
+#include "CompileAllBlueprintsCommandlet.h"
+
+#include "AssetRegistryModule.h"
+#include "CompilerResultsLog.h"
+#include "EngineUtils.h"
+#include "ISourceControlModule.h"
+#include "Misc/FileHelper.h"
+#include "Paths.h"
+
+DEFINE_LOG_CATEGORY_STATIC(LogCompileAllBlueprintsCommandlet, Log, All);
+
+#define KISMET_COMPILER_MODULENAME "KismetCompiler"
+
+UCompileAllBlueprintsCommandlet::UCompileAllBlueprintsCommandlet(const FObjectInitializer& ObjectInitializer)
+: Super(ObjectInitializer)
+{
+	bResultsOnly = false;
+	bCompileSkeletonOnly = false;
+	bCookedOnly = false;
+	bDirtyOnly = false;
+	bSimpleAssetList = false;
+
+	TotalNumFailedLoads = 0;
+	TotalNumFatalIssues = 0;
+	TotalNumWarnings = 0;
+}
+
+int32 UCompileAllBlueprintsCommandlet::Main(const FString& Params)
+{
+	InitCommandLine(Params);
+	InitKismetBlueprintCompiler();
+
+	BuildBlueprintAssetList();
+	BuildBlueprints();
+	
+	LogResults();
+
+	return (TotalNumFatalIssues + TotalNumFailedLoads);
+}
+
+void UCompileAllBlueprintsCommandlet::InitCommandLine(const FString& Params)
+{
+	TArray<FString> Tokens;
+	TArray<FString> Switches;
+	TMap<FString, FString> SwitchParams;
+	ParseCommandLine(*Params, Tokens, Switches, SwitchParams);
+
+	bResultsOnly = Switches.Contains(TEXT("ShowResultsOnly"));
+	bDirtyOnly = Switches.Contains(TEXT("DirtyOnly"));
+	bCookedOnly = Switches.Contains(TEXT("CookedOnly"));
+	bCompileSkeletonOnly = Switches.Contains(TEXT("CompileSkeletonOnly"));
+	bSimpleAssetList = Switches.Contains(TEXT("SimpleAssetList"));
+	
+	RequireAssetTags.Empty();
+	if (SwitchParams.Contains(TEXT("RequireTags")))
+	{
+		const FString& FullTagInfo = SwitchParams[TEXT("RequireTags")];
+		ParseTagPairs(FullTagInfo, RequireAssetTags);
+	}
+
+	ExcludeAssetTags.Empty();
+	if (SwitchParams.Contains(TEXT("ExcludeTags")))
+	{
+		const FString& FullTagInfo = SwitchParams[TEXT("ExcludeTags")];
+		ParseTagPairs(FullTagInfo, ExcludeAssetTags);
+	}
+
+	IgnoreFolders.Empty();
+	if (SwitchParams.Contains(TEXT("IgnoreFolder")))
+	{
+		const FString& AllIgnoreFolders = SwitchParams[TEXT("IgnoreFolder")];
+		ParseIgnoreFolders(AllIgnoreFolders);
+	}
+
+	WhitelistFiles.Empty();
+	if (SwitchParams.Contains(TEXT("WhitelistFile")))
+	{
+		const FString& WhitelistFullPath = SwitchParams[TEXT("WhitelistFile")];
+		ParseWhitelist(WhitelistFullPath);
+	}
+}
+
+void UCompileAllBlueprintsCommandlet::ParseTagPairs(const FString& FullTagString, TArray<TPair<FString, TArray<FString>>>& OutputAssetTags)
+{
+	TArray<FString> AllTagPairs;
+	FullTagString.ParseIntoArray(AllTagPairs, TEXT(";"));
+
+	//Break All Tag Pairs into individual tags and values
+	for (const FString& StringTagPair : AllTagPairs)
+	{
+		TArray<FString> ParsedTagPairs;
+		StringTagPair.ParseIntoArray(ParsedTagPairs, TEXT(","));
+
+		if (ParsedTagPairs.Num() > 0)
+		{
+			TArray<FString> TagValues;
+
+			//Start AssetTagIndex at 1, as the first one is the key. We will be using index 0 in all adds as the key
+			for (int AssetTagIndex = 1; AssetTagIndex < ParsedTagPairs.Num(); ++AssetTagIndex)
+			{
+				TagValues.Add(ParsedTagPairs[AssetTagIndex]);
+			}
+
+			TPair<FString, TArray<FString>> NewPair(ParsedTagPairs[0], TagValues);
+			OutputAssetTags.Add(NewPair);
+		}
+	}
+}
+
+void UCompileAllBlueprintsCommandlet::ParseIgnoreFolders(const FString& FullIgnoreFolderString)
+{
+	TArray<FString> ParsedIgnoreFolders;
+	FullIgnoreFolderString.ParseIntoArray(ParsedIgnoreFolders, TEXT(","));
+
+	for (const FString& IgnoreFolder : ParsedIgnoreFolders)
+	{
+		IgnoreFolders.Add(IgnoreFolder.TrimQuotes());
+	}
+}
+
+void UCompileAllBlueprintsCommandlet::ParseWhitelist(const FString& WhitelistFilePath)
+{
+	const FString FilePath = FPaths::ProjectDir() + WhitelistFilePath;
+	if (!FFileHelper::LoadANSITextFileToStrings(*FilePath, &IFileManager::Get(), WhitelistFiles))
+	{
+		UE_LOG(LogCompileAllBlueprintsCommandlet, Error, TEXT("Failed to Load Whitelist File! : %s"), *FilePath);
+	}
+}
+
+void UCompileAllBlueprintsCommandlet::BuildBlueprints()
+{
+	for (FAssetData const& Asset : BlueprintAssetList)
+	{
+		if (ShouldBuildAsset(Asset))
+		{
+			FString const AssetPath = Asset.ObjectPath.ToString();
+			UE_LOG(LogCompileAllBlueprintsCommandlet, Display, TEXT("Loading and Compiling: '%s'..."), *AssetPath);
+
+			//Load with LOAD_NoWarn and LOAD_DisableCompileOnLoad as we are covering those explicitly with CompileBlueprint errors.
+			UBlueprint* LoadedBlueprint = Cast<UBlueprint>(StaticLoadObject(Asset.GetClass(), /*Outer =*/nullptr, *AssetPath,nullptr, LOAD_NoWarn | LOAD_DisableCompileOnLoad));
+			if (LoadedBlueprint == nullptr)
+			{
+				++TotalNumFailedLoads;
+				UE_LOG(LogCompileAllBlueprintsCommandlet, Error, TEXT("Failed to Load : '%s'."), *AssetPath);
+				continue;
+			}
+			else
+			{
+				CompileBlueprint(LoadedBlueprint);
+			}
+		}
+	}
+}
+
+void UCompileAllBlueprintsCommandlet::BuildBlueprintAssetList()
+{
+	BlueprintAssetList.Empty();
+
+	UE_LOG(LogCompileAllBlueprintsCommandlet, Display, TEXT("Loading Asset Registry..."));
+	FAssetRegistryModule& AssetRegistryModule = FModuleManager::LoadModuleChecked<FAssetRegistryModule>(AssetRegistryConstants::ModuleName);
+	AssetRegistryModule.Get().SearchAllAssets(/*bSynchronousSearch =*/true);
+	UE_LOG(LogCompileAllBlueprintsCommandlet, Display, TEXT("Finished Loading Asset Registry."));
+	
+	UE_LOG(LogCompileAllBlueprintsCommandlet, Display, TEXT("Gathering All Blueprints From Asset Registry..."));
+	AssetRegistryModule.Get().GetAssetsByClass(UBlueprint::StaticClass()->GetFName(), BlueprintAssetList, true);
+}
+
+bool UCompileAllBlueprintsCommandlet::ShouldBuildAsset(FAssetData const& Asset) const
+{
+	bool bShouldBuild = true;
+
+	if (bCookedOnly && Asset.GetClass() && !Asset.GetClass()->bCooked)
+	{
+		FString const AssetPath = Asset.ObjectPath.ToString();
+		UE_LOG(LogCompileAllBlueprintsCommandlet, Verbose, TEXT("Skipping Building %s: As is not cooked"), *AssetPath);
+		bShouldBuild = false;
+	}
+
+	if (IgnoreFolders.Num() > 0)
+	{
+		for (const FString& IgnoreFolder : IgnoreFolders)
+		{
+			if (Asset.ObjectPath.ToString().StartsWith(IgnoreFolder))
+			{
+				FString const AssetPath = Asset.ObjectPath.ToString();
+				UE_LOG(LogCompileAllBlueprintsCommandlet, Verbose, TEXT("Skipping Building %s: As Object is in an Ignored Folder"), *AssetPath);
+				bShouldBuild = false;
+			}
+		}
+	}
+
+	if ((ExcludeAssetTags.Num() > 0) && (CheckHasTagInList(Asset, ExcludeAssetTags)))
+	{
+		FString const AssetPath = Asset.ObjectPath.ToString();
+		UE_LOG(LogCompileAllBlueprintsCommandlet, Verbose, TEXT("Skipping Building %s: As has an excluded tag"), *AssetPath);
+		bShouldBuild = false;
+	}
+
+	if ((RequireAssetTags.Num() > 0) && (!CheckHasTagInList(Asset, RequireAssetTags)))
+	{
+		FString const AssetPath = Asset.ObjectPath.ToString();
+		UE_LOG(LogCompileAllBlueprintsCommandlet, Verbose, TEXT("Skipping Building %s: As the asset is missing a required tag"), *AssetPath);
+		bShouldBuild = false;
+	}
+
+	if ((WhitelistFiles.Num() > 0) && (CheckInWhitelist(Asset)))
+	{
+		FString const AssetPath = Asset.ObjectPath.ToString();
+		UE_LOG(LogCompileAllBlueprintsCommandlet, Verbose, TEXT("Skipping Building %s: As the asset is part of the whitelist"), *AssetPath);
+		bShouldBuild = false;
+	}
+
+	if (bDirtyOnly)
+	{
+		const UPackage* AssetPackage = Asset.GetPackage();
+		if ((AssetPackage == nullptr) || !AssetPackage->IsDirty())
+		{
+			FString const AssetPath = Asset.ObjectPath.ToString();
+			UE_LOG(LogCompileAllBlueprintsCommandlet, Verbose, TEXT("Skipping Building %s: As Package is not dirty"), *AssetPath);
+			bShouldBuild = false;
+		}
+	}
+
+	return bShouldBuild;
+}
+
+bool UCompileAllBlueprintsCommandlet::CheckHasTagInList(FAssetData const& Asset, const TArray<TPair<FString, TArray<FString>>>& TagCollectionToCheck) const 
+{
+	bool bContainedTag = false;
+
+	for (const TPair<FString, TArray<FString>>& SingleTagAndValues : TagCollectionToCheck)
+	{
+		if (Asset.TagsAndValues.Contains(FName(*SingleTagAndValues.Key)))
+		{
+			const TArray<FString>& TagValuesToCheck = SingleTagAndValues.Value;
+			if (TagValuesToCheck.Num() > 0)
+			{
+				for (const FString& IndividualValueToCheck : TagValuesToCheck)
+				{
+					if ((Asset.TagsAndValues.Find(FName(*SingleTagAndValues.Key))) && (*Asset.TagsAndValues.Find(FName(*SingleTagAndValues.Key)) == IndividualValueToCheck))
+					{
+						bContainedTag = true;
+						break;
+					}
+				}
+			}
+			//if we don't have any values to check, just return true as the tag was included
+			else
+			{
+				bContainedTag = true;
+				break;
+			}
+		}
+	}
+
+	return bContainedTag;
+}
+
+bool UCompileAllBlueprintsCommandlet::CheckInWhitelist(FAssetData const& Asset) const
+{
+	bool bIsInWhitelist = false;
+
+	const FString& AssetFilePath = Asset.ObjectPath.ToString();
+	for (const FString& WhiteList : WhitelistFiles)
+	{
+		if (AssetFilePath == WhiteList)
+		{
+			bIsInWhitelist = true;
+			break;
+		}
+	}
+
+	return bIsInWhitelist;
+}
+
+void UCompileAllBlueprintsCommandlet::CompileBlueprint(UBlueprint* Blueprint)
+{
+	if (KismetBlueprintCompilerModule && Blueprint)
+	{
+		//Have to create a new MessageLog for each asset as the warning / error counts are cumulative
+		FCompilerResultsLog MessageLog;
+		//Need to prevent the Compiler Results Log from automatically outputting results if verbosity is too low
+		if (bResultsOnly)
+		{
+			MessageLog.bSilentMode = true;
+		}
+		else
+		{
+			MessageLog.bAnnotateMentionedNodes = true;
+		}
+	
+		FKismetCompilerOptions CompileOptions;
+		//Defaults to use for CompileOptions
+		CompileOptions.CompileType = EKismetCompileType::Full;
+		CompileOptions.bSaveIntermediateProducts = false;
+		CompileOptions.bRegenerateSkelton = false;
+		CompileOptions.bIsDuplicationInstigated = false;
+		CompileOptions.bAddInstrumentation = false;
+		CompileOptions.bReinstanceAndStubOnFailure = false;
+
+		if (bCompileSkeletonOnly)
+		{
+			CompileOptions.CompileType = EKismetCompileType::SkeletonOnly;
+		}
+		
+		KismetBlueprintCompilerModule->CompileBlueprint(Blueprint, CompileOptions, MessageLog);
+
+		if ((MessageLog.NumErrors + MessageLog.NumWarnings) > 0)
+		{
+			AssetsWithErrorsOrWarnings.Add(Blueprint->GetPathName());
+		
+			TotalNumFatalIssues += MessageLog.NumErrors;
+			TotalNumWarnings += MessageLog.NumWarnings;
+		}
+		
+		for (TSharedRef<class FTokenizedMessage>& Message : MessageLog.Messages)
+		{
+			UE_LOG(LogCompileAllBlueprintsCommandlet, Display, TEXT("%s"), *Message->ToText().ToString());
+		}
+	}
+}
+
+void UCompileAllBlueprintsCommandlet::InitKismetBlueprintCompiler()
+{
+	UE_LOG(LogCompileAllBlueprintsCommandlet, Display, TEXT("Loading Kismit Blueprint Compiler..."));
+	//Get Kismet Compiler Setup. Static so that the expensive stuff only happens once per run.
+	KismetBlueprintCompilerModule = &FModuleManager::LoadModuleChecked<IKismetCompilerInterface>(TEXT(KISMET_COMPILER_MODULENAME));
+	UE_LOG(LogCompileAllBlueprintsCommandlet, Display, TEXT("Finished Loading Kismit Blueprint Compiler..."));
+}
+
+void UCompileAllBlueprintsCommandlet::LogResults()
+{
+	//results output
+	UE_LOG(LogCompileAllBlueprintsCommandlet, Display, TEXT("\n\n\n===================================================================================\nCompiling Completed with %d errors and %d warnings and %d blueprints that failed to load.\n===================================================================================\n\n\n"), TotalNumFatalIssues, TotalNumWarnings, TotalNumFailedLoads);
+
+	//Assets with problems listing
+	if (bSimpleAssetList && (AssetsWithErrorsOrWarnings.Num() > 0))
+	{
+		UE_LOG(LogCompileAllBlueprintsCommandlet, Warning, TEXT("\n===================================================================================\nAssets With Errors or Warnings:\n===================================================================================\n"));
+
+		for (const FString& Asset : AssetsWithErrorsOrWarnings)
+		{
+			UE_LOG(LogCompileAllBlueprintsCommandlet, Warning, TEXT("%s"), *Asset);
+		}
+
+		UE_LOG(LogCompileAllBlueprintsCommandlet, Warning, TEXT("\n===================================================================================\nEnd of Asset List\n===================================================================================\n"));
+	}
+}
\ No newline at end of file
diff --git a/Engine/Source/Editor/UnrealEd/Private/Commandlets/PackageUtilities.cpp b/Engine/Source/Editor/UnrealEd/Private/Commandlets/PackageUtilities.cpp
--- a/Engine/Source/Editor/UnrealEd/Private/Commandlets/PackageUtilities.cpp
+++ b/Engine/Source/Editor/UnrealEd/Private/Commandlets/PackageUtilities.cpp
@@ -2119,22 +2119,29 @@ struct CompressAnimationsFunctor
 				continue;
 			}
 
-			// Set version since we've checked this animation for recompression.
-			if (AnimSeq->CompressCommandletVersion != CompressCommandletVersion)
-			{
-				AnimSeq->CompressCommandletVersion = CompressCommandletVersion;
-				bDirtyPackage = true;
-			}
-
 			UE_LOG(LogPackageUtilities, Warning, TEXT("Compressing animation '%s' [#%d / %d in package '%s']"),
 				*AnimSeq->GetName(),
 				ActiveAnimationIndex,
 				NumAnimationsInPackage,
 				*PackageFileName);
 
+			// First set automatic compressor and call it.
+			// This will run through a bunch of compressors and pick the best.
+			// Problem is this is going to create a DDC key with 'Automatic Compressor'
 			UAnimCompress* CompressionAlgorithm = NewObject<UAnimCompress_Automatic>();
 			AnimSeq->CompressionScheme = static_cast<UAnimCompress*>(StaticDuplicateObject(CompressionAlgorithm, AnimSeq));
-			AnimSeq->RequestAnimCompression(false, false, false);
+			AnimSeq->RequestAnimCompression(false, true, false);
+
+			// Automatic compression should have picked a suitable compressor that is not UAnimCompress_Automatic
+			if (ensure(!AnimSeq->CompressionScheme->IsA(UAnimCompress_Automatic::StaticClass())))
+			{
+				// Update CompressCommandletVersion in that case, and create a proper DDC entry
+				// (with actual compressor)
+				AnimSeq->CompressCommandletVersion = CompressCommandletVersion;
+				AnimSeq->RequestAnimCompression(false, false, false);
+				bDirtyPackage = true;
+			}
+
 			{
 				NewSize = AnimSeq->GetResourceSizeBytes(EResourceSizeMode::Inclusive);
 			}
diff --git a/Engine/Source/Editor/UnrealEd/Private/ObjectTools.cpp b/Engine/Source/Editor/UnrealEd/Private/ObjectTools.cpp
--- a/Engine/Source/Editor/UnrealEd/Private/ObjectTools.cpp
+++ b/Engine/Source/Editor/UnrealEd/Private/ObjectTools.cpp
@@ -143,7 +143,7 @@ namespace ObjectTools
 		bool bIsSupported = false;
 
 		// Check object prerequisites
-		if ( Obj->IsAsset() )
+		if (ensure(Obj) && Obj->IsAsset() )
 		{
 			UPackage* ObjectPackage = Obj->GetOutermost();
 			if( ObjectPackage != NULL )
diff --git a/Engine/Source/Programs/AutomationTool/AutomationUtils/ProjectParams.cs b/Engine/Source/Programs/AutomationTool/AutomationUtils/ProjectParams.cs
--- a/Engine/Source/Programs/AutomationTool/AutomationUtils/ProjectParams.cs
+++ b/Engine/Source/Programs/AutomationTool/AutomationUtils/ProjectParams.cs
@@ -224,6 +224,7 @@ namespace AutomationTool
 			this.AddPatchLevel = InParams.AddPatchLevel;
 			this.StageBaseReleasePaks = InParams.StageBaseReleasePaks;
 			this.DLCFile = InParams.DLCFile;
+			this.GenerateRemaster = InParams.GenerateRemaster;
             this.DLCIncludeEngineContent = InParams.DLCIncludeEngineContent;
             this.DiffCookedContentPath = InParams.DiffCookedContentPath;
             this.AdditionalCookerOptions = InParams.AdditionalCookerOptions;
@@ -399,6 +400,7 @@ namespace AutomationTool
             bool? GeneratePatch = null,
 			bool? AddPatchLevel = null,
 			bool? StageBaseReleasePaks = null,
+			bool? GenerateRemaster = null,
 			string DLCName = null,
             string DiffCookedContentPath = null,
             bool? DLCIncludeEngineContent = null,
@@ -533,6 +535,7 @@ namespace AutomationTool
             this.GeneratePatch = GetParamValueIfNotSpecified(Command, GeneratePatch, this.GeneratePatch, "GeneratePatch");
             this.AddPatchLevel = GetParamValueIfNotSpecified(Command, AddPatchLevel, this.AddPatchLevel, "AddPatchLevel");
 			this.StageBaseReleasePaks = GetParamValueIfNotSpecified(Command, StageBaseReleasePaks, this.StageBaseReleasePaks, "StageBaseReleasePaks");
+			this.GenerateRemaster = GetParamValueIfNotSpecified(Command, GenerateRemaster, this.GenerateRemaster, "GenerateRemaster");
 			this.AdditionalCookerOptions = ParseParamValueIfNotSpecified(Command, AdditionalCookerOptions, "AdditionalCookerOptions", String.Empty);
 
 			DLCName = ParseParamValueIfNotSpecified(Command, DLCName, "DLCName", String.Empty);
@@ -1346,21 +1349,23 @@ namespace AutomationTool
         /// see also CreateReleaseVersion, BasedOnReleaseVersion
         /// </summary>
         public bool GeneratePatch;
+		
+		/// <summary>
+		/// Are we generating a remaster, generate a patch from a previously released version of the game (use CreateReleaseVersion to create a release). 
+		/// this requires BasedOnReleaseVersion
+		/// see also CreateReleaseVersion, BasedOnReleaseVersion
+		/// </summary>
+		public bool GenerateRemaster;
 
-        /// <summary>
-        /// Are we adding a new patch tier (otherwise patch will modify the most recent patch tier), this requires GeneratePatch and BasedOnReleaseVersion
-        /// see also GeneratePatch, BasedOnReleaseVersion
+		/// <summary>
         /// </summary>
         public bool AddPatchLevel;
-
         /// <summary>
         /// Are we staging the unmodified pak files from the base release
-        /// </summary>
         public bool StageBaseReleasePaks;
 
-        /// <summary>
-        /// Plugin file for dlc to cook and package (if this paramter is supplied cooks the dlc and packages it into the dlc directory)
-        /// </summary>
+        /// Name of dlc to cook and package (if this paramter is supplied cooks the dlc and packages it into the dlc directory)
+		/// </summary>
         public FileReference DLCFile;
 
         /// <summary>
@@ -2173,6 +2178,14 @@ namespace AutomationTool
             get { return StageBaseReleasePaks; }
         }
 
+		/// <summary>
+		/// True if we are generating a patch
+		/// </summary>
+		public bool IsGeneratingRemaster
+		{
+			get { return GenerateRemaster; }
+		}
+
 		/// <summary>
 		/// Filename of the target game exe (or program exe).
 		/// </summary>
@@ -2441,6 +2454,7 @@ namespace AutomationTool
 				CommandUtils.LogLog("GeneratePatch={0}", GeneratePatch);
 				CommandUtils.LogLog("AddPatchLevel={0}", AddPatchLevel);
 				CommandUtils.LogLog("StageBaseReleasePaks={0}", StageBaseReleasePaks);
+				CommandUtils.LogLog("GenerateRemaster={0}", GenerateRemaster);
 				CommandUtils.LogLog("CreateReleaseVersion={0}", CreateReleaseVersion);
                 CommandUtils.LogLog("BasedOnReleaseVersion={0}", BasedOnReleaseVersion);
                 CommandUtils.LogLog("DLCFile={0}", DLCFile);
diff --git a/Engine/Source/Programs/AutomationTool/BuildGraph/Schema.cs b/Engine/Source/Programs/AutomationTool/BuildGraph/Schema.cs
--- a/Engine/Source/Programs/AutomationTool/BuildGraph/Schema.cs
+++ b/Engine/Source/Programs/AutomationTool/BuildGraph/Schema.cs
@@ -637,6 +637,7 @@ namespace AutomationTool
 			AggregateType.Attributes.Add(CreateSchemaAttribute("Submitters", ScriptSchemaStandardType.BalancedString, XmlSchemaUse.Optional));
 			AggregateType.Attributes.Add(CreateSchemaAttribute("Warnings", ScriptSchemaStandardType.Boolean, XmlSchemaUse.Optional));
 			AggregateType.Attributes.Add(CreateSchemaAttribute("If", ScriptSchemaStandardType.BalancedString, XmlSchemaUse.Optional));
+			AggregateType.Attributes.Add(CreateSchemaAttribute("Absolute", ScriptSchemaStandardType.BalancedString, XmlSchemaUse.Optional));
 			return AggregateType;
 		}
 
diff --git a/Engine/Source/Programs/AutomationTool/BuildGraph/Script.cs b/Engine/Source/Programs/AutomationTool/BuildGraph/Script.cs
--- a/Engine/Source/Programs/AutomationTool/BuildGraph/Script.cs
+++ b/Engine/Source/Programs/AutomationTool/BuildGraph/Script.cs
@@ -1135,6 +1135,7 @@ namespace AutomationTool
 				string[] Users = ReadListAttribute(Element, "Users");
 				string[] Submitters = ReadListAttribute(Element, "Submitters");
 				bool? bWarnings = Element.HasAttribute("Warnings") ? (bool?)ReadBooleanAttribute(Element, "Warnings", true) : null;
+				bool bAbsolute = Element.HasAttribute("Absolute") ? ReadBooleanAttribute(Element, "Absolute", true) : false;
 
 				// Find the list of targets which are included, and recurse through all their dependencies
 				HashSet<Node> Nodes = new HashSet<Node>();
@@ -1167,11 +1168,25 @@ namespace AutomationTool
 				{
 					if (Users != null)
 					{
-						Node.NotifyUsers.UnionWith(Users);
+						if (bAbsolute)
+						{
+							Node.NotifyUsers = new HashSet<string>(Users);
+						}
+						else
+						{
+							Node.NotifyUsers.UnionWith(Users);
+						}
 					}
 					if (Submitters != null)
 					{
-						Node.NotifySubmitters.UnionWith(Submitters);
+						if (bAbsolute)
+						{
+							Node.NotifySubmitters = new HashSet<string>(Submitters);
+						}
+						else
+						{
+							Node.NotifySubmitters.UnionWith(Submitters);
+						}
 					}
 					if (bWarnings.HasValue)
 					{
diff --git a/Engine/Source/Programs/AutomationTool/Win/WinPlatform.Automation.cs b/Engine/Source/Programs/AutomationTool/Win/WinPlatform.Automation.cs
--- a/Engine/Source/Programs/AutomationTool/Win/WinPlatform.Automation.cs
+++ b/Engine/Source/Programs/AutomationTool/Win/WinPlatform.Automation.cs
@@ -222,26 +222,38 @@ public abstract class BaseWinPlatform : Platform
 		PathVariables["EngineDir"] = SC.EngineRoot.FullName;
 		PathVariables["ProjectDir"] = SC.ProjectRoot.FullName;
 
-		string ExpandedAppLocalDir = Utils.ExpandVariables(Params.AppLocalDirectory, PathVariables);
-
-		DirectoryReference BaseAppLocalDependenciesPath = Path.IsPathRooted(ExpandedAppLocalDir) ? new DirectoryReference(CombinePaths(ExpandedAppLocalDir, PlatformDir)) : DirectoryReference.Combine(SC.ProjectRoot, ExpandedAppLocalDir, PlatformDir);
-		if (DirectoryReference.Exists(BaseAppLocalDependenciesPath))
+		// support multiple comma-separated paths
+		string[] AppLocalDirectories = Params.AppLocalDirectory.Split(';');
+		foreach (string AppLocalDirectory in AppLocalDirectories)
 		{
-			StagedDirectoryReference ProjectBinaryPath = new StagedDirectoryReference(SC.ProjectBinariesFolder.MakeRelativeTo(SC.ProjectRoot.ParentDirectory));
-			StagedDirectoryReference EngineBinaryPath = new StagedDirectoryReference(CombinePaths("Engine", "Binaries", PlatformDir));
+			string ExpandedAppLocalDir = Utils.ExpandVariables(AppLocalDirectory, PathVariables);
+
+			DirectoryReference BaseAppLocalDependenciesPath = Path.IsPathRooted(ExpandedAppLocalDir) ? new DirectoryReference(CombinePaths(ExpandedAppLocalDir, PlatformDir)) : DirectoryReference.Combine(SC.ProjectRoot, ExpandedAppLocalDir, PlatformDir);
+			if (DirectoryReference.Exists(BaseAppLocalDependenciesPath))
+			{
+				StagedDirectoryReference ProjectBinaryPath = new StagedDirectoryReference(SC.ProjectBinariesFolder.MakeRelativeTo(SC.ProjectRoot.ParentDirectory));
+				StagedDirectoryReference EngineBinaryPath = new StagedDirectoryReference(CombinePaths("Engine", "Binaries", PlatformDir));
+
+				Log("Copying AppLocal dependencies from {0} to {1} and {2}", BaseAppLocalDependenciesPath, ProjectBinaryPath, EngineBinaryPath);
+
 
-			Log("Copying AppLocal dependencies from {0} to {1} and {2}", BaseAppLocalDependenciesPath, ProjectBinaryPath, EngineBinaryPath);
 
-			foreach (DirectoryReference DependencyDirectory in DirectoryReference.EnumerateDirectories(BaseAppLocalDependenciesPath))
+				// Stage files in subdirs
+				foreach (DirectoryReference DependencyDirectory in DirectoryReference.EnumerateDirectories(BaseAppLocalDependenciesPath))
+				{	
+					SC.StageFiles(StagedFileType.NonUFS, DependencyDirectory, StageFilesSearch.TopDirectoryOnly, ProjectBinaryPath);
+					SC.StageFiles(StagedFileType.NonUFS, DependencyDirectory, StageFilesSearch.TopDirectoryOnly, EngineBinaryPath);
+				}
+				
+				// stage loose files here
+				SC.StageFiles(StagedFileType.NonUFS, BaseAppLocalDependenciesPath, StageFilesSearch.AllDirectories, ProjectBinaryPath);
+				SC.StageFiles(StagedFileType.NonUFS, BaseAppLocalDependenciesPath, StageFilesSearch.AllDirectories, EngineBinaryPath);
+			}
+			else
 			{
-				SC.StageFiles(StagedFileType.NonUFS, DependencyDirectory, StageFilesSearch.TopDirectoryOnly, ProjectBinaryPath);
-				SC.StageFiles(StagedFileType.NonUFS, DependencyDirectory, StageFilesSearch.TopDirectoryOnly, EngineBinaryPath);
+				LogWarning("Unable to deploy AppLocalDirectory dependencies. No such path: {0}", BaseAppLocalDependenciesPath);
 			}
 		}
-		else
-		{
-			throw new AutomationException("Unable to deploy AppLocalDirectory dependencies. No such path: {0}", BaseAppLocalDependenciesPath);
-		}
 	}
 
     /// <summary>
diff --git a/Engine/Source/Programs/UnrealBuildTool/System/RulesAssembly.cs b/Engine/Source/Programs/UnrealBuildTool/System/RulesAssembly.cs
--- a/Engine/Source/Programs/UnrealBuildTool/System/RulesAssembly.cs
+++ b/Engine/Source/Programs/UnrealBuildTool/System/RulesAssembly.cs
@@ -263,7 +263,7 @@ namespace UnrealBuildTool
 			{
 				if (Parent == null)
 				{
-					throw new BuildException("Could not find definition for module '{0}'", ModuleName);
+					throw new BuildException("Could not find definition for module '{0}' (referenced via {1})", ModuleName, ReferenceChain);
 				}
 				else
 				{
diff --git a/Engine/Source/Runtime/AIModule/Classes/BehaviorTree/BehaviorTreeComponent.h b/Engine/Source/Runtime/AIModule/Classes/BehaviorTree/BehaviorTreeComponent.h
--- a/Engine/Source/Runtime/AIModule/Classes/BehaviorTree/BehaviorTreeComponent.h
+++ b/Engine/Source/Runtime/AIModule/Classes/BehaviorTree/BehaviorTreeComponent.h
@@ -89,7 +89,6 @@ class AIMODULE_API UBehaviorTreeComponent : public UBrainComponent
 
 	// Begin UBrainComponent overrides
 	virtual void RestartLogic() override;
-protected:
 	virtual void StopLogic(const FString& Reason) override;
 	virtual void PauseLogic(const FString& Reason) override;
 	virtual EAILogicResuming::Type ResumeLogic(const FString& Reason) override;
diff --git a/Engine/Source/Runtime/AIModule/Classes/BrainComponent.h b/Engine/Source/Runtime/AIModule/Classes/BrainComponent.h
--- a/Engine/Source/Runtime/AIModule/Classes/BrainComponent.h
+++ b/Engine/Source/Runtime/AIModule/Classes/BrainComponent.h
@@ -146,8 +146,6 @@ public:
 
 	/** AI logic won't be needed anymore, stop all activity and run cleanup */
 	virtual void Cleanup() {}
-
-protected:
 	virtual void PauseLogic(const FString& Reason) {}
 	/** MUST be called by child implementations!
 	 *	@return indicates whether child class' ResumeLogic should be called (true) or has it been 
diff --git a/Engine/Source/Runtime/AIModule/Classes/Perception/AIPerceptionComponent.h b/Engine/Source/Runtime/AIModule/Classes/Perception/AIPerceptionComponent.h
--- a/Engine/Source/Runtime/AIModule/Classes/Perception/AIPerceptionComponent.h
+++ b/Engine/Source/Runtime/AIModule/Classes/Perception/AIPerceptionComponent.h
@@ -299,6 +299,11 @@ public:
 	UFUNCTION(BlueprintCallable, Category = "AI|Perception")
 	bool GetActorsPerception(AActor* Actor, FActorPerceptionBlueprintInfo& Info);
 
+	/** Note that this works only if given sense has been already configured for
+	 *	this component instance */
+	UFUNCTION(BlueprintCallable, Category = "AI|Perception")
+	void SetSenseEnabled(TSubclassOf<UAISense> SenseClass, const bool bEnable);
+
 	//////////////////////////////////////////////////////////////////////////
 	// Might want to move these to special "BP_AIPerceptionComponent"
 	//////////////////////////////////////////////////////////////////////////
diff --git a/Engine/Source/Runtime/AIModule/Private/Perception/AIPerceptionComponent.cpp b/Engine/Source/Runtime/AIModule/Private/Perception/AIPerceptionComponent.cpp
--- a/Engine/Source/Runtime/AIModule/Private/Perception/AIPerceptionComponent.cpp
+++ b/Engine/Source/Runtime/AIModule/Private/Perception/AIPerceptionComponent.cpp
@@ -685,6 +685,15 @@ bool UAIPerceptionComponent::GetActorsPerception(AActor* Actor, FActorPerception
 	return bInfoFound;
 }
 
+void UAIPerceptionComponent::SetSenseEnabled(TSubclassOf<UAISense> SenseClass, const bool bEnable)
+{
+	const FAISenseID SenseID = UAISense::GetSenseID(SenseClass);
+	if (SenseID.IsValid())
+	{
+		UpdatePerceptionWhitelist(SenseID, bEnable);
+	}
+}
+
 //----------------------------------------------------------------------//
 // debug
 //----------------------------------------------------------------------//
diff --git a/Engine/Source/Runtime/AIModule/Private/Perception/AISense_Sight.cpp b/Engine/Source/Runtime/AIModule/Private/Perception/AISense_Sight.cpp
--- a/Engine/Source/Runtime/AIModule/Private/Perception/AISense_Sight.cpp
+++ b/Engine/Source/Runtime/AIModule/Private/Perception/AISense_Sight.cpp
@@ -268,7 +268,8 @@ float UAISense_Sight::Update()
 						}
 					}
 				}
-				else
+				// communicate failure only if we've seen give actor before
+				else if (SightQuery->bLastResult)
 				{
 					SIGHT_LOG_SEGMENT(Listener.Listener.Get()->GetOwner(), Listener.CachedLocation, TargetLocation, FColor::Red, TEXT("%s"), *(Target.TargetId.ToString()));
 					Listener.RegisterStimulus(TargetActor, FAIStimulus(*this, 0.f, TargetLocation, Listener.CachedLocation, FAIStimulus::SensingFailed));
diff --git a/Engine/Source/Runtime/AIModule/Private/Tasks/AITask_MoveTo.cpp b/Engine/Source/Runtime/AIModule/Private/Tasks/AITask_MoveTo.cpp
--- a/Engine/Source/Runtime/AIModule/Private/Tasks/AITask_MoveTo.cpp
+++ b/Engine/Source/Runtime/AIModule/Private/Tasks/AITask_MoveTo.cpp
@@ -158,7 +158,7 @@ void UAITask_MoveTo::PerformMove()
 		PathFinishDelegateHandle = PFComp->OnRequestFinished.AddUObject(this, &UAITask_MoveTo::OnRequestFinished);
 		SetObservedPath(FollowedPath);
 
-		if (TaskState == EGameplayTaskState::Finished)
+		if (IsFinished())
 		{
 			UE_VLOG(GetGameplayTasksComponent(), LogGameplayTasks, Error, TEXT("%s> re-Activating Finished task!"), *GetName());
 		}
diff --git a/Engine/Source/Runtime/ActorSequence/Private/ActorSequenceComponent.cpp b/Engine/Source/Runtime/ActorSequence/Private/ActorSequenceComponent.cpp
--- a/Engine/Source/Runtime/ActorSequence/Private/ActorSequenceComponent.cpp
+++ b/Engine/Source/Runtime/ActorSequence/Private/ActorSequenceComponent.cpp
@@ -26,12 +26,15 @@ void UActorSequenceComponent::BeginPlay()
 {
 	Super::BeginPlay();
 
-	SequencePlayer = NewObject<UActorSequencePlayer>(this, "SequencePlayer");
-	SequencePlayer->Initialize(Sequence, PlaybackSettings);
-
-	if (bAutoPlay)
+	if (Sequence != nullptr)
 	{
-		SequencePlayer->Play();
+		SequencePlayer = NewObject<UActorSequencePlayer>(this, "SequencePlayer");
+		SequencePlayer->Initialize(Sequence, PlaybackSettings);
+
+		if (bAutoPlay)
+		{
+			SequencePlayer->Play();
+		}
 	}
 }
 
diff --git a/Engine/Source/Runtime/ActorSequence/Public/ActorSequenceComponent.h b/Engine/Source/Runtime/ActorSequence/Public/ActorSequenceComponent.h
--- a/Engine/Source/Runtime/ActorSequence/Public/ActorSequenceComponent.h
+++ b/Engine/Source/Runtime/ActorSequence/Public/ActorSequenceComponent.h
@@ -27,6 +27,11 @@ public:
 	{
 		return Sequence;
 	}
+
+	UActorSequencePlayer* GetSequencePlayer() const 
+	{
+		return SequencePlayer;
+	}
 	
 	virtual void PostInitProperties() override;
 	virtual void BeginPlay() override;
diff --git a/Engine/Source/Runtime/AssetRegistry/Private/PackageReader.cpp b/Engine/Source/Runtime/AssetRegistry/Private/PackageReader.cpp
--- a/Engine/Source/Runtime/AssetRegistry/Private/PackageReader.cpp
+++ b/Engine/Source/Runtime/AssetRegistry/Private/PackageReader.cpp
@@ -185,7 +185,7 @@ bool FPackageReader::ReadAssetRegistryData (TArray<FAssetData*>& AssetDataList)
 		if (ObjectPath.StartsWith(TEXT("/"), ESearchCase::CaseSensitive))
 		{
 			// This should never happen, it means that package A has an export with an outer of package B
-			UE_LOG(LogAssetRegistry, Warning, TEXT("Package %s has invalid export %s, resave source package!"), *PackageName, *ObjectPath);
+			UE_ASSET_LOG(LogAssetRegistry, Warning, *PackageName, TEXT("Package has invalid export %s, resave source package!"), *ObjectPath);
 			continue;
 		}
 
diff --git a/Engine/Source/Runtime/Core/Private/HAL/MallocBinned2.cpp b/Engine/Source/Runtime/Core/Private/HAL/MallocBinned2.cpp
--- a/Engine/Source/Runtime/Core/Private/HAL/MallocBinned2.cpp
+++ b/Engine/Source/Runtime/Core/Private/HAL/MallocBinned2.cpp
@@ -61,6 +61,12 @@ int64 AllocatedLargePoolMemory = 0; // memory requests to the OS which don't fit
 int64 AllocatedLargePoolMemoryWAlignment = 0; // when we allocate at OS level we need to align to a size
 #endif
 
+#if BINNED2_ALLOCATOR_STATS_VALIDATION
+int64 AllocatedSmallPoolMemoryValidation = 0;
+FCriticalSection ValidationCriticalSection;
+int32 RecursionCounter = 0;
+#endif
+
 // Block sizes are based around getting the maximum amount of allocations per pool, with as little alignment waste as possible.
 // Block sizes should be close to even divisors of the system page size, and well distributed.
 // They must be 16-byte aligned as well.
@@ -508,10 +514,45 @@ struct FMallocBinned2::Private
 			Bundle = NextBundle;
 		}
 	}
+
+
+	static FCriticalSection FreeBlockListsRegistrationMutex;
+	static TArray<FPerThreadFreeBlockLists*> RegisteredFreeBlockLists;
+	static void RegisterThreadFreeBlockLists( FPerThreadFreeBlockLists* FreeBlockLists )
+	{
+		FScopeLock Lock(&FreeBlockListsRegistrationMutex);
+#if BINNED2_ALLOCATOR_STATS_VALIDATION
+		++RecursionCounter;
+#endif
+		RegisteredFreeBlockLists.Add(FreeBlockLists);
+#if BINNED2_ALLOCATOR_STATS_VALIDATION
+		--RecursionCounter;
+#endif
+	}
+	static void UnregisterThreadFreeBlockLists( FPerThreadFreeBlockLists* FreeBlockLists )
+	{
+		FScopeLock Lock(&FreeBlockListsRegistrationMutex);
+#if BINNED2_ALLOCATOR_STATS_VALIDATION
+		++RecursionCounter;
+#endif
+		RegisteredFreeBlockLists.Remove(FreeBlockLists);
+#if BINNED2_ALLOCATOR_STATS_VALIDATION
+		--RecursionCounter;
+#endif
+#if BINNED2_ALLOCATOR_STATS
+		FMallocBinned2::FPerThreadFreeBlockLists::ConsolidatedMemory += FreeBlockLists->AllocatedMemory;
+#endif
+	}
 };
 
 FMallocBinned2::Private::FGlobalRecycler FMallocBinned2::Private::GGlobalRecycler;
 
+TArray<FMallocBinned2::FPerThreadFreeBlockLists*> FMallocBinned2::Private::RegisteredFreeBlockLists;
+#if BINNED2_ALLOCATOR_STATS
+int64 FMallocBinned2::FPerThreadFreeBlockLists::ConsolidatedMemory = 0;
+#endif
+FCriticalSection FMallocBinned2::Private::FreeBlockListsRegistrationMutex;
+
 FORCEINLINE bool FMallocBinned2::FPoolList::IsEmpty() const
 {
 	return Front == nullptr;
@@ -649,11 +690,6 @@ void* FMallocBinned2::MallocExternal(SIZE_T Size, uint32 Alignment)
 	if ((Size <= BINNED2_MAX_SMALL_POOL_SIZE) & (Alignment <= BINNED2_MINIMUM_ALIGNMENT)) // one branch, not two
 	{
 		uint32 PoolIndex = BoundSizeToPoolIndex(Size);
-
-#if BINNED2_ALLOCATOR_STATS
-		FPlatformAtomics::InterlockedAdd(&AllocatedSmallPoolMemory, (int64)PoolIndexToBlockSize(PoolIndex));
-#endif
-
 		FPerThreadFreeBlockLists* Lists = GMallocBinned2PerThreadCaches ? FPerThreadFreeBlockLists::Get() : nullptr;
 		if (Lists)
 		{
@@ -661,6 +697,10 @@ void* FMallocBinned2::MallocExternal(SIZE_T Size, uint32 Alignment)
 			{
 				if (void* Result = Lists->Malloc(PoolIndex))
 				{
+#if BINNED2_ALLOCATOR_STATS
+					uint32 BlockSize = PoolIndexToBlockSize(PoolIndex);
+					Lists->AllocatedMemory += BlockSize;
+#endif
 					return Result;
 				}
 			}
@@ -682,10 +722,14 @@ void* FMallocBinned2::MallocExternal(SIZE_T Size, uint32 Alignment)
 		}
 
 		void* Result = Pool->AllocateRegularBlock();
+#if BINNED2_ALLOCATOR_STATS
+		AllocatedSmallPoolMemory += PoolIndexToBlockSize(PoolIndex);
+#endif // BINNED2_ALLOCATOR_STATS
 		if (GMallocBinned2AllocExtra)
 		{
 			if (Lists)
 			{
+				// prefill the free list with some allocations so we are less likely to hit this slow path with the mutex 
 				for (int32 Index = 0; Index < GMallocBinned2AllocExtra && Pool->HasFreeRegularBlock(); Index++)
 				{
 					if (!Lists->Free(Result, PoolIndex, Table.BlockSize))
@@ -817,9 +861,6 @@ void FMallocBinned2::FreeExternal(void* Ptr)
 		uint32 BlockSize = BasePtr->BlockSize;
 		uint32 PoolIndex = BasePtr->PoolIndex;
 
-#if BINNED2_ALLOCATOR_STATS
-		FPlatformAtomics::InterlockedAdd(&AllocatedSmallPoolMemory, -((int64)(BlockSize)));
-#endif
 		FBundleNode* BundlesToRecycle = nullptr;
 		FPerThreadFreeBlockLists* Lists = GMallocBinned2PerThreadCaches ? FPerThreadFreeBlockLists::Get() : nullptr;
 		if (Lists)
@@ -827,6 +868,9 @@ void FMallocBinned2::FreeExternal(void* Ptr)
 			BundlesToRecycle = Lists->RecycleFullBundle(BasePtr->PoolIndex);
 			bool bPushed = Lists->Free(Ptr, PoolIndex, BlockSize);
 			check(bPushed);
+#if BINNED2_ALLOCATOR_STATS
+			Lists->AllocatedMemory -= BlockSize;
+#endif
 		}
 		else
 		{
@@ -838,6 +882,13 @@ void FMallocBinned2::FreeExternal(void* Ptr)
 			BundlesToRecycle->NextBundle = nullptr;
 			FScopeLock Lock(&Mutex);
 			Private::FreeBundles(*this, BundlesToRecycle, BlockSize, PoolIndex);
+#if BINNED2_ALLOCATOR_STATS
+			if (!Lists)
+			{
+				// lists track their own stat track them instead in the global stat if we don't have lists
+				AllocatedSmallPoolMemory -= ((int64)(BlockSize));
+			}
+#endif
 		}
 	}
 	else if (Ptr)
@@ -1068,12 +1119,18 @@ void FMallocBinned2::FPerThreadFreeBlockLists::SetTLS()
 		LLM_SCOPED_TAG_WITH_ENUM(ELLMScopeTag::SmallBinnedAllocation);
 		ThreadSingleton = new (FPlatformMemory::BinnedAllocFromOS(Align(sizeof(FPerThreadFreeBlockLists), FMallocBinned2::OsAllocationGranularity))) FPerThreadFreeBlockLists();
 		FPlatformTLS::SetTlsValue(FMallocBinned2::Binned2TlsSlot, ThreadSingleton);
+		FMallocBinned2::Private::RegisterThreadFreeBlockLists(ThreadSingleton);
 	}
 }
 
 void FMallocBinned2::FPerThreadFreeBlockLists::ClearTLS()
 {
 	check(FMallocBinned2::Binned2TlsSlot);
+	FPerThreadFreeBlockLists* ThreadSingleton = (FPerThreadFreeBlockLists*)FPlatformTLS::GetTlsValue(FMallocBinned2::Binned2TlsSlot);
+	if ( ThreadSingleton )
+	{
+		FMallocBinned2::Private::UnregisterThreadFreeBlockLists(ThreadSingleton);
+	}
 	FPlatformTLS::SetTlsValue(FMallocBinned2::Binned2TlsSlot, nullptr);
 }
 
@@ -1082,30 +1139,54 @@ void FMallocBinned2::FFreeBlock::CanaryFail() const
 	UE_LOG(LogMemory, Fatal, TEXT("FMallocBinned2 Attempt to realloc an unrecognized block %p   canary == 0x%x != 0x%x"), (void*)this, (int32)Canary, (int32)FMallocBinned2::FFreeBlock::CANARY_VALUE);
 }
 
+#if BINNED2_ALLOCATOR_STATS
+int64 FMallocBinned2::GetTotalAllocatedSmallPoolMemory() const
+{
+	int64 FreeBlockAllocatedMemory = 0;
+	{
+		FScopeLock Lock(&Private::FreeBlockListsRegistrationMutex);
+		for (const FPerThreadFreeBlockLists* FreeBlockLists : Private::RegisteredFreeBlockLists)
+		{
+			FreeBlockAllocatedMemory += FreeBlockLists->AllocatedMemory;
+		}
+		FreeBlockAllocatedMemory += FPerThreadFreeBlockLists::ConsolidatedMemory;
+	}
+
+	return AllocatedSmallPoolMemory + FreeBlockAllocatedMemory;
+}
+#endif
+
 void FMallocBinned2::GetAllocatorStats( FGenericMemoryStats& OutStats )
 {
 #if BINNED2_ALLOCATOR_STATS
-	OutStats.Add(TEXT("AllocatedSmallPoolMemory"), AllocatedSmallPoolMemory );
+
+	int64 TotalAllocatedSmallPoolMemory = GetTotalAllocatedSmallPoolMemory();
+
+	OutStats.Add(TEXT("AllocatedSmallPoolMemory"), TotalAllocatedSmallPoolMemory);
 	OutStats.Add(TEXT("AllocatedOSSmallPoolMemory"), AllocatedOSSmallPoolMemory);
 	OutStats.Add(TEXT("AllocatedLargePoolMemory"), AllocatedLargePoolMemory);
 	OutStats.Add(TEXT("AllocatedLargePoolMemoryWAlignment"), AllocatedLargePoolMemoryWAlignment);
 	OutStats.Add(TEXT("PageAllocatorFreeCacheSize"), CachedOSPageAllocator.GetCachedFreeTotal());
 
-	uint64 TotalAllocated = AllocatedSmallPoolMemory + AllocatedLargePoolMemory;
+	uint64 TotalAllocated = TotalAllocatedSmallPoolMemory + AllocatedLargePoolMemory;
 	uint64 TotalOSAllocated = AllocatedOSSmallPoolMemory + AllocatedLargePoolMemoryWAlignment + CachedOSPageAllocator.GetCachedFreeTotal();
 
 	OutStats.Add(TEXT("TotalAllocated"), TotalAllocated);
 	OutStats.Add(TEXT("TotalOSAllocated"), TotalOSAllocated);
 #endif
+	FMalloc::GetAllocatorStats(OutStats);
 }
 
 void FMallocBinned2::DumpAllocatorStats(class FOutputDevice& Ar)
 {
 #if BINNED2_ALLOCATOR_STATS
+
+	int64 TotalAllocatedSmallPoolMemory = GetTotalAllocatedSmallPoolMemory();
+
 	Ar.Logf(TEXT("FMallocBinned2 Mem report"));
 	Ar.Logf(TEXT("Small Pool"));
-	Ar.Logf(TEXT("Requested Allocations: %fmb  (including block size padding)"), ((double)AllocatedSmallPoolMemory) / (1024.0f * 1024.0f) );
-	Ar.Logf(TEXT("OS Allocated: %fmb"), ((double)AllocatedOSSmallPoolMemory) / (1024.0f * 1024.0f));
+	Ar.Logf(TEXT("Requested Allocations: %fmb  (including block size padding)"), ((double)TotalAllocatedSmallPoolMemory) / (1024.0f * 1024.0f) );
+	Ar.Logf(TEXT("OS Allocated: %fmb"), ((double)TotalAllocatedSmallPoolMemory) / (1024.0f * 1024.0f));
 	Ar.Logf(TEXT("Large Pool"));
 	Ar.Logf(TEXT("Requested Allocations: %fmb"), ((double)AllocatedLargePoolMemory) / (1024.0f * 1024.0f));
 	Ar.Logf(TEXT("OS Allocated: %fmb"), ((double)AllocatedLargePoolMemoryWAlignment) / (1024.0f * 1024.0f));
diff --git a/Engine/Source/Runtime/Core/Private/Linux/LinuxPlatformMisc.cpp b/Engine/Source/Runtime/Core/Private/Linux/LinuxPlatformMisc.cpp
--- a/Engine/Source/Runtime/Core/Private/Linux/LinuxPlatformMisc.cpp
+++ b/Engine/Source/Runtime/Core/Private/Linux/LinuxPlatformMisc.cpp
@@ -673,9 +673,15 @@ FString FLinuxPlatformMisc::GetCPUBrand()
 #if !UE_BUILD_SHIPPING
 bool FLinuxPlatformMisc::IsDebuggerPresent()
 {
+	extern CORE_API bool GIgnoreDebugger;
+	if (GIgnoreDebugger)
+	{
+		return false;
+	}
+
 	// If a process is tracing this one then TracerPid in /proc/self/status will
 	// be the id of the tracing process. Use SignalHandler safe functions 
-
+	
 	int StatusFile = open("/proc/self/status", O_RDONLY);
 	if (StatusFile == -1) 
 	{
diff --git a/Engine/Source/Runtime/Core/Private/Misc/TimeGuard.cpp b/Engine/Source/Runtime/Core/Private/Misc/TimeGuard.cpp
--- a/Engine/Source/Runtime/Core/Private/Misc/TimeGuard.cpp
+++ b/Engine/Source/Runtime/Core/Private/Misc/TimeGuard.cpp
@@ -11,6 +11,8 @@ TMap<const TCHAR*, FLightweightTimeGuard::FGuardInfo>  FLightweightTimeGuard::Hi
 bool FLightweightTimeGuard::bEnabled;
 float FLightweightTimeGuard::FrameTimeThresholdMS = 1000.0 / 30.0;
 FCriticalSection FLightweightTimeGuard::ReportMutex;
+TSet<const TCHAR *> FLightweightTimeGuard::VolatileNames; // any names which come in volatile we allocate them to a static string and put them in this array
+
 
 void FLightweightTimeGuard::SetEnabled(bool InEnable)
 {
@@ -34,10 +36,28 @@ void FLightweightTimeGuard::GetData(TMap<const TCHAR*, FGuardInfo>& Dest)
 	Dest = HitchData;
 }
 
-void FLightweightTimeGuard::ReportHitch(const TCHAR* InName, const float TimeMS)
+void FLightweightTimeGuard::ReportHitch(const TCHAR* VolatileInName, const float TimeMS, bool VolatileName)
 {
 	FScopeLock lock(&ReportMutex);
 
+	const TCHAR* InName = VolatileInName;
+	if ( VolatileName )
+	{
+		const TCHAR** CachedName = VolatileNames.Find(VolatileInName);
+		if ( CachedName == nullptr )
+		{
+			int32 StringLength = FCString::Strlen(VolatileInName) + 1;
+			TCHAR *NewString = new TCHAR[StringLength];
+			FCString::Strcpy(NewString, StringLength, VolatileInName);
+			VolatileNames.Add(NewString);
+			InName = NewString;
+		}
+		else
+		{
+			InName = *CachedName;
+		}
+	}
+
 	FGuardInfo& Data = HitchData.FindOrAdd(InName);
 
 	if (Data.Count == 0)
diff --git a/Engine/Source/Runtime/Core/Private/Stats/StatsCommand.cpp b/Engine/Source/Runtime/Core/Private/Stats/StatsCommand.cpp
--- a/Engine/Source/Runtime/Core/Private/Stats/StatsCommand.cpp
+++ b/Engine/Source/Runtime/Core/Private/Stats/StatsCommand.cpp
@@ -695,14 +695,14 @@ static void DumpHitch(int64 Frame)
 				GameThread = ChildIter.Value();
 				UE_LOG( LogStats, Log, TEXT( "------------------ Game Thread %.2fms" ), GameThreadTime * 1000.0f );
 				GameThread->CullByCycles( MinCycles );
-				GameThread->DebugPrint();
+				GameThread->DebugPrint(nullptr, 127);
 			}
 			else if( ThreadName == FName( NAME_RenderThread ) )
 			{
 				RenderThread = ChildIter.Value();
 				UE_LOG( LogStats, Log, TEXT( "------------------ Render Thread (%s) %.2fms" ), *RenderThread->Meta.NameAndInfo.GetRawName().ToString(), RenderThreadTime * 1000.0f );
 				RenderThread->CullByCycles( MinCycles );
-				RenderThread->DebugPrint();
+				RenderThread->DebugPrint(nullptr, 127);
 			}
 		}
 
diff --git a/Engine/Source/Runtime/Core/Public/HAL/MallocBinned2.h b/Engine/Source/Runtime/Core/Public/HAL/MallocBinned2.h
--- a/Engine/Source/Runtime/Core/Public/HAL/MallocBinned2.h
+++ b/Engine/Source/Runtime/Core/Public/HAL/MallocBinned2.h
@@ -51,24 +51,31 @@
 #endif
 
 
-#if UE_BUILD_SHIPPING && !WITH_EDITOR
-	#define BINNED2_ALLOCATOR_STATS 0
-#else
-	#define BINNED2_ALLOCATOR_STATS 1
+#ifndef BINNED2_ALLOCATOR_STATS
+	#if UE_BUILD_SHIPPING && !WITH_EDITOR
+		#define BINNED2_ALLOCATOR_STATS 0	
+	#else
+		#define BINNED2_ALLOCATOR_STATS 0
+	#endif
 #endif
 
 
-#if BINNED2_ALLOCATOR_STATS
-extern int64 AllocatedSmallPoolMemory; // memory that's requested to be allocated by the game
-
+#define BINNED2_ALLOCATOR_STATS_VALIDATION (BINNED2_ALLOCATOR_STATS && 0)
 
+#if BINNED2_ALLOCATOR_STATS
 //////////////////////////////////////////////////////////////////////////
 // the following don't need a critical section because they are covered by the critical section called Mutex
+extern int64 AllocatedSmallPoolMemory; // memory that's requested to be allocated by the game
 extern int64 AllocatedOSSmallPoolMemory;
 extern int64 AllocatedLargePoolMemory; // memory requests to the OS which don't fit in the small pool
 extern int64 AllocatedLargePoolMemoryWAlignment; // when we allocate at OS level we need to align to a size
 #endif
-
+#if BINNED2_ALLOCATOR_STATS_VALIDATION
+#include "Misc/ScopeLock.h"
+extern int64 AllocatedSmallPoolMemoryValidation;
+extern FCriticalSection ValidationCriticalSection;
+extern int32 RecursionCounter;
+#endif
 
 
 //
@@ -342,6 +349,12 @@ class CORE_API FMallocBinned2 final : public FMalloc
 		static void SetTLS();
 		static void ClearTLS();
 
+		FPerThreadFreeBlockLists() 
+#if BINNED2_ALLOCATOR_STATS
+			: AllocatedMemory(0) 
+#endif
+		{ }
+
 		FORCEINLINE void* Malloc(uint32 InPoolIndex)
 		{
 			return FreeLists[InPoolIndex].PopFromFront(InPoolIndex);
@@ -370,6 +383,11 @@ class CORE_API FMallocBinned2 final : public FMalloc
 		{
 			return FreeLists[InPoolIndex].PopBundles(InPoolIndex);
 		}
+#if BINNED2_ALLOCATOR_STATS
+	public:
+		int64 AllocatedMemory;
+		static int64 ConsolidatedMemory;
+#endif
 	private:
 		FFreeBlockList FreeLists[BINNED2_SMALL_POOL_COUNT];
 	};
@@ -389,9 +407,35 @@ public:
 	// FMalloc interface.
 	virtual bool IsInternallyThreadSafe() const override;
 	FORCEINLINE virtual void* Malloc(SIZE_T Size, uint32 Alignment) override
+	{
+#if BINNED2_ALLOCATOR_STATS_VALIDATION
+		FScopeLock Lock(&ValidationCriticalSection);
+		++RecursionCounter;
+		void *Result = MallocInline(Size, Alignment);
+		if ( !IsOSAllocation(Result))
+		{
+			SIZE_T OutSize;
+			ensure( GetAllocationSize(Result, OutSize) );
+			AllocatedSmallPoolMemoryValidation += OutSize;
+			if (RecursionCounter==1)
+			{
+				check(GetTotalAllocatedSmallPoolMemory() == AllocatedSmallPoolMemoryValidation);
+				if (GetTotalAllocatedSmallPoolMemory() != AllocatedSmallPoolMemoryValidation)
+				{
+					FPlatformMisc::DebugBreak();
+				}
+			}
+		}
+		--RecursionCounter;
+		return Result;
+#else
+		return MallocInline(Size, Alignment);
+#endif
+	}
+	FORCEINLINE void* MallocInline(SIZE_T Size, uint32 Alignment )
 	{
 		void* Result = nullptr;
-
+	
 		// Only allocate from the small pools if the size is small enough and the alignment isn't crazy large.
 		// With large alignments, we'll waste a lot of memory allocating an entire page, but such alignments are highly unlikely in practice.
 		if ((Size <= BINNED2_MAX_SMALL_POOL_SIZE) & (Alignment <= BINNED2_MINIMUM_ALIGNMENT)) // one branch, not two
@@ -399,7 +443,15 @@ public:
 			FPerThreadFreeBlockLists* Lists = GMallocBinned2PerThreadCaches ? FPerThreadFreeBlockLists::Get() : nullptr;
 			if (Lists)
 			{
-				Result = Lists->Malloc(BoundSizeToPoolIndex(Size));
+				uint32 PoolIndex = BoundSizeToPoolIndex(Size);
+				uint32 BlockSize = PoolIndexToBlockSize(PoolIndex);
+				Result = Lists->Malloc(PoolIndex);
+#if BINNED2_ALLOCATOR_STATS
+				if (Result)
+				{
+					Lists->AllocatedMemory += BlockSize;
+				}
+#endif
 			}
 		}
 		if (Result == nullptr)
@@ -409,8 +461,43 @@ public:
 
 		return Result;
 	}
-
 	FORCEINLINE virtual void* Realloc(void* Ptr, SIZE_T NewSize, uint32 Alignment) override
+	{
+#if BINNED2_ALLOCATOR_STATS_VALIDATION
+		bool bOldIsOsAllocation = IsOSAllocation(Ptr);
+		SIZE_T OldSize;
+		if (!bOldIsOsAllocation)
+		{
+			ensure(GetAllocationSize(Ptr, OldSize));
+		}
+		FScopeLock Lock(&ValidationCriticalSection);
+		++RecursionCounter;
+		void *Result = ReallocInline(Ptr, NewSize, Alignment);
+		if ( !bOldIsOsAllocation )
+		{
+			AllocatedSmallPoolMemoryValidation -= OldSize;
+		}
+		if (!IsOSAllocation(Result))
+		{
+			SIZE_T OutSize;
+			ensure(GetAllocationSize(Result, OutSize));
+			AllocatedSmallPoolMemoryValidation += OutSize;
+		}
+		if (RecursionCounter == 1)
+		{
+			check(GetTotalAllocatedSmallPoolMemory() == AllocatedSmallPoolMemoryValidation);
+			if (GetTotalAllocatedSmallPoolMemory() != AllocatedSmallPoolMemoryValidation)
+			{
+				FPlatformMisc::DebugBreak();
+			}
+		}
+		--RecursionCounter;
+		return Result;
+#else
+		return ReallocInline(Ptr, NewSize, Alignment);
+#endif
+	}
+	FORCEINLINE void* ReallocInline(void* Ptr, SIZE_T NewSize, uint32 Alignment) 
 	{
 		if (NewSize <= BINNED2_MAX_SMALL_POOL_SIZE && Alignment <= BINNED2_MINIMUM_ALIGNMENT) // one branch, not two
 		{
@@ -436,7 +523,15 @@ public:
 				}
 				if (bCanFree)
 				{
-					void* Result = NewSize ? Lists->Malloc(BoundSizeToPoolIndex(NewSize)) : nullptr;
+					uint32 NewPoolIndex = BoundSizeToPoolIndex(NewSize);
+					uint32 NewBlockSize = PoolIndexToBlockSize(NewPoolIndex);
+					void* Result = NewSize ? Lists->Malloc(NewPoolIndex) : nullptr;
+#if BINNED2_ALLOCATOR_STATS
+					if (Result)
+					{
+						Lists->AllocatedMemory += NewBlockSize;
+					}
+#endif
 					if (Result || !NewSize)
 					{
 						if (Result && Ptr)
@@ -447,6 +542,9 @@ public:
 						{
 							bool bDidPush = Lists->Free(Ptr, PoolIndex, BlockSize);
 							checkSlow(bDidPush);
+#if BINNED2_ALLOCATOR_STATS
+							Lists->AllocatedMemory -= BlockSize;
+#endif
 						}
 
 						return Result;
@@ -459,6 +557,32 @@ public:
 	}
 
 	FORCEINLINE virtual void Free(void* Ptr) override
+	{
+#if BINNED2_ALLOCATOR_STATS_VALIDATION
+		FScopeLock Lock(&ValidationCriticalSection);
+		++RecursionCounter;
+		if (!IsOSAllocation(Ptr))
+		{
+			SIZE_T OutSize;
+			ensure(GetAllocationSize(Ptr, OutSize));
+			AllocatedSmallPoolMemoryValidation -= OutSize;
+		}
+		FreeInline(Ptr);
+		if (RecursionCounter == 1)
+		{
+			check(GetTotalAllocatedSmallPoolMemory() == AllocatedSmallPoolMemoryValidation);
+			if (GetTotalAllocatedSmallPoolMemory() != AllocatedSmallPoolMemoryValidation)
+			{
+				FPlatformMisc::DebugBreak();
+			}
+		}
+		--RecursionCounter;
+#else
+		FreeInline(Ptr);
+#endif
+	}
+
+	FORCEINLINE void FreeInline(void* Ptr)
 	{
 		if (!IsOSAllocation(Ptr))
 		{
@@ -466,15 +590,18 @@ public:
 			if (Lists)
 			{
 				FFreeBlock* BasePtr = GetPoolHeaderFromPointer(Ptr);
+				int32 BlockSize = BasePtr->BlockSize;
 				if (BasePtr->IsCanaryOk() && Lists->Free(Ptr, BasePtr->PoolIndex, BasePtr->BlockSize))
 				{
+#if BINNED2_ALLOCATOR_STATS
+					Lists->AllocatedMemory -= BasePtr->BlockSize;
+#endif
 					return;
 				}
 			}
 		}
 		FreeExternal(Ptr);
 	}
-
 	FORCEINLINE virtual bool GetAllocationSize(void *Ptr, SIZE_T &SizeOut) override
 	{
 		if (!IsOSAllocation(Ptr))
@@ -521,6 +648,9 @@ public:
 	void FreeExternal(void *Ptr);
 	bool GetAllocationSizeExternal(void* Ptr, SIZE_T& SizeOut);
 
+#if BINNED2_ALLOCATOR_STATS
+	int64 GetTotalAllocatedSmallPoolMemory() const;
+#endif
 	virtual void GetAllocatorStats( FGenericMemoryStats& out_Stats ) override;
 	/** Dumps current allocator stats to the log. */
 	virtual void DumpAllocatorStats(class FOutputDevice& Ar) override;
diff --git a/Engine/Source/Runtime/Core/Public/IOS/IOSPlatformMisc.h b/Engine/Source/Runtime/Core/Public/IOS/IOSPlatformMisc.h
--- a/Engine/Source/Runtime/Core/Public/IOS/IOSPlatformMisc.h
+++ b/Engine/Source/Runtime/Core/Public/IOS/IOSPlatformMisc.h
@@ -58,6 +58,12 @@ struct CORE_API FIOSPlatformMisc : public FGenericPlatformMisc
 	{
 		// Based on http://developer.apple.com/library/mac/#qa/qa1361/_index.html
 
+		extern CORE_API bool GIgnoreDebugger;
+		if (GIgnoreDebugger)
+		{
+			return false;
+		}
+
 		struct kinfo_proc Info;
 		int32 Mib[] = { CTL_KERN, KERN_PROC, KERN_PROC_PID, getpid() };
 		SIZE_T Size = sizeof(Info);
diff --git a/Engine/Source/Runtime/Core/Public/Mac/MacPlatformMisc.h b/Engine/Source/Runtime/Core/Public/Mac/MacPlatformMisc.h
--- a/Engine/Source/Runtime/Core/Public/Mac/MacPlatformMisc.h
+++ b/Engine/Source/Runtime/Core/Public/Mac/MacPlatformMisc.h
@@ -35,6 +35,13 @@ struct CORE_API FMacPlatformMisc : public FGenericPlatformMisc
 #if !UE_BUILD_SHIPPING
 	FORCEINLINE static bool IsDebuggerPresent()
 	{
+		extern CORE_API bool GIgnoreDebugger;
+
+		if (GIgnoreDebugger)
+		{
+			return false;
+		}
+
 		// Based on http://developer.apple.com/library/mac/#qa/qa1361/_index.html
 
 #pragma clang diagnostic push
diff --git a/Engine/Source/Runtime/Core/Public/Misc/TimeGuard.h b/Engine/Source/Runtime/Core/Public/Misc/TimeGuard.h
--- a/Engine/Source/Runtime/Core/Public/Misc/TimeGuard.h
+++ b/Engine/Source/Runtime/Core/Public/Misc/TimeGuard.h
@@ -7,14 +7,18 @@
 #include "CoreGlobals.h"
 #include "Misc/DateTime.h"
 #include "HAL/PlatformTime.h"
+#include "Delegates/Delegate.h"
 
 #ifndef DO_TIMEGUARD
-	// By default we are enabled based on STATS, but DO_TIMEGUARD can be set in XXX.Target.cs if so desired%
-	#define DO_TIMEGUARD STATS
+	// By default we are disabled, DO_TIMEGUARD can be set in XXX.Target.cs if so desired%
+	#define DO_TIMEGUARD 0
 #endif
 
 #if DO_TIMEGUARD
 
+
+DECLARE_DELEGATE_RetVal(const FString, FTimerNameDelegate);
+
 class FLightweightTimeGuard
 {
 
@@ -38,13 +42,28 @@ public:
 	};
 
 public:
-	FORCEINLINE FLightweightTimeGuard(TCHAR const* InName, const float InTargetMS = 0.0)
-		: Name(nullptr)
+
+
+	FORCEINLINE FLightweightTimeGuard(FTimerNameDelegate InNameDelegate, const float InTargetMS = 0.0)
+		: Name(nullptr), ObjectName(NAME_None)
+	{
+		NameDelegate = (bEnabled && IsInGameThread()) ? InNameDelegate : nullptr;
+
+		if (NameDelegate.IsBound())
+		{
+			TargetTimeMS = (InTargetMS > 0) ? InTargetMS : FrameTimeThresholdMS;
+			StartTime = FPlatformTime::Seconds();
+		}
+	}
+
+	FORCEINLINE FLightweightTimeGuard(TCHAR const* InName, FName InObjectName = NAME_None, const float InTargetMS = 0.0)
+		: Name(nullptr), ObjectName(NAME_None)
 	{
 		Name = (bEnabled && IsInGameThread()) ? InName : nullptr;
 
 		if (Name)
 		{
+			ObjectName = InObjectName;
 			TargetTimeMS = (InTargetMS > 0) ? InTargetMS : FrameTimeThresholdMS;
 			StartTime = FPlatformTime::Seconds();
 		}
@@ -61,7 +80,24 @@ public:
 
 			if (delta > TargetTimeMS)
 			{
-				ReportHitch(Name, delta);
+				if ( ObjectName != NAME_None )
+				{
+					ReportHitch( *FString::Printf(TEXT("%s %s"), Name, *ObjectName.ToString()), delta, true );
+				}
+				else
+				{
+					ReportHitch(Name, delta, false);
+				}
+			}
+		}
+		else if (NameDelegate.IsBound())
+		{
+			double delta = (FPlatformTime::Seconds() - StartTime) * 1000;
+
+			if (delta > TargetTimeMS)
+			{
+				FString Temp = NameDelegate.Execute();
+				ReportHitch(*Temp, delta, true);
 			}
 		}
 	}
@@ -75,8 +111,9 @@ public:
 protected:
 
 	// Used for reporting
-	static CORE_API void						ReportHitch(const TCHAR* InName, float TimeMS);
+	static CORE_API void						ReportHitch(const TCHAR* InName, float TimeMS, bool VolatileName );
 	static TMap<const TCHAR*, FGuardInfo>		HitchData;
+	static TSet<const TCHAR *>					VolatileNames; // any names which come in volatile we allocate them to a static string and put them in this array
 	static FCriticalSection						ReportMutex;
 	static CORE_API bool						bEnabled;
 	static CORE_API float						FrameTimeThresholdMS;
@@ -85,21 +122,44 @@ protected:
 
 	// Per-stat tracking
 	TCHAR const*	Name;
+	FName			ObjectName;
+	FTimerNameDelegate NameDelegate;
 	float			TargetTimeMS;
 	double			StartTime;
-
 };
 
 #define SCOPE_TIME_GUARD(name) \
 	FLightweightTimeGuard ANONYMOUS_VARIABLE(TimeGuard)(name);
 
 #define SCOPE_TIME_GUARD_MS(name, timeMs) \
-	FLightweightTimeGuard ANONYMOUS_VARIABLE(TimeGuard)(name, timeMs);
+	FLightweightTimeGuard ANONYMOUS_VARIABLE(TimeGuard)(name, NAME_None, timeMs);
+
+#define SCOPE_TIME_GUARD_NAMED(name, fname) \
+	FLightweightTimeGuard ANONYMOUS_VARIABLE(TimeGuard)(name, fname);
+
+#define SCOPE_TIME_GUARD_NAMED_MS(name, fname, timems) \
+	FLightweightTimeGuard ANONYMOUS_VARIABLE(TimeGuard)(name, fname, timems);
+
+
+
+
+#define SCOPE_TIME_GUARD_DELEGATE(inDelegate) \
+	FLightweightTimeGuard ANONYMOUS_VARIABLE(TimeGuard)(inDelegate);
+
+#define SCOPE_TIME_GUARD_DELEGATE_MS(inDelegate, timems) \
+	FLightweightTimeGuard ANONYMOUS_VARIABLE(TimeGuard)(inDelegate, timems);
+
+#define CLEAR_TIME_GUARDS FLightweightTimeGuard::ClearData
 
+#define ENABLE_TIME_GUARDS(bEnabled) FLightweightTimeGuard::SetEnabled(bEnabled)
 
 #else
 
 #define SCOPE_TIME_GUARD(name)
 #define SCOPE_TIME_GUARD_MS(name, timeMs)
+#define SCOPE_TIME_GUARD_NAMED(name, fname)
+#define SCOPE_TIME_GUARD_NAMED_MS(name, fname, timeMs)
+#define CLEAR_TIME_GUARDS()
+#define ENABLE_TIME_GUARDS(bEnabled)
 
 #endif // DO_TIMEGUARD
diff --git a/Engine/Source/Runtime/CoreUObject/Private/Serialization/AsyncLoading.cpp b/Engine/Source/Runtime/CoreUObject/Private/Serialization/AsyncLoading.cpp
--- a/Engine/Source/Runtime/CoreUObject/Private/Serialization/AsyncLoading.cpp
+++ b/Engine/Source/Runtime/CoreUObject/Private/Serialization/AsyncLoading.cpp
@@ -3001,7 +3001,7 @@ UObject* FAsyncPackage::EventDrivenIndexToObject(FPackageIndex Index, bool bChec
 
 		if (!Result)
 		{
-			UE_LOG(LogStreaming, Fatal, TEXT("Missing Dependency, request for %s but it hasn't been created yet."), *Linker->ImpExp(Index).ObjectName.ToString());
+			UE_LOG(LogStreaming, Error, TEXT("Missing Dependency, request for %s but it hasn't been created yet."), *Linker->ImpExp(Index).ObjectName.ToString());
 		}
 		else if (EventNodeArray.GetNode(MyDependentNode, false).bAddedToGraph || !EventNodeArray.GetNode(MyDependentNode, false).bFired)
 		{
diff --git a/Engine/Source/Runtime/CoreUObject/Private/UObject/GarbageCollection.cpp b/Engine/Source/Runtime/CoreUObject/Private/UObject/GarbageCollection.cpp
--- a/Engine/Source/Runtime/CoreUObject/Private/UObject/GarbageCollection.cpp
+++ b/Engine/Source/Runtime/CoreUObject/Private/UObject/GarbageCollection.cpp
@@ -1435,7 +1435,6 @@ void CollectGarbageInternal(EObjectFlags KeepFlags, bool bPerformFullPurge)
 
 			FCoreUObjectDelegates::PreGarbageCollectConditionalBeginDestroy.Broadcast();
 
-			// This nulls all weak references to unreachable objects, and may clear the pools
 			FGCArrayPool::Get().ClearWeakReferences(bPerformFullPurge);
 
 			// Unhash all unreachable objects.
@@ -1492,6 +1491,8 @@ void CollectGarbageInternal(EObjectFlags KeepFlags, bool bPerformFullPurge)
 					FScopedCBDProfile Profile(Object);
 					Object->ConditionalBeginDestroy();
 				}
+
+
 			}
 
 			UE_LOG(LogGarbage, Log, TEXT("%f ms for unhashing unreachable objects. Clusters removed: %d.   Items %d Cluster Items %d"), (FPlatformTime::Seconds() - StartTime) * 1000, ClustersRemoved, Items, ClusterItems);
diff --git a/Engine/Source/Runtime/CoreUObject/Private/UObject/UObjectGlobals.cpp b/Engine/Source/Runtime/CoreUObject/Private/UObject/UObjectGlobals.cpp
--- a/Engine/Source/Runtime/CoreUObject/Private/UObject/UObjectGlobals.cpp
+++ b/Engine/Source/Runtime/CoreUObject/Private/UObject/UObjectGlobals.cpp
@@ -7,6 +7,7 @@
 #include "UObject/UObjectGlobals.h"
 #include "HAL/PlatformFilemanager.h"
 #include "HAL/FileManager.h"
+#include "Misc/Paths.h"
 #include "Misc/ITransaction.h"
 #include "Serialization/ArchiveProxy.h"
 #include "Misc/CommandLine.h"
@@ -2244,7 +2245,7 @@ UObject* StaticAllocateObject
 	if (bCreatingCDO)
 	{
 		check(InClass->GetClass());
-		if( !GIsDuplicatingClassForReinstancing )
+		if( !GIsDuplicatingClassForReinstancing || InClass->HasAnyClassFlags(CLASS_Native) )
 		{
 			InName = InClass->GetDefaultObjectName();
 		}
@@ -3740,6 +3741,66 @@ COREUOBJECT_API UFunction* FindDelegateSignature(FName DelegateSignatureName)
 	return nullptr;
 }
 
+/**
+ * Takes a path of some sort and attempts to turn it into the asset log's canonical path.
+ */
+FString FAssetMsg::FormatPathForAssetLog(const TCHAR* InPath)
+{
+	static bool Once = false;
+	static bool ShowDiskPath = true;
+
+	if (!Once)
+	{
+		GConfig->GetBool(TEXT("Core.System"), TEXT("AssetLogShowsDiskPath"), ShowDiskPath, GEngineIni);
+		Once = true;
+	}
+
+	if (FPlatformProperties::RequiresCookedData() || !ShowDiskPath)
+	{
+		return FString(InPath);
+	}
+	
+	FString AssetPath = InPath;
+	FString FilePath;
+
+	// check for /Game/Path/Package.obj and turn it into a package reference
+	if (FPackageName::IsValidObjectPath(AssetPath))
+	{
+		AssetPath = FPackageName::ObjectPathToPackageName(AssetPath);
+	}
+
+	// Try to convert this to a file path
+	if (FPackageName::DoesPackageExist(AssetPath, 0, &FilePath) == false)
+	{
+		// if failed, assume we were given something that's a file path (e.g. ../../../Game/Whatever)
+		FilePath = AssetPath;
+	}
+
+	// if that succeeded FilePath will be a relative path to a  file, if not just assume that's what we were given and proceed...
+	if (IFileManager::Get().FileExists(*FilePath) == false)
+	{
+		return FString::Printf(TEXT("%s (no disk path found)"), InPath);
+	}
+
+	// turn this into an absolute path for error logging
+	FilePath = IFileManager::Get().ConvertToAbsolutePathForExternalAppForRead(*FilePath);
+	
+	// turn into a native platform file
+	FPaths::MakePlatformFilename(FilePath);
+	return FilePath;
+}
+
+/**
+ * Format the path of the passed in object
+ */
+FString FAssetMsg::FormatPathForAssetLog(const UObject* Object)
+{
+	FString ResolvedPath;
+	ResolvedPath = FormatPathForAssetLog(*Object->GetPathName());
+	return ResolvedPath;
+}
+
+
 namespace UE4CodeGen_Private
 {
 	void ConstructUProperty(UObject* Outer, const FPropertyParamsBase* const*& PropertyArray, int32& NumProperties)
@@ -3914,7 +3975,6 @@ namespace UE4CodeGen_Private
 
 					SetBit(Buffer.Get());
 
-					// Here we are making the assumption that bitfields are aligned in the struct. Probably true.
 					// If not, it may be ok unless we are on a page boundary or something, but the check will fire in that case.
 					// Have faith.
 					for (uint32 TestOffset = 0; TestOffset < SizeOf; TestOffset++)
@@ -4060,7 +4120,6 @@ namespace UE4CodeGen_Private
 				const FArrayPropertyParams* Prop = (const FArrayPropertyParams*)PropBase;
 				NewProp = new (EC_InternalUseOnlyConstructor, Outer, UTF8_TO_TCHAR(Prop->NameUTF8), Prop->ObjectFlags) UArrayProperty(FObjectInitializer(), EC_CppProperty, Prop->Offset, Prop->PropertyFlags);
 
-				// Next property is the array inner
 				ReadMore = 1;
 
 #if WITH_METADATA
@@ -4074,7 +4133,6 @@ namespace UE4CodeGen_Private
 			{
 				const FMapPropertyParams* Prop = (const FMapPropertyParams*)PropBase;
 				NewProp = new (EC_InternalUseOnlyConstructor, Outer, UTF8_TO_TCHAR(Prop->NameUTF8), Prop->ObjectFlags) UMapProperty(FObjectInitializer(), EC_CppProperty, Prop->Offset, Prop->PropertyFlags);
-
 				// Next two properties are the map key and value inners
 				ReadMore = 2;
 
@@ -4089,7 +4147,6 @@ namespace UE4CodeGen_Private
 			{
 				const FSetPropertyParams* Prop = (const FSetPropertyParams*)PropBase;
 				NewProp = new (EC_InternalUseOnlyConstructor, Outer, UTF8_TO_TCHAR(Prop->NameUTF8), Prop->ObjectFlags) USetProperty(FObjectInitializer(), EC_CppProperty, Prop->Offset, Prop->PropertyFlags);
-
 				// Next property is the set inner
 				ReadMore = 1;
 
@@ -4153,7 +4210,6 @@ namespace UE4CodeGen_Private
 				const FEnumPropertyParams* Prop = (const FEnumPropertyParams*)PropBase;
 				NewProp = new (EC_InternalUseOnlyConstructor, Outer, UTF8_TO_TCHAR(Prop->NameUTF8), Prop->ObjectFlags) UEnumProperty(FObjectInitializer(), EC_CppProperty, Prop->Offset, Prop->PropertyFlags, Prop->EnumFunc ? Prop->EnumFunc() : nullptr);
 
-				// Next property is the underlying integer property
 				ReadMore = 1;
 
 #if WITH_METADATA
diff --git a/Engine/Source/Runtime/CoreUObject/Public/UObject/AssetPtr.h b/Engine/Source/Runtime/CoreUObject/Public/UObject/AssetPtr.h
--- a/Engine/Source/Runtime/CoreUObject/Public/UObject/AssetPtr.h
+++ b/Engine/Source/Runtime/CoreUObject/Public/UObject/AssetPtr.h
@@ -11,6 +11,7 @@
 #include "Templates/Casts.h"
 #include "UObject/PersistentObjectPtr.h"
 #include "Misc/StringAssetReference.h"
+#include "PackageName.h"
 
 /**
  * FAssetPtr is a type of weak pointer to a UObject, that also keeps track of the path to the object on disk.
@@ -287,6 +288,17 @@ public:
 		return Get() != nullptr;
 	}
 
+	/**
+	* Test if this points to an asset on disk
+	*
+	* @return true if the file exists
+	*/
+	FORCEINLINE bool DoesAssetExist() const
+	{
+		// This does the runtime type check
+		return FPackageName::DoesPackageExist(GetLongPackageName());
+	}
+
 	/**  
 	 * Test if this does not point to a live UObject, but may in the future
 	 *
diff --git a/Engine/Source/Runtime/CoreUObject/Public/UObject/UObjectBaseUtility.h b/Engine/Source/Runtime/CoreUObject/Public/UObject/UObjectBaseUtility.h
--- a/Engine/Source/Runtime/CoreUObject/Public/UObject/UObjectBaseUtility.h
+++ b/Engine/Source/Runtime/CoreUObject/Public/UObject/UObjectBaseUtility.h
@@ -134,7 +134,8 @@ public:
 	/**
 	 * Marks this object as RF_PendingKill.
 	 */
-	FORCEINLINE void MarkPendingKill()
+	// Daniel L: Temporarily Virtual to track down OR-40458 remove virtual!
+	virtual void MarkPendingKill()
 	{
 		check(!IsRooted());
 		GUObjectArray.IndexToObject(InternalIndex)->SetPendingKill();
diff --git a/Engine/Source/Runtime/CoreUObject/Public/UObject/UObjectGlobals.h b/Engine/Source/Runtime/CoreUObject/Public/UObject/UObjectGlobals.h
--- a/Engine/Source/Runtime/CoreUObject/Public/UObject/UObjectGlobals.h
+++ b/Engine/Source/Runtime/CoreUObject/Public/UObject/UObjectGlobals.h
@@ -1940,6 +1940,52 @@ struct COREUOBJECT_API FDynamicClassStaticData
 
 COREUOBJECT_API TMap<FName, FDynamicClassStaticData>& GetDynamicClassMap();
 
+/**
+* FAssetMsg
+* This struct contains functions for asset-related messaging
+**/
+struct FAssetMsg
+{
+	/** Formats a path for the UE_ASSET_LOG macro */
+	static COREUOBJECT_API FString FormatPathForAssetLog(const TCHAR* Path);
+
+	/** If possible, finds a path to the underlying asset for the provided object and formats it for the UE_ASSET_LOG macro */
+	static COREUOBJECT_API FString FormatPathForAssetLog(const UObject* Object);
+};
+
+#if NO_LOGGING
+	#define UE_ASSET_LOG(...)
+#else
+	/**
+	* A  macro that outputs a formatted message to log with a canonical reference to an asset if a given logging category is active at a given verbosity level
+	* @param CategoryName name of the logging category
+	* @param Verbosity, verbosity level to test against
+	* @param Asset, Object or asset path to format
+	* @param Format, format text
+	***/
+	#define UE_ASSET_LOG(CategoryName, Verbosity, Asset, Format, ...) \
+	{ \
+		static_assert(IS_TCHAR_ARRAY(Format), "Formatting string must be a TCHAR array."); \
+		static_assert((ELogVerbosity::Verbosity & ELogVerbosity::VerbosityMask) < ELogVerbosity::NumVerbosity && ELogVerbosity::Verbosity > 0, "Verbosity must be constant and in range."); \
+		CA_CONSTANT_IF((ELogVerbosity::Verbosity & ELogVerbosity::VerbosityMask) <= ELogVerbosity::COMPILED_IN_MINIMUM_VERBOSITY && (ELogVerbosity::Warning & ELogVerbosity::VerbosityMask) <= FLogCategory##CategoryName::CompileTimeVerbosity) \
+		{ \
+			UE_LOG_EXPAND_IS_FATAL(Verbosity, PREPROCESSOR_NOTHING, if (!CategoryName.IsSuppressed(ELogVerbosity::Verbosity))) \
+			{ \
+				FString NewFormat = FString::Printf(TEXT("%s: %s"), *FAssetMsg::FormatPathForAssetLog(Asset), Format);\
+				FMsg::Logf_Internal(__FILE__, __LINE__, CategoryName.GetCategoryName(), ELogVerbosity::Verbosity, *NewFormat, ##__VA_ARGS__); \
+				UE_LOG_EXPAND_IS_FATAL(Verbosity, \
+					{ \
+						_DebugBreakAndPromptForRemote(); \
+						FDebug::AssertFailed("", __FILE__, __LINE__, *NewFormat, ##__VA_ARGS__); \
+						CA_ASSUME(false); \
+					}, \
+					PREPROCESSOR_NOTHING \
+				) \
+			} \
+		} \
+	}
+#endif // NO_LOGGING
+
 #if WITH_EDITOR
 /** 
  * Returns if true if the object is editor-only:
diff --git a/Engine/Source/Runtime/Engine/Classes/Animation/AnimInstance.h b/Engine/Source/Runtime/Engine/Classes/Animation/AnimInstance.h
--- a/Engine/Source/Runtime/Engine/Classes/Animation/AnimInstance.h
+++ b/Engine/Source/Runtime/Engine/Classes/Animation/AnimInstance.h
@@ -1129,7 +1129,7 @@ protected:
 	void UpdateMontageEvaluationData();
 
 	/** Called to setup for updates */
-	void PreUpdateAnimation(float DeltaSeconds);
+	virtual void PreUpdateAnimation(float DeltaSeconds);
 
 	/** update animation curves to component */
 	void UpdateCurvesToComponents(USkeletalMeshComponent* Component);
diff --git a/Engine/Source/Runtime/Engine/Classes/Animation/AnimSequence.h b/Engine/Source/Runtime/Engine/Classes/Animation/AnimSequence.h
--- a/Engine/Source/Runtime/Engine/Classes/Animation/AnimSequence.h
+++ b/Engine/Source/Runtime/Engine/Classes/Animation/AnimSequence.h
@@ -683,6 +683,11 @@ public:
 	// End Transform related functions 
 
 	// Begin Memory related functions
+
+	/** @return	estimate uncompressed raw size. This is *not* the real raw size. 
+				Here we estimate what it would be with no trivial compression. */
+	int32 GetUncompressedRawSize() const;
+
 	/**
 	 * @return		The approximate size of raw animation data.
 	 */
diff --git a/Engine/Source/Runtime/Engine/Classes/Components/PrimitiveComponent.h b/Engine/Source/Runtime/Engine/Classes/Components/PrimitiveComponent.h
--- a/Engine/Source/Runtime/Engine/Classes/Components/PrimitiveComponent.h
+++ b/Engine/Source/Runtime/Engine/Classes/Components/PrimitiveComponent.h
@@ -197,6 +197,12 @@ public:
 	UPROPERTY(Category=LOD, AdvancedDisplay, VisibleAnywhere, BlueprintReadOnly, meta=(DisplayName="Current Max Draw Distance") )
 	float CachedMaxDrawDistance;
 
+#if WITH_EDITORONLY_DATA
+	/** If true, and if World setting has bEnableHierarchicalLOD equal to true, then this component will be included when generating a Proxy mesh for the parent Actor */
+	UPROPERTY(EditAnywhere, AdvancedDisplay, Category = HLOD, meta = (DisplayName = "Include Component for HLOD Mesh generation"))
+	uint8 bEnableAutoLODGeneration : 1;
+#endif 
+
 	/** The scene depth priority group to draw the primitive in. */
 	UPROPERTY()
 	TEnumAsByte<enum ESceneDepthPriorityGroup> DepthPriorityGroup;
diff --git a/Engine/Source/Runtime/Engine/Classes/Engine/Texture.h b/Engine/Source/Runtime/Engine/Classes/Engine/Texture.h
--- a/Engine/Source/Runtime/Engine/Classes/Engine/Texture.h
+++ b/Engine/Source/Runtime/Engine/Classes/Engine/Texture.h
@@ -714,6 +714,9 @@ public:
 	ENGINE_API virtual uint32 GetMaximumDimension() const;
 #endif
 
+	// Daniel L: Temporary code to track down OR-40458 REMOVE ME
+	ENGINE_API virtual void MarkPendingKill() override;
+
 	/** @return the width of the surface represented by the texture. */
 	virtual float GetSurfaceWidth() const PURE_VIRTUAL(UTexture::GetSurfaceWidth,return 0;);
 
diff --git a/Engine/Source/Runtime/Engine/Classes/Engine/World.h b/Engine/Source/Runtime/Engine/Classes/Engine/World.h
--- a/Engine/Source/Runtime/Engine/Classes/Engine/World.h
+++ b/Engine/Source/Runtime/Engine/Classes/Engine/World.h
@@ -3240,9 +3240,12 @@ public:
 	// Post duplication event.
 	static FWorldPostDuplicateEvent OnPostDuplicate;
 
-	// Callback for world cleanup
+	// Callback for world cleanup start
 	static FWorldCleanupEvent OnWorldCleanup;
 
+	// Callback for world cleanup end
+	static FWorldCleanupEvent OnPostWorldCleanup;
+
 	// Callback for world destruction (only called for initialized worlds)
 	static FWorldEvent OnPreWorldFinishDestroy;
 
diff --git a/Engine/Source/Runtime/Engine/Classes/GameFramework/Actor.h b/Engine/Source/Runtime/Engine/Classes/GameFramework/Actor.h
--- a/Engine/Source/Runtime/Engine/Classes/GameFramework/Actor.h
+++ b/Engine/Source/Runtime/Engine/Classes/GameFramework/Actor.h
@@ -618,7 +618,7 @@ public:
 	uint8 bIgnoresOriginShifting:1;
 	
 	/** If true, and if World setting has bEnableHierarchicalLOD equal to true, then it will generate LODActor from groups of clustered Actor */
-	UPROPERTY(EditAnywhere, AdvancedDisplay, Category=Actor)
+	UPROPERTY(EditAnywhere, AdvancedDisplay, Category = HLOD, meta = (DisplayName = "Include Actor for HLOD Mesh generation"))
 	uint8 bEnableAutoLODGeneration:1;
 
 private:
diff --git a/Engine/Source/Runtime/Engine/Classes/GameFramework/PlayerController.h b/Engine/Source/Runtime/Engine/Classes/GameFramework/PlayerController.h
--- a/Engine/Source/Runtime/Engine/Classes/GameFramework/PlayerController.h
+++ b/Engine/Source/Runtime/Engine/Classes/GameFramework/PlayerController.h
@@ -1555,7 +1555,7 @@ public:
 	virtual bool DefaultCanUnpause();
 
 	/** @return true if game is currently paused. */
-	bool IsPaused();
+	bool IsPaused() const;
 
 	bool InputEnabled() const { return bInputEnabled; }
 
diff --git a/Engine/Source/Runtime/Engine/Classes/Kismet/KismetMathLibrary.h b/Engine/Source/Runtime/Engine/Classes/Kismet/KismetMathLibrary.h
--- a/Engine/Source/Runtime/Engine/Classes/Kismet/KismetMathLibrary.h
+++ b/Engine/Source/Runtime/Engine/Classes/Kismet/KismetMathLibrary.h
@@ -1482,6 +1482,11 @@ class ENGINE_API UKismetMathLibrary : public UBlueprintFunctionLibrary
 	UFUNCTION(BlueprintPure, Category = "Math|Vector2D", meta = (NativeBreakFunc))
 	static void GetYawPitchFromVector(FVector InVec, float& Yaw, float& Pitch);
 
+	/** Breaks a direction vector apart into Azimuth (Yaw) and Elevation (Pitch) rotation values given in degrees. (non-clamped)
+	 Relative to the provided reference frame (an Actor's WorldTransform for example) */
+	UFUNCTION(BlueprintPure, Category = "Math|Vector", meta = (NativeBreakFunc))
+	static void GetAzimuthAndElevation(FVector InDirection, const FTransform& ReferenceFrame, float& Azimuth, float& Elevation);
+
 	/** Makes a rotator {Roll, Pitch, Yaw} from rotation values supplied in degrees */
 	UFUNCTION(BlueprintPure, Category="Math|Rotator", meta=(Keywords="construct build rotation rotate rotator makerotator", NativeMakeFunc))
 	static FRotator MakeRotator(
diff --git a/Engine/Source/Runtime/Engine/Private/ActorEditor.cpp b/Engine/Source/Runtime/Engine/Private/ActorEditor.cpp
--- a/Engine/Source/Runtime/Engine/Private/ActorEditor.cpp
+++ b/Engine/Source/Runtime/Engine/Private/ActorEditor.cpp
@@ -877,7 +877,7 @@ void AActor::SetLODParent(UPrimitiveComponent* InLODParent, float InParentDrawDi
 	TArray<UPrimitiveComponent*> ComponentsToBeReplaced;
 	GetComponents(ComponentsToBeReplaced);
 
-	for(auto& Component : ComponentsToBeReplaced)
+	for(UPrimitiveComponent* Component : ComponentsToBeReplaced)
 	{
 		// parent primitive will be null if no LOD parent is selected
 		Component->SetLODParentPrimitive(InLODParent);
diff --git a/Engine/Source/Runtime/Engine/Private/Animation/AnimCompress.cpp b/Engine/Source/Runtime/Engine/Private/Animation/AnimCompress.cpp
--- a/Engine/Source/Runtime/Engine/Private/Animation/AnimCompress.cpp
+++ b/Engine/Source/Runtime/Engine/Private/Animation/AnimCompress.cpp
@@ -859,39 +859,7 @@ void UAnimCompress::FilterTrivialKeys(
 
 void UAnimCompress::FilterAnimRotationOnlyKeys(TArray<FTranslationTrack> & PositionTracks, UAnimSequence* AnimSeq)
 {
-#if( REDUCE_ANIMROTATIONONLY_TRACKS )
 
-	USkeleton* Skeleton = AnimSeq->Skeleton;
-	if (!Skeleton)
-	{
-		return;
-	}
-
-	// Go through all tracks
-	for(int32 TrackIndex=0; TrackIndex<PositionTracks.Num(); TrackIndex++)
-	{
-		FTranslationTrack & Track = PositionTracks[TrackIndex];
-
-		// We only care if we have more than 1 key. Can't reduce otherwise.
-		if( Track.Times.Num() > 1 )
-		{
-			int32 BoneTreeIndex = AnimSeq->GetSkeletonIndexFromRawDataTrackIndex(TrackIndex);
-			FName const BoneName = Skeleton->GetBoneName(BoneTreeIndex);
-			const bool bReduceTranslationTrack = (Skeleton->GetBoneTranslationRetargetingMode(BoneTreeIndex) == EBoneTranslationRetargetingMode::Skeleton);
-
-			// If it's a bAnimRotationOnly track, then reduce it to one key.
-			if( bReduceTranslationTrack )	
-			{
-				Track.PosKeys.RemoveAt(1, Track.PosKeys.Num()- 1);
-				Track.PosKeys.Shrink();
-				Track.Times.RemoveAt(1, Track.Times.Num()- 1);
-				Track.Times.Shrink();
-				Track.Times[0] = 0.0f;
-			}
-		}
-	}
-
-#endif
 }
 
 
diff --git a/Engine/Source/Runtime/Engine/Private/Animation/AnimCompress_PerTrackCompression.cpp b/Engine/Source/Runtime/Engine/Private/Animation/AnimCompress_PerTrackCompression.cpp
--- a/Engine/Source/Runtime/Engine/Private/Animation/AnimCompress_PerTrackCompression.cpp
+++ b/Engine/Source/Runtime/Engine/Private/Animation/AnimCompress_PerTrackCompression.cpp
@@ -1177,20 +1177,9 @@ void UAnimCompress_PerTrackCompression::CompressUsingUnderlyingCompressor(
 			AnimSeq->CompressedScaleOffsets.SetOffsetData(TrackIndex, 0, ScaleOffset);
 		}
 
-		// See if we can skip saving translation track.
-		bool bSkipTranslationTrack = false;
-#if( SKIP_FORCEMESHTRANSLATION_TRACKS || SKIP_ANIMROTATIONONLY_TRACKS )
-		{
-			USkeleton* Skeleton = AnimSeq->Skeleton;
-			check (Skeleton);
-			int32 const BoneTreeIndex = AnimSeq->TrackToSkeletonMapTable[TrackIndex].BoneTreeIndex;
-			bSkipTranslationTrack = (Skeleton->GetBoneTranslationRetargetingMode(BoneTreeIndex) == EBoneTranslationRetargetingMode::Skeleton);
-		}
-#endif
-
 		// Now write out compression and translation frames into the stream
 		int32 TranslationOffset = INDEX_NONE;
-		if( !bSkipTranslationTrack && BestTranslation.CompressedBytes.Num() > 0 )
+		if (BestTranslation.CompressedBytes.Num() > 0 )
 		{
 			check(BestTranslation.ActualCompressionMode < ACF_MAX);
 			TranslationOffset = AnimSeq->CompressedByteStream.Num();
diff --git a/Engine/Source/Runtime/Engine/Private/Animation/AnimCompress_RemoveLinearKeys.cpp b/Engine/Source/Runtime/Engine/Private/Animation/AnimCompress_RemoveLinearKeys.cpp
--- a/Engine/Source/Runtime/Engine/Private/Animation/AnimCompress_RemoveLinearKeys.cpp
+++ b/Engine/Source/Runtime/Engine/Private/Animation/AnimCompress_RemoveLinearKeys.cpp
@@ -138,26 +138,49 @@ void FilterLinearKeysTemplate(
 	{
 		int32 LowKey = 0;
 		int32 HighKey = KeyCount-1;
-		int32 PrevKey = 0;
+
+		TArray<uint32> KnownParentTimes;
+		KnownParentTimes.Empty(KeyCount);
+		const int32 ParentKeyCount = ParentTimes ? ParentTimes->Num() : 0;
+		for (int32 TimeIndex = 0, ParentTimeIndex = 0; TimeIndex < KeyCount; TimeIndex++)
+		{
+			while ((ParentTimeIndex < ParentKeyCount) && (Times[TimeIndex] > (*ParentTimes)[ParentTimeIndex]))
+			{
+				ParentTimeIndex++;
+			}
+
+			KnownParentTimes.Add((ParentTimeIndex < ParentKeyCount) && (Times[TimeIndex] == (*ParentTimes)[ParentTimeIndex]));
+		}
+
+		TArray<FTransform> CachedInvRawBases;
+		CachedInvRawBases.Empty(KeyCount);
+		for (int32 FrameIndex = 0; FrameIndex < KeyCount; ++FrameIndex)
+		{
+			const FTransform& RawBase = RawWorldBones[(BoneIndex*NumFrames) + FrameIndex];
+			CachedInvRawBases.Add(RawBase.Inverse());
+		}
 		
 		// copy the low key (this one is a given)
 		NewTimes.Add(Times[0]);
 		NewKeys.Add(Keys[0]);
 
-		FTransform DummyBone(FQuat::Identity, FVector(END_EFFECTOR_SOCKET_DUMMY_BONE_SIZE, END_EFFECTOR_SOCKET_DUMMY_BONE_SIZE, END_EFFECTOR_SOCKET_DUMMY_BONE_SIZE));
+		const FTransform EndEffectorDummyBoneSocket(FQuat::Identity, FVector(END_EFFECTOR_DUMMY_BONE_LENGTH_SOCKET));
+		const FTransform EndEffectorDummyBone(FQuat::Identity, FVector(END_EFFECTOR_DUMMY_BONE_LENGTH));
 
-		float const DeltaThreshold = (BoneData[BoneIndex].IsEndEffector() && (BoneData[BoneIndex].bHasSocket || BoneData[BoneIndex].bKeyEndEffector)) ? EffectorDiffSocket : MaxTargetDelta;
+		const float DeltaThreshold = (BoneData[BoneIndex].IsEndEffector() && (BoneData[BoneIndex].bHasSocket || BoneData[BoneIndex].bKeyEndEffector)) ? EffectorDiffSocket : MaxTargetDelta;
 
 		// We will test within a sliding window between LowKey and HighKey.
 		// Therefore, we are done when the LowKey exceeds the range
-		while (LowKey < KeyCount-1)
+		while (LowKey + 1 < KeyCount)
 		{
-			// high key always starts at the top of the range
-			HighKey = KeyCount-1;
-
-			// keep testing until the window is closed
-			while (HighKey > LowKey+1)
+			int32 GoodHighKey = LowKey + 1;
+			int32 BadHighKey = KeyCount;
+			
+			// bisect until we find the lowest acceptable high key
+			while (BadHighKey - GoodHighKey >= 2)
 			{
+				HighKey = GoodHighKey + (BadHighKey - GoodHighKey) / 2;
+
 				// get the parameters of the window we are testing
 				const float LowTime = Times[LowKey];
 				const float HighTime = Times[HighKey];
@@ -190,7 +213,7 @@ void FilterLinearKeysTemplate(
 					{
 						// get the raw world transform for this bone (the original world-space position)
 						const int32 FrameIndex = TestKey;
-						const FTransform& RawBase = RawWorldBones[(BoneIndex*NumFrames) + FrameIndex];
+						const FTransform& InvRawBase = CachedInvRawBases[FrameIndex];
 						
 						// generate the proposed local bone atom and transform (local space)
 						FTransform ProposedTM = UpdateBoneAtom(BoneIndex, BoneAtoms[FrameIndex], LerpValue);
@@ -205,24 +228,30 @@ void FilterLinearKeysTemplate(
 							// find the offset transform from the raw base to the end effector
 							const int32 TargetBoneIndex = TargetBoneIndices[TargetIndex];
 							FTransform RawTarget = RawWorldBones[(TargetBoneIndex*NumFrames) + FrameIndex];
-							const FTransform& RelTM = RawTarget.GetRelativeTransform(RawBase); 
+							const FTransform RelTM = RawTarget * InvRawBase;
 
 							// forecast where the new end effector would be using our proposed key
 							FTransform ProposedTarget = RelTM * ProposedBase;
 
-							// If this is an EndEffector with a Socket attached to it, add an extra bone, to measure error introduced by effector rotation compression.
-							if( BoneData[TargetIndex].bHasSocket || BoneData[TargetIndex].bKeyEndEffector )
+							// If this is an EndEffector, add a dummy bone to measure the effect of compressing the rotation.
+							// Sockets and Key EndEffectors have a longer dummy bone to maintain higher precision.
+							if (BoneData[TargetIndex].bHasSocket || BoneData[TargetIndex].bKeyEndEffector)
 							{
-								ProposedTarget = DummyBone * ProposedTarget;
-								RawTarget = DummyBone * RawTarget;
+								ProposedTarget = EndEffectorDummyBoneSocket * ProposedTarget;
+								RawTarget = EndEffectorDummyBoneSocket * RawTarget;
+							}
+							else
+							{
+								ProposedTarget = EndEffectorDummyBone * ProposedTarget;
+								RawTarget = EndEffectorDummyBone * RawTarget;
 							}
 
 							// determine the extend of error at the target end effector
-							float ThisError = (ProposedTarget.GetTranslation() - RawTarget.GetTranslation()).Size();
+							const float ThisError = (ProposedTarget.GetTranslation() - RawTarget.GetTranslation()).Size();
 							TargetError = FMath::Max(TargetError, ThisError); 
 
 							// exit early when we encounter a large delta
-							float const TargetDeltaThreshold = BoneData[TargetIndex].bHasSocket ? EffectorDiffSocket : DeltaThreshold;
+							const float TargetDeltaThreshold = BoneData[TargetIndex].bHasSocket ? EffectorDiffSocket : DeltaThreshold;
 							if( TargetError > TargetDeltaThreshold )
 							{ 
 								break;
@@ -235,7 +264,7 @@ void FilterLinearKeysTemplate(
 					// making the skeleton more uniform in key distribution.
 					if (ParentTimes)
 					{
-						if (ParentTimes->Find(TestTime) != INDEX_NONE)
+						if (KnownParentTimes[TestKey])
 						{
 							// our parent has a key at this time, 
 							// inflate our perceived error to increase our sensitivity
@@ -259,39 +288,20 @@ void FilterLinearKeysTemplate(
 				}
 
 				// determine if the span succeeded. That is, the worst errors found are within tolerances
-				if (MaxLerpError <= MaxDelta &&
-					MaxTargetError <= DeltaThreshold)
+				if ((MaxLerpError <= MaxDelta) && (MaxTargetError <= DeltaThreshold))
 				{
-					// save the high end of the test span as our next key
-					NewTimes.Add(Times[HighKey]);
-					NewKeys.Add(Keys[HighKey]);
-
-					// start testing a new span
-					LowKey = HighKey;
-					HighKey =  KeyCount-1;
+					GoodHighKey = HighKey;
 				}
 				else
 				{
-					// we failed, shrink the test span window and repeat
-					--HighKey;
+					BadHighKey = HighKey;
 				}
 			}
 
-			// if the test window is still valid, accept the high key
-			if (HighKey > LowKey)
-			{
-				NewTimes.Add(Times[HighKey]);
-				NewKeys.Add(Keys[HighKey]);
-			}
-			LowKey= HighKey;
-		}
+			NewTimes.Add(Times[GoodHighKey]);
+			NewKeys.Add(Keys[GoodHighKey]);
 
-		// The process has ended, but we must make sure the last key is accounted for
-		if (NewTimes.Last() != Times.Last() &&
-			CalcDelta(Keys.Last(), NewKeys.Last()) >= MaxDelta )
-		{
-			NewTimes.Add(Times.Last());
-			NewKeys.Add(Keys.Last());
+			LowKey = GoodHighKey;
 		}
 
 		// return the new key set to the caller
diff --git a/Engine/Source/Runtime/Engine/Private/Animation/AnimCompressionDerivedData.h b/Engine/Source/Runtime/Engine/Private/Animation/AnimCompressionDerivedData.h
--- a/Engine/Source/Runtime/Engine/Private/Animation/AnimCompressionDerivedData.h
+++ b/Engine/Source/Runtime/Engine/Private/Animation/AnimCompressionDerivedData.h
@@ -44,7 +44,7 @@ public:
 		// This is a version string that mimics the old versioning scheme. If you
 		// want to bump this version, generate a new guid using VS->Tools->Create GUID and
 		// return it here. Ex.
-		return TEXT("4E35CA976C2D4FAB86128830166699A0");
+		return TEXT("486652B60D9C4989BFD91BDF463DB28D");
 	}
 
 	virtual FString GetPluginSpecificCacheKeySuffix() const override;
diff --git a/Engine/Source/Runtime/Engine/Private/Animation/AnimMontage.cpp b/Engine/Source/Runtime/Engine/Private/Animation/AnimMontage.cpp
--- a/Engine/Source/Runtime/Engine/Private/Animation/AnimMontage.cpp
+++ b/Engine/Source/Runtime/Engine/Private/Animation/AnimMontage.cpp
@@ -585,8 +585,8 @@ void UAnimMontage::AddBranchingPointMarker(FBranchingPointMarker TickMarker, TMa
 	FAnimNotifyEvent** FoundNotifyEventPtr = TriggerTimes.Find(TickMarker.TriggerTime);
 	if (FoundNotifyEventPtr)
 	{
-		UE_LOG(LogAnimMontage, Warning, TEXT("Branching Point '%s' overlaps with '%s' at time: %f. One of them will not get triggered! (%s)"),
-			*Notifies[TickMarker.NotifyIndex].NotifyName.ToString(), *(*FoundNotifyEventPtr)->NotifyName.ToString(), TickMarker.TriggerTime, *GetPathName());
+		UE_ASSET_LOG(LogAnimMontage, Warning, this, TEXT("Branching Point '%s' overlaps with '%s' at time: %f. One of them will not get triggered!"),
+			*Notifies[TickMarker.NotifyIndex].NotifyName.ToString(), *(*FoundNotifyEventPtr)->NotifyName.ToString(), TickMarker.TriggerTime);
 	}
 	else
 	{
diff --git a/Engine/Source/Runtime/Engine/Private/Animation/AnimNotifyState_Trail.cpp b/Engine/Source/Runtime/Engine/Private/Animation/AnimNotifyState_Trail.cpp
--- a/Engine/Source/Runtime/Engine/Private/Animation/AnimNotifyState_Trail.cpp
+++ b/Engine/Source/Runtime/Engine/Private/Animation/AnimNotifyState_Trail.cpp
@@ -92,7 +92,7 @@ void UAnimNotifyState_Trail::NotifyBegin(class USkeletalMeshComponent * MeshComp
 
 	if(PSTemplate == nullptr)
 	{
-		UE_LOG(LogParticles, Warning, TEXT("Trail Notify: Null PSTemplate for trail notify in anim: %s"), *GetNameSafe(Animation)); 
+		UE_LOG(LogParticles, Warning, TEXT("Trail Notify: Null PSTemplate for trail notify in anim: %s"), *GetPathNameSafe(Animation));
 		return;
 	}
 
diff --git a/Engine/Source/Runtime/Engine/Private/Animation/AnimSequence.cpp b/Engine/Source/Runtime/Engine/Private/Animation/AnimSequence.cpp
--- a/Engine/Source/Runtime/Engine/Private/Animation/AnimSequence.cpp
+++ b/Engine/Source/Runtime/Engine/Private/Animation/AnimSequence.cpp
@@ -172,14 +172,18 @@ void UAnimSequence::GetAssetRegistryTags(TArray<FAssetRegistryTag>& OutTags) con
 	}
 #endif
 
-	const float CompressionRatio = CompressedRawDataSize > 0 ? (float)GetApproxCompressedSize() / (float)CompressedRawDataSize : -1.f;
 
-	OutTags.Add(FAssetRegistryTag(TEXT("Compression Ratio"), FString::Printf(TEXT("%.03f"), CompressionRatio), FAssetRegistryTag::TT_Numerical));
+	OutTags.Add(FAssetRegistryTag(TEXT("Compression Ratio"), FString::Printf(TEXT("%.03f"), (float)GetApproxCompressedSize() / (float)GetUncompressedRawSize()), FAssetRegistryTag::TT_Numerical));
 	OutTags.Add(FAssetRegistryTag(TEXT("Compressed Size (KB)"), FString::Printf(TEXT("%.02f"), (float)GetApproxCompressedSize() / 1024.0f), FAssetRegistryTag::TT_Numerical));
 
 	Super::GetAssetRegistryTags(OutTags);
 }
 
+int32 UAnimSequence::GetUncompressedRawSize() const
+{
+	return ((sizeof(FVector) + sizeof(FQuat) + sizeof(FVector)) * RawAnimationData.Num() * NumFrames);
+}
+
 int32 UAnimSequence::GetApproxRawSize() const
 {
 	int32 Total = sizeof(FRawAnimSequenceTrack) * RawAnimationData.Num();
@@ -2103,9 +2107,12 @@ void UAnimSequence::RequestAnimCompression(bool bAsyncCompression, TSharedPtr<FA
 		FDerivedDataAnimationCompression* AnimCompressor = new FDerivedDataAnimationCompression(this, CompressContext, bDoCompressionInPlace);
 		// For debugging DDC/Compression issues		
 		const bool bSkipDDC = false;
-		if (bSkipDDC)
+		if (bSkipDDC || (CompressCommandletVersion == INDEX_NONE))
 		{
 			AnimCompressor->Build(OutData);
+
+			delete AnimCompressor;
+			AnimCompressor = nullptr;
 		}
 		else
 		{
