Copying //UE4/Dev-Framework to //UE4/Dev-Main (Source: //UE4/Dev-Framework @ 3582324)

diff --git a/Engine/Source/Editor/UnrealEd/Private/Kismet2/BlueprintEditorUtils.cpp b/Engine/Source/Editor/UnrealEd/Private/Kismet2/BlueprintEditorUtils.cpp
--- a/Engine/Source/Editor/UnrealEd/Private/Kismet2/BlueprintEditorUtils.cpp
+++ b/Engine/Source/Editor/UnrealEd/Private/Kismet2/BlueprintEditorUtils.cpp
@@ -123,8 +123,6 @@ DEFINE_STAT(EKismetCompilerStats_NotifyBlueprintChanged);
 DECLARE_CYCLE_STAT(TEXT("Mark Blueprint as Structurally Modified"), EKismetCompilerStats_MarkBlueprintasStructurallyModified, STATGROUP_KismetCompiler);
 DECLARE_CYCLE_STAT(TEXT("Refresh External DependencyNodes"), EKismetCompilerStats_RefreshExternalDependencyNodes, STATGROUP_KismetCompiler);
 
-FBlueprintEditorUtils::FFixLevelScriptActorBindingsEvent FBlueprintEditorUtils::FixLevelScriptActorBindingsEvent;
-
 struct FCompareNodePriority
 {
 	FORCEINLINE bool operator()( const UK2Node& A, const UK2Node& B ) const
@@ -762,7 +760,7 @@ UClass* FBlueprintEditorUtils::FindFirstNativeClass(UClass* Class)
 	return Class;
 }
 
-void FBlueprintEditorUtils::GetAllGraphNames(const UBlueprint* Blueprint, TArray<FName>& GraphNames)
+void FBlueprintEditorUtils::GetAllGraphNames(const UBlueprint* Blueprint, TSet<FName>& GraphNames)
 {
 	TArray< UEdGraph* > GraphList;
 	Blueprint->GetAllGraphs(GraphList);
@@ -782,7 +780,7 @@ void FBlueprintEditorUtils::GetAllGraphNames(const UBlueprint* Blueprint, TArray
 
 		for(int32 FunctionIndex = 0; FunctionIndex < ParentBP->FunctionGraphs.Num(); ++FunctionIndex)
 		{
-			GraphNames.AddUnique(ParentBP->FunctionGraphs[FunctionIndex]->GetFName());
+			GraphNames.Add(ParentBP->FunctionGraphs[FunctionIndex]->GetFName());
 		}
 	}
 }
@@ -1430,12 +1428,6 @@ UClass* FBlueprintEditorUtils::RegenerateBlueprintClass(UBlueprint* Blueprint, U
 
 		Blueprint->bIsRegeneratingOnLoad = false;
 
-		if( Package )
-		{
-			// Tell the linker to try to find exports in memory first, so that it gets the new, regenerated versions
-			Package->FindExportsInMemoryFirst(true);
-		}
-
 		bRegenerated = bShouldBeRecompiled;
 
 		if (!FKismetEditorUtilities::IsClassABlueprintSkeleton(ClassToRegenerate))
@@ -3038,7 +3030,7 @@ bool FBlueprintEditorUtils::IsDataOnlyBlueprint(const UBlueprint* Blueprint)
 		return false;
 	}
 
-	if(USimpleConstructionScript* SimpleConstructionScript = Blueprint->SimpleConstructionScript)
+	if (USimpleConstructionScript* SimpleConstructionScript = Blueprint->SimpleConstructionScript)
 	{
 		const TArray<USCS_Node*>& Nodes = SimpleConstructionScript->GetAllNodes();
 		if (Nodes.Num() > 1)
@@ -3084,20 +3076,13 @@ bool FBlueprintEditorUtils::IsDataOnlyBlueprint(const UBlueprint* Blueprint)
 		}
 	}
 
-	// No extra eventgraphs
-	if( Blueprint->UbergraphPages.Num() > 1 )
-	{
-		return false;
-	}
-
-	// EventGraph is empty
-	UEdGraph* EventGraph = (Blueprint->UbergraphPages.Num() == 1) ? Blueprint->UbergraphPages[0] : nullptr;
-	if( EventGraph && EventGraph->Nodes.Num() > 0 )
+	// All EventGraphs are empty (at least of non-ghost, non-disabled nodes)
+	for (UEdGraph* EventGraph : Blueprint->UbergraphPages)
 	{
-		for(UEdGraphNode* GraphNode : EventGraph->Nodes)
+		for (UEdGraphNode* GraphNode : EventGraph->Nodes)
 		{
-			// If there is an enabled node in the event graph (or a node that was explicitly disabled by the user), the Blueprint is not data only
-			if (GraphNode && (GraphNode->IsNodeEnabled() || GraphNode->bUserSetEnabledState))
+			// If there is an enabled node in the event graph, the Blueprint is not data only
+			if (GraphNode && (GraphNode->GetDesiredEnabledState() != ENodeEnabledState::Disabled))
 			{
 				return false;
 			}
@@ -3105,7 +3090,7 @@ bool FBlueprintEditorUtils::IsDataOnlyBlueprint(const UBlueprint* Blueprint)
 	}
 
 	// No implemented interfaces
-	if( Blueprint->ImplementedInterfaces.Num() > 0 )
+	if (Blueprint->ImplementedInterfaces.Num() > 0)
 	{
 		return false;
 	}
@@ -3415,20 +3400,6 @@ bool FBlueprintEditorUtils::MoveVariableBeforeVariable(UBlueprint* Blueprint, FN
 	return bMoved;
 }
 
-int32 FBlueprintEditorUtils::FindFirstNewVarOfCategory(const UBlueprint* Blueprint, FText Category)
-{
-	for(int32 VarIdx=0; VarIdx<Blueprint->NewVariables.Num(); VarIdx++)
-	{
-		if(Blueprint->NewVariables[VarIdx].Category.EqualTo(Category))
-		{
-			return VarIdx;
-		}
-	}
-	return INDEX_NONE;
-}
-
-
-
 int32 FBlueprintEditorUtils::FindTimelineIndex(const UBlueprint* Blueprint, const FName& InName) 
 {
 	const FName TimelineTemplateName = *UTimelineTemplate::TimelineVariableNameToTemplateName(InName);
@@ -3443,7 +3414,7 @@ int32 FBlueprintEditorUtils::FindTimelineIndex(const UBlueprint* Blueprint, cons
 	return INDEX_NONE;
 }
 
-void FBlueprintEditorUtils::GetSCSVariableNameList(const UBlueprint* Blueprint, TArray<FName>& VariableNames)
+void FBlueprintEditorUtils::GetSCSVariableNameList(const UBlueprint* Blueprint, TSet<FName>& VariableNames)
 {
 	if(Blueprint != nullptr && Blueprint->SimpleConstructionScript != nullptr)
 	{
@@ -3454,14 +3425,14 @@ void FBlueprintEditorUtils::GetSCSVariableNameList(const UBlueprint* Blueprint,
 				const FName VariableName = SCS_Node->GetVariableName();
 				if(VariableName != NAME_None)
 				{
-					VariableNames.AddUnique(VariableName);
+					VariableNames.Add(VariableName);
 				}
 			}
 		}
 	}
 }
 
-void FBlueprintEditorUtils::GetImplementingBlueprintsFunctionNameList(const UBlueprint* Blueprint, TArray<FName>& FunctionNames)
+void FBlueprintEditorUtils::GetImplementingBlueprintsFunctionNameList(const UBlueprint* Blueprint, TSet<FName>& FunctionNames)
 {
 	if(Blueprint != nullptr && FBlueprintEditorUtils::IsInterfaceBlueprint(Blueprint))
 	{
@@ -3835,17 +3806,27 @@ void FBlueprintEditorUtils::RemoveBlueprintVariableMetaData(UBlueprint* Blueprin
 
 void FBlueprintEditorUtils::SetBlueprintVariableCategory(UBlueprint* Blueprint, const FName& VarName, const UStruct* InLocalVarScope, const FText& NewCategory, bool bDontRecompile)
 {
-	const FScopedTransaction Transaction( LOCTEXT("ChangeVariableCategory", "Change Variable Category") );
-	Blueprint->Modify();
+	if (Blueprint == nullptr)
+	{
+		return;
+	}
 
 	// Ensure we always set a category
-	const UEdGraphSchema_K2* K2Schema = GetDefault<UEdGraphSchema_K2>();
 	FText SetCategory = NewCategory;
 	if (SetCategory.IsEmpty())
 	{
-		SetCategory = K2Schema->VR_DefaultCategory; 
+		SetCategory = UEdGraphSchema_K2::VR_DefaultCategory;
+	}
+	
+	const FText OldCategory = GetBlueprintVariableCategory(Blueprint, VarName, InLocalVarScope);
+	if (OldCategory.EqualTo(SetCategory))
+	{
+		return;
 	}
 
+	const FScopedTransaction Transaction(LOCTEXT("ChangeVariableCategory", "Change Variable Category"));
+	Blueprint->Modify();
+
 	UClass* SkeletonGeneratedClass = Blueprint->SkeletonGeneratedClass;
 	if (UProperty* TargetProperty = FindField<UProperty>(SkeletonGeneratedClass, VarName))
 	{
@@ -3860,57 +3841,126 @@ void FBlueprintEditorUtils::SetBlueprintVariableCategory(UBlueprint* Blueprint,
 			const int32 VarIndex = FBlueprintEditorUtils::FindNewVariableIndex(Blueprint, VarName);
 			if (VarIndex != INDEX_NONE)
 			{
-				bIsCategoryChanged = !Blueprint->NewVariables[VarIndex].Category.EqualTo(SetCategory);
-				
-				if(bIsCategoryChanged)
-				{
-					Blueprint->NewVariables[VarIndex].Category = SetCategory;
-				}
+				Blueprint->NewVariables[VarIndex].Category = SetCategory;
 			}
 			else
 			{
 				const int32 SCS_NodeIndex = FBlueprintEditorUtils::FindSCS_Node(Blueprint, VarName);
 				if (SCS_NodeIndex != INDEX_NONE)
 				{
-					bIsCategoryChanged = !Blueprint->SimpleConstructionScript->GetAllNodes()[SCS_NodeIndex]->CategoryName.EqualTo(SetCategory);
-					
-					if(bIsCategoryChanged)
-					{
-						Blueprint->SimpleConstructionScript->GetAllNodes()[SCS_NodeIndex]->Modify();
-						Blueprint->SimpleConstructionScript->GetAllNodes()[SCS_NodeIndex]->CategoryName = SetCategory;
-					}
+					Blueprint->SimpleConstructionScript->GetAllNodes()[SCS_NodeIndex]->Modify();
+					Blueprint->SimpleConstructionScript->GetAllNodes()[SCS_NodeIndex]->CategoryName = SetCategory;
 				}
 			}
+		}
+	}
+	else if (InLocalVarScope)
+	{
+		UK2Node_FunctionEntry* OutFunctionEntryNode;
+		if (FBPVariableDescription* LocalVariable = FindLocalVariable(Blueprint, InLocalVarScope, VarName, &OutFunctionEntryNode))
+		{
+			OutFunctionEntryNode->Modify();
+			LocalVariable->SetMetaData(TEXT("Category"), *SetCategory.ToString());
+			LocalVariable->Category = SetCategory;
+		}
+	}
+
+	if (bDontRecompile == false)
+	{
+		FBlueprintEditorUtils::MarkBlueprintAsStructurallyModified(Blueprint);
+	}
+}
 
-			if (bDontRecompile == false && bIsCategoryChanged)
+void FBlueprintEditorUtils::SetBlueprintFunctionOrMacroCategory(UEdGraph* Graph, const FText& InCategoryName, bool bDontRecompile)
+{
+	UBlueprint* Blueprint = FBlueprintEditorUtils::FindBlueprintForGraphChecked(Graph);
+	if (FKismetUserDeclaredFunctionMetadata* MetaData = FBlueprintEditorUtils::GetGraphFunctionMetaData(Graph))
+	{
+		UFunction* Function = nullptr;
+		for (TFieldIterator<UFunction> FunctionIt(Blueprint->SkeletonGeneratedClass, EFieldIteratorFlags::IncludeSuper); FunctionIt; ++FunctionIt)
+		{
+			if (FunctionIt->GetName() == Graph->GetName())
 			{
-				FBlueprintEditorUtils::MarkBlueprintAsStructurallyModified(Blueprint);
+				Function = *FunctionIt;
+				break;
 			}
 		}
+
+		const FText& NewCategory = InCategoryName.IsEmpty() ? UEdGraphSchema_K2::VR_DefaultCategory : InCategoryName;
+		MetaData->Category = NewCategory;
+
+		if (Function)
+		{
+			check(!Function->IsNative()); // Should never get here with a native function, as we wouldn't have been able to find metadata for it
+			Function->Modify();
+			Function->SetMetaData(FBlueprintMetadata::MD_FunctionCategory, *NewCategory.ToString());
+		}
+
+		if (!bDontRecompile)
+		{
+			FBlueprintEditorUtils::MarkBlueprintAsModified(Blueprint);
+		}
 	}
-	else if(InLocalVarScope)
+}
+
+int32 FBlueprintEditorUtils::FindIndexOfGraphInParent(UEdGraph* Graph)
+{
+	int32 Result = INDEX_NONE;
+
+	if (UBlueprint* Blueprint = FBlueprintEditorUtils::FindBlueprintForGraph(Graph))
 	{
-		UK2Node_FunctionEntry* OutFunctionEntryNode;
-		if(FBPVariableDescription* LocalVariable = FindLocalVariable(Blueprint, InLocalVarScope, VarName, &OutFunctionEntryNode))
+		Result = Blueprint->FunctionGraphs.IndexOfByKey(Graph);
+		if (Result == INDEX_NONE)
 		{
-			// If the category does not change, we will not recompile the Blueprint
-			bool bIsCategoryChanged = !LocalVariable->Category.EqualTo(SetCategory);
+			Result = Blueprint->MacroGraphs.IndexOfByKey(Graph);
+		}
+	}
+
+	return Result;
+}
+
+bool FBlueprintEditorUtils::MoveGraphBeforeOtherGraph(UEdGraph* Graph, int32 NewIndex, bool bDontRecompile)
+{
+	if (UBlueprint* Blueprint = FBlueprintEditorUtils::FindBlueprintForGraph(Graph))
+	{
+		bool bModified = false;
 
-			if(bIsCategoryChanged)
+		const int32 OldFunctionIndex = Blueprint->FunctionGraphs.IndexOfByKey(Graph);
+		if (OldFunctionIndex != INDEX_NONE)
+		{
+			if ((OldFunctionIndex != NewIndex) && Blueprint->FunctionGraphs.IsValidIndex(NewIndex))
 			{
-				OutFunctionEntryNode->Modify();
-				LocalVariable->SetMetaData(TEXT("Category"), *SetCategory.ToString());
-				LocalVariable->Category = SetCategory;
+				Blueprint->Modify();
+				Blueprint->FunctionGraphs.Insert(Graph, NewIndex);
+				Blueprint->FunctionGraphs.RemoveAt((OldFunctionIndex < NewIndex) ? OldFunctionIndex : (OldFunctionIndex + 1));
+				bModified = true;
 			}
+		}
 
-			if (bDontRecompile == false && bIsCategoryChanged)
+		const int32 OldMacroIndex = Blueprint->MacroGraphs.IndexOfByKey(Graph);
+		if (OldMacroIndex != INDEX_NONE)
+		{
+			if ((OldMacroIndex != NewIndex) && Blueprint->MacroGraphs.IsValidIndex(NewIndex))
 			{
-				FBlueprintEditorUtils::MarkBlueprintAsStructurallyModified(Blueprint);
+				Blueprint->Modify();
+				Blueprint->MacroGraphs.Insert(Graph, NewIndex);
+				Blueprint->MacroGraphs.RemoveAt((OldMacroIndex < NewIndex) ? OldMacroIndex : (OldMacroIndex + 1));
+				bModified = true;
 			}
 		}
+
+		if (bModified && !bDontRecompile)
+		{
+			FBlueprintEditorUtils::MarkBlueprintAsStructurallyModified(Blueprint);
+		}
+
+		return bModified;
 	}
+
+	return false;
 }
 
+
 FText FBlueprintEditorUtils::GetBlueprintVariableCategory(UBlueprint* Blueprint, const FName& VarName, const UStruct* InLocalVarScope)
 {
 	FText CategoryName;
@@ -3971,9 +4021,7 @@ void FBlueprintEditorUtils::SetBlueprintVariableRepNotifyFunc(UBlueprint* Bluepr
 	}
 }
 
-
-// Gets a list of function names currently in use in the blueprint, based on the skeleton class
-void FBlueprintEditorUtils::GetFunctionNameList(const UBlueprint* Blueprint, TArray<FName>& FunctionNames)
+void FBlueprintEditorUtils::GetFunctionNameList(const UBlueprint* Blueprint, TSet<FName>& FunctionNames)
 {
 	if( UClass* SkeletonClass = Blueprint->SkeletonGeneratedClass )
 	{
@@ -3984,7 +4032,7 @@ void FBlueprintEditorUtils::GetFunctionNameList(const UBlueprint* Blueprint, TAr
 	}
 }
 
-void FBlueprintEditorUtils::GetDelegateNameList(const UBlueprint* Blueprint, TArray<FName>& FunctionNames)
+void FBlueprintEditorUtils::GetDelegateNameList(const UBlueprint* Blueprint, TSet<FName>& FunctionNames)
 {
 	check(Blueprint);
 	for (int32 It = 0; It < Blueprint->DelegateSignatureGraphs.Num(); It++)
@@ -4095,8 +4143,7 @@ bool FBlueprintEditorUtils::IsPinTypeValid(const FEdGraphPinType& Type)
 	return true;
 }
 
-// Gets the visible class variable list.  This includes both variables introduced here and in all superclasses.
-void FBlueprintEditorUtils::GetClassVariableList(const UBlueprint* Blueprint, TArray<FName>& VisibleVariables, bool bIncludePrivateVars) 
+void FBlueprintEditorUtils::GetClassVariableList(const UBlueprint* Blueprint, TSet<FName>& VisibleVariables, bool bIncludePrivateVars) 
 {
 	// Existing variables in the parent class and above, when using the compilation manager the previous SkeletonGeneratedClass will have been cleared when
 	// we're regenerating the SkeletonGeneratedClass. Using this function in the skeleton pass at all is highly dubious, but I am leaving it until the 
@@ -4128,14 +4175,14 @@ void FBlueprintEditorUtils::GetClassVariableList(const UBlueprint* Blueprint, TA
 
 				for (int32 VariableIndex = 0; VariableIndex < ParentBP->NewVariables.Num(); ++VariableIndex)
 				{
-					VisibleVariables.AddUnique(ParentBP->NewVariables[VariableIndex].VarName);
+					VisibleVariables.Add(ParentBP->NewVariables[VariableIndex].VarName);
 				}
 
 				for (UTimelineTemplate* Timeline : ParentBP->Timelines)
 				{
 					if (Timeline)
 					{
-						VisibleVariables.AddUnique(Timeline->GetFName());
+						VisibleVariables.Add(Timeline->GetFName());
 					}
 				}
 			}
@@ -4191,7 +4238,7 @@ bool FBlueprintEditorUtils::AddMemberVariable(UBlueprint* Blueprint, const FName
 	}
 
 	// First we need to see if there is already a variable with that name, in this blueprint or parent class
-	TArray<FName> CurrentVars;
+	TSet<FName> CurrentVars;
 	FBlueprintEditorUtils::GetClassVariableList(Blueprint, CurrentVars);
 	if(CurrentVars.Contains(NewVarName))
 	{
@@ -4220,7 +4267,7 @@ bool FBlueprintEditorUtils::AddMemberVariable(UBlueprint* Blueprint, const FName
 		PostSetupObjectPinType(Blueprint, NewVar);
 	}
 	NewVar.ReplicationCondition = COND_None;
-	NewVar.Category = K2Schema->VR_DefaultCategory;
+	NewVar.Category = UEdGraphSchema_K2::VR_DefaultCategory;
 	NewVar.DefaultValue = DefaultValue;
 
 	// user created variables should be none of these things
@@ -4807,7 +4854,7 @@ bool FBlueprintEditorUtils::AddLocalVariable(UBlueprint* Blueprint, UEdGraph* In
 		NewVar.VarType = InNewVarType;
 		NewVar.PropertyFlags |= CPF_BlueprintVisible;
 		NewVar.FriendlyName = FName::NameToDisplayString( NewVar.VarName.ToString(), (NewVar.VarType.PinCategory == K2Schema->PC_Boolean) ? true : false );
-		NewVar.Category = K2Schema->VR_DefaultCategory;
+		NewVar.Category = UEdGraphSchema_K2::VR_DefaultCategory;
 		NewVar.DefaultValue = DefaultValue;
 
 		PostSetupObjectPinType(Blueprint, NewVar);
@@ -5975,7 +6022,7 @@ namespace
 			}
 			else
 			{
-				TArray<FName> DummyExtraNameList;
+				TSet<FName> DummyExtraNameList;
 				UEdGraphNode* CustomEventNode = CurrentGraph->GetSchema()->CreateSubstituteNode(EventNode, CurrentGraph, nullptr, DummyExtraNameList);
 				if (ensure(CustomEventNode))
 				{
@@ -6105,7 +6152,7 @@ static void ConformInterfaceByName(UBlueprint* Blueprint, FBPInterfaceDescriptio
 
 		// check to make sure that there aren't any interface methods that we originally 
 		// implemented as events, but have since switched to functions 
-		TArray<FName> ExtraNameList;
+		TSet<FName> ExtraNameList;
 		for (UK2Node_Event* EventNode : ImplementedEvents)
 		{
 			// if this event belongs to something other than this interface
@@ -7058,17 +7105,11 @@ void FBlueprintEditorUtils::GetActorReferenceMap(UWorld* InWorld, TArray<UClass*
 	}
 }
 
-FBlueprintEditorUtils::FFixLevelScriptActorBindingsEvent& FBlueprintEditorUtils::OnFixLevelScriptActorBindings()
-{
-	return FixLevelScriptActorBindingsEvent;
-}
-
-
-bool FBlueprintEditorUtils::FixLevelScriptActorBindings(ALevelScriptActor* LevelScriptActor, const ULevelScriptBlueprint* ScriptBlueprint)
+void FBlueprintEditorUtils::FixLevelScriptActorBindings(ALevelScriptActor* LevelScriptActor, const ULevelScriptBlueprint* ScriptBlueprint)
 {
 	if( ScriptBlueprint->BlueprintType != BPTYPE_LevelScript )
 	{
-		return false;
+		return;
 	}
 
 	UPackage* ActorPackage = LevelScriptActor->GetOutermost();
@@ -7079,11 +7120,9 @@ bool FBlueprintEditorUtils::FixLevelScriptActorBindings(ALevelScriptActor* Level
 	// prevents us from cross-binding instantiated (PIE) actors to editor objects
 	if (ActorPackage != BlueprintPkg)
 	{
-		return false;
+		return;
 	}
 
-	bool bWasSuccessful = true;
-
 	TArray<UEdGraph*> AllGraphs;
 	ScriptBlueprint->GetAllGraphs(AllGraphs);
 
@@ -7114,34 +7153,6 @@ bool FBlueprintEditorUtils::FixLevelScriptActorBindings(ALevelScriptActor* Level
 						Delegate.BindUFunction(LevelScriptActor, TargetFunction);
 						TargetDelegate->AddUnique(Delegate);
 					}
-					else
-					{
-						FString ErrorFormat = LOCTEXT("FailedLSABinding_NoDelegate", "Unable to bind event node %s. Target instance is missing the expected delegate - is it a Blueprint in need of compile?").ToString();
-						if (FCompilerResultsLog* ActiveLog = FCompilerResultsLog::GetEventTarget())
-						{
-							ActiveLog->Warning(*FString::Printf(*ErrorFormat, TEXT("@@")), EventNode);
-						}
-						else
-						{
-							// For some reason, we don't have a valid entry point for the event in the LSA...
-							UE_LOG(LogBlueprint, Warning, TEXT("%s"), *FString::Printf(*ErrorFormat, *EventNode->GetName()));
-						}
-						bWasSuccessful = false;
-					}
-				}
-				else
-				{
-					FString ErrorFormat = LOCTEXT("FailedLSABinding_NoFunction", "Unable to bind event node %s. Please recompile the level Blueprint.").ToString();
-					if (FCompilerResultsLog* ActiveLog = FCompilerResultsLog::GetEventTarget())
-					{
-						ActiveLog->Warning(*FString::Printf(*ErrorFormat, TEXT("@@")), EventNode);
-					}
-					else
-					{
-						// For some reason, we don't have a valid entry point for the event in the LSA...
-						UE_LOG(LogBlueprint, Warning, TEXT("%s"), *FString::Printf(*ErrorFormat, (EventNode ? *EventNode->GetName() : TEXT("[unknown]"))));
-					}
-					bWasSuccessful = false;
 				}
 			}
 		}
@@ -7160,12 +7171,6 @@ bool FBlueprintEditorUtils::FixLevelScriptActorBindings(ALevelScriptActor* Level
 			}
 		}
 	}
-
-	// Allow external sub-systems perform changes to the level script actor
-	FixLevelScriptActorBindingsEvent.Broadcast( LevelScriptActor, ScriptBlueprint, bWasSuccessful );
-
-
-	return bWasSuccessful;
 }
 
 void FBlueprintEditorUtils::ListPackageContents(UPackage* Package, FOutputDevice& Ar)
@@ -8496,34 +8501,30 @@ void FBlueprintEditorUtils::GetEntryAndResultNodes(const UEdGraph* InGraph, TWea
 	}
 }
 
-FKismetUserDeclaredFunctionMetadata* FBlueprintEditorUtils::GetGraphFunctionMetaData( UEdGraph* InGraph )
+FKismetUserDeclaredFunctionMetadata* FBlueprintEditorUtils::GetGraphFunctionMetaData(const UEdGraph* InGraph)
 {
-	FKismetUserDeclaredFunctionMetadata* FunctionMetaData = nullptr;
-
-	if( InGraph )
+	if (InGraph)
 	{
-		TArray<UK2Node_FunctionEntry*> EntryNodes;
-		InGraph->GetNodesOfClass( EntryNodes );
-
-		if( EntryNodes.Num() )
+		UK2Node_EditablePinBase* FunctionEntryNode = GetEntryNode(InGraph);
+		if (UK2Node_FunctionEntry* TypedEntryNode = Cast<UK2Node_FunctionEntry>(FunctionEntryNode))
+		{
+			return &(TypedEntryNode->MetaData);
+		}
+		else if (UK2Node_Tunnel* TunnelNode = ExactCast<UK2Node_Tunnel>(FunctionEntryNode))
 		{
-			FunctionMetaData = &EntryNodes[ 0 ]->MetaData;
+			// Must be exactly a tunnel, not a macro instance
+			return &(TunnelNode->MetaData);
 		}
 	}
-	return FunctionMetaData;
+
+	return nullptr;
 }
 
 FText FBlueprintEditorUtils::GetGraphDescription(const UEdGraph* InGraph)
 {
-	UK2Node_EditablePinBase* FunctionEntryNode = GetEntryNode(InGraph);
-	if (UK2Node_FunctionEntry* TypedEntryNode = Cast<UK2Node_FunctionEntry>(FunctionEntryNode))
-	{
-		return TypedEntryNode->MetaData.ToolTip;
-	}
-	else if (UK2Node_Tunnel* TunnelNode = ExactCast<UK2Node_Tunnel>(FunctionEntryNode))
+	if (FKismetUserDeclaredFunctionMetadata* MetaData = GetGraphFunctionMetaData(InGraph))
 	{
-		// Must be exactly a tunnel, not a macro instance
-		return TunnelNode->MetaData.ToolTip;
+		 return MetaData->ToolTip;
 	}
 
 	return LOCTEXT( "NoGraphTooltip", "(None)" );
diff --git a/Engine/Source/Editor/UnrealEd/Private/Kismet2/CompilerResultsLog.cpp b/Engine/Source/Editor/UnrealEd/Private/Kismet2/CompilerResultsLog.cpp
--- a/Engine/Source/Editor/UnrealEd/Private/Kismet2/CompilerResultsLog.cpp
+++ b/Engine/Source/Editor/UnrealEd/Private/Kismet2/CompilerResultsLog.cpp
@@ -130,7 +130,6 @@ FCompilerResultsLog::FCompilerResultsLog(bool bIsCompatibleWithEvents/* = true*/
 	, bLogInfoOnly(false)
 	, bAnnotateMentionedNodes(true)
 	, bLogDetailedResults(false)
-	, bTreatCompositeGraphsAsTunnels(false)
 	, EventDisplayThresholdMs(0)
 {
 	CurrentEventScope = nullptr;
@@ -696,7 +695,7 @@ void FCompilerResultsLog::RegisterIntermediateTunnelInstance(UEdGraphNode* Inter
 {
 	for (auto Tunnel : ActiveTunnels)
 	{
-		if (Tunnel.Get()->IsA<UK2Node_MacroInstance>() || bTreatCompositeGraphsAsTunnels)
+		if (Tunnel.Get()->IsA<UK2Node_MacroInstance>())
 		{
 			IntermediateTunnelInstanceHierarchyMap.Add(IntermediateTunnel, Tunnel);
 		}
@@ -704,7 +703,7 @@ void FCompilerResultsLog::RegisterIntermediateTunnelInstance(UEdGraphNode* Inter
 	UEdGraphNode* TrueSourceTunnelInstance = Cast<UEdGraphNode>(SourceBacktrackMap.FindSourceObject(IntermediateTunnel));
 	if (!ActiveTunnels.Contains(TrueSourceTunnelInstance))
 	{
-		if (TrueSourceTunnelInstance->IsA<UK2Node_MacroInstance>() || bTreatCompositeGraphsAsTunnels)
+		if (TrueSourceTunnelInstance->IsA<UK2Node_MacroInstance>())
 		{
 			IntermediateTunnelInstanceHierarchyMap.Add(IntermediateTunnel, TrueSourceTunnelInstance);
 		}
@@ -716,7 +715,7 @@ UEdGraphNode* FCompilerResultsLog::GetSourceNode(const UEdGraphNode* Intermediat
 	UEdGraphNode* Result = nullptr;
 	if (UEdGraphNode* SourceTunnelInstance = GetSourceTunnelInstance(IntermediateNode))
 	{
-		if (SourceTunnelInstance->IsA<UK2Node_MacroInstance>() || bTreatCompositeGraphsAsTunnels)
+		if (SourceTunnelInstance->IsA<UK2Node_MacroInstance>())
 		{
 			Result = SourceTunnelInstance;
 		}
diff --git a/Engine/Source/Editor/UnrealEd/Private/Kismet2/Kismet2.cpp b/Engine/Source/Editor/UnrealEd/Private/Kismet2/Kismet2.cpp
--- a/Engine/Source/Editor/UnrealEd/Private/Kismet2/Kismet2.cpp
+++ b/Engine/Source/Editor/UnrealEd/Private/Kismet2/Kismet2.cpp
@@ -14,7 +14,7 @@
 #include "UObject/GarbageCollection.h"
 #include "UObject/Class.h"
 #include "UObject/Package.h"
-#include "Misc/StringAssetReference.h"
+#include "UObject/SoftObjectPath.h"
 #include "UObject/MetaData.h"
 #include "Templates/SubclassOf.h"
 #include "UObject/UnrealType.h"
@@ -79,7 +79,6 @@
 #include "Editor/UnrealEd/Public/PackageTools.h"
 #include "Framework/Notifications/NotificationManager.h"
 #include "Widgets/Notifications/SNotificationList.h"
-#include "Developer/BlueprintProfiler/Public/BlueprintProfilerModule.h"
 #include "Engine/InheritableComponentHandler.h"
 
 DECLARE_CYCLE_STAT(TEXT("Compile Blueprint"), EKismetCompilerStats_CompileBlueprint, STATGROUP_KismetCompiler);
@@ -643,8 +642,6 @@ UK2Node_Event* FKismetEditorUtilities::AddDefaultEventNode(UBlueprint* InBluepri
 		NewEventNode->PostPlacedNewNode();
 		NewEventNode->SetFlags(RF_Transactional);
 		NewEventNode->AllocateDefaultPins();
-		NewEventNode->DisableNode();
-		NewEventNode->NodeComment = LOCTEXT("DisabledNodeComment", "This node is disabled and will not be called.\nDrag off pins to build functionality.").ToString();
 		NewEventNode->bCommentBubblePinned = true;
 		NewEventNode->bCommentBubbleVisible = true;
 		NewEventNode->NodePosY = InOutNodePosY;
@@ -677,12 +674,10 @@ UK2Node_Event* FKismetEditorUtilities::AddDefaultEventNode(UBlueprint* InBluepri
 			UEdGraphSchema_K2::SetNodeMetaData(ParentFunctionNode, FNodeMetadata::DefaultGraphNode);
 			FunctionNodeCreator.Finalize();
 
-			ParentFunctionNode->DisableNode();
-
-			// Adding the call to parent and connecting it will reset this value
-			NewEventNode->DisableNode();
-			NewEventNode->NodeComment = LOCTEXT("DisabledNodeComment", "This node is disabled and will not be called.\nDrag off pins to build functionality.").ToString();
+			ParentFunctionNode->MakeAutomaticallyPlacedGhostNode();
 		}
+
+		NewEventNode->MakeAutomaticallyPlacedGhostNode();
 	}
 
 	return NewEventNode;
@@ -691,7 +686,7 @@ UK2Node_Event* FKismetEditorUtilities::AddDefaultEventNode(UBlueprint* InBluepri
 UBlueprint* FKismetEditorUtilities::ReloadBlueprint(UBlueprint* StaleBlueprint)
 {
 	check(StaleBlueprint->IsAsset());
-	FStringAssetReference BlueprintAssetRef(StaleBlueprint);
+	FSoftObjectPath BlueprintAssetRef(StaleBlueprint);
 
 	FBlueprintUnloader Unloader(StaleBlueprint);
 	Unloader.UnloadBlueprint(/*bResetPackage =*/true);
@@ -755,7 +750,6 @@ void FKismetEditorUtilities::CompileBlueprint(UBlueprint* BlueprintObj, EBluepri
 	const bool bSaveIntermediateProducts	= (CompileFlags & EBlueprintCompileOptions::SaveIntermediateProducts	) != EBlueprintCompileOptions::None;
 	const bool bSkeletonUpToDate			= (CompileFlags & EBlueprintCompileOptions::SkeletonUpToDate			) != EBlueprintCompileOptions::None;
 	const bool bBatchCompile				= (CompileFlags & EBlueprintCompileOptions::BatchCompile				) != EBlueprintCompileOptions::None;
-	const bool bAddInstrumentation			= (CompileFlags & EBlueprintCompileOptions::AddInstrumentation			) != EBlueprintCompileOptions::None;
 	const bool bSkipReinstancing			= (CompileFlags & EBlueprintCompileOptions::SkipReinstancing			) != EBlueprintCompileOptions::None;
 
 	FSecondsCounterScope Timer(BlueprintCompileAndLoadTimerData); 
@@ -827,13 +821,7 @@ void FKismetEditorUtilities::CompileBlueprint(UBlueprint* BlueprintObj, EBluepri
 	FKismetCompilerOptions CompileOptions;
 	CompileOptions.bSaveIntermediateProducts = bSaveIntermediateProducts;
 	CompileOptions.bRegenerateSkelton = !bSkeletonUpToDate;
-	CompileOptions.bAddInstrumentation = bAddInstrumentation;
 	CompileOptions.bReinstanceAndStubOnFailure = !bSkipReinstancing;
-	if (pResults)
-	{
-		// enable debug information for composite graph instances if we are instrumenting the blueprint.
-		pResults->bTreatCompositeGraphsAsTunnels = bAddInstrumentation;
-	}
 	Compiler.CompileBlueprint(BlueprintObj, CompileOptions, Results, ReinstanceHelper);
 
 	FBlueprintEditorUtils::UpdateDelegatesInBlueprint(BlueprintObj);
@@ -2016,35 +2004,28 @@ void FKismetEditorUtilities::UpgradeCosmeticallyStaleBlueprint(UBlueprint* Bluep
 
 void FKismetEditorUtilities::CreateNewBoundEventForActor(AActor* Actor, FName EventName)
 {
-	if(Actor != NULL && EventName != NAME_None)
+	if ((Actor != nullptr) && (EventName != NAME_None))
 	{
 		// First, find the property we want to bind to
-		UMulticastDelegateProperty* DelegateProperty = FindField<UMulticastDelegateProperty>(Actor->GetClass(), EventName);
-		if(DelegateProperty != NULL)
+		if (UMulticastDelegateProperty* DelegateProperty = FindField<UMulticastDelegateProperty>(Actor->GetClass(), EventName))
 		{
 			// Get the correct level script blueprint
-			ULevelScriptBlueprint* LSB = Actor->GetLevel()->GetLevelScriptBlueprint();
-			UEdGraph* TargetGraph = NULL;
-			if(LSB != NULL && LSB->UbergraphPages.Num() > 0)
+			if (ULevelScriptBlueprint* LSB = Actor->GetLevel()->GetLevelScriptBlueprint())
 			{
-				TargetGraph = LSB->UbergraphPages[0]; // Just use the forst graph
-			}
-
-			if(TargetGraph != NULL)
-			{
-				// Figure out a decent place to stick the node
-				const FVector2D NewNodePos = TargetGraph->GetGoodPlaceForNewNode();
-
-				// Create a new event node
-				UK2Node_ActorBoundEvent* EventNodeTemplate = NewObject<UK2Node_ActorBoundEvent>();
-				EventNodeTemplate->InitializeActorBoundEventParams(Actor, DelegateProperty);
+				if (UEdGraph* TargetGraph = LSB->GetLastEditedUberGraph())
+				{
+					// Figure out a decent place to stick the node
+					const FVector2D NewNodePos = TargetGraph->GetGoodPlaceForNewNode();
 
-				UK2Node_ActorBoundEvent* EventNode = FEdGraphSchemaAction_K2NewNode::SpawnNodeFromTemplate<UK2Node_ActorBoundEvent>(TargetGraph, EventNodeTemplate, NewNodePos);
+					// Create a new event node
+					UK2Node_ActorBoundEvent* EventNodeTemplate = NewObject<UK2Node_ActorBoundEvent>();
+					EventNodeTemplate->InitializeActorBoundEventParams(Actor, DelegateProperty);
 
-				// Finally, bring up kismet and jump to the new node
-				if(EventNode != NULL)
-				{
-					BringKismetToFocusAttentionOnObject(EventNode);
+					// Finally, bring up kismet and jump to the new node
+					if (UK2Node_ActorBoundEvent* EventNode = FEdGraphSchemaAction_K2NewNode::SpawnNodeFromTemplate<UK2Node_ActorBoundEvent>(TargetGraph, EventNodeTemplate, NewNodePos))
+					{
+						BringKismetToFocusAttentionOnObject(EventNode);
+					}
 				}
 			}
 		}
@@ -2067,11 +2048,7 @@ void FKismetEditorUtilities::CreateNewBoundEventForClass(UClass* Class, FName Ev
 		UMulticastDelegateProperty* DelegateProperty = FindField<UMulticastDelegateProperty>(Class, EventName);
 		if ( DelegateProperty != nullptr )
 		{
-			UEdGraph* TargetGraph = nullptr;
-			if(Blueprint->UbergraphPages.Num() > 0)
-			{
-				TargetGraph = Blueprint->UbergraphPages[0]; // Just use the first graph
-			}
+			UEdGraph* TargetGraph = Blueprint->GetLastEditedUberGraph();
 
 			if ( TargetGraph != nullptr )
 			{
diff --git a/Engine/Source/Editor/UnrealEd/Private/Kismet2/KismetDebugUtilities.cpp b/Engine/Source/Editor/UnrealEd/Private/Kismet2/KismetDebugUtilities.cpp
--- a/Engine/Source/Editor/UnrealEd/Private/Kismet2/KismetDebugUtilities.cpp
+++ b/Engine/Source/Editor/UnrealEd/Private/Kismet2/KismetDebugUtilities.cpp
@@ -781,15 +781,7 @@ void FKismetDebugUtilities::SetBreakpointInternal(UBreakpoint* Breakpoint, bool
 	{
 		if (uint8* InstallSite = InstallSites[i])
 		{
-			if (*InstallSite == EX_InstrumentationEvent)
-			{
-				InstallSite++;
-				*InstallSite = bShouldBeEnabled ? EScriptInstrumentation::NodeDebugSite : EScriptInstrumentation::NodeEntry;
-			}
-			else
-			{
-				*InstallSite = bShouldBeEnabled ? EX_Breakpoint : EX_Tracepoint;
-			}
+			*InstallSite = bShouldBeEnabled ? EX_Breakpoint : EX_Tracepoint;
 		}
 	}
 }
diff --git a/Engine/Source/Editor/UnrealEd/Private/Kismet2/KismetReinstanceUtilities.cpp b/Engine/Source/Editor/UnrealEd/Private/Kismet2/KismetReinstanceUtilities.cpp
--- a/Engine/Source/Editor/UnrealEd/Private/Kismet2/KismetReinstanceUtilities.cpp
+++ b/Engine/Source/Editor/UnrealEd/Private/Kismet2/KismetReinstanceUtilities.cpp
@@ -74,7 +74,7 @@ struct FReplaceReferenceHelper
 		}
 	}
 
-	static void FindAndReplaceReferences(const TArray<UObject*>& SourceObjects, TSet<UObject*>* ObjectsThatShouldUseOldStuff, const TArray<UObject*>& ObjectsToReplace, const TMap<UObject*, UObject*>& OldToNewInstanceMap, const TMap<FStringAssetReference, UObject*>& ReinstancedObjectsWeakReferenceMap)
+	static void FindAndReplaceReferences(const TArray<UObject*>& SourceObjects, TSet<UObject*>* ObjectsThatShouldUseOldStuff, const TArray<UObject*>& ObjectsToReplace, const TMap<UObject*, UObject*>& OldToNewInstanceMap, const TMap<FSoftObjectPath, UObject*>& ReinstancedObjectsWeakReferenceMap)
 	{
 		if(SourceObjects.Num() == 0 && ObjectsToReplace.Num() == 0 )
 		{
@@ -98,22 +98,22 @@ struct FReplaceReferenceHelper
 				{
 					// The class for finding and replacing weak references.
 					// We can't relay on "standard" weak references replacement as
-					// it depends on FStringAssetReference::ResolveObject, which
+					// it depends on FSoftObjectPath::ResolveObject, which
 					// tries to find the object with the stored path. It is
 					// impossible, cause above we deleted old actors (after
 					// spawning new ones), so during objects traverse we have to
-					// find FStringAssetReferences with the raw given path taken
+					// find FSoftObjectPath with the raw given path taken
 					// before deletion of old actors and fix them.
 					class ReferenceReplace : public FArchiveReplaceObjectRef < UObject >
 					{
 					public:
-						ReferenceReplace(UObject* InSearchObject, const TMap<UObject*, UObject*>& InReplacementMap, const TMap<FStringAssetReference, UObject*>& InWeakReferencesMap)
+						ReferenceReplace(UObject* InSearchObject, const TMap<UObject*, UObject*>& InReplacementMap, const TMap<FSoftObjectPath, UObject*>& InWeakReferencesMap)
 							: FArchiveReplaceObjectRef<UObject>(InSearchObject, InReplacementMap, false, false, false, true), WeakReferencesMap(InWeakReferencesMap)
 						{
 							SerializeSearchObject();
 						}
 
-						FArchive& operator<<(FStringAssetReference& Ref) override
+						FArchive& operator<<(FSoftObjectPath& Ref) override
 						{
 							const UObject*const* PtrToObjPtr = WeakReferencesMap.Find(Ref);
 
@@ -125,13 +125,13 @@ struct FReplaceReferenceHelper
 							return *this;
 						}
 
-						FArchive& operator<<(FAssetPtr& Ref) override
+						FArchive& operator<<(FSoftObjectPtr& Ref) override
 						{
 							return operator<<(Ref.GetUniqueID());
 						}
 
 					private:
-						const TMap<FStringAssetReference, UObject*>& WeakReferencesMap;
+						const TMap<FSoftObjectPath, UObject*>& WeakReferencesMap;
 					};
 
 					ReferenceReplace ReplaceAr(Obj, OldToNewInstanceMap, ReinstancedObjectsWeakReferenceMap);
@@ -573,7 +573,7 @@ void FBlueprintCompileReinstancer::FinalizeFastReinstancing(TArray<UObject*>& Ob
 {
 	TArray<UObject*> SourceObjects;
 	TMap<UObject*, UObject*> OldToNewInstanceMap;
-	TMap<FStringAssetReference, UObject*> ReinstancedObjectsWeakReferenceMap;
+	TMap<FSoftObjectPath, UObject*> ReinstancedObjectsWeakReferenceMap;
 	FReplaceReferenceHelper::IncludeCDO(DuplicatedClass, ClassToReinstance, OldToNewInstanceMap, SourceObjects, OriginalCDO);
 
 	if (IsClassObjectReplaced())
@@ -989,11 +989,8 @@ void FBlueprintCompileReinstancer::UpdateBytecodeReferences()
 			for( TFieldIterator<UFunction> FuncIter(BPClass, EFieldIteratorFlags::ExcludeSuper); FuncIter; ++FuncIter )
 			{
 				UFunction* CurrentFunction = *FuncIter;
-				if( CurrentFunction->Script.Num() > 0 )
-				{
-					FArchiveReplaceObjectRef<UObject> ReplaceAr(CurrentFunction, FieldMappings, /*bNullPrivateRefs=*/ false, /*bIgnoreOuterRef=*/ true, /*bIgnoreArchetypeRef=*/ true);
-					bBPWasChanged |= (0 != ReplaceAr.GetCount());
-				}
+				FArchiveReplaceObjectRef<UObject> ReplaceAr(CurrentFunction, FieldMappings, /*bNullPrivateRefs=*/ false, /*bIgnoreOuterRef=*/ true, /*bIgnoreArchetypeRef=*/ true);
+				bBPWasChanged |= (0 != ReplaceAr.GetCount());
 			}
 
 			FArchiveReplaceObjectRef<UObject> ReplaceInBPAr(*DependentBP, FieldMappings, false, true, true);
@@ -1144,7 +1141,7 @@ struct FActorReplacementHelper
 	 * Runs construction scripts on the new actor and then finishes it off by
 	 * attaching it to the same attachments that its predecessor was set with. 
 	 */
-	void Finalize(const TMap<UObject*, UObject*>& OldToNewInstanceMap, TSet<UObject*>* ObjectsThatShouldUseOldStuff, const TArray<UObject*>& ObjectsToReplace, const TMap<FStringAssetReference, UObject*>& ReinstancedObjectsWeakReferenceMap);
+	void Finalize(const TMap<UObject*, UObject*>& OldToNewInstanceMap, TSet<UObject*>* ObjectsThatShouldUseOldStuff, const TArray<UObject*>& ObjectsToReplace, const TMap<FSoftObjectPath, UObject*>& ReinstancedObjectsWeakReferenceMap);
 
 	/**
 	* Takes the cached child actors, as well as the old AttachParent, and sets
@@ -1153,7 +1150,7 @@ struct FActorReplacementHelper
 	*
 	* @param OldToNewInstanceMap Mapping of reinstanced objects.
 	*/
-	void ApplyAttachments(const TMap<UObject*, UObject*>& OldToNewInstanceMap, TSet<UObject*>* ObjectsThatShouldUseOldStuff, const TArray<UObject*>& ObjectsToReplace, const TMap<FStringAssetReference, UObject*>& ReinstancedObjectsWeakReferenceMap);
+	void ApplyAttachments(const TMap<UObject*, UObject*>& OldToNewInstanceMap, TSet<UObject*>* ObjectsThatShouldUseOldStuff, const TArray<UObject*>& ObjectsToReplace, const TMap<FSoftObjectPath, UObject*>& ReinstancedObjectsWeakReferenceMap);
 
 private:
 	/**
@@ -1175,7 +1172,7 @@ private:
 	TMap<FName, UActorComponent*> OldActorComponentNameMap;
 };
 
-void FActorReplacementHelper::Finalize(const TMap<UObject*, UObject*>& OldToNewInstanceMap, TSet<UObject*>* ObjectsThatShouldUseOldStuff, const TArray<UObject*>& ObjectsToReplace, const TMap<FStringAssetReference, UObject*>& ReinstancedObjectsWeakReferenceMap)
+void FActorReplacementHelper::Finalize(const TMap<UObject*, UObject*>& OldToNewInstanceMap, TSet<UObject*>* ObjectsThatShouldUseOldStuff, const TArray<UObject*>& ObjectsToReplace, const TMap<FSoftObjectPath, UObject*>& ReinstancedObjectsWeakReferenceMap)
 {
 	
 	// because this is an editor context it's important to use this execution guard
@@ -1253,7 +1250,7 @@ void FActorReplacementHelper::Finalize(const TMap<UObject*, UObject*>& OldToNewI
 	}
 }
 
-void FActorReplacementHelper::ApplyAttachments(const TMap<UObject*, UObject*>& OldToNewInstanceMap, TSet<UObject*>* ObjectsThatShouldUseOldStuff, const TArray<UObject*>& ObjectsToReplace, const TMap<FStringAssetReference, UObject*>& ReinstancedObjectsWeakReferenceMap)
+void FActorReplacementHelper::ApplyAttachments(const TMap<UObject*, UObject*>& OldToNewInstanceMap, TSet<UObject*>* ObjectsThatShouldUseOldStuff, const TArray<UObject*>& ObjectsToReplace, const TMap<FSoftObjectPath, UObject*>& ReinstancedObjectsWeakReferenceMap)
 {
 	USceneComponent* NewRootComponent = NewActor->GetRootComponent();
 	if (NewRootComponent == nullptr)
@@ -1520,7 +1517,7 @@ UClass* FBlueprintCompileReinstancer::MoveCDOToNewClass(UClass* OwnerClass, cons
 	{
 		if(bAvoidCDODuplication)
 		{
-			OldCDO->Rename(*OldCDO->GetName(), CopyOfOwnerClass, REN_DoNotDirty | REN_DontCreateRedirectors | REN_ForceNoResetLoaders);
+			OldCDO->Rename(nullptr, CopyOfOwnerClass->GetOuter(), REN_DoNotDirty | REN_DontCreateRedirectors | REN_ForceNoResetLoaders);
 			CopyOfOwnerClass->ClassDefaultObject = OldCDO;
 		}
 		OldCDO->SetClass(CopyOfOwnerClass);
@@ -1549,7 +1546,7 @@ void FBlueprintCompileReinstancer::ReplaceInstancesOfClass_Inner(TMap<UClass*, U
 	// Map of old objects to new name (used to assist with reinstancing archetypes)
 	TMap<UObject*, FName> OldToNewNameMap;
 
-	TMap<FStringAssetReference, UObject*> ReinstancedObjectsWeakReferenceMap;
+	TMap<FSoftObjectPath, UObject*> ReinstancedObjectsWeakReferenceMap;
 
 	// actors being replace
 	TArray<FActorReplacementHelper> ReplacementActors;
@@ -2159,28 +2156,6 @@ void FBlueprintCompileReinstancer::CopyPropertiesForUnrelatedObjects(UObject* Ol
 	InstancedPropertyUtils::FArchiveInsertInstancedSubObjects InstancedSubObjSpawner(NewObject, InstancedPropertyMap);
 }
 
-UObject* FBlueprintCompileReinstancer::GetClassCDODuplicate(UObject* CDO, FName Name)
-{
-	UObject* DupCDO = nullptr;
-
-	FCDODuplicatesProvider& CDODupProvider = GetCDODuplicatesProviderDelegate();
-
-	if (!CDODupProvider.IsBound() || (DupCDO = CDODupProvider.Execute(CDO, Name)) == nullptr)
-	{
-		GIsDuplicatingClassForReinstancing = true;
-		DupCDO = (UObject*)StaticDuplicateObject(CDO, GetTransientPackage(), Name);
-		GIsDuplicatingClassForReinstancing = false;
-	}
-
-	return DupCDO;
-}
-
-FBlueprintCompileReinstancer::FCDODuplicatesProvider& FBlueprintCompileReinstancer::GetCDODuplicatesProviderDelegate()
-{
-	static FCDODuplicatesProvider Delegate;
-	return Delegate;
-}
-
 FRecreateUberGraphFrameScope::FRecreateUberGraphFrameScope(UClass* InClass, bool bRecreate)
 	: RecompiledClass(InClass)
 {
diff --git a/Engine/Source/Editor/UnrealEd/Private/Kismet2/StructureEditorUtils.cpp b/Engine/Source/Editor/UnrealEd/Private/Kismet2/StructureEditorUtils.cpp
--- a/Engine/Source/Editor/UnrealEd/Private/Kismet2/StructureEditorUtils.cpp
+++ b/Engine/Source/Editor/UnrealEd/Private/Kismet2/StructureEditorUtils.cpp
@@ -417,8 +417,8 @@ bool FStructureEditorUtils::ChangeVariableDefaultValue(UUserDefinedStruct* Struc
 		else if ((PinType.PinCategory == K2Schema->PC_Object) 
 			|| (PinType.PinCategory == K2Schema->PC_Interface) 
 			|| (PinType.PinCategory == K2Schema->PC_Class)
-			|| (PinType.PinCategory == K2Schema->PC_AssetClass)
-			|| (PinType.PinCategory == K2Schema->PC_Asset))
+			|| (PinType.PinCategory == K2Schema->PC_SoftClass)
+			|| (PinType.PinCategory == K2Schema->PC_SoftObject))
 		{
 			// K2Schema->DefaultValueSimpleValidation finds an object, passed by path, invalid
 			bResult = true;
@@ -695,6 +695,7 @@ bool FStructureEditorUtils::ChangeTooltip(UUserDefinedStruct* Struct, const FStr
 		StructEditorData->ToolTip = InTooltip;
 
 		Struct->SetMetaData(FBlueprintMetadata::MD_Tooltip, *StructEditorData->ToolTip);
+		Struct->PostEditChange();
 
 		return true;
 	}
diff --git a/Engine/Source/Editor/UnrealEd/Private/PlayLevel.cpp b/Engine/Source/Editor/UnrealEd/Private/PlayLevel.cpp
--- a/Engine/Source/Editor/UnrealEd/Private/PlayLevel.cpp
+++ b/Engine/Source/Editor/UnrealEd/Private/PlayLevel.cpp
@@ -15,10 +15,9 @@
 #include "UObject/UObjectIterator.h"
 #include "UObject/Package.h"
 #include "UObject/LazyObjectPtr.h"
-#include "Misc/StringAssetReference.h"
+#include "UObject/SoftObjectPtr.h"
 #include "Serialization/ArchiveTraceRoute.h"
 #include "Misc/PackageName.h"
-#include "Misc/StringClassReference.h"
 #include "InputCoreTypes.h"
 #include "Layout/Margin.h"
 #include "Layout/SlateRect.h"
@@ -157,6 +156,9 @@ void UEditorEngine::EndPlayMap()
 
 	FEditorDelegates::PrePIEEnded.Broadcast( bIsSimulatingInEditor );
 
+	// Clean up Soft Object Path remaps
+	FSoftObjectPath::ClearPIEPackageNames();
+
 	FlushAsyncLoading();
 
 	// Monitoring when PIE corrupts references between the World and the PIE generated World for UE-20486
@@ -862,13 +864,10 @@ void UEditorEngine::RequestPlaySession( bool bAtPlayerStart, TSharedPtr<class IL
 	// Unless we have been asked to play in a specific viewport window, this index will be -1
 	PlayInEditorViewportIndex = -1;
 
-	// @todo gmp: temp hack for Rocket demo
 	bPlayOnLocalPcSession = false;
 	bPlayUsingLauncher = false;
 }
 
-
-// @todo gmp: temp hack for Rocket demo
 void UEditorEngine::RequestPlaySession( const FVector* StartLocation, const FRotator* StartRotation, bool MobilePreview, bool VulkanPreview , const FString& MobilePreviewTargetDevice)
 {
 	bPlayOnLocalPcSession = true;
@@ -916,6 +915,32 @@ void UEditorEngine::CancelRequestPlaySession()
 	PlayUsingMobilePreviewTargetDevice.Reset();
 }
 
+bool UEditorEngine::SaveMapsForPlaySession()
+{
+	// Prompt the user to save the level if it has not been saved before. 
+	// An unmodified but unsaved blank template level does not appear in the dirty packages check below.
+	if (FEditorFileUtils::GetFilename(GWorld).Len() == 0)
+	{
+		if (!FEditorFileUtils::SaveCurrentLevel())
+		{
+			CancelRequestPlaySession();
+			return false;
+		}
+	}
+
+	// Also save dirty packages, this is required because we're going to be launching a session outside of our normal process
+	bool bPromptUserToSave = true;
+	bool bSaveMapPackages = true;
+	bool bSaveContentPackages = true;
+	if (!FEditorFileUtils::SaveDirtyPackages(bPromptUserToSave, bSaveMapPackages, bSaveContentPackages))
+	{
+		CancelRequestPlaySession();
+		return false;
+	}
+
+	return true;
+}
+
 void UEditorEngine::PlaySessionPaused()
 {
 	FEditorDelegates::PausePIE.Broadcast(bIsSimulatingInEditor);
@@ -1156,12 +1181,11 @@ void UEditorEngine::StartQueuedPlayMapRequest()
 	const EPlayNetMode PlayNetMode = [&PlayInSettings]{ EPlayNetMode NetMode(PIE_Standalone); return (PlayInSettings->GetPlayNetMode(NetMode) ? NetMode : PIE_Standalone); }();
 	const bool CanRunUnderOneProcess = [&PlayInSettings]{ bool RunUnderOneProcess(false); return (PlayInSettings->GetRunUnderOneProcess(RunUnderOneProcess) && RunUnderOneProcess); }();
 
-	// World composition does not copy levels to a separate folder to reduce startup time, and instead uses same files as Editor
-	// This causes issues with network replication as server object names will collide with Editor loaded world
-	const bool bWorldCompositionActive = GetEditorWorldContext().World()->WorldComposition != nullptr;
-	if (bWorldCompositionActive && !(CanRunUnderOneProcess || PlayNetMode == PIE_Standalone))
+	const bool bRequestSave = bPlayOnLocalPcSession || bPlayUsingLauncher || (!CanRunUnderOneProcess && PlayNetMode != PIE_Standalone);
+	if (bRequestSave && !SaveMapsForPlaySession())
 	{
-		FText ErrorMsg = LOCTEXT("WorldCompPIESingleProcessError", "World Composition does not support multiplayer Play in Editor using separate processes. Please set 'Use Single Process' in the 'Level Editor - Play' settings under Editor Preferences.");
+		// Maps did not save, print a warning
+		FText ErrorMsg = LOCTEXT("PIEWorldSaveFail", "PIE failed because map save was canceled");
 		UE_LOG(LogPlayLevel, Warning, TEXT("%s"), *ErrorMsg.ToString());
 		FMessageLog(NAME_CategoryPIE).Warning(ErrorMsg);
 		FMessageLog(NAME_CategoryPIE).Open();
@@ -1192,13 +1216,14 @@ void UEditorEngine::StartQueuedPlayMapRequest()
 		// If we're playing in the editor
 		if (!bPlayOnLocalPcSession)
 		{
-				PlayInEditor(GetEditorWorldContext().World(), bWantSimulateInEditor);
+			PlayInEditor(GetEditorWorldContext().World(), bWantSimulateInEditor);
 
-				// Editor counts as a client
-				NumClients++;
-			}
+			// Editor counts as a client
+			NumClients++;
+		}
 
 		// Spawn number of clients
+		const bool WillAutoConnectToServer = [&PlayInSettings] { bool AutoConnectToServer(false); return (PlayInSettings->GetAutoConnectToServer(AutoConnectToServer) && AutoConnectToServer); }();
 		const int32 PlayNumberOfClients = [&PlayInSettings]{ int32 NumberOfClients(0); return (PlayInSettings->GetPlayNumberOfClients(NumberOfClients) ? NumberOfClients : 0); }();
 		for (int32 i = NumClients; i < PlayNumberOfClients; ++i)
 		{
@@ -1223,135 +1248,6 @@ void UEditorEngine::StartQueuedPlayMapRequest()
 	}
 }
 
-/* Temporarily renames streaming levels for pie saving */
-class FScopedRenameStreamingLevels
-{
-public:
-	FScopedRenameStreamingLevels( UWorld* InWorld, const FString& AutosavePackagePrefix, const FString& MapnamePrefix )
-		: World(InWorld)
-	{
-		if(InWorld->StreamingLevels.Num() > 0)
-		{
-			for(int32 LevelIndex=0; LevelIndex < InWorld->StreamingLevels.Num(); ++LevelIndex)
-			{
-				ULevelStreaming* StreamingLevel = InWorld->StreamingLevels[LevelIndex];
-				if ( StreamingLevel )
-				{
-					const FString WorldAssetPackageName = StreamingLevel->GetWorldAssetPackageName();
-					const FName WorldAssetPackageFName = StreamingLevel->GetWorldAssetPackageFName();
-					PreviousStreamingPackageNames.Add( WorldAssetPackageFName );
-					FString StreamingLevelPackageName = FString::Printf(TEXT("%s%s/%s%s"), *AutosavePackagePrefix, *FPackageName::GetLongPackagePath( WorldAssetPackageName ), *MapnamePrefix, *FPackageName::GetLongPackageAssetName( WorldAssetPackageName ));
-					StreamingLevelPackageName.ReplaceInline(TEXT("//"), TEXT("/"));
-					StreamingLevel->SetWorldAssetByPackageName(FName(*StreamingLevelPackageName));
-				}
-			}
-		}
-
-		World->StreamingLevelsPrefix = MapnamePrefix;
-	}
-
-	~FScopedRenameStreamingLevels()
-	{
-		check(World.IsValid());
-		check( PreviousStreamingPackageNames.Num() == World->StreamingLevels.Num() );
-		if(World->StreamingLevels.Num() > 0)
-		{
-			for(int32 LevelIndex=0; LevelIndex < World->StreamingLevels.Num(); ++LevelIndex)
-			{
-				ULevelStreaming* StreamingLevel = World->StreamingLevels[LevelIndex];
-				if ( StreamingLevel )
-				{
-					StreamingLevel->SetWorldAssetByPackageName(PreviousStreamingPackageNames[LevelIndex]);
-				}
-			}
-		}
-
-		World->StreamingLevelsPrefix.Empty();
-	}
-private:
-	TWeakObjectPtr<UWorld> World;
-	TArray<FName> PreviousStreamingPackageNames;
-};
-
-
-void UEditorEngine::SaveWorldForPlay(TArray<FString>& SavedMapNames)
-{
-	UWorld* World = GWorld;
-
-	// check if PersistentLevel has any external references
-	if( PackageUsingExternalObjects(World->PersistentLevel) && EAppReturnType::Yes != FMessageDialog::Open( EAppMsgType::YesNo, NSLOCTEXT("UnrealEd", "Warning_UsingExternalPackage", "This map is using externally referenced packages which won't be found when in a game and all references will be broken. Perform a map check for more details.\n\nWould you like to continue?")) )
-	{
-		return;
-	}
-
-	const FString PlayOnConsolePackageName = FPackageName::FilenameToLongPackageName(FPaths::Combine(*FPaths::ProjectSavedDir(), *PlayOnConsoleSaveDir)) + TEXT("/");
-
-	// make a per-platform name for the map
-	const FString ConsoleName = FString(TEXT("PC"));
-	const FString Prefix = FString(PLAYWORLD_CONSOLE_BASE_PACKAGE_PREFIX) + ConsoleName;
-
-	// Temporarily rename streaming levels for pie saving
-	FScopedRenameStreamingLevels ScopedRenameStreamingLevels( World, PlayOnConsolePackageName, Prefix );
-	
-	// spawn a play-from-here player start or a temporary player start
-	AActor* PlayerStart = NULL;
-	bool bCreatedPlayerStart = false;
-
-	// Gross hack to avoid marking the level dirty while this temporary spawn is done while saving we flag the editor as loading package
-	bool bTempIsEditorLoadingPackage = GIsEditorLoadingPackage;
-	GIsEditorLoadingPackage = true;
-
-	{
-		// Do the spawning in a transaction so that we can undo it below
-		FScopedTransaction SpawnPlayFromHereStartTransaction(LOCTEXT("SpawnPlayFromHereStart","Spawn Play From Here Start"));
-
-		SpawnPlayFromHereStart( World, PlayerStart, PlayWorldLocation, PlayWorldRotation );
-
-		if (PlayerStart != NULL)
-		{
-			bCreatedPlayerStart = true;
-		}
-		else
-		{
-			PlayerStart = CheckForPlayerStart();
-	
-			if( PlayerStart == NULL )
-			{
-				FActorSpawnParameters SpawnInfo;
-				SpawnInfo.SpawnCollisionHandlingOverride = ESpawnActorCollisionHandlingMethod::AlwaysSpawn;
-				PlayerStart = World->SpawnActor<AActor>( APlayerStart::StaticClass(), FVector::ZeroVector, FRotator::ZeroRotator, SpawnInfo );
-
-				bCreatedPlayerStart = true;
-			}
-		}
-	}
-	
-	// save out all open map packages
-	TArray<FString> SavedWorldFileNames;
-	bool bSavedWorld = SavePlayWorldPackages(World, *Prefix, /*out*/ SavedWorldFileNames);
-
-	// Undo the (potentially) spawning transaction and don't display a notification banner about it
-	bool bTempSquelched = GEditor->bSquelchTransactionNotification;
-	GEditor->bSquelchTransactionNotification = true;
-	GEditor->UndoTransaction(false);
-	GEditor->bSquelchTransactionNotification = bTempSquelched;
-
-	GIsEditorLoadingPackage = bTempIsEditorLoadingPackage;
-	
-	if (bSavedWorld)
-	{
-		// Convert the filenames into map names
-		SavedMapNames.Reserve(SavedMapNames.Num() + SavedWorldFileNames.Num());
-		for (int32 Index = 0; Index < SavedWorldFileNames.Num(); ++Index)
-		{
-			const FString MapName = FPackageName::FilenameToLongPackageName(SavedWorldFileNames[Index]);
-			SavedMapNames.Add(MapName);
-		}
-	}
-
-}
-
-// @todo gmp: temp hack for Rocket demo
 void UEditorEngine::EndPlayOnLocalPc( )
 {
 	for (int32 i=0; i < PlayOnLocalPCSessions.Num(); ++i)
@@ -1369,7 +1265,6 @@ void UEditorEngine::EndPlayOnLocalPc( )
 	PlayOnLocalPCSessions.Empty();
 }
 
-// @todo gmp: temp hack for Rocket demo
 void UEditorEngine::PlayStandaloneLocalPc(FString MapNameOverride, FIntPoint* WindowPos, int32 PIENum, bool bIsServer)
 {
 	const ULevelEditorPlaySettings* PlayInSettings = GetDefault<ULevelEditorPlaySettings>();
@@ -1389,16 +1284,10 @@ void UEditorEngine::PlayStandaloneLocalPc(FString MapNameOverride, FIntPoint* Wi
 	if (MapNameOverride.IsEmpty())
 	{
 		FWorldContext & EditorContext = GetEditorWorldContext();
-		if (EditorContext.World()->WorldComposition)
-		{
-			// Open world composition from original folder
-			FString MapName = EditorContext.World()->GetOutermost()->GetName();
-			SavedMapNames.Add(MapName);
-		}
-		else
-		{
-			SaveWorldForPlay(SavedMapNames);
-		}
+
+		// Open original map
+		FString MapName = EditorContext.World()->GetOutermost()->GetName();
+		SavedMapNames.Add(MapName);
 	}
 	else
 	{
@@ -2035,50 +1924,10 @@ void UEditorEngine::PlayUsingLauncher()
 
 		TArray<FString> MapNames;
 		FWorldContext & EditorContext = GetEditorWorldContext();
-		if (EditorContext.World()->WorldComposition || (LauncherProfile->GetCookMode() == ELauncherProfileCookModes::ByTheBookInEditor) || (LauncherProfile->GetCookMode() == ELauncherProfileCookModes::OnTheFlyInEditor) )
-		{
-			// Prompt the user to save the level if it has not been saved before. 
-			// An unmodified but unsaved blank template level does not appear in the dirty packages check below.
-			if (FEditorFileUtils::GetFilename(GWorld).Len() == 0)
-			{
-				if (!FEditorFileUtils::SaveCurrentLevel())
-				{
-					CancelRequestPlaySession();
-					return;
-				}
-			}
 
-			// Daniel: Only reason we actually need to save any packages is because if a new package is created it won't be on disk yet and CookOnTheFly will early out if the package doesn't exist (even though it could be in memory and not require loading at all)
-			//			future me can optimize this by either adding extra allowances to CookOnTheFlyServer code or only saving packages which doesn't exist if it becomes a problem
-			// if this returns false, it means we should stop what we're doing and return to the editor
-			bool bPromptUserToSave = true;
-			bool bSaveMapPackages = true;
-			bool bSaveContentPackages = true;
-			if (!FEditorFileUtils::SaveDirtyPackages(bPromptUserToSave, bSaveMapPackages, bSaveContentPackages))
-			{
-				CancelRequestPlaySession();
-				return;
-			}
-
-
-			// Open world composition from original folder
-			// Or if using by book in editor don't need to resave the package just cook it by the book 
-			FString MapName = EditorContext.World()->GetOutermost()->GetName();
-			MapNames.Add(MapName);
-
-
-			
-		}
-		else
-		{
-			SaveWorldForPlay(MapNames);
-
-			if (MapNames.Num() == 0)
-			{
-				CancelRequestPlaySession();
-				return;
-			}
-		}
+		// Load maps in place as we saved them above
+		FString EditorMapName = EditorContext.World()->GetOutermost()->GetName();
+		MapNames.Add(EditorMapName);
 	
 		FString InitialMapName;
 		if (MapNames.Num() > 0)
@@ -2093,7 +1942,6 @@ void UEditorEngine::PlayUsingLauncher()
 			LauncherProfile->AddCookedMap(MapName);
 		}
 
-
 		if ( LauncherProfile->GetCookMode() == ELauncherProfileCookModes::ByTheBookInEditor )
 		{
 			TArray<ITargetPlatform*> TargetPlatforms;
@@ -2215,30 +2063,6 @@ void UEditorEngine::RequestEndPlayMap()
 	}
 }
 
-bool UEditorEngine::SavePlayWorldPackages(UWorld* InWorld, const TCHAR* Prefix, TArray<FString>& OutSavedFilenames)
-{
-	{
-		// if this returns false, it means we should stop what we're doing and return to the editor
-		bool bPromptUserToSave = true;
-		bool bSaveMapPackages = false;
-		bool bSaveContentPackages = true;
-		if (!FEditorFileUtils::SaveDirtyPackages(bPromptUserToSave, bSaveMapPackages, bSaveContentPackages))
-		{
-			return false;
-		}
-	}
-
-	// Update cull distance volumes before saving.
-	InWorld->UpdateCullDistanceVolumes();
-
-	// Clean up any old worlds.
-	CollectGarbage( GARBAGE_COLLECTION_KEEPFLAGS );
-
-	// Save temporary copies of all levels to be used for playing in editor or using standalone PC/console
-	return FEditorFileUtils::SaveWorlds(InWorld, FPaths::Combine(*FPaths::ProjectSavedDir(), *PlayOnConsoleSaveDir), Prefix, OutSavedFilenames);
-}
-
-
 FString UEditorEngine::BuildPlayWorldURL(const TCHAR* MapName, bool bSpectatorMode, FString AdditionalURLOptions)
 {
 	// the URL we are building up
@@ -2669,11 +2493,7 @@ void UEditorEngine::SpawnIntraProcessPIEWorlds(bool bAnyBlueprintErrors, bool bS
 		}
 
 		UGameInstance* const ClientGameInstance = CreatePIEGameInstance(PIEInstance, bInSimulateInEditor, bAnyBlueprintErrors, bStartInSpectatorMode, false, PIEStartTime);
-		if (ClientGameInstance)
-		{
-			ClientGameInstance->GetWorldContext()->PIERemapPrefix = ServerPrefix;
-		}
-		else
+		if (!ClientGameInstance)
 		{
 			// Failed, abort
 			return;
@@ -2708,31 +2528,6 @@ bool UEditorEngine::CreatePIEWorldFromLogin(FWorldContext& PieWorldContext, EPla
 	{
 		GameInstance->GetWorldContext()->bWaitingOnOnlineSubsystem = false;
 
-		if (PlayNetMode == EPlayNetMode::PIE_ListenServer)
-		{
-			// If any clients finished before us, update their PIERemapPrefix
-			for (FWorldContext &WorldContext : WorldList)
-			{
-				if (WorldContext.WorldType == EWorldType::PIE && WorldContext.World() != NULL && WorldContext.ContextHandle != PieWorldContext.ContextHandle)
-				{
-					WorldContext.PIERemapPrefix = PieWorldContext.PIEPrefix;
-				}
-			}
-		}
-		else
-		{
-			// Grab a valid PIERemapPrefix
-			for (FWorldContext &WorldContext : WorldList)
-			{
-				// This relies on the server being the first in the WorldList. Might be risky.
-				if (WorldContext.WorldType == EWorldType::PIE && WorldContext.World() != NULL && WorldContext.ContextHandle != PieWorldContext.ContextHandle)
-				{
-					PieWorldContext.PIERemapPrefix = WorldContext.PIEPrefix;
-					break;
-				}
-			}
-		}
-
 		return true;
 	}
 	else
@@ -2994,30 +2789,10 @@ void UEditorEngine::RequestLateJoin()
 	}
 	else
 	{
-
-		//FRED_TODO: do stuff here similar to what happens above
-		// Only launch as clients if they should connect
-
-		// 	// this is kind of sketchy but should work
-		// 	int32 SettingsIndex = PIEInstance;
-
-		//	if (WillAutoConnectToServer)
-		{
-			
-		}
-		// 	else
-		// 	{
-		// 		PlayInSettings->SetPlayNetMode(EPlayNetMode::PIE_Standalone);
-		// 	}
-
 		GetMultipleInstancePositions(SettingsIndex++, NextX, NextY);
 
 		UGameInstance* const ClientGameInstance = CreatePIEGameInstance(PIEInstance, false, false, bStartLateJoinersInSpectatorMode, false, PIEStartTime);
-		if (ClientGameInstance)
-		{
-			ClientGameInstance->GetWorldContext()->PIERemapPrefix = ServerPrefix;
-		}
-		else
+		if (!ClientGameInstance)
 		{
 			// Failed, abort
 			return;
@@ -3049,7 +2824,7 @@ UGameInstance* UEditorEngine::CreatePIEGameInstance(int32 InPIEInstance, bool bI
 	}
 
 	// create a new GameInstance
-	FStringClassReference GameInstanceClassName = GetDefault<UGameMapsSettings>()->GameInstanceClass;
+	FSoftClassPath GameInstanceClassName = GetDefault<UGameMapsSettings>()->GameInstanceClass;
 	UClass* GameInstanceClass = (GameInstanceClassName.IsValid() ? LoadObject<UClass>(NULL, *GameInstanceClassName.ToString()) : UGameInstance::StaticClass());
 
 	// If the GameInstance class from the settings cannot be found, fall back to the base class
@@ -3739,57 +3514,6 @@ bool UEditorEngine::PackageUsingExternalObjects( ULevel* LevelToCheck, bool bAdd
 	return bFoundExternal;
 }
 
-UWorld* UEditorEngine::CreatePIEWorldBySavingToTemp(FWorldContext& WorldContext, UWorld* InWorld, FString &PlayWorldMapName)
-{
-	double StartTime = FPlatformTime::Seconds();
-	UWorld * LoadedWorld = NULL;
-
-	// We haven't saved it off yet
-	TArray<FString> SavedMapNames;
-	SaveWorldForPlay(SavedMapNames);
-
-	if (SavedMapNames.Num() == 0)
-	{
-		UE_LOG(LogPlayLevel, Warning, TEXT("PIE: Unable to save editor world to temp file"));
-		return LoadedWorld;
-	}
-
-	// Before loading the map, we need to set these flags to true so that postload will work properly
-	GIsPlayInEditorWorld = true;
-
-	const FName SavedMapFName = FName(*SavedMapNames[0]);
-	UWorld::WorldTypePreLoadMap.FindOrAdd(SavedMapFName) = EWorldType::PIE;
-
-	// Load the package we saved
-	UPackage* EditorLevelPackage = LoadPackage(NULL, *SavedMapNames[0], LOAD_PackageForPIE);
-
-	// Clean up the world type list now that PostLoad has occurred
-	UWorld::WorldTypePreLoadMap.Remove(SavedMapFName);
-
-	if( EditorLevelPackage )
-	{
-		// Find world object and use its PersistentLevel pointer.
-		LoadedWorld = UWorld::FindWorldInPackage(EditorLevelPackage);
-
-		if (LoadedWorld)
-		{
-			PostCreatePIEWorld(LoadedWorld);
-			UE_LOG(LogPlayLevel, Log, TEXT("PIE: Created PIE world by saving and reloading to %s (%fs)"), *LoadedWorld->GetPathName(), float(FPlatformTime::Seconds() - StartTime));
-		}
-		else
-		{
-			UE_LOG(LogPlayLevel, Warning, TEXT("PIE: Unable to find World in loaded package: %s"), *EditorLevelPackage->GetPathName());
-		}
-	}
-
-	// After loading the map, reset these so that things continue as normal
-	GIsPlayInEditorWorld = false;
-
-	PlayWorldMapName = SavedMapNames[0];
-
-	return LoadedWorld;
-}
-
 UWorld* UEditorEngine::CreatePIEWorldByDuplication(FWorldContext &WorldContext, UWorld* InWorld, FString &PlayWorldMapName)
 {
 	double StartTime = FPlatformTime::Seconds();
@@ -3830,19 +3554,16 @@ UWorld* UEditorEngine::CreatePIEWorldByDuplication(FWorldContext &WorldContext,
 		FLazyObjectPtr::ResetPIEFixups();
 
 		// Prepare string asset references for fixup
-		TArray<FString> PackageNamesBeingDuplicatedForPIE;
-		PackageNamesBeingDuplicatedForPIE.Add(PlayWorldMapName);
+		FSoftObjectPath::AddPIEPackageName(FName(*PlayWorldMapName));
 		for (ULevelStreaming* StreamingLevel : InWorld->StreamingLevels)
 		{
 			if ( StreamingLevel )
 			{
 				FString StreamingLevelPIEName = UWorld::ConvertToPIEPackageName(StreamingLevel->GetWorldAssetPackageName(), WorldContext.PIEInstance);
-				PackageNamesBeingDuplicatedForPIE.Add(MoveTemp(StreamingLevelPIEName));
+				FSoftObjectPath::AddPIEPackageName(FName(*StreamingLevelPIEName));
 			}
 		}
 
-		FStringAssetReference::SetPackageNamesBeingDuplicatedForPIE(PackageNamesBeingDuplicatedForPIE);
-
 		// NULL GWorld before various PostLoad functions are called, this makes it easier to debug invalid GWorld accesses
 		GWorld = NULL;
 
@@ -3856,8 +3577,6 @@ UWorld* UEditorEngine::CreatePIEWorldByDuplication(FWorldContext &WorldContext,
 			EDuplicateMode::PIE
 			) );
 
-		FStringAssetReference::ClearPackageNamesBeingDuplicatedForPIE();
-
 		// Store prefix we used to rename this world and streaming levels package names
 		NewPIEWorld->StreamingLevelsPrefix = UWorld::BuildPIEPackagePrefix(WorldContext.PIEInstance);
 		// Fixup model components. The index buffers have been created for the components in the source world and the order
diff --git a/Engine/Source/Editor/UnrealEd/Private/ScriptDisassembler.cpp b/Engine/Source/Editor/UnrealEd/Private/ScriptDisassembler.cpp
--- a/Engine/Source/Editor/UnrealEd/Private/ScriptDisassembler.cpp
+++ b/Engine/Source/Editor/UnrealEd/Private/ScriptDisassembler.cpp
@@ -737,9 +737,9 @@ void FKismetBytecodeDisassembler::ProcessCommon(int32& ScriptIndex, EExprToken O
 			Ar.Logf(TEXT("%s $%X: EX_ObjectConst (%p:%s)"), *Indents, (int32)Opcode, Pointer, *Pointer->GetFullName());
 			break;
 		}
-	case EX_AssetConst:
+	case EX_SoftObjectConst:
 		{
-			Ar.Logf(TEXT("%s $%X: EX_AssetConst"), *Indents, (int32)Opcode);
+			Ar.Logf(TEXT("%s $%X: EX_SoftObjectConst"), *Indents, (int32)Opcode);
 			SerializeExpr(ScriptIndex);
 			break;
 		}
diff --git a/Engine/Source/Editor/UnrealEd/Private/Settings/SettingsClasses.cpp b/Engine/Source/Editor/UnrealEd/Private/Settings/SettingsClasses.cpp
--- a/Engine/Source/Editor/UnrealEd/Private/Settings/SettingsClasses.cpp
+++ b/Engine/Source/Editor/UnrealEd/Private/Settings/SettingsClasses.cpp
@@ -4,7 +4,6 @@
 #include "HAL/FileManager.h"
 #include "Misc/Paths.h"
 #include "Modules/ModuleManager.h"
-#include "Misc/StringAssetReference.h"
 #include "Misc/PackageName.h"
 #include "InputCoreTypes.h"
 #include "EditorStyleSettings.h"
@@ -31,7 +30,6 @@
 #include "AutoReimport/AutoReimportUtilities.h"
 #include "Misc/ConfigCacheIni.h" // for FConfigCacheIni::GetString()
 #include "SourceCodeNavigation.h"
-#include "Developer/BlueprintProfiler/Public/BlueprintProfilerModule.h"
 #include "IProjectManager.h"
 #include "ProjectDescriptor.h"
 #include "Settings/SkeletalMeshEditorSettings.h"
@@ -113,7 +111,6 @@ USkeletalMeshEditorSettings::USkeletalMeshEditorSettings(const FObjectInitialize
 UEditorExperimentalSettings::UEditorExperimentalSettings( const FObjectInitializer& ObjectInitializer )
 	: Super(ObjectInitializer)
 	, bEnableLocalizationDashboard(true)
-	, bBlueprintPerformanceAnalysisTools(false)
 	, bUseOpenCLForConvexHullDecomp(false)
 	, bAllowPotentiallyUnsafePropertyEditing(false)
 {
@@ -138,18 +135,6 @@ void UEditorExperimentalSettings::PostEditChangeProperty( struct FPropertyChange
 			FModuleManager::Get().LoadModule(TEXT("EnvironmentQueryEditor"));
 		}
 	}
-	else if (Name == FName(TEXT("bBlueprintPerformanceAnalysisTools")))
-	{
-		if (!bBlueprintPerformanceAnalysisTools)
-		{
-			IBlueprintProfilerInterface* ProfilerInterface = FModuleManager::GetModulePtr<IBlueprintProfilerInterface>("BlueprintProfiler");
-			if (ProfilerInterface && ProfilerInterface->IsProfilerEnabled())
-			{
-				// Force Profiler off
-				ProfilerInterface->ToggleProfilingCapture();
-			}
-		}
-	}
 
 	if (!FUnrealEdMisc::Get().IsDeletePreferences())
 	{
@@ -422,7 +407,7 @@ ULevelEditorViewportSettings::ULevelEditorViewportSettings( const FObjectInitial
 	MeasuringToolUnits = MeasureUnits_Centimeters;
 
 	// Set a default preview mesh
-	PreviewMeshes.Add(FStringAssetReference("/Engine/EditorMeshes/ColorCalibrator/SM_ColorCalibrator.SM_ColorCalibrator"));
+	PreviewMeshes.Add(FSoftObjectPath("/Engine/EditorMeshes/ColorCalibrator/SM_ColorCalibrator.SM_ColorCalibrator"));
 }
 
 void ULevelEditorViewportSettings::PostInitProperties()
diff --git a/Engine/Source/Editor/UnrealEd/Private/Toolkits/AssetEditorToolkit.cpp b/Engine/Source/Editor/UnrealEd/Private/Toolkits/AssetEditorToolkit.cpp
--- a/Engine/Source/Editor/UnrealEd/Private/Toolkits/AssetEditorToolkit.cpp
+++ b/Engine/Source/Editor/UnrealEd/Private/Toolkits/AssetEditorToolkit.cpp
@@ -983,7 +983,7 @@ void FAssetEditorToolkit::GenerateToolbar()
 	ToolbarBuilder.BeginSection("Asset");
 	{
 		ToolbarBuilder.AddToolBarButton(FAssetEditorCommonCommands::Get().SaveAsset);
-		ToolbarBuilder.AddToolBarButton(FGlobalEditorCommonCommands::Get().FindInContentBrowser, NAME_None, LOCTEXT("FindInContentBrowserButton", "Find in CB"));
+		ToolbarBuilder.AddToolBarButton(FGlobalEditorCommonCommands::Get().FindInContentBrowser, NAME_None, LOCTEXT("FindInContentBrowserButton", "Browse"));
 	}
 	ToolbarBuilder.EndSection();
 
diff --git a/Engine/Source/Editor/UnrealEd/Private/Toolkits/GlobalEditorCommonCommands.cpp b/Engine/Source/Editor/UnrealEd/Private/Toolkits/GlobalEditorCommonCommands.cpp
--- a/Engine/Source/Editor/UnrealEd/Private/Toolkits/GlobalEditorCommonCommands.cpp
+++ b/Engine/Source/Editor/UnrealEd/Private/Toolkits/GlobalEditorCommonCommands.cpp
@@ -32,7 +32,7 @@ void FGlobalEditorCommonCommands::RegisterCommands()
 
 	UI_COMMAND( SummonOpenAssetDialog, "Open Asset...", "Summons an asset picker", EUserInterfaceActionType::Button, FInputChord(EModifierKey::Control, EKeys::P) );
 	UI_COMMAND( SummonOpenAssetDialogAlternate, "Open Asset...", "Summons an asset picker", EUserInterfaceActionType::Button, FInputChord(EModifierKey::Alt | EModifierKey::Shift, EKeys::O));
-	UI_COMMAND( FindInContentBrowser, "Find in Content Browser", "Summons the Content Browser and navigates to the selected asset", EUserInterfaceActionType::Button, FInputChord(EModifierKey::Control, EKeys::B));
+	UI_COMMAND( FindInContentBrowser, "Browse to Asset", "Browses to the associated asset and selects it in the most recently used Content Browser (summoning one if necessary)", EUserInterfaceActionType::Button, FInputChord(EModifierKey::Control, EKeys::B));
 	UI_COMMAND( ViewReferences, "Reference Viewer...", "Launches the reference viewer showing the selected assets' references", EUserInterfaceActionType::Button, FInputChord(EModifierKey::Shift | EModifierKey::Alt, EKeys::R));
 	UI_COMMAND( ViewSizeMap, "Size Map...", "Displays an interactive map showing the approximate size of this asset and everything it references", EUserInterfaceActionType::Button, FInputChord(EModifierKey::Shift | EModifierKey::Alt, EKeys::M));	// @todo sizemap: Make sure key is not used already
 	
diff --git a/Engine/Source/Editor/UnrealEd/Private/UnrealEdPrivatePCH.h b/Engine/Source/Editor/UnrealEd/Private/UnrealEdPrivatePCH.h
--- a/Engine/Source/Editor/UnrealEd/Private/UnrealEdPrivatePCH.h
+++ b/Engine/Source/Editor/UnrealEd/Private/UnrealEdPrivatePCH.h
@@ -280,9 +280,7 @@
 #include "UObject/CoreNative.h"
 #include "UObject/WeakObjectPtr.h"
 #include "Templates/Casts.h"
-#include "Misc/StringAssetReference.h"
-#include "UObject/PersistentObjectPtr.h"
-#include "UObject/AssetPtr.h"
+#include "UObject/SoftObjectPtr.h"
 #include "UObject/Interface.h"
 #include "Templates/SubclassOf.h"
 #include "UObject/LazyObjectPtr.h"
@@ -296,7 +294,6 @@
 #include "UObject/GCObject.h"
 #include "UObject/ScriptMacros.h"
 #include "UObject/Stack.h"
-#include "Misc/StringClassReference.h"
 #include "UObject/Package.h"
 #include "Misc/PackageName.h"
 #include "UObject/ObjectResource.h"
diff --git a/Engine/Source/Editor/UnrealEd/Public/ActorGroupingUtils.h b/Engine/Source/Editor/UnrealEd/Public/ActorGroupingUtils.h
--- a/Engine/Source/Editor/UnrealEd/Public/ActorGroupingUtils.h
+++ b/Engine/Source/Editor/UnrealEd/Public/ActorGroupingUtils.h
@@ -77,7 +77,7 @@ public:
 
 protected:
 	static bool bGroupingActive;
-	static FStringClassReference ClassToUse;
+	static FSoftClassPath ClassToUse;
 };
 
 
diff --git a/Engine/Source/Editor/UnrealEd/Public/EdGraphNode_Comment.h b/Engine/Source/Editor/UnrealEd/Public/EdGraphNode_Comment.h
--- a/Engine/Source/Editor/UnrealEd/Public/EdGraphNode_Comment.h
+++ b/Engine/Source/Editor/UnrealEd/Public/EdGraphNode_Comment.h
@@ -41,8 +41,16 @@ public:
 	UPROPERTY(EditAnywhere, Category=Comment)
 	FLinearColor CommentColor;
 
+	/** Size of the text in the comment box */
+	UPROPERTY(EditAnywhere, Category=Comment, meta=(ClampMin=1, ClampMax=1000))
+	int32 FontSize;
+
+	/** Whether to show a zoom-invariant comment bubble when zoomed out (making the comment readable at any distance). */
+	UPROPERTY(EditAnywhere, Category=Comment, meta=(DisplayName="Show Bubble When Zoomed"))
+	uint32 bCommentBubbleVisible_InDetailsPanel:1;
+
 	/** Whether to use Comment Color to color the background of the comment bubble shown when zoomed out. */
-	UPROPERTY(EditAnywhere, Category=Comment, meta=(DisplayName="Color Bubble"))
+	UPROPERTY(EditAnywhere, Category=Comment, meta=(DisplayName="Color Bubble", EditCondition=bCommentBubbleVisible_InDetailsPanel))
 	uint32 bColorCommentBubble:1;
 
 	/** Whether the comment should move any fully enclosed nodes around when it is moved */
@@ -57,6 +65,7 @@ public:
 
 	//~ Begin UObject Interface
 	UNREALED_API static void AddReferencedObjects(UObject* InThis, FReferenceCollector& Collector);
+	virtual void PostEditChangeProperty(FPropertyChangedEvent& PropertyChangedEvent) override;
 	//~ End UObject Interface
 
 	//~ Begin UEdGraphNode Interface
@@ -70,7 +79,6 @@ public:
 	UNREALED_API virtual void PostPlacedNewNode() override;
 	UNREALED_API virtual void OnRenameNode(const FString& NewName) override;
 	UNREALED_API virtual TSharedPtr<class INameValidatorInterface> MakeNameValidator() const override;
-	virtual bool ShouldDrawNodeAsComment() const override { return true; }
 	UNREALED_API virtual FString GetDocumentationLink() const override;
 	UNREALED_API virtual FString GetDocumentationExcerptName() const override;
 	UNREALED_API virtual FSlateIcon GetIconAndTint(FLinearColor& OutColor) const override;
diff --git a/Engine/Source/Editor/UnrealEd/Public/EdGraphToken.h b/Engine/Source/Editor/UnrealEd/Public/EdGraphToken.h
--- a/Engine/Source/Editor/UnrealEd/Public/EdGraphToken.h
+++ b/Engine/Source/Editor/UnrealEd/Public/EdGraphToken.h
@@ -18,6 +18,7 @@ public:
 	/** Factory method, tokens can only be constructed as shared refs */
 	UNREALED_API static TSharedRef<IMessageToken> Create(const UObject* InObject, const FCompilerResultsLog* Log, UEdGraphNode*& OutSourceNode);
 	UNREALED_API static TSharedRef<IMessageToken> Create(const UEdGraphPin* InPin, const FCompilerResultsLog* Log, UEdGraphNode*& OutSourceNode);
+	UNREALED_API static TSharedRef<IMessageToken> Create(const TCHAR* String, const FCompilerResultsLog* Log, UEdGraphNode*& OutSourceNode);
 
 	/** Begin IMessageToken interface */
 	virtual EMessageToken::Type GetType() const override
diff --git a/Engine/Source/Editor/UnrealEd/Public/EdGraphUtilities.h b/Engine/Source/Editor/UnrealEd/Public/EdGraphUtilities.h
--- a/Engine/Source/Editor/UnrealEd/Public/EdGraphUtilities.h
+++ b/Engine/Source/Editor/UnrealEd/Public/EdGraphUtilities.h
@@ -77,13 +77,13 @@ public:
 	  * @param OutClonedNodes			Will populate with a full list of cloned nodes if provided
 	  * @param bInIsCompiling			TRUE if the function is being called during compilation, this will eliminate some nodes that will not be compiled
 	  */
-	static void CloneAndMergeGraphIn(UEdGraph* MergeTarget, UEdGraph* SourceGraph, FCompilerResultsLog& MessageLog, bool bRequireSchemaMatch, bool bInIsCompiling = false, bool bCreateBoundaryNodes = false, TArray<UEdGraphNode*>* OutClonedNodes = NULL);
+	static void CloneAndMergeGraphIn(UEdGraph* MergeTarget, UEdGraph* SourceGraph, FCompilerResultsLog& MessageLog, bool bRequireSchemaMatch, bool bInIsCompiling = false, TArray<UEdGraphNode*>* OutClonedNodes = NULL);
 
 	/**
 	 * Moves the contents of all of the children graphs of ParentGraph (recursively) into the MergeTarget graph.
 	 * This does not clone, it's destructive to the ParentGraph
 	 */
-	static void MergeChildrenGraphsIn(UEdGraph* MergeTarget, UEdGraph* ParentGraph, bool bRequireSchemaMatch, bool bInIsCompiling = false, FCompilerResultsLog* MessageLog = nullptr, bool bWantBoundaryNodes = false);
+	static void MergeChildrenGraphsIn(UEdGraph* MergeTarget, UEdGraph* ParentGraph, bool bRequireSchemaMatch, bool bInIsCompiling = false, FCompilerResultsLog* MessageLog = nullptr);
 
 	/** Tries to rename the graph to have a name similar to BaseName */
 	static void RenameGraphCloseToName(UEdGraph* Graph, const FString& BaseName, int32 StartIndex = 1);
diff --git a/Engine/Source/Editor/UnrealEd/Public/Kismet2/BlueprintEditorUtils.h b/Engine/Source/Editor/UnrealEd/Public/Kismet2/BlueprintEditorUtils.h
--- a/Engine/Source/Editor/UnrealEd/Public/Kismet2/BlueprintEditorUtils.h
+++ b/Engine/Source/Editor/UnrealEd/Public/Kismet2/BlueprintEditorUtils.h
@@ -617,7 +617,7 @@ public:
 	 *
 	 * @param [in,out]	GraphNames	The graph names will be appended to this array.
 	 */
-	static void GetAllGraphNames(const UBlueprint* Blueprint, TArray<FName>& GraphNames);
+	static void GetAllGraphNames(const UBlueprint* Blueprint, TSet<FName>& GraphNames);
 
 	/**
 	 * Gets the compiler-relevant (i.e. non-ignorable) node links from the given pin.
@@ -654,7 +654,7 @@ public:
 	 * @param			Blueprint		The blueprint to check
 	 * @param [in,out]	FunctionNames	List of function names currently in use
 	 */
-	static void GetFunctionNameList(const UBlueprint* Blueprint, TArray<FName>& FunctionNames);
+	static void GetFunctionNameList(const UBlueprint* Blueprint, TSet<FName>& FunctionNames);
 
 	/**
 	 * Gets a list of delegates names in the blueprint, based on the skeleton class
@@ -662,7 +662,7 @@ public:
 	 * @param			Blueprint		The blueprint to check
 	 * @param [in,out]	DelegatesNames	List of function names currently in use
 	 */
-	static void GetDelegateNameList(const UBlueprint* Blueprint, TArray<FName>& DelegatesNames);
+	static void GetDelegateNameList(const UBlueprint* Blueprint, TSet<FName>& DelegatesNames);
 
 	/** 
 	 * Get a graph for delegate signature with given name, from given blueprint.
@@ -713,15 +713,15 @@ public:
 	/**
 	 * Gets the visible class variable list.  This includes both variables introduced here and in all superclasses.
 	 *
-	 * @param [in,out]	VisibleVariables	The visible variables will be appened to this array.
+	 * @param [in,out]	VisibleVariables	The visible variables will be appended to this array.
 	 */
-	static void GetClassVariableList(const UBlueprint* Blueprint, TArray<FName>& VisibleVariables, bool bIncludePrivateVars=false);
+	static void GetClassVariableList(const UBlueprint* Blueprint, TSet<FName>& VisibleVariables, bool bIncludePrivateVars=false);
 
 	/**
 	 * Gets variables of specified type
 	 *
 	 * @param 			FEdGraphPinType	 			Type of variables to look for
-	 * @param [in,out]	VisibleVariables			The visible variables will be appened to this array.
+	 * @param [in,out]	VisibleVariables			The visible variables will be appended to this array.
 	 */
 	static void GetNewVariablesOfType( const UBlueprint* Blueprint, const FEdGraphPinType& Type, TArray<FName>& OutVars);
 
@@ -729,7 +729,7 @@ public:
 	 * Gets local variables of specified type
 	 *
 	 * @param 			FEdGraphPinType	 			Type of variables to look for
-	 * @param [in,out]	VisibleVariables			The visible variables will be appened to this array.
+	 * @param [in,out]	VisibleVariables			The visible variables will be appended to this array.
 	 */
 	static void GetLocalVariablesOfType( const UEdGraph* Graph, const FEdGraphPinType& Type, TArray<FName>& OutVars);
 
@@ -1028,6 +1028,23 @@ public:
 	 */
 	static void SetBlueprintVariableCategory(UBlueprint* Blueprint, const FName& VarName, const UStruct* InLocalVarScope, const FText& NewCategory, bool bDontRecompile=false);
 
+
+	/**
+	 * Sets the custom category on the function or macro
+	 * @note: Will not change the category for functions defined via native classes.
+	 *
+	 * @param	Graph				Graph associated with the function or macro
+	 * @param	NewCategory			The new value of the custom category for the function
+	 * @param	bDontRecompile		If true, the blueprint will not be marked as modified, and will not be recompiled.  
+	 */
+	static void SetBlueprintFunctionOrMacroCategory(UEdGraph* Graph, const FText& NewCategory, bool bDontRecompile=false);
+
+	/** Finds the index of the specified graph (function or macro) in the parent (if it is not reorderable, then we will return INDEX_NONE) */
+	static int32 FindIndexOfGraphInParent(UEdGraph* Graph);
+
+	/** Reorders the specified graph (function or macro) to be at the new index in the parent (moving whatever was there to be after it), assuming it is reorderable and that is a valid index */
+	static bool MoveGraphBeforeOtherGraph(UEdGraph* Graph, int32 NewIndex, bool bDontRecompile);
+
 	/**
 	 * Gets the custom category on the variable with the specified name.
 	 *
@@ -1061,9 +1078,6 @@ public:
 	/** Change the order of variables in the Blueprint */
 	static bool MoveVariableBeforeVariable(UBlueprint* Blueprint, FName VarNameToMove, FName TargetVarName, bool bDontRecompile);
 
-	/** Find first variable of the supplied category */
-	static int32 FindFirstNewVarOfCategory(const UBlueprint* Blueprint, FText Category);
-
 	/**
 	 * Find the index of a timeline first declared in this blueprint. Returns INDEX_NONE if not found.
 	 *
@@ -1078,7 +1092,7 @@ public:
 	 *
 	 * @param [in,out]	VariableNames		The list of variable names for the SCS node array.
 	 */
-	static void GetSCSVariableNameList(const UBlueprint* Blueprint, TArray<FName>& VariableNames);
+	static void GetSCSVariableNameList(const UBlueprint* Blueprint, TSet<FName>& VariableNames);
 
 	/** 
 	 * Gets a list of function names in blueprints that implement the interface defined by the given blueprint.
@@ -1086,7 +1100,7 @@ public:
 	 * @param [in,out]	Blueprint			The interface blueprint to check.
 	 * @param [in,out]	VariableNames		The list of function names for implementing blueprints.
 	 */
-	static void GetImplementingBlueprintsFunctionNameList(const UBlueprint* Blueprint, TArray<FName>& FunctionNames);
+	static void GetImplementingBlueprintsFunctionNameList(const UBlueprint* Blueprint, TSet<FName>& FunctionNames);
 
 	/**
 	 * Find the index of a SCS_Node first declared in this blueprint. Returns INDEX_NONE if not found.
@@ -1287,7 +1301,7 @@ public:
 	 * @param	LevelScriptActor	The newly-created level script actor that should be (re-)bound to
 	 * @param	ScriptBlueprint		The level scripting blueprint that contains the bound events to try and bind delegates to this actor for
 	 */
-	static bool FixLevelScriptActorBindings(ALevelScriptActor* LevelScriptActor, const class ULevelScriptBlueprint* ScriptBlueprint);
+	static void FixLevelScriptActorBindings(ALevelScriptActor* LevelScriptActor, const class ULevelScriptBlueprint* ScriptBlueprint);
 
 	/**
 	 * Find how many actors reference the supplied actor
@@ -1315,11 +1329,6 @@ public:
 	// Diagnostic exec commands
 	static bool KismetDiagnosticExec(const TCHAR* Stream, FOutputDevice& Ar);
 
-	/** Called by FixLevelScriptActorBindings() to allow external modules to update level script actors when they are
-	    recreated after a blueprint change */
-	DECLARE_EVENT_ThreeParams( FBlueprintEditorUtils, FFixLevelScriptActorBindingsEvent, ALevelScriptActor* /* LevelScriptActor */, const ULevelScriptBlueprint* /* LevelScriptBlueprint */, bool& /* bWasSuccessful */ );
-	static FFixLevelScriptActorBindingsEvent& OnFixLevelScriptActorBindings();
-
 	/**
 	 * Searches the world for any blueprints that are open and do not have a debug instances set and sets one if possible.
 	 * It will favor a selected instance over a non selected one
@@ -1385,7 +1394,7 @@ public:
 	 * @param InGraph			The graph to search through
 	 * @return					If valid a pointer to the user declared function meta data structure otherwise nullptr.
 	 */
-	static FKismetUserDeclaredFunctionMetadata* GetGraphFunctionMetaData( UEdGraph* InGraph );
+	static FKismetUserDeclaredFunctionMetadata* GetGraphFunctionMetaData(const UEdGraph* InGraph);
 
 	/**
 	 * Returns the description of the graph from the metadata
@@ -1451,9 +1460,6 @@ protected:
 	// Removes all NULL graph references in the specified array
 	static void CleanNullGraphReferencesInArray(UBlueprint* Blueprint, TArray<UEdGraph*>& GraphArray);
 
-	/** Static global event that is broadcast when a blueprint is changed and FixLevelScriptActorBindings() is called */
-	static FFixLevelScriptActorBindingsEvent FixLevelScriptActorBindingsEvent;
-
 	/**
 	 * Checks that the actor type matches the blueprint type (or optionally is BASED on the same type. 
 	 *
diff --git a/Engine/Source/Editor/UnrealEd/Public/Kismet2/CompilerResultsLog.h b/Engine/Source/Editor/UnrealEd/Public/Kismet2/CompilerResultsLog.h
--- a/Engine/Source/Editor/UnrealEd/Public/Kismet2/CompilerResultsLog.h
+++ b/Engine/Source/Editor/UnrealEd/Public/Kismet2/CompilerResultsLog.h
@@ -101,9 +101,6 @@ public:
 	// Should detailed results be appended to the final summary log?
 	bool bLogDetailedResults;
 
-	// Should composite graphs be treated the same as macro graphs/tunnels to generate more accurate debug data.
-	bool bTreatCompositeGraphsAsTunnels;
-
 	// Minimum event time (ms) for inclusion into the final summary log
 	int EventDisplayThresholdMs;
 
@@ -129,12 +126,6 @@ protected:
 	// Map to track active nested tunnels for intermediate tunnel instances.
 	TMultiMap<TWeakObjectPtr<UEdGraphNode>, TWeakObjectPtr<UEdGraphNode>> IntermediateTunnelInstanceHierarchyMap;
 
-	// Map to track compile time only expansion nodes.
-	TSet<TWeakObjectPtr<const UEdGraphNode>> ExpansionNodes;
-
-	// Per pin filtering of trace signals for instrumentation.
-	TSet<const UEdGraphPin*> PinTraceFilter;
-
 public:
 	FCompilerResultsLog(bool bIsCompatibleWithEvents = true);
 	virtual ~FCompilerResultsLog();
@@ -269,14 +260,6 @@ public:
 	void NotifyIntermediateObjectCreation(UObject* NewObject, UObject* SourceObject);
 	void NotifyIntermediatePinCreation(UEdGraphPin* NewObject, UEdGraphPin* SourceObject);
 
-	/** Compile time node expansion tracking and testing */
-	void RegisterExpansionNode(UEdGraphNode* ExpansionNode) { ExpansionNodes.Add(ExpansionNode); }
-	bool IsExpansionNode(const UEdGraphNode* Node) const { return ExpansionNodes.Contains(Node); }
-
-	/** Per pin trace signal filtering */
-	void AddPinTraceFilter(UEdGraphPin* PinToFilter) { PinTraceFilter.Add(PinToFilter); }
-	bool IsPinTraceSuppressed(const UEdGraphPin* PotentiallyFilteredPin) const { return PinTraceFilter.Contains(PotentiallyFilteredPin); }
-
 	/** Registers intermediate tunnel nodes, both the node and tunnel instance should be intermediate nodes */
 	void RegisterIntermediateTunnelNode(UEdGraphNode* Node, UEdGraphNode* OwningTunnelInstance);
 
diff --git a/Engine/Source/Editor/UnrealEd/Public/Kismet2/Kismet2NameValidators.h b/Engine/Source/Editor/UnrealEd/Public/Kismet2/Kismet2NameValidators.h
--- a/Engine/Source/Editor/UnrealEd/Public/Kismet2/Kismet2NameValidators.h
+++ b/Engine/Source/Editor/UnrealEd/Public/Kismet2/Kismet2NameValidators.h
@@ -95,7 +95,7 @@ public:
 	// End FNameValidatorInterface
 private:
 	/** Name set to validate */
-	TArray<FName> Names;
+	TSet<FName> Names;
 	/** The blueprint to check for validity within */
 	const UBlueprint* BlueprintObject;
 	/** The current name of the object being validated */
diff --git a/Engine/Source/Editor/UnrealEd/Public/Kismet2/KismetEditorUtilities.h b/Engine/Source/Editor/UnrealEd/Public/Kismet2/KismetEditorUtilities.h
--- a/Engine/Source/Editor/UnrealEd/Public/Kismet2/KismetEditorUtilities.h
+++ b/Engine/Source/Editor/UnrealEd/Public/Kismet2/KismetEditorUtilities.h
@@ -42,8 +42,6 @@ enum class EBlueprintCompileOptions
 	SkeletonUpToDate = 0x8, 
 	/** Indicates this is a batch compile and that BroadcastCompiled and BroadcastBlueprintCompiled should be skipped */
 	BatchCompile = 0x10,
-	/** Triggers profiling information to be generated */
-	AddInstrumentation = 0x20,
 	/** Skips creating a reinstancer and running reinstancing routines - useful if calling code is performing reinstancing */
 	SkipReinstancing = 0x40,
 };
diff --git a/Engine/Source/Editor/UnrealEd/Public/Kismet2/KismetReinstanceUtilities.h b/Engine/Source/Editor/UnrealEd/Public/Kismet2/KismetReinstanceUtilities.h
--- a/Engine/Source/Editor/UnrealEd/Public/Kismet2/KismetReinstanceUtilities.h
+++ b/Engine/Source/Editor/UnrealEd/Public/Kismet2/KismetReinstanceUtilities.h
@@ -38,17 +38,6 @@ public:
 
 class UNREALED_API FBlueprintCompileReinstancer : public TSharedFromThis<FBlueprintCompileReinstancer>, public FGCObject
 {
-public:
-	/**
-	 * CDO duplicates provider delegate type.
-	 */
-	DECLARE_DELEGATE_RetVal_TwoParams(UObject*, FCDODuplicatesProvider, UObject*, FName);
-
-	/**
-	 * Gets CDO duplicates provider delegate.
-	 */
-	static FCDODuplicatesProvider& GetCDODuplicatesProviderDelegate();
-
 protected:
 	friend struct FBlueprintCompilationManagerImpl;
 
@@ -198,15 +187,6 @@ protected:
 	static void CopyPropertiesForUnrelatedObjects(UObject* OldObject, UObject* NewObject, bool bClearExternalReferences);
 
 private:
-	/**
-	 * Gets class CDO duplicate from cache (or creates it if not available or not
-	 * during hot-reload) for given class.
-	 *
-	 * @param CDO Object to get duplicate of.
-	 * @param Name The name that CDO has to be renamed with (or created with).
-	 */
-	static UObject* GetClassCDODuplicate(UObject* CDO, FName Name);
-
 	/** Handles the work of ReplaceInstancesOfClass, handling both normal replacement of instances and batch */
 	static void ReplaceInstancesOfClass_Inner(TMap<UClass*, UClass*>& InOldToNewClassMap, UObject* InOriginalCDO, TSet<UObject*>* ObjectsThatShouldUseOldStuff = NULL, bool bClassObjectReplaced = false, bool bPreserveRootComponent = true);
 };
diff --git a/Engine/Source/Editor/UnrealEd/Public/UnrealEdSharedPCH.h b/Engine/Source/Editor/UnrealEd/Public/UnrealEdSharedPCH.h
--- a/Engine/Source/Editor/UnrealEd/Public/UnrealEdSharedPCH.h
+++ b/Engine/Source/Editor/UnrealEd/Public/UnrealEdSharedPCH.h
@@ -271,11 +271,9 @@
 #include "UObject/WeakObjectPtr.h"
 #include "UObject/Interface.h"
 #include "Templates/Casts.h"
-#include "Misc/StringAssetReference.h"
 #include "UObject/CoreNetTypes.h"
-#include "UObject/PersistentObjectPtr.h"
 #include "UObject/LazyObjectPtr.h"
-#include "UObject/AssetPtr.h"
+#include "UObject/SoftObjectPtr.h"
 #include "UObject/GCObject.h"
 #include "Templates/SubclassOf.h"
 #include "Misc/WorldCompositionUtility.h"
@@ -287,7 +285,6 @@
 #include "UObject/ScriptMacros.h"
 #include "UObject/Stack.h"
 #include "UObject/Package.h"
-#include "Misc/StringClassReference.h"
 #include "Misc/PackageName.h"
 #include "UObject/LinkerLoad.h"
 #include "UObject/ObjectResource.h"
diff --git a/Engine/Source/Programs/UnrealBuildTool/Windows/VCToolChain.cs b/Engine/Source/Programs/UnrealBuildTool/Windows/VCToolChain.cs
--- a/Engine/Source/Programs/UnrealBuildTool/Windows/VCToolChain.cs
+++ b/Engine/Source/Programs/UnrealBuildTool/Windows/VCToolChain.cs
@@ -1377,7 +1377,10 @@ namespace UnrealBuildTool
 				// Preprocessor definitions.
 				foreach (string Definition in CompileEnvironment.Definitions)
 				{
-					AddDefinition(Arguments, Definition);
+					if (!Definition.Contains("_API"))
+					{
+						AddDefinition(Arguments, Definition);
+					}
 				}
 
 				// Add the RES file to the produced item list.
diff --git a/Engine/Source/Programs/UnrealHeaderTool/Private/CodeGenerator.cpp b/Engine/Source/Programs/UnrealHeaderTool/Private/CodeGenerator.cpp
--- a/Engine/Source/Programs/UnrealHeaderTool/Private/CodeGenerator.cpp
+++ b/Engine/Source/Programs/UnrealHeaderTool/Private/CodeGenerator.cpp
@@ -1139,10 +1139,10 @@ void FNativeClassHeaderGenerator::PropertyNew(FOutputDevice& Out, UProperty* Pro
 		return;
 	}
 
-	if (UAssetClassProperty* TypedProp = Cast<UAssetClassProperty>(Prop))
+	if (USoftClassProperty* TypedProp = Cast<USoftClassProperty>(Prop))
 	{
 		Out.Logf(
-			TEXT("%sstatic const UE4CodeGen_Private::FAssetClassPropertyParams %s = { UE4CodeGen_Private::EPropertyClass::AssetClass, %s, %s, 0x%016llx, %s, %s, %s, %s, %s };%s\r\n"),
+			TEXT("%sstatic const UE4CodeGen_Private::FSoftClassPropertyParams %s = { UE4CodeGen_Private::EPropertyClass::SoftClass, %s, %s, 0x%016llx, %s, %s, %s, %s, %s };%s\r\n"),
 			Spaces,
 			Name,
 			*PropName,
@@ -1199,10 +1199,10 @@ void FNativeClassHeaderGenerator::PropertyNew(FOutputDevice& Out, UProperty* Pro
 		return;
 	}
 
-	if (UAssetObjectProperty* TypedProp = Cast<UAssetObjectProperty>(Prop))
+	if (USoftObjectProperty* TypedProp = Cast<USoftObjectProperty>(Prop))
 	{
 		Out.Logf(
-			TEXT("%sstatic const UE4CodeGen_Private::FAssetObjectPropertyParams %s = { UE4CodeGen_Private::EPropertyClass::AssetObject, %s, %s, 0x%016llx, %s, %s, %s, %s, %s };%s\r\n"),
+			TEXT("%sstatic const UE4CodeGen_Private::FSoftObjectPropertyParams %s = { UE4CodeGen_Private::EPropertyClass::SoftObject, %s, %s, 0x%016llx, %s, %s, %s, %s, %s };%s\r\n"),
 			Spaces,
 			Name,
 			*PropName,
@@ -3106,7 +3106,10 @@ void FNativeClassHeaderGenerator::ExportGeneratedStructBodyMacros(FOutputDevice&
 		Out.Logf(TEXT("}\r\n"));
 
 		Out.Logf(TEXT("static FCompiledInDeferStruct Z_CompiledInDeferStruct_UScriptStruct_%s(%s::StaticStruct, TEXT(\"%s\"), TEXT(\"%s\"), %s, %s, %s);\r\n"),
-			StructNameCPP, StructNameCPP, *Struct->GetOutermost()->GetName(), *ActualStructName,
+			StructNameCPP,
+			StructNameCPP,
+			bIsDynamic ? *FClass::GetTypePackageName(Struct) : *Struct->GetOutermost()->GetName(),
+			*ActualStructName,
 			bIsDynamic ? TEXT("true") : TEXT("false"),
 			bIsDynamic ? *AsTEXT(FClass::GetTypePackageName(Struct)) : TEXT("nullptr"),
 			bIsDynamic ? *AsTEXT(FNativeClassHeaderGenerator::GetOverriddenPathName(Struct)) : TEXT("nullptr"));
