Fix Lightmass ApplicationCore module not being staged for Swarm, causing missing DLL errors when building lighting.

diff --git a/Engine/Source/Editor/UnrealEd/Private/Lightmass/Lightmass.cpp b/Engine/Source/Editor/UnrealEd/Private/Lightmass/Lightmass.cpp
--- a/Engine/Source/Editor/UnrealEd/Private/Lightmass/Lightmass.cpp
+++ b/Engine/Source/Editor/UnrealEd/Private/Lightmass/Lightmass.cpp
@@ -2627,6 +2627,7 @@ bool FLightmassProcessor::BeginRun()
 		TEXT("../DotNET/SwarmInterface.dll"),
 		TEXT("../Win32/AgentInterface.dll"),
 		TEXT("../Win32/UnrealLightmass-SwarmInterface.dll"),
+		TEXT("../Win32/UnrealLightmass-ApplicationCore.dll"),
 		TEXT("../Win32/UnrealLightmass-Core.dll"),
 		TEXT("../Win32/UnrealLightmass-CoreUObject.dll"),
 		TEXT("../Win32/UnrealLightmass-Projects.dll"),
@@ -2642,6 +2643,7 @@ bool FLightmassProcessor::BeginRun()
 		TEXT("../DotNET/SwarmInterface.dll"),
 		TEXT("../Win64/AgentInterface.dll"),
 		TEXT("../Win64/UnrealLightmass-SwarmInterface.dll"),
+		TEXT("../Win64/UnrealLightmass-ApplicationCore.dll"),
 		TEXT("../Win64/UnrealLightmass-Core.dll"),
 		TEXT("../Win64/UnrealLightmass-CoreUObject.dll"),
 		TEXT("../Win64/UnrealLightmass-Projects.dll"),
@@ -2655,6 +2657,7 @@ bool FLightmassProcessor::BeginRun()
 	const TCHAR* RequiredDependencyPaths64[] =
 	{
 		TEXT("../DotNET/Mac/AgentInterface.dll"),
+		TEXT("../Mac/UnrealLightmass-ApplicationCore.dylib"),
 		TEXT("../Mac/UnrealLightmass-Core.dylib"),
 		TEXT("../Mac/UnrealLightmass-CoreUObject.dylib"),
 		TEXT("../Mac/UnrealLightmass-Json.dylib"),
@@ -2669,6 +2672,7 @@ bool FLightmassProcessor::BeginRun()
 	const TCHAR* RequiredDependencyPaths64[] =
 	{
 		TEXT("../DotNET/Linux/AgentInterface.dll"),
+		TEXT("../Linux/libUnrealLightmass-ApplicationCore.so"),
 		TEXT("../Linux/libUnrealLightmass-Core.so"),
 		TEXT("../Linux/libUnrealLightmass-CoreUObject.so"),
 		TEXT("../Linux/libUnrealLightmass-Json.so"),
