Copying //UE4/Dev-Framework to //UE4/Dev-Main (Source: //UE4/Dev-Framework @ 3628051)

diff --git a/Engine/Source/Editor/Kismet/Private/FindInBlueprintManager.cpp b/Engine/Source/Editor/Kismet/Private/FindInBlueprintManager.cpp
--- a/Engine/Source/Editor/Kismet/Private/FindInBlueprintManager.cpp
+++ b/Engine/Source/Editor/Kismet/Private/FindInBlueprintManager.cpp
@@ -30,6 +30,12 @@
 #include "FileHelpers.h"
 #include "EdGraphSchema_K2.h"
 #include "K2Node_FunctionEntry.h"
+#include "EditorStyleSet.h"
+#include "BlueprintEditorSettings.h"
+#include "Framework/Docking/TabManager.h"
+#include "Widgets/Docking/SDockTab.h"
+#include "WorkspaceMenuStructure.h"
+#include "WorkspaceMenuStructureModule.h"
 
 #include "Engine/SimpleConstructionScript.h"
 #include "Kismet2/BlueprintEditorUtils.h"
@@ -1339,6 +1345,11 @@ FFindInBlueprintSearchManager::FFindInBlueprintSearchManager()
 	, AssetRegistryModule(nullptr)
 	, CachingObject(nullptr)
 {
+	for (int32 TabIdx = 0; TabIdx < ARRAY_COUNT(GlobalFindResultsTabIDs); TabIdx++)
+	{
+		const FName TabID = FName(*FString::Printf(TEXT("GlobalFindResults_%02d"), TabIdx + 1));
+		GlobalFindResultsTabIDs[TabIdx] = TabID;
+	}
 }
 
 FFindInBlueprintSearchManager::~FFindInBlueprintSearchManager()
@@ -1359,6 +1370,9 @@ FFindInBlueprintSearchManager::~FFindInBlueprintSearchManager()
 		IHotReloadInterface& HotReloadSupport = FModuleManager::GetModuleChecked<IHotReloadInterface>("HotReload");
 		HotReloadSupport.OnHotReload().RemoveAll(this);
 	}
+
+	// Shut down the global find results tab feature.
+	EnableGlobalFindResults(false);
 }
 
 void FFindInBlueprintSearchManager::Initialize()
@@ -1397,6 +1411,12 @@ void FFindInBlueprintSearchManager::Initialize()
 		// Do an immediate load of the cache to catch any Blueprints that were discovered by the asset registry before we initialized.
 		BuildCache();
 	}
+
+	// Register global find results tabs if the feature is enabled.
+	if (GetDefault<UBlueprintEditorSettings>()->bHostFindInBlueprintsInGlobalTab)
+	{
+		EnableGlobalFindResults(true);
+	}
 }
 
 void FFindInBlueprintSearchManager::OnAssetAdded(const FAssetData& InAssetData)
@@ -2262,4 +2282,198 @@ TSharedPtr< FJsonObject > FFindInBlueprintSearchManager::ConvertJsonStringToObje
 	return JsonObject;
 }
 
+void FFindInBlueprintSearchManager::GlobalFindResultsClosed(const TSharedRef<SFindInBlueprints>& FindResults)
+{
+	for (TWeakPtr<SFindInBlueprints> FindResultsPtr : GlobalFindResults)
+	{
+		if (FindResultsPtr.Pin() == FindResults)
+		{
+			GlobalFindResults.Remove(FindResultsPtr);
+			break;
+		}
+	}
+}
+
+FText FFindInBlueprintSearchManager::GetGlobalFindResultsTabLabel(int32 TabIdx)
+{
+	int32 NumOpenGlobalFindResultsTabs = 0;
+	for (int32 i = GlobalFindResults.Num() - 1; i >= 0; --i)
+	{
+		if (GlobalFindResults[i].IsValid())
+		{
+			++NumOpenGlobalFindResultsTabs;
+		}
+		else
+		{
+			GlobalFindResults.RemoveAt(i);
+		}
+	}
+
+	if (NumOpenGlobalFindResultsTabs > 1 || TabIdx > 0)
+	{
+		return FText::Format(LOCTEXT("GlobalFindResultsTabNameWithIndex", "Find in Blueprints {0}"), FText::AsNumber(TabIdx + 1));
+	}
+	else
+	{
+		return LOCTEXT("GlobalFindResultsTabName", "Find in Blueprints");
+	}
+}
+
+TSharedRef<SDockTab> FFindInBlueprintSearchManager::SpawnGlobalFindResultsTab(const FSpawnTabArgs& SpawnTabArgs, int32 TabIdx)
+{
+	TAttribute<FText> Label = TAttribute<FText>::Create(TAttribute<FText>::FGetter::CreateRaw(this, &FFindInBlueprintSearchManager::GetGlobalFindResultsTabLabel, TabIdx));
+
+	TSharedRef<SDockTab> NewTab = SNew(SDockTab)
+		.TabRole(ETabRole::NomadTab)
+		.Label(Label)
+		.ToolTipText(LOCTEXT("GlobalFindResultsTabTooltip", "Search for a string in all Blueprint assets."));
+
+	TSharedRef<SFindInBlueprints> FindResults = SNew(SFindInBlueprints)
+		.bIsSearchWindow(false)
+		.ContainingTab(NewTab);
+
+	GlobalFindResults.Add(FindResults);
+
+	NewTab->SetContent(FindResults);
+
+	return NewTab;
+}
+
+TSharedPtr<SFindInBlueprints> FFindInBlueprintSearchManager::OpenGlobalFindResultsTab()
+{
+	TSet<FName> OpenGlobalTabIDs;
+
+	for (TWeakPtr<SFindInBlueprints> FindResultsPtr : GlobalFindResults)
+	{
+		TSharedPtr<SFindInBlueprints> FindResults = FindResultsPtr.Pin();
+		if (FindResults.IsValid())
+		{
+			OpenGlobalTabIDs.Add(FindResults->GetHostTabId());
+		}
+	}
+
+	for (int32 Idx = 0; Idx < ARRAY_COUNT(GlobalFindResultsTabIDs); ++Idx)
+	{
+		const FName GlobalTabId = GlobalFindResultsTabIDs[Idx];
+		if (!OpenGlobalTabIDs.Contains(GlobalTabId))
+		{
+			TSharedRef<SDockTab> NewTab = FGlobalTabmanager::Get()->InvokeTab(GlobalTabId);
+			return StaticCastSharedRef<SFindInBlueprints>(NewTab->GetContent());
+		}
+	}
+
+	return TSharedPtr<SFindInBlueprints>();
+}
+
+TSharedPtr<SFindInBlueprints> FFindInBlueprintSearchManager::GetGlobalFindResults()
+{
+	TSharedPtr<SFindInBlueprints> FindResultsToUse;
+
+	for (TWeakPtr<SFindInBlueprints> FindResultsPtr : GlobalFindResults)
+	{
+		TSharedPtr<SFindInBlueprints> FindResults = FindResultsPtr.Pin();
+		if (FindResults.IsValid() && !FindResults->IsLocked())
+		{
+			FindResultsToUse = FindResults;
+			break;
+		}
+	}
+
+	if (FindResultsToUse.IsValid())
+	{
+		FGlobalTabmanager::Get()->InvokeTab(FindResultsToUse->GetHostTabId());
+	}
+	else
+	{
+		FindResultsToUse = OpenGlobalFindResultsTab();
+	}
+
+	return FindResultsToUse;
+}
+
+void FFindInBlueprintSearchManager::EnableGlobalFindResults(bool bEnable)
+{
+	const TSharedRef<FGlobalTabmanager>& GlobalTabManager = FGlobalTabmanager::Get();
+
+	if (bEnable)
+	{
+		// Register the spawners for all global Find Results tabs
+		const FSlateIcon GlobalFindResultsIcon(FEditorStyle::GetStyleSetName(), "Kismet.Tabs.FindResults");
+		GlobalFindResultsMenuItem = WorkspaceMenu::GetMenuStructure().GetToolsCategory()->AddGroup(
+			LOCTEXT("WorkspaceMenu_GlobalFindResultsCategory", "Find in Blueprints"),
+			LOCTEXT("GlobalFindResultsMenuTooltipText", "Find references to functions, events and variables in all Blueprints."),
+			GlobalFindResultsIcon,
+			true);
+
+		for (int32 TabIdx = 0; TabIdx < ARRAY_COUNT(GlobalFindResultsTabIDs); TabIdx++)
+		{
+			const FName TabID = GlobalFindResultsTabIDs[TabIdx];
+			if (!GlobalTabManager->CanSpawnTab(TabID))
+			{
+				const FText DisplayName = FText::Format(LOCTEXT("GlobalFindResultsDisplayName", "Find in Blueprints {0}"), FText::AsNumber(TabIdx + 1));
+
+				GlobalTabManager->RegisterNomadTabSpawner(TabID, FOnSpawnTab::CreateRaw(this, &FFindInBlueprintSearchManager::SpawnGlobalFindResultsTab, TabIdx))
+					.SetDisplayName(DisplayName)
+					.SetIcon(GlobalFindResultsIcon)
+					.SetGroup(GlobalFindResultsMenuItem.ToSharedRef());
+			}
+		}
+	}
+	else
+	{
+		// Close all Global Find Results tabs when turning the feature off, since these may not get closed along with the Blueprint Editor contexts above.
+		TSet<TSharedPtr<SFindInBlueprints>> FindResultsToClose;
+
+		for (TWeakPtr<SFindInBlueprints> FindResultsPtr : GlobalFindResults)
+		{
+			TSharedPtr<SFindInBlueprints> FindResults = FindResultsPtr.Pin();
+			if (FindResults.IsValid())
+			{
+				FindResultsToClose.Add(FindResults);
+			}
+		}
+
+		for (TSharedPtr<SFindInBlueprints> FindResults : FindResultsToClose)
+		{
+			FindResults->CloseHostTab();
+		}
+
+		GlobalFindResults.Empty();
+
+		for (int32 TabIdx = 0; TabIdx < ARRAY_COUNT(GlobalFindResultsTabIDs); TabIdx++)
+		{
+			const FName TabID = GlobalFindResultsTabIDs[TabIdx];
+			if (GlobalTabManager->CanSpawnTab(TabID))
+			{
+				GlobalTabManager->UnregisterNomadTabSpawner(TabID);
+			}
+		}
+
+		if (GlobalFindResultsMenuItem.IsValid())
+		{
+			WorkspaceMenu::GetMenuStructure().GetToolsCategory()->RemoveItem(GlobalFindResultsMenuItem.ToSharedRef());
+			GlobalFindResultsMenuItem.Reset();
+		}
+	}
+}
+
+void FFindInBlueprintSearchManager::CloseOrphanedGlobalFindResultsTabs(TSharedPtr<class FTabManager> TabManager)
+{
+	if (TabManager.IsValid())
+	{
+		for (int32 TabIdx = 0; TabIdx < ARRAY_COUNT(GlobalFindResultsTabIDs); TabIdx++)
+		{
+			const FName TabID = GlobalFindResultsTabIDs[TabIdx];
+			if (!FGlobalTabmanager::Get()->CanSpawnTab(TabID))
+			{
+				TSharedPtr<SDockTab> OrphanedTab = TabManager->FindExistingLiveTab(FTabId(TabID));
+				if (OrphanedTab.IsValid())
+				{
+					OrphanedTab->RequestCloseTab();
+				}
+			}
+		}
+	}
+}
+
 #undef LOCTEXT_NAMESPACE
diff --git a/Engine/Source/Editor/Kismet/Private/FindInBlueprints.cpp b/Engine/Source/Editor/Kismet/Private/FindInBlueprints.cpp
--- a/Engine/Source/Editor/Kismet/Private/FindInBlueprints.cpp
+++ b/Engine/Source/Editor/Kismet/Private/FindInBlueprints.cpp
@@ -30,7 +30,10 @@
 #include "ImaginaryBlueprintData.h"
 #include "FiBSearchInstance.h"
 #include "BlueprintEditorTabs.h"
+#include "BlueprintEditorSettings.h"
 #include "HAL/PlatformApplicationMisc.h"
+#include "Framework/Commands/UICommandList.h"
+#include "Widgets/Docking/SDockTab.h"
 
 #define LOCTEXT_NAMESPACE "FindInBlueprints"
 
@@ -107,6 +110,17 @@ bool FindInBlueprintsHelpers::ParsePinType(FText InKey, FText InValue, FEdGraphP
 	return bParsed;
 }
 
+void FindInBlueprintsHelpers::ExpandAllChildren(FSearchResult InTreeNode, TSharedPtr<STreeView<TSharedPtr<FFindInBlueprintsResult>>> InTreeView)
+{
+	if (InTreeNode->Children.Num())
+	{
+		InTreeView->SetItemExpansion(InTreeNode, true);
+		for (int32 i = 0; i < InTreeNode->Children.Num(); i++)
+		{
+			ExpandAllChildren(InTreeNode->Children[i], InTreeView);
+		}
+	}
+}
 
 //////////////////////////////////////////////////////////////////////////
 // FBlueprintSearchResult
@@ -196,18 +210,6 @@ UBlueprint* FFindInBlueprintsResult::GetParentBlueprint() const
 	return ResultBlueprint;
 }
 
-void FFindInBlueprintsResult::ExpandAllChildren( TSharedPtr< STreeView< TSharedPtr< FFindInBlueprintsResult > > > InTreeView )
-{
-	if( Children.Num() )
-	{
-		InTreeView->SetItemExpansion(this->AsShared(), true);
-		for( int32 i = 0; i < Children.Num(); i++ )
-		{
-			Children[i]->ExpandAllChildren(InTreeView);
-		}
-	}
-}
-
 FText FFindInBlueprintsResult::GetDisplayString() const
 {
 	return DisplayText;
@@ -590,12 +592,23 @@ void SFindInBlueprints::Construct( const FArguments& InArgs, TSharedPtr<FBluepri
 	LastSearchedFiBVersion = EFiBVersion::FIB_VER_LATEST;
 	BlueprintEditorPtr = InBlueprintEditor;
 
+	HostTab = InArgs._ContainingTab;
+	bIsLocked = false;
+
+	if (HostTab.IsValid())
+	{
+		HostTab.Pin()->SetOnTabClosed(SDockTab::FOnTabClosedCallback::CreateSP(this, &SFindInBlueprints::OnHostTabClosed));
+	}
+
 	if (InArgs._bIsSearchWindow)
 	{
 		RegisterCommands();
 	}
 
-	bIsInFindWithinBlueprintMode = true;
+	bIsInFindWithinBlueprintMode = BlueprintEditorPtr.IsValid();
+	bHasGlobalSearchResults = !bIsInFindWithinBlueprintMode;
+
+	const bool bHostFindInBlueprintsInGlobalTab = GetDefault<UBlueprintEditorSettings>()->bHostFindInBlueprintsInGlobalTab;
 	
 	this->ChildSlot
 		[
@@ -614,18 +627,47 @@ void SFindInBlueprints::Construct( const FArguments& InArgs, TSharedPtr<FBluepri
 					.Visibility(InArgs._bHideSearchBar? EVisibility::Collapsed : EVisibility::Visible)
 				]
 				+SHorizontalBox::Slot()
+				.Padding(4.f, 0.f, 2.f, 0.f)
+				.AutoWidth()
+				[
+					SNew(SButton)
+					.OnClicked(this, &SFindInBlueprints::OnOpenGlobalFindResults)
+					.Visibility(BlueprintEditorPtr.IsValid() && bHostFindInBlueprintsInGlobalTab ? EVisibility::Visible : EVisibility::Collapsed)
+					.ToolTipText(LOCTEXT("OpenInGlobalFindResultsButtonTooltip", "Find in all Blueprints"))
+					[
+						SNew(STextBlock)
+						.TextStyle(FEditorStyle::Get(), "FindResults.FindInBlueprints")
+						.Text(FText::FromString(FString(TEXT("\xf1e5"))) /*fa-binoculars*/)
+					]
+				]
+				+SHorizontalBox::Slot()
 				.Padding(2.f, 0.f)
 				.AutoWidth()
 				[
 					SNew(SCheckBox)
 					.OnCheckStateChanged(this, &SFindInBlueprints::OnFindModeChanged)
 					.IsChecked(this, &SFindInBlueprints::OnGetFindModeChecked)
-					.Visibility(InArgs._bHideSearchBar? EVisibility::Collapsed : EVisibility::Visible)
+					.Visibility(InArgs._bHideSearchBar || bHostFindInBlueprintsInGlobalTab ? EVisibility::Collapsed : EVisibility::Visible)
 					[
 						SNew(STextBlock)
 						.Text(LOCTEXT("BlueprintSearchModeChange", "Find In Current Blueprint Only"))
 					]
 				]
+				+SHorizontalBox::Slot()
+				.AutoWidth()
+				.VAlign(VAlign_Center)
+				[
+					SNew(SButton)
+					.VAlign(EVerticalAlignment::VAlign_Center)
+					.ButtonStyle(FEditorStyle::Get(), "FlatButton")
+					.ContentPadding(FMargin(1, 0))
+					.OnClicked(this, &SFindInBlueprints::OnLockButtonClicked)
+					.Visibility(!InArgs._bHideSearchBar && !BlueprintEditorPtr.IsValid() ? EVisibility::Visible : EVisibility::Collapsed)
+					[
+						SNew(SImage)
+						.Image(this, &SFindInBlueprints::OnGetLockButtonImage)
+					]
+				]
 			]
 			+SVerticalBox::Slot()
 			.FillHeight(1.0f)
@@ -893,7 +935,7 @@ EActiveTimerReturnType SFindInBlueprints::UpdateSearchResults( double InCurrentT
 		{
 			for ( auto& Item : BackgroundItemsFound )
 			{
-				Item->ExpandAllChildren( TreeView );
+				FindInBlueprintsHelpers::ExpandAllChildren(Item, TreeView);
 				ItemsFound.Add( Item );
 			}
 			TreeView->RequestTreeRefresh();
@@ -933,12 +975,12 @@ EActiveTimerReturnType SFindInBlueprints::UpdateSearchResults( double InCurrentT
 
 void SFindInBlueprints::RegisterCommands()
 {
-	TSharedPtr<FUICommandList> ToolKitCommandList = BlueprintEditorPtr.Pin()->GetToolkitCommands();
+	CommandList = BlueprintEditorPtr.IsValid() ? BlueprintEditorPtr.Pin()->GetToolkitCommands() : MakeShareable(new FUICommandList());
 
-	ToolKitCommandList->MapAction( FGenericCommands::Get().Copy,
+	CommandList->MapAction( FGenericCommands::Get().Copy,
 		FExecuteAction::CreateSP(this, &SFindInBlueprints::OnCopyAction) );
 
-	ToolKitCommandList->MapAction( FGenericCommands::Get().SelectAll,
+	CommandList->MapAction( FGenericCommands::Get().SelectAll,
 		FExecuteAction::CreateSP(this, &SFindInBlueprints::OnSelectAllAction) );
 }
 
@@ -1038,7 +1080,7 @@ void SFindInBlueprints::MakeSearchQuery(FString InSearchString, bool bInIsFindWi
 			{
 				for(auto Item : ItemsFound)
 				{
-					Item->ExpandAllChildren(TreeView);
+					FindInBlueprintsHelpers::ExpandAllChildren(Item, TreeView);
 				}
 			}
 
@@ -1048,6 +1090,8 @@ void SFindInBlueprints::MakeSearchQuery(FString InSearchString, bool bInIsFindWi
 		{
 			LaunchStreamThread(InSearchString, InSearchFilterForImaginaryDataReturn, InMinimiumVersionRequirement, InOnSearchComplete);
 		}
+
+		bHasGlobalSearchResults = !bInIsFindWithinBlueprint;
 	}
 }
 
@@ -1346,7 +1390,7 @@ void SFindInBlueprints::OnCacheComplete()
 TSharedPtr<SWidget> SFindInBlueprints::OnContextMenuOpening()
 {
 	const bool bShouldCloseWindowAfterMenuSelection = true;
-	FMenuBuilder MenuBuilder( bShouldCloseWindowAfterMenuSelection,BlueprintEditorPtr.Pin()->GetToolkitCommands());
+	FMenuBuilder MenuBuilder( bShouldCloseWindowAfterMenuSelection, CommandList);
 
 	MenuBuilder.BeginSection("BasicOperations");
 	{
@@ -1408,4 +1452,58 @@ void SFindInBlueprints::OnCopyAction()
 	FPlatformApplicationMisc::ClipboardCopy( *SelectedText );
 }
 
+FReply SFindInBlueprints::OnOpenGlobalFindResults()
+{
+	TSharedPtr<SFindInBlueprints> GlobalFindResults = FFindInBlueprintSearchManager::Get().GetGlobalFindResults();
+	if (GlobalFindResults.IsValid())
+	{
+		GlobalFindResults->FocusForUse(false, SearchValue, true);
+	}
+
+	return FReply::Handled();
+}
+
+void SFindInBlueprints::OnHostTabClosed(TSharedRef<SDockTab> DockTab)
+{
+	FFindInBlueprintSearchManager::Get().GlobalFindResultsClosed(SharedThis(this));
+}
+
+FReply SFindInBlueprints::OnLockButtonClicked()
+{
+	bIsLocked = !bIsLocked;
+	return FReply::Handled();
+}
+
+const FSlateBrush* SFindInBlueprints::OnGetLockButtonImage() const
+{
+	if (bIsLocked)
+	{
+		return FEditorStyle::GetBrush("FindResults.LockButton_Locked");
+	}
+	else
+	{
+		return FEditorStyle::GetBrush("FindResults.LockButton_Unlocked");
+	}
+}
+
+FName SFindInBlueprints::GetHostTabId() const
+{
+	TSharedPtr<SDockTab> HostTabPtr = HostTab.Pin();
+	if (HostTabPtr.IsValid())
+	{
+		return HostTabPtr->GetLayoutIdentifier().TabType;
+	}
+
+	return NAME_None;
+}
+
+void SFindInBlueprints::CloseHostTab()
+{
+	TSharedPtr<SDockTab> HostTabPtr = HostTab.Pin();
+	if (HostTabPtr.IsValid())
+	{
+		HostTabPtr->RequestCloseTab();
+	}
+}
+
 #undef LOCTEXT_NAMESPACE
diff --git a/Engine/Source/Editor/Kismet/Private/SBlueprintContextTargetMenu.cpp b/Engine/Source/Editor/Kismet/Private/SBlueprintContextTargetMenu.cpp
--- a/Engine/Source/Editor/Kismet/Private/SBlueprintContextTargetMenu.cpp
+++ b/Engine/Source/Editor/Kismet/Private/SBlueprintContextTargetMenu.cpp
@@ -1,4 +1,4 @@
-// Copyright 1998-2017 Epic Games, Inc. All Rights Reserved.
+﻿// Copyright 1998-2017 Epic Games, Inc. All Rights Reserved.
 
 #include "SBlueprintContextTargetMenu.h"
 #include "Engine/Blueprint.h"
@@ -567,7 +567,7 @@ void SBlueprintContextTargetMenu::Construct(const FArguments& InArgs, const FBlu
 				{
 					if (!ProfileRef.IsTargetEnabled(ContextTarget))
 					{
-						return LOCTEXT("DisabledTargetTooltip", "This target is invalid or redundent for this context.");
+						return LOCTEXT("DisabledTargetTooltip", "This target is invalid or redundant for this context.");
 					}
 					else if (TargetEnum != nullptr)
 					{
diff --git a/Engine/Source/Editor/Kismet/Private/SBlueprintDiff.cpp b/Engine/Source/Editor/Kismet/Private/SBlueprintDiff.cpp
--- a/Engine/Source/Editor/Kismet/Private/SBlueprintDiff.cpp
+++ b/Engine/Source/Editor/Kismet/Private/SBlueprintDiff.cpp
@@ -806,6 +806,10 @@ void SBlueprintDiff::Construct( const FArguments& InArgs)
 		];
 
 	SetCurrentMode( FBlueprintEditorApplicationModes::StandardBlueprintEditorMode );
+
+	// Bind to blueprint changed events as they may be real in memory blueprints that will be modified
+	const_cast<UBlueprint*>(PanelNew.Blueprint)->OnChanged().AddSP(this, &SBlueprintDiff::OnBlueprintChanged);
+	const_cast<UBlueprint*>(PanelOld.Blueprint)->OnChanged().AddSP(this, &SBlueprintDiff::OnBlueprintChanged);
 }
 END_SLATE_FUNCTION_BUILD_OPTIMIZATION
 
@@ -843,6 +847,18 @@ void SBlueprintDiff::OnGraphChanged(FListItemGraphToDiff* Diff)
 	}
 }
 
+void SBlueprintDiff::OnBlueprintChanged(UBlueprint* InBlueprint)
+{
+	if (InBlueprint == PanelOld.Blueprint && PanelOld.GraphEditor.IsValid())
+	{
+		PanelOld.GraphEditor.Pin()->NotifyGraphChanged();
+	}
+	if (InBlueprint == PanelNew.Blueprint && PanelNew.GraphEditor.IsValid())
+	{
+		PanelNew.GraphEditor.Pin()->NotifyGraphChanged();
+	}
+}
+
 TSharedRef<SWidget> SBlueprintDiff::DefaultEmptyPanel()
 {
 	return SNew(SHorizontalBox)
diff --git a/Engine/Source/Editor/Kismet/Private/SSCSEditor.cpp b/Engine/Source/Editor/Kismet/Private/SSCSEditor.cpp
--- a/Engine/Source/Editor/Kismet/Private/SSCSEditor.cpp
+++ b/Engine/Source/Editor/Kismet/Private/SSCSEditor.cpp
@@ -1,4 +1,4 @@
-// Copyright 1998-2017 Epic Games, Inc. All Rights Reserved.
+﻿// Copyright 1998-2017 Epic Games, Inc. All Rights Reserved.
 
 
 #include "SSCSEditor.h"
@@ -318,12 +318,24 @@ void FSCSRowDragDropOp::HoverTargetChanged()
 		}
 	}
 
-	UProperty* VariableProperty = GetVariableProperty();
-	if (!bHoverHandled && (VariableProperty == nullptr))
+	if (!bHoverHandled)
 	{
-		FText Message = LOCTEXT("CannotFindProperty", "Cannot find corresponding variable (make sure component has been assigned to one)");
-		SetSimpleFeedbackMessage(ErrorSymbol, IconTint, Message);
+		if (UProperty* VariableProperty = GetVariableProperty())
+		{
+			const FSlateBrush* PrimarySymbol;
+			const FSlateBrush* SecondarySymbol;
+			FSlateColor PrimaryColor;
+			FSlateColor SecondaryColor;
+			GetDefaultStatusSymbol(/*out*/ PrimarySymbol, /*out*/ PrimaryColor, /*out*/ SecondarySymbol, /*out*/ SecondaryColor);
 
+			//Create feedback message with the function name.
+			SetSimpleFeedbackMessage(PrimarySymbol, PrimaryColor, VariableProperty->GetDisplayNameText(), SecondarySymbol, SecondaryColor);
+		}
+		else
+		{
+			FText Message = LOCTEXT("CannotFindProperty", "Cannot find corresponding variable (make sure component has been assigned to one)");
+			SetSimpleFeedbackMessage(ErrorSymbol, IconTint, Message);
+		}
 		bHoverHandled = true;
 	}
 
diff --git a/Engine/Source/Editor/Kismet/Public/FindInBlueprintManager.h b/Engine/Source/Editor/Kismet/Public/FindInBlueprintManager.h
--- a/Engine/Source/Editor/Kismet/Public/FindInBlueprintManager.h
+++ b/Engine/Source/Editor/Kismet/Public/FindInBlueprintManager.h
@@ -8,13 +8,18 @@
 #include "Types/WidgetActiveTimerDelegate.h"
 #include "Dom/JsonObject.h"
 #include "HAL/Runnable.h"
+#include "SlateFwd.h"
+#include "Input/Reply.h"
 
 struct FAssetData;
 class FFindInBlueprintsResult;
 class FImaginaryBlueprint;
 class FImaginaryFiBData;
+class FSpawnTabArgs;
 class SFindInBlueprints;
 
+#define MAX_GLOBAL_FIND_RESULTS 4
+
 /**
  *Const values for Find-in-Blueprints to tag searchable data
  */
@@ -241,6 +246,80 @@ struct KISMET_API FFiBMD
 	static const FString FiBSearchableHiddenExplicitMD;
 };
 
+////////////////////////////////////
+// FFindInBlueprintsResult
+
+/* Item that matched the search results */
+class FFindInBlueprintsResult : public TSharedFromThis< FFindInBlueprintsResult >
+{
+public:
+	/* Create a root */
+	FFindInBlueprintsResult(const FText& InDisplayText);
+
+	/* Create a listing for a search result*/
+	FFindInBlueprintsResult(const FText& InDisplayText, TSharedPtr<FFindInBlueprintsResult> InParent);
+
+	virtual ~FFindInBlueprintsResult() {}
+
+	/* Called when user clicks on the search item */
+	virtual FReply OnClick();
+
+	/* Get Category for this search result */
+	virtual FText GetCategory() const;
+
+	/* Create an icon to represent the result */
+	virtual TSharedRef<SWidget>	CreateIcon() const;
+
+	/** Finalizes any content for the search data that was unsafe to do on a separate thread */
+	virtual void FinalizeSearchData() {};
+
+	/* Gets the comment on this node if any */
+	FString GetCommentText() const;
+
+	/** gets the blueprint housing all these search results */
+	UBlueprint* GetParentBlueprint() const;
+
+	/**
+	* Parses search info for specific data important for displaying the search result in an easy to understand format
+	*
+	* @param	InTokens		The search tokens to check results against
+	* @param	InKey			This is the tag for the data, describing what it is so special handling can occur if needed
+	* @param	InValue			Compared against search query to see if it passes the filter, sometimes data is rejected because it is deemed unsearchable
+	* @param	InParent		The parent search result
+	*/
+	virtual void ParseSearchInfo(FText InKey, FText InValue) {};
+
+	/** Returns the Object represented by this search information give the Blueprint it can be found in */
+	virtual UObject* GetObject(UBlueprint* InBlueprint) const;
+
+	/**
+	* Adds extra search info, anything that does not have a predestined place in the search result. Adds a sub-item to the searches and formats its description so the tag displays
+	*
+	* @param	InKey			This is the tag for the data, describing what it is so special handling can occur if needed
+	* @param	InValue			Compared against search query to see if it passes the filter, sometimes data is rejected because it is deemed unsearchable
+	* @param	InParent		The parent search result
+	*/
+	void AddExtraSearchInfo(FText InKey, FText InValue, TSharedPtr< FFindInBlueprintsResult > InParent);
+
+	/** Returns the display string for the row */
+	FText GetDisplayString() const;
+
+public:
+	/*Any children listed under this category */
+	TArray< TSharedPtr<FFindInBlueprintsResult> > Children;
+
+	/*If it exists it is the blueprint*/
+	TWeakPtr<FFindInBlueprintsResult> Parent;
+
+	/*The display text for this item */
+	FText DisplayText;
+
+	/** Display text for comment information */
+	FString CommentText;
+};
+
+typedef TSharedPtr<FFindInBlueprintsResult> FSearchResult;
+
 ////////////////////////////////////
 // FStreamSearch
 
@@ -450,6 +529,17 @@ public:
 
 	bool IsGatheringDataEnabled() const { return bEnableGatheringData; }
 
+	/** Find or create the global find results widget */
+	TSharedPtr<SFindInBlueprints> GetGlobalFindResults();
+
+	/** Enable or disable the global find results tab feature */
+	void EnableGlobalFindResults(bool bEnable);
+
+	/** Close any orphaned global find results tabs for a particular tab manager */
+	void CloseOrphanedGlobalFindResultsTabs(TSharedPtr<class FTabManager> TabManager);
+
+	void GlobalFindResultsClosed(const TSharedRef<SFindInBlueprints>& FindResults);
+
 private:
 	/** Initializes the FiB manager */
 	void Initialize();
@@ -501,6 +591,15 @@ private:
 	/** Begins the process of extracting unloaded FiB data */
 	void ExtractUnloadedFiBData(const FAssetData& InAssetData, const FString& InFiBData, bool bIsVersioned);
 
+	/** Determines the global find results tab label */
+	FText GetGlobalFindResultsTabLabel(int32 TabIdx);
+
+	/** Handler for a request to spawn a new global find results tab */
+	TSharedRef<SDockTab> SpawnGlobalFindResultsTab(const FSpawnTabArgs& SpawnTabArgs, int32 TabIdx);
+
+	/** Creates and opens a new global find results tab */
+	TSharedPtr<SFindInBlueprints> OpenGlobalFindResultsTab();
+
 protected:
 	/** Tells if gathering data is currently allowed */
 	bool bEnableGatheringData;
@@ -546,6 +645,15 @@ protected:
 
 	/** Mapping between a class name and its UClass instance - used for faster look up in FFindInBlueprintSearchManager::OnAssetAdded */
 	TMap<FName, const UClass*> CachedAssetClasses;
+
+	/** The tab identifier/instance name for global find results */
+	FName GlobalFindResultsTabIDs[MAX_GLOBAL_FIND_RESULTS];
+
+	/** Array of open global find results widgets */
+	TArray<TWeakPtr<SFindInBlueprints>> GlobalFindResults;
+
+	/** Global Find Results workspace menu item */
+	TSharedPtr<class FWorkspaceItem> GlobalFindResultsMenuItem;
 };
 
 struct KISMET_API FDisableGatheringDataOnScope
diff --git a/Engine/Source/Editor/Kismet/Public/FindInBlueprints.h b/Engine/Source/Editor/Kismet/Public/FindInBlueprints.h
--- a/Engine/Source/Editor/Kismet/Public/FindInBlueprints.h
+++ b/Engine/Source/Editor/Kismet/Public/FindInBlueprints.h
@@ -20,10 +20,9 @@
 #include "FindInBlueprintManager.h"
 
 class FBlueprintEditor;
-class FFindInBlueprintsResult;
 class FImaginaryFiBData;
+class FUICommandList;
 
-typedef TSharedPtr<class FFindInBlueprintsResult> FSearchResult;
 typedef STreeView<FSearchResult>  STreeViewType;
 
 DECLARE_DELEGATE_OneParam(FOnSearchComplete, TArray<TSharedPtr<class FImaginaryFiBData>>&);
@@ -79,81 +78,12 @@ namespace FindInBlueprintsHelpers
 	 * @return						TRUE when the parsing is successful
 	 */
 	bool ParsePinType(FText InKey, FText InValue, FEdGraphPinType& InOutPinType);
-}
-
-/* Item that matched the search results */
-class FFindInBlueprintsResult : public TSharedFromThis< FFindInBlueprintsResult >
-{
-public: 
-	/* Create a root */
-	FFindInBlueprintsResult(const FText& InDisplayText);
-	
-	/* Create a listing for a search result*/
-	FFindInBlueprintsResult(const FText& InDisplayText, TSharedPtr<FFindInBlueprintsResult> InParent);
-
-	virtual ~FFindInBlueprintsResult() {}
-	
-	/* Called when user clicks on the search item */
-	virtual FReply OnClick();
-
-	/* Get Category for this search result */
-	virtual FText GetCategory() const;
-
-	/* Create an icon to represent the result */
-	virtual TSharedRef<SWidget>	CreateIcon() const;
-
-	/** Finalizes any content for the search data that was unsafe to do on a separate thread */
-	virtual void FinalizeSearchData() {};
-
-	/* Gets the comment on this node if any */
-	FString GetCommentText() const;
-
-	/** gets the blueprint housing all these search results */
-	UBlueprint* GetParentBlueprint() const;
 
 	/**
-	 * Parses search info for specific data important for displaying the search result in an easy to understand format
-	 *
-	 * @param	InTokens		The search tokens to check results against
-	 * @param	InKey			This is the tag for the data, describing what it is so special handling can occur if needed
-	 * @param	InValue			Compared against search query to see if it passes the filter, sometimes data is rejected because it is deemed unsearchable
-	 * @param	InParent		The parent search result
-	 */
-	virtual void ParseSearchInfo(FText InKey, FText InValue) {};
-
-	/** Returns the Object represented by this search information give the Blueprint it can be found in */
-	virtual UObject* GetObject(UBlueprint* InBlueprint) const;
-
-	/**
-	 * Adds extra search info, anything that does not have a predestined place in the search result. Adds a sub-item to the searches and formats its description so the tag displays
-	 *
-	 * @param	InKey			This is the tag for the data, describing what it is so special handling can occur if needed
-	 * @param	InValue			Compared against search query to see if it passes the filter, sometimes data is rejected because it is deemed unsearchable
-	 * @param	InParent		The parent search result
-	 */
-	void AddExtraSearchInfo(FText InKey, FText InValue, TSharedPtr< FFindInBlueprintsResult > InParent);
-
-	/**
-	 * Iterates through all this node's children, and tells the tree view to expand them
-	 */
-	void ExpandAllChildren( TSharedPtr< STreeView< TSharedPtr< FFindInBlueprintsResult > > > InTreeView );
-
-	/** Returns the display string for the row */
-	FText GetDisplayString() const;
-
-public:
-	/*Any children listed under this category */
-	TArray< TSharedPtr<FFindInBlueprintsResult> > Children;
-
-	/*If it exists it is the blueprint*/
-	TWeakPtr<FFindInBlueprintsResult> Parent;
-
-	/*The display text for this item */
-	FText DisplayText;
-
-	/** Display text for comment information */
-	FString CommentText;
-};
+	* Iterates through all the given tree node's children and tells the tree view to expand them
+	*/
+	void ExpandAllChildren(FSearchResult InTreeNode, TSharedPtr< STreeView< TSharedPtr< FFindInBlueprintsResult > > > InTreeView);
+}
 
 /** Graph nodes use this class to store their data */
 class FFindInBlueprintsGraphNode : public FFindInBlueprintsResult
@@ -263,12 +193,14 @@ public:
 	SLATE_BEGIN_ARGS( SFindInBlueprints )
 		: _bIsSearchWindow(true)
 		, _bHideSearchBar(false)
+		, _ContainingTab()
 	{}
 		SLATE_ARGUMENT(bool, bIsSearchWindow)
 		SLATE_ARGUMENT(bool, bHideSearchBar)
+		SLATE_ARGUMENT(TSharedPtr<SDockTab>, ContainingTab)
 	SLATE_END_ARGS()
 
-	void Construct(const FArguments& InArgs, TSharedPtr<class FBlueprintEditor> InBlueprintEditor);
+	void Construct(const FArguments& InArgs, TSharedPtr<class FBlueprintEditor> InBlueprintEditor = nullptr);
 	~SFindInBlueprints();
 
 	/** Focuses this widget's search box, and changes the mode as well, and optionally the search terms */
@@ -296,6 +228,18 @@ public:
 	 */
 	void CacheAllBlueprints(FSimpleDelegate InOnFinished = FSimpleDelegate(), EFiBVersion InMinimiumVersionRequirement = EFiBVersion::FIB_VER_LATEST);
 
+	/** If this is a global find results widget, returns the host tab's unique ID. Otherwise, returns NAME_None. */
+	FName GetHostTabId() const;
+
+	/** If this is a global find results widget, ask the host tab to close */
+	void CloseHostTab();
+
+	/** Determines if this context does not accept syncing from an external source */
+	bool IsLocked() const
+	{
+		return bIsLocked;
+	}
+
 private:
 	/** Processes results of the ongoing async stream search */
 	EActiveTimerReturnType UpdateSearchResults( double InCurrentTime, float InDeltaTime );
@@ -392,6 +336,18 @@ private:
 	/** Callback when user attempts to copy their selection in the Find-in-Blueprints */
 	void OnCopyAction();
 
+	/** Called when the user clicks the global find results button */
+	FReply OnOpenGlobalFindResults();
+
+	/** Called when the host tab is closed (if valid) */
+	void OnHostTabClosed(TSharedRef<SDockTab> DockTab);
+
+	/** Called when the lock button is clicked in a global find results tab */
+	FReply OnLockButtonClicked();
+
+	/** Returns the image used for the lock button in a global find results tab */
+	const FSlateBrush* OnGetLockButtonImage() const;
+
 private:
 	/** Pointer back to the blueprint editor that owns us */
 	TWeakPtr<class FBlueprintEditor> BlueprintEditorPtr;
@@ -434,4 +390,16 @@ private:
 
 	/** Cached version that was last searched */
 	EFiBVersion LastSearchedFiBVersion;
+
+	/** Commands handled by this widget */
+	TSharedPtr< FUICommandList > CommandList;
+
+	/** Tab hosting this widget. May be invalid. */
+	TWeakPtr<SDockTab> HostTab;
+
+	/** True if current search should not be changed by an external source */
+	bool bIsLocked;
+
+	/** True if the most recent search was a global search */
+	bool bHasGlobalSearchResults;
 };
diff --git a/Engine/Source/Editor/Kismet/Public/SBlueprintDiff.h b/Engine/Source/Editor/Kismet/Public/SBlueprintDiff.h
--- a/Engine/Source/Editor/Kismet/Public/SBlueprintDiff.h
+++ b/Engine/Source/Editor/Kismet/Public/SBlueprintDiff.h
@@ -131,6 +131,9 @@ public:
 	/** Called when a new Graph is clicked on by user */
 	void OnGraphChanged(struct FListItemGraphToDiff* Diff);
 
+	/** Called when blueprint is modified */
+	void OnBlueprintChanged(UBlueprint* InBlueprint);
+
 	/** Helper function for generating an empty widget */
 	static TSharedRef<SWidget> DefaultEmptyPanel();
 
diff --git a/Engine/Source/Editor/KismetCompiler/Private/AnimBlueprintCompiler.cpp b/Engine/Source/Editor/KismetCompiler/Private/AnimBlueprintCompiler.cpp
--- a/Engine/Source/Editor/KismetCompiler/Private/AnimBlueprintCompiler.cpp
+++ b/Engine/Source/Editor/KismetCompiler/Private/AnimBlueprintCompiler.cpp
@@ -510,7 +510,7 @@ void FAnimBlueprintCompiler::ProcessAnimationNode(UAnimGraphNode_Base* VisualAni
 				else
 				{
 					// Dynamic value that needs to be wired up and evaluated each frame
-					FString EvaluationHandlerStr = SourcePinProperty->GetMetaData(AnimGraphDefaultSchema->NAME_OnEvaluate);
+					const FString& EvaluationHandlerStr = SourcePinProperty->GetMetaData(AnimGraphDefaultSchema->NAME_OnEvaluate);
 					FName EvaluationHandlerName(*EvaluationHandlerStr);
 					if (EvaluationHandlerName == NAME_None)
 					{
diff --git a/Engine/Source/Editor/KismetCompiler/Private/KismetCompiler.cpp b/Engine/Source/Editor/KismetCompiler/Private/KismetCompiler.cpp
--- a/Engine/Source/Editor/KismetCompiler/Private/KismetCompiler.cpp
+++ b/Engine/Source/Editor/KismetCompiler/Private/KismetCompiler.cpp
@@ -392,14 +392,7 @@ void FKismetCompilerContext::ValidateLink(const UEdGraphPin* PinA, const UEdGrap
 		if (bForbiddenConnection || bMissingConversion)
 		{
 			const FString ErrorMessage = FString::Printf(*LOCTEXT("PinTypeMismatch_Error", "Can't connect pins @@ and @@: %s").ToString(), *ConnectResponse.Message.ToString());
-			if (ConnectResponse.IsFatal())
-			{
-				MessageLog.Error(*ErrorMessage, PinA, PinB);
-			}
-			else
-			{
-				MessageLog.Warning(*ErrorMessage, PinA, PinB);
-			}
+			MessageLog.Error(*ErrorMessage, PinA, PinB);
 		}
 	}
 
@@ -1587,6 +1580,11 @@ void FKismetCompilerContext::PrecompileFunction(FKismetFunctionContext& Context,
 		if (Context.EntryPoint)
 		{
 			Context.Function->FunctionFlags |= (EFunctionFlags)Context.EntryPoint->GetExtraFlags();
+			
+			if (UEdGraphPin* WorldContextPin = Context.EntryPoint->GetAutoWorldContextPin())
+			{
+				Context.Function->SetMetaData(FBlueprintMetadata::MD_WorldContext, *WorldContextPin->PinName);
+			}
 		}
 
 		// First try to get the overriden function from the super class
diff --git a/Engine/Source/Editor/KismetCompiler/Private/KismetCompilerMisc.cpp b/Engine/Source/Editor/KismetCompiler/Private/KismetCompilerMisc.cpp
--- a/Engine/Source/Editor/KismetCompiler/Private/KismetCompilerMisc.cpp
+++ b/Engine/Source/Editor/KismetCompiler/Private/KismetCompilerMisc.cpp
@@ -285,7 +285,7 @@ bool FKismetCompilerUtilities::IsTypeCompatibleWithProperty(UEdGraphPin* SourceP
 			if(OwningFunction)
 			{
 				// Check for the magic ArrayParm property, which always matches array types
-				FString ArrayPointerMetaData = OwningFunction->GetMetaData(FBlueprintMetadata::MD_ArrayParam);
+				const FString& ArrayPointerMetaData = OwningFunction->GetMetaData(FBlueprintMetadata::MD_ArrayParam);
 				TArray<FString> ArrayPinComboNames;
 				ArrayPointerMetaData.ParseIntoArray(ArrayPinComboNames, TEXT(","), true);
 
diff --git a/Engine/Source/Editor/KismetCompiler/Private/UserDefinedStructureCompilerUtils.cpp b/Engine/Source/Editor/KismetCompiler/Private/UserDefinedStructureCompilerUtils.cpp
--- a/Engine/Source/Editor/KismetCompiler/Private/UserDefinedStructureCompilerUtils.cpp
+++ b/Engine/Source/Editor/KismetCompiler/Private/UserDefinedStructureCompilerUtils.cpp
@@ -434,8 +434,9 @@ void FUserDefinedStructureCompilerUtils::CompileStruct(class UUserDefinedStruct*
 
 			if (bReconstruct)
 			{
-				if (UBlueprint* FoundBlueprint = Node->GetBlueprint())
+				if (Node->HasValidBlueprint())
 				{
+					UBlueprint* FoundBlueprint = Node->GetBlueprint();
 					// The blueprint skeleton needs to be updated before we reconstruct the node
 					// or else we may have member references that point to the old skeleton
 					if (!BlueprintsThatHaveBeenRecompiled.Contains(FoundBlueprint))
diff --git a/Engine/Source/Editor/Matinee/Private/MatineeMenus.cpp b/Engine/Source/Editor/Matinee/Private/MatineeMenus.cpp
--- a/Engine/Source/Editor/Matinee/Private/MatineeMenus.cpp
+++ b/Engine/Source/Editor/Matinee/Private/MatineeMenus.cpp
@@ -5499,7 +5499,7 @@ TSharedPtr<SWidget> FMatinee::CreateGroupMenu()
 			//	UInterpTrack* DefTrack = TrackClass->GetDefaultObject<UInterpTrack>();
 			//	if ( !DefTrack->bDirGroupOnly && !DefTrack->bSubTrackOnly )
 			//	{
-			//		FString DisplayGroup = TrackClass->GetMetaData("DisplayGroup");
+			//		const FString& DisplayGroup = TrackClass->GetMetaData("DisplayGroup");
 			//		TrackGroups.FindOrAdd(DisplayGroup).Add(TrackClass);
 			//	}
 			//}
diff --git a/Engine/Source/Editor/PropertyEditor/Private/CategoryPropertyNode.cpp b/Engine/Source/Editor/PropertyEditor/Private/CategoryPropertyNode.cpp
--- a/Engine/Source/Editor/PropertyEditor/Private/CategoryPropertyNode.cpp
+++ b/Engine/Source/Editor/PropertyEditor/Private/CategoryPropertyNode.cpp
@@ -53,7 +53,7 @@ void FCategoryPropertyNode::InitChildNodes()
 			if (!bShowHiddenProperties)
 			{
 				static const FName Name_bShowOnlyWhenTrue("bShowOnlyWhenTrue");
-				FString MetaDataVisibilityCheckString = It->GetMetaData(Name_bShowOnlyWhenTrue);
+				const FString& MetaDataVisibilityCheckString = It->GetMetaData(Name_bShowOnlyWhenTrue);
 				if (MetaDataVisibilityCheckString.Len())
 				{
 					//ensure that the metadata visibility string is actually set to true in order to show this property
diff --git a/Engine/Source/Editor/PropertyEditor/Private/ObjectPropertyNode.cpp b/Engine/Source/Editor/PropertyEditor/Private/ObjectPropertyNode.cpp
--- a/Engine/Source/Editor/PropertyEditor/Private/ObjectPropertyNode.cpp
+++ b/Engine/Source/Editor/PropertyEditor/Private/ObjectPropertyNode.cpp
@@ -438,7 +438,7 @@ void FObjectPropertyNode::InternalInitChildNodes( FName SinglePropertyName )
 
 		bool bMetaDataAllowVisible = true;
 		static const FName Name_bShowOnlyWhenTrue("bShowOnlyWhenTrue");
-		FString MetaDataVisibilityCheckString = It->GetMetaData(Name_bShowOnlyWhenTrue);
+		const FString& MetaDataVisibilityCheckString = It->GetMetaData(Name_bShowOnlyWhenTrue);
 		if (MetaDataVisibilityCheckString.Len())
 		{
 			//ensure that the metadata visibility string is actually set to true in order to show this property
diff --git a/Engine/Source/Editor/PropertyEditor/Private/PropertyCustomizationHelpers.cpp b/Engine/Source/Editor/PropertyEditor/Private/PropertyCustomizationHelpers.cpp
--- a/Engine/Source/Editor/PropertyEditor/Private/PropertyCustomizationHelpers.cpp
+++ b/Engine/Source/Editor/PropertyEditor/Private/PropertyCustomizationHelpers.cpp
@@ -378,14 +378,14 @@ void SObjectPropertyEntryBox::Construct( const FArguments& InArgs )
 		PropertyHandle = InArgs._PropertyHandle;
 
 		// check if the property metadata wants us to display a thumbnail
-		FString DisplayThumbnailString = PropertyHandle->GetProperty()->GetMetaData(TEXT("DisplayThumbnail"));
+		const FString& DisplayThumbnailString = PropertyHandle->GetProperty()->GetMetaData(TEXT("DisplayThumbnail"));
 		if(DisplayThumbnailString.Len() > 0)
 		{
 			bDisplayThumbnail = DisplayThumbnailString == TEXT("true");
 		}
 
 		// check if the property metadata has an override to the thumbnail size
-		FString ThumbnailSizeString = PropertyHandle->GetProperty()->GetMetaData(TEXT("ThumbnailSize"));
+		const FString& ThumbnailSizeString = PropertyHandle->GetProperty()->GetMetaData(TEXT("ThumbnailSize"));
 		if ( ThumbnailSizeString.Len() > 0 )
 		{
 			FVector2D ParsedVector;
diff --git a/Engine/Source/Editor/PropertyEditor/Private/PropertyEditorHelpers.cpp b/Engine/Source/Editor/PropertyEditor/Private/PropertyEditorHelpers.cpp
--- a/Engine/Source/Editor/PropertyEditor/Private/PropertyEditorHelpers.cpp
+++ b/Engine/Source/Editor/PropertyEditor/Private/PropertyEditorHelpers.cpp
@@ -524,7 +524,7 @@ namespace PropertyEditorHelpers
 				else
 				{
 
-					FString EnumName = Property->GetMetaData(TEXT("Enum"));
+					const FString& EnumName = Property->GetMetaData(TEXT("Enum"));
 					Enum = FindObject<UEnum>(ANY_PACKAGE, *EnumName, true);
 				}
 
diff --git a/Engine/Source/Editor/PropertyEditor/Private/PropertyHandleImpl.cpp b/Engine/Source/Editor/PropertyEditor/Private/PropertyHandleImpl.cpp
--- a/Engine/Source/Editor/PropertyEditor/Private/PropertyHandleImpl.cpp
+++ b/Engine/Source/Editor/PropertyEditor/Private/PropertyHandleImpl.cpp
@@ -439,16 +439,15 @@ FPropertyAccess::Result FPropertyValueImpl::ImportText( const TArray<FObjectBase
 
 				bool bIsInContainer = false;
 
-				if (Property && Property->IsA(UArrayProperty::StaticClass()))
+				if (UArrayProperty* ArrayProperty = Cast<UArrayProperty>(Property))
 				{
-					bIsInContainer = Cast<UArrayProperty>(Property)->Inner == NodeProperty;
+					bIsInContainer = (ArrayProperty->Inner == NodeProperty);
 				}
-				else if (Property && Property->IsA(USetProperty::StaticClass()))
+				else if (USetProperty* SetProp = Cast<USetProperty>(Property))
 				{
 					// If the element is part of a set, check for duplicate elements
-					USetProperty* SetProp = Cast<USetProperty>(Property);
 					bIsInContainer = SetProp->ElementProp == NodeProperty;
-					
+
 					if (bIsInContainer)
 					{
 						FScriptSetHelper SetHelper(SetProp, ParentNode->GetValueBaseAddress((uint8*)Cur.Object));
@@ -462,9 +461,8 @@ FPropertyAccess::Result FPropertyValueImpl::ImportText( const TArray<FObjectBase
 						}
 					}
 				}
-				else if (Property && Property->IsA(UMapProperty::StaticClass()))
+				else if (UMapProperty* MapProperty = Cast<UMapProperty>(Property))
 				{
-					UMapProperty* MapProperty = Cast<UMapProperty>(Property);
 					bIsInContainer = MapProperty->KeyProp == NodeProperty;
 
 					if (bIsInContainer)
@@ -2897,7 +2895,7 @@ bool FPropertyHandleBase::GeneratePossibleValues(TArray< TSharedPtr<FString> >&
 	}
 	else if ( Property->IsA(UStrProperty::StaticClass()) && Property->HasMetaData( TEXT("Enum") ) )
 	{
-		FString EnumName = Property->GetMetaData(TEXT("Enum"));
+		const FString& EnumName = Property->GetMetaData(TEXT("Enum"));
 		Enum = FindObject<UEnum>(ANY_PACKAGE, *EnumName, true);
 		check( Enum );
 	}
@@ -3762,7 +3760,8 @@ FPropertyAccess::Result FPropertyHandleObject::SetValue(const FAssetData& NewVal
 FPropertyAccess::Result FPropertyHandleObject::SetValueFromFormattedString(const FString& InValue, EPropertyValueSetFlags::Type Flags)
 {
 	UProperty* Property = GetProperty();
-	FString AllowedClassesString = Property ? Property->GetMetaData("AllowedClasses") : TEXT("");
+	const FString EmptyString;
+	const FString& AllowedClassesString = Property ? Property->GetMetaData("AllowedClasses") : EmptyString;
 
 	if (Property && !AllowedClassesString.IsEmpty())
 	{
diff --git a/Engine/Source/Editor/PropertyEditor/Private/PropertyNode.cpp b/Engine/Source/Editor/PropertyEditor/Private/PropertyNode.cpp
--- a/Engine/Source/Editor/PropertyEditor/Private/PropertyNode.cpp
+++ b/Engine/Source/Editor/PropertyEditor/Private/PropertyNode.cpp
@@ -2441,7 +2441,7 @@ void FPropertyNode::NotifyPostChange( FPropertyChangedEvent& InPropertyChangedEv
 	if( OriginalActiveProperty )
 	{
 		//if i have metadata forcing other property windows to rebuild
-		FString MetaData = OriginalActiveProperty->GetMetaData(TEXT("ForceRebuildProperty"));
+		const FString& MetaData = OriginalActiveProperty->GetMetaData(TEXT("ForceRebuildProperty"));
 
 		if( MetaData.Len() > 0 )
 		{
diff --git a/Engine/Source/Editor/PropertyEditor/Private/UserInterface/PropertyEditor/SPropertyEditorAsset.cpp b/Engine/Source/Editor/PropertyEditor/Private/UserInterface/PropertyEditor/SPropertyEditorAsset.cpp
--- a/Engine/Source/Editor/PropertyEditor/Private/UserInterface/PropertyEditor/SPropertyEditorAsset.cpp
+++ b/Engine/Source/Editor/PropertyEditor/Private/UserInterface/PropertyEditor/SPropertyEditorAsset.cpp
@@ -58,7 +58,7 @@ bool SPropertyEditorAsset::ShouldDisplayThumbnail( const FArguments& InArgs, con
 				PropertyToCheck = MapParent;
 			}
 
-			FString DisplayThumbnailString = PropertyToCheck->GetMetaData(TEXT("DisplayThumbnail"));
+			const FString& DisplayThumbnailString = PropertyToCheck->GetMetaData(TEXT("DisplayThumbnail"));
 			if (DisplayThumbnailString.Len() > 0)
 			{
 				bDisplayThumbnail = DisplayThumbnailString == TEXT("true");
@@ -106,24 +106,24 @@ void SPropertyEditorAsset::Construct( const FArguments& InArgs, const TSharedPtr
 	// Account for the allowed classes specified in the property metadata
 	if (Property)
 	{
-		FString ClassFilterString;
+		const FString* ClassFilterString;
 		if (UArrayProperty* ArrayParent = Cast<UArrayProperty>(Property->GetOuter()))
 		{
-			ClassFilterString = ArrayParent->GetMetaData("AllowedClasses");
+			ClassFilterString = &ArrayParent->GetMetaData("AllowedClasses");
 		}
 		else
 		{
-			ClassFilterString = Property->GetMetaData("AllowedClasses");
+			ClassFilterString = &Property->GetMetaData("AllowedClasses");
 		}
 
-		if (ClassFilterString.IsEmpty())
+		if (ClassFilterString->IsEmpty())
 		{
 			CustomClassFilters.Add(ObjectClass);
 		}
 		else
 		{
 			TArray<FString> CustomClassFilterNames;
-			ClassFilterString.ParseIntoArray(CustomClassFilterNames, TEXT(","), true);
+			ClassFilterString->ParseIntoArray(CustomClassFilterNames, TEXT(","), true);
 
 			for (auto It = CustomClassFilterNames.CreateIterator(); It; ++It)
 			{
diff --git a/Engine/Source/Editor/Sequencer/Public/AnimatedPropertyKey.h b/Engine/Source/Editor/Sequencer/Public/AnimatedPropertyKey.h
--- a/Engine/Source/Editor/Sequencer/Public/AnimatedPropertyKey.h
+++ b/Engine/Source/Editor/Sequencer/Public/AnimatedPropertyKey.h
@@ -39,7 +39,10 @@ struct FAnimatedPropertyKey
 		}
 		else if (const UObjectPropertyBase* ObjectProperty = Cast<const UObjectPropertyBase>(Property))
 		{
-			Definition.ObjectTypeName = ObjectProperty->PropertyClass->GetFName();
+			if (ObjectProperty->PropertyClass)
+			{
+				Definition.ObjectTypeName = ObjectProperty->PropertyClass->GetFName();
+			}
 		}
 		else if(const UArrayProperty* ArrayProperty = Cast<const UArrayProperty>(Property))
 		{
diff --git a/Engine/Source/Editor/StatsViewer/Private/SStatsViewer.cpp b/Engine/Source/Editor/StatsViewer/Private/SStatsViewer.cpp
--- a/Engine/Source/Editor/StatsViewer/Private/SStatsViewer.cpp
+++ b/Engine/Source/Editor/StatsViewer/Private/SStatsViewer.cpp
@@ -338,11 +338,11 @@ void SStatsViewer::Tick( const FGeometry& AllottedGeometry, const double InCurre
 				{
 					TSharedPtr< FPropertyPath > PropertyPath = Column->GetDataSource()->AsPropertyPath();
 					const FPropertyInfo& PropertyInfo = PropertyPath->GetRootProperty();
-					const FString ColumnWidthString = PropertyInfo.Property->GetMetaData(StatsViewerMetadata::ColumnWidth);
+					const FString& ColumnWidthString = PropertyInfo.Property->GetMetaData(StatsViewerMetadata::ColumnWidth);
 					const float ColumnWidth = ColumnWidthString.Len() > 0 ? FCString::Atof( *ColumnWidthString ) : 100.0f;
 					Column->SetWidth( ColumnWidth );
 
-					const FString SortModeString = PropertyInfo.Property->GetMetaData(StatsViewerMetadata::SortMode);
+					const FString& SortModeString = PropertyInfo.Property->GetMetaData(StatsViewerMetadata::SortMode);
 					if( SortModeString.Len() > 0 )
 					{
 						EColumnSortMode::Type SortType = SortModeString == TEXT( "Ascending" ) ? EColumnSortMode::Ascending : EColumnSortMode::Descending;
@@ -496,7 +496,7 @@ FReply SStatsViewer::OnExportClicked()
 				{
 					TSharedPtr< FPropertyPath > PropertyPath = Column->GetDataSource()->AsPropertyPath();
 					const FPropertyInfo& PropertyInfo = PropertyPath->GetRootProperty();
-					FString ShowTotal = PropertyInfo.Property->GetMetaData(TEXT("ShowTotal"));
+					const FString& ShowTotal = PropertyInfo.Property->GetMetaData(TEXT("ShowTotal"));
 					if( ShowTotal.Len() > 0 )
 					{
 						FText* TotalText = CustomColumn->TotalsMap.Find( PropertyInfo.Property->GetNameCPP() );
diff --git a/Engine/Source/Editor/UnrealEd/Classes/CookOnTheSide/CookOnTheFlyServer.h b/Engine/Source/Editor/UnrealEd/Classes/CookOnTheSide/CookOnTheFlyServer.h
--- a/Engine/Source/Editor/UnrealEd/Classes/CookOnTheSide/CookOnTheFlyServer.h
+++ b/Engine/Source/Editor/UnrealEd/Classes/CookOnTheSide/CookOnTheFlyServer.h
@@ -970,6 +970,9 @@ private:
 	/** Map of platform name to asset registry generators, which hold the state of asset registry data for a platform */
 	TMap<FName, FAssetRegistryGenerator*> RegistryGenerators;
 
+	/** List of filenames that may be out of date in the asset registry */
+	TSet<FName> ModifiedAssetFilenames;
+
 	//////////////////////////////////////////////////////////////////////////
 	// iterative ini settings checking
 	// growing list of ini settings which are accessed over the course of the cook
diff --git a/Engine/Source/Editor/UnrealEd/Classes/Editor/TransBuffer.h b/Engine/Source/Editor/UnrealEd/Classes/Editor/TransBuffer.h
--- a/Engine/Source/Editor/UnrealEd/Classes/Editor/TransBuffer.h
+++ b/Engine/Source/Editor/UnrealEd/Classes/Editor/TransBuffer.h
@@ -19,7 +19,7 @@ public:
     GENERATED_BODY()
 	// Variables.
 	/** The queue of transaction records */
-	TArray<FTransaction> UndoBuffer;
+	TArray<TSharedRef<FTransaction>> UndoBuffer;
 
 	/** Number of transactions that have been undone, and are eligible to be redone */
 	int32 UndoCount;
@@ -74,24 +74,42 @@ protected:
 			Result = ActiveCount;
 			if (ActiveCount++ == 0)
 			{
-				// Cancel redo buffer.
-				if (UndoCount)
+				// Cache the redo buffer in case the transaction is canceled so we can restore the state
+				int32 TransactionsToRemove = UndoBuffer.Num();
+				PreviousUndoCount = UndoCount;
+				UndoCount = 0;
+
+				// Determine if any additional entries need to be removed due to memory, avoid n^2 by unrolling GetUndoSize()
+				SIZE_T AccumulatedBufferDataSize = 0;
+				for (int32 i = 0; i < UndoBuffer.Num() - PreviousUndoCount; i++)
 				{
-					UndoBuffer.RemoveAt(UndoBuffer.Num() - UndoCount, UndoCount);
+					AccumulatedBufferDataSize += UndoBuffer[i]->DataSize();
+					if (AccumulatedBufferDataSize <= MaxMemory)
+					{
+						--TransactionsToRemove;
+					}
+					else
+					{
+						break;
+					}
 				}
-				UndoCount = 0;
 
-				// Purge previous transactions if too much data occupied.
-				while (GetUndoSize() > MaxMemory)
+				if (TransactionsToRemove > 0)
 				{
-					UndoBuffer.RemoveAt(0);
+					RemovedTransactions.Reserve(TransactionsToRemove);
+					for (int32 i = UndoBuffer.Num() - TransactionsToRemove; i < UndoBuffer.Num(); ++i)
+					{
+						RemovedTransactions.Add(UndoBuffer[i]);
+					}
+					UndoBuffer.RemoveAt(UndoBuffer.Num() - TransactionsToRemove, TransactionsToRemove, false);
 				}
 
 				// Begin a new transaction.
-				GUndo = new(UndoBuffer) TTransaction(SessionContext, Description, 1);
+				UndoBuffer.Emplace(MakeShareable(new TTransaction(SessionContext, Description, 1)));
+				GUndo = &UndoBuffer.Last().Get();
 			}
 			const int32 PriorRecordsCount = (Result > 0 ? ActiveRecordCounts[Result - 1] : 0);
-			ActiveRecordCounts.Add(UndoBuffer.Last().GetRecordCount() - PriorRecordsCount);
+			ActiveRecordCounts.Add(UndoBuffer.Last()->GetRecordCount() - PriorRecordsCount);
 			CheckState();
 		}
 		return Result;
@@ -191,4 +209,11 @@ private:
 
 	// Reference to the current transaction, nullptr when not transacting:
 	FTransaction* CurrentTransaction;
+
+	// Cached previous undo count while a transaction is being built in case we cancel it and want to restore the previous undo buffer
+	int32 PreviousUndoCount;
+
+	// The list of transactions that were removed when a transaction was begun in case it is canceled and we want to restore the original state
+	TArray<TSharedRef<FTransaction>> RemovedTransactions;
+
 };
diff --git a/Engine/Source/Editor/UnrealEd/Classes/Editor/Transactor.h b/Engine/Source/Editor/UnrealEd/Classes/Editor/Transactor.h
--- a/Engine/Source/Editor/UnrealEd/Classes/Editor/Transactor.h
+++ b/Engine/Source/Editor/UnrealEd/Classes/Editor/Transactor.h
@@ -406,6 +406,12 @@ public:
 		return Ar << T.Records << T.Title << T.ObjectMap << T.Context << T.PrimaryObject;
 	}
 
+	/** Serializes a reference to a transaction in a given archive. */
+	friend FArchive& operator<<(FArchive& Ar, TSharedRef<FTransaction>& SharedT)
+	{
+		return Ar << SharedT.Get();
+	}
+
 	/** Used by GC to collect referenced objects. */
 	void AddReferencedObjects( FReferenceCollector& Collector );
 
diff --git a/Engine/Source/Editor/UnrealEd/Private/CookOnTheFlyServer.cpp b/Engine/Source/Editor/UnrealEd/Private/CookOnTheFlyServer.cpp
--- a/Engine/Source/Editor/UnrealEd/Private/CookOnTheFlyServer.cpp
+++ b/Engine/Source/Editor/UnrealEd/Private/CookOnTheFlyServer.cpp
@@ -2702,7 +2702,22 @@ void UCookOnTheFlyServer::OnObjectPropertyChanged(UObject* ObjectBeingModified,
 
 void UCookOnTheFlyServer::OnObjectSaved( UObject* ObjectSaved )
 {
-	OnObjectUpdated( ObjectSaved );
+	if (GIsCookerLoadingPackage)
+	{
+		// This is the cooker saving a cooked package, ignore
+		return;
+	}
+
+	UPackage *Package = ObjectSaved->GetOutermost();
+
+	MarkPackageDirtyForCooker(Package);
+
+	FName PackageFFileName = GetCachedStandardPackageFileFName(Package);
+
+	if (PackageFFileName != NAME_None)
+	{
+		ModifiedAssetFilenames.Add(PackageFFileName);
+	}
 }
 
 void UCookOnTheFlyServer::OnObjectUpdated( UObject *Object )
@@ -4442,15 +4457,16 @@ void UCookOnTheFlyServer::PopulateCookedPackagesFromDisk(const TArray<ITargetPla
 			const FName CookedFile = CookedPaths.Value;
 			const FName UncookedFilename = CookedPaths.Key;
 			const FName* FoundPackageName = GetCachedPackageFilenameToPackageFName(UncookedFilename);
-			const FName PackageName = FoundPackageName ? *FoundPackageName : FName();
-			bool bShouldKeep = true;
-
 			if ( !FoundPackageName )
 			{
+				// Source file no longer exists
 				++NumPackagesRemoved;
-				bShouldKeep = false;
+				continue;
 			}
-			else if (ModifiedPackages.Contains(PackageName))
+			const FName PackageName = *FoundPackageName;
+			bool bShouldKeep = true;
+
+			if (ModifiedPackages.Contains(PackageName))
 			{
 				++NumPackagesFileHashMismatch;
 				bShouldKeep = false;
@@ -4653,6 +4669,23 @@ void UCookOnTheFlyServer::GenerateAssetRegistry()
 
 	if (!!(CookFlags & ECookInitializationFlags::GeneratedAssetRegistry))
 	{
+		// Force a rescan of modified package files
+		TArray<FString> ModifiedPackageFileList;
+
+		for (FName ModifiedPackage : ModifiedAssetFilenames)
+		{
+			ModifiedPackageFileList.Add(ModifiedPackage.ToString());
+		}
+
+		AssetRegistry->ScanModifiedAssetFiles(ModifiedPackageFileList);
+
+		ModifiedAssetFilenames.Reset();
+
+		// This is cook in the editor on a second pass, so refresh the generators
+		for (TPair<FName, FAssetRegistryGenerator*>& Pair : RegistryGenerators)
+		{
+			Pair.Value->Initialize(CookByTheBookOptions ? CookByTheBookOptions->StartupPackages : TArray<FName>());
+		}
 		return;
 	}
 	CookFlags |= ECookInitializationFlags::GeneratedAssetRegistry;
@@ -5706,6 +5739,11 @@ void UCookOnTheFlyServer::InitializeSandbox()
 
 		}
 	}
+	else
+	{
+		// This is an in-editor cook, do an iterative clean
+		CleanSandbox(true);
+	}
 }
 
 void UCookOnTheFlyServer::TermSandbox()
@@ -5857,13 +5895,6 @@ void UCookOnTheFlyServer::StartCookByTheBook( const FCookByTheBookStartupOptions
 		}
 	}
 
-	// need to test this out
-	/*if (RecompileChangedShaders(TargetPlatformNames))
-	{
-		// clean everything :(
-		CleanSandbox(false);
-	}*/
-
 	CookByTheBookOptions->bLeakTest = (CookOptions & ECookByTheBookOptions::LeakTest) != ECookByTheBookOptions::None; // this won't work from the editor this needs to be standalone
 	check(!CookByTheBookOptions->bLeakTest || CurrentCookMode == ECookMode::CookByTheBook);
 
@@ -5878,7 +5909,7 @@ void UCookOnTheFlyServer::StartCookByTheBook( const FCookByTheBookStartupOptions
 
 	if (!IsChildCooker())
 	{
-		for (const auto &Platform : TargetPlatforms)
+		for (ITargetPlatform* Platform : TargetPlatforms)
 		{
 			FName PlatformName = FName(*Platform->PlatformName());
 
diff --git a/Engine/Source/Editor/UnrealEd/Private/EdMode.cpp b/Engine/Source/Editor/UnrealEd/Private/EdMode.cpp
--- a/Engine/Source/Editor/UnrealEd/Private/EdMode.cpp
+++ b/Engine/Source/Editor/UnrealEd/Private/EdMode.cpp
@@ -993,7 +993,7 @@ struct FPropertyWidgetInfoChainElement
 			check(Element.Property);
 			const FString Postfix = (Element.Index != INDEX_NONE) ? FString::Printf(TEXT("[%d]."), Element.Index) : SimplePostFix;
 			const FString PropertyName = Element.Property->GetName() + Postfix;
-			const FString DisplayName = Element.Property->GetMetaData(TEXT("DisplayName"));
+			const FString& DisplayName = Element.Property->GetMetaData(TEXT("DisplayName"));
 
 			WidgetInfo.PropertyName += PropertyName;
 			WidgetInfo.DisplayName += (!DisplayName.IsEmpty()) ? (DisplayName + Postfix) : PropertyName;
@@ -1001,7 +1001,7 @@ struct FPropertyWidgetInfoChainElement
 
 		{
 			const FString PropertyName = CurrentProp->GetName();
-			const FString DisplayName = CurrentProp->GetMetaData(TEXT("DisplayName"));
+			const FString& DisplayName = CurrentProp->GetMetaData(TEXT("DisplayName"));
 
 			WidgetInfo.PropertyName += PropertyName;
 			WidgetInfo.DisplayName += (!DisplayName.IsEmpty()) ? DisplayName : PropertyName;
diff --git a/Engine/Source/Editor/UnrealEd/Private/EditorEngine.cpp b/Engine/Source/Editor/UnrealEd/Private/EditorEngine.cpp
--- a/Engine/Source/Editor/UnrealEd/Private/EditorEngine.cpp
+++ b/Engine/Source/Editor/UnrealEd/Private/EditorEngine.cpp
@@ -4526,7 +4526,7 @@ FString UEditorEngine::GetFriendlyName( const UProperty* Property, UStruct* Owne
 
 	if ( !DidFindText )
 	{
-		FString DefaultFriendlyName = Property->GetMetaData(TEXT("DisplayName"));
+		const FString& DefaultFriendlyName = Property->GetMetaData(TEXT("DisplayName"));
 		if ( DefaultFriendlyName.IsEmpty() )
 		{
 			const bool bIsBool = Cast<const UBoolProperty>(Property) != NULL;
diff --git a/Engine/Source/Editor/UnrealEd/Private/EditorTransaction.cpp b/Engine/Source/Editor/UnrealEd/Private/EditorTransaction.cpp
--- a/Engine/Source/Editor/UnrealEd/Private/EditorTransaction.cpp
+++ b/Engine/Source/Editor/UnrealEd/Private/EditorTransaction.cpp
@@ -612,9 +612,13 @@ void UTransBuffer::AddReferencedObjects(UObject* InThis, FReferenceCollector& Co
 		// We cannot support undoing across GC if we allow it to eliminate references so we need
 		// to suppress it.
 		Collector.AllowEliminatingReferences(false);
-		for( int32 Index = 0; Index < This->UndoBuffer.Num(); Index++ )
+		for (const TSharedRef<FTransaction>& SharedTrans : This->UndoBuffer)
 		{
-			This->UndoBuffer[ Index ].AddReferencedObjects( Collector );
+			SharedTrans->AddReferencedObjects( Collector );
+		}
+		for (const TSharedRef<FTransaction>& SharedTrans : This->RemovedTransactions)
+		{
+			SharedTrans->AddReferencedObjects(Collector);
 		}
 		Collector.AllowEliminatingReferences(true);
 	}
@@ -648,7 +652,9 @@ int32 UTransBuffer::End()
 				static_cast<FTransaction*>(GUndo)->DumpObjectMap( *GLog );
 			}
 #endif
-			GUndo = NULL;
+			GUndo = nullptr;
+			PreviousUndoCount = INDEX_NONE;
+			RemovedTransactions.Reset();
 		}
 		ActiveRecordCounts.Pop();
 		CheckState();
@@ -703,10 +709,21 @@ void UTransBuffer::Cancel( int32 StartIndex /*=0*/ )
 		if ( StartIndex == 0 )
 		{
 			// clear the global pointer to the soon-to-be-deleted transaction
-			GUndo = NULL;
+			GUndo = nullptr;
 			
 			// remove the currently active transaction from the buffer
-			UndoBuffer.Pop();
+			UndoBuffer.Pop(false);
+
+			// replace the removed transactions
+			UndoBuffer.Reserve(UndoBuffer.Num() + RemovedTransactions.Num());
+			for (TSharedRef<FTransaction>& Transaction : RemovedTransactions)
+			{
+				UndoBuffer.Add(Transaction);
+			}
+			RemovedTransactions.Reset();
+
+			UndoCount = PreviousUndoCount;
+			PreviousUndoCount = INDEX_NONE;
 		}
 		else
 		{
@@ -716,7 +733,7 @@ void UTransBuffer::Cancel( int32 StartIndex /*=0*/ )
 				RecordsToKeep += ActiveRecordCounts[ActiveIndex];
 			}
 
-			FTransaction& Transaction = UndoBuffer.Last();
+			FTransaction& Transaction = UndoBuffer.Last().Get();
 			Transaction.RemoveRecords(Transaction.GetRecordCount() - RecordsToKeep);
 		}
 
@@ -793,7 +810,7 @@ const FTransaction* UTransBuffer::GetTransaction( int32 QueueIndex ) const
 {
 	if (UndoBuffer.Num() > QueueIndex && QueueIndex != INDEX_NONE)
 	{
-		return &UndoBuffer[QueueIndex];
+		return &UndoBuffer[QueueIndex].Get();
 	}
 
 	return NULL;
@@ -810,7 +827,7 @@ FUndoSessionContext UTransBuffer::GetUndoContext( bool bCheckWhetherUndoPossible
 		return Context;
 	}
 
-	const FTransaction* Transaction = &UndoBuffer[ UndoBuffer.Num() - (UndoCount + 1) ];
+	TSharedRef<FTransaction>& Transaction = UndoBuffer[ UndoBuffer.Num() - (UndoCount + 1) ];
 	return Transaction->GetContext();
 }
 
@@ -825,7 +842,7 @@ FUndoSessionContext UTransBuffer::GetRedoContext()
 		return Context;
 	}
 
-	const FTransaction* Transaction = &UndoBuffer[ UndoBuffer.Num() - UndoCount ];
+	TSharedRef<FTransaction>& Transaction = UndoBuffer[ UndoBuffer.Num() - UndoCount ];
 	return Transaction->GetContext();
 }
 
@@ -865,7 +882,7 @@ bool UTransBuffer::Undo(bool bCanRedo)
 	// Apply the undo changes.
 	GIsTransacting = true;
 	{
-		FTransaction& Transaction = UndoBuffer[ UndoBuffer.Num() - ++UndoCount ];
+		FTransaction& Transaction = UndoBuffer[ UndoBuffer.Num() - ++UndoCount ].Get();
 		UE_LOG(LogEditorTransaction, Log,  TEXT("Undo %s"), *Transaction.GetTitle().ToString() );
 		CurrentTransaction = &Transaction;
 
@@ -902,7 +919,7 @@ bool UTransBuffer::Redo()
 	// Apply the redo changes.
 	GIsTransacting = true;
 	{
-		FTransaction& Transaction = UndoBuffer[ UndoBuffer.Num() - UndoCount-- ];
+		FTransaction& Transaction = UndoBuffer[ UndoBuffer.Num() - UndoCount-- ].Get();
 		UE_LOG(LogEditorTransaction, Log,  TEXT("Redo %s"), *Transaction.GetTitle().ToString() );
 		CurrentTransaction = &Transaction;
 
@@ -935,7 +952,7 @@ SIZE_T UTransBuffer::GetUndoSize() const
 	SIZE_T Result=0;
 	for( int32 i=0; i<UndoBuffer.Num(); i++ )
 	{
-		Result += UndoBuffer[i].DataSize();
+		Result += UndoBuffer[i]->DataSize();
 	}
 	return Result;
 }
@@ -961,7 +978,7 @@ void UTransBuffer::SetPrimaryUndoObject(UObject* PrimaryObject)
 
 		if ( CurrentTransactionIdx >= 0 )
 		{
-			FTransaction* Transaction = &UndoBuffer[ CurrentTransactionIdx ];
+			TSharedRef<FTransaction>& Transaction = UndoBuffer[ CurrentTransactionIdx ];
 			Transaction->SetPrimaryObject(PrimaryObject);
 		}
 	}
@@ -970,9 +987,9 @@ void UTransBuffer::SetPrimaryUndoObject(UObject* PrimaryObject)
 bool UTransBuffer::IsObjectInTransationBuffer( const UObject* Object ) const
 {
 	TArray<UObject*> TransactionObjects;
-	for( const FTransaction& Transaction : UndoBuffer )
+	for( const TSharedRef<FTransaction>& Transaction : UndoBuffer )
 	{
-		Transaction.GetTransactionObjects(TransactionObjects);
+		Transaction->GetTransactionObjects(TransactionObjects);
 
 		if( TransactionObjects.Contains(Object) )
 		{
@@ -998,9 +1015,9 @@ bool UTransBuffer::IsObjectTransacting(const UObject* Object) const
 
 bool UTransBuffer::ContainsPieObject() const
 {
-	for( const FTransaction& Transaction : UndoBuffer )
+	for( const TSharedRef<FTransaction>& Transaction : UndoBuffer )
 	{
-		if( Transaction.ContainsPieObject() )
+		if( Transaction->ContainsPieObject() )
 		{
 			return true;
 		}
diff --git a/Engine/Source/Editor/UnrealEd/Private/Kismet2/EnumEditorUtils.cpp b/Engine/Source/Editor/UnrealEd/Private/Kismet2/EnumEditorUtils.cpp
--- a/Engine/Source/Editor/UnrealEd/Private/Kismet2/EnumEditorUtils.cpp
+++ b/Engine/Source/Editor/UnrealEd/Private/Kismet2/EnumEditorUtils.cpp
@@ -581,7 +581,7 @@ void FEnumEditorUtils::UpgradeDisplayNamesFromMetaData(UUserDefinedEnum* Enum)
 		bool bDidUpgradeDisplayNames = false;
 		for (int32 Index = 0; Index < EnumeratorsToEnsure; ++Index)
 		{
-			const FString MetaDataEntryDisplayName = Enum->GetMetaData(FEnumEditorUtilsHelper::DisplayName(), Index);
+			const FString& MetaDataEntryDisplayName = Enum->GetMetaData(FEnumEditorUtilsHelper::DisplayName(), Index);
 			if (!MetaDataEntryDisplayName.IsEmpty())
 			{
 				bDidUpgradeDisplayNames = true;
diff --git a/Engine/Source/Editor/UnrealEd/Private/Kismet2/KismetReinstanceUtilities.cpp b/Engine/Source/Editor/UnrealEd/Private/Kismet2/KismetReinstanceUtilities.cpp
--- a/Engine/Source/Editor/UnrealEd/Private/Kismet2/KismetReinstanceUtilities.cpp
+++ b/Engine/Source/Editor/UnrealEd/Private/Kismet2/KismetReinstanceUtilities.cpp
@@ -1479,11 +1479,6 @@ UClass* FBlueprintCompileReinstancer::MoveCDOToNewClass(UClass* OwnerClass, cons
 	OwnerClass->ClassFlags |= CLASS_NewerVersionExists;
 
 	UObject* OldCDO = OwnerClass->ClassDefaultObject;
-	if(bAvoidCDODuplication)
-	{
-		// prevent duplication of CDO by hiding it from StaticDuplicateObject
-		OwnerClass->ClassDefaultObject = nullptr;
-	}
 	const FName ReinstanceName = MakeUniqueObjectName(GetTransientPackage(), OwnerClass->GetClass(), *(FString(TEXT("REINST_")) + *OwnerClass->GetName()));
 	UClass* CopyOfOwnerClass = CastChecked<UClass>(StaticDuplicateObject(OwnerClass, GetTransientPackage(), ReinstanceName, ~RF_Transactional));
 
@@ -1513,10 +1508,13 @@ UClass* FBlueprintCompileReinstancer::MoveCDOToNewClass(UClass* OwnerClass, cons
 
 	CopyOfOwnerClass->Bind();
 	CopyOfOwnerClass->StaticLink(true);
+
 	if(OldCDO)
 	{
+		// @todo: #dano, rename bAvoidCDODuplication because it's really a flag to move the CDO aside not 'prevent duplication':
 		if(bAvoidCDODuplication)
 		{
+			OwnerClass->ClassDefaultObject = nullptr;
 			OldCDO->Rename(nullptr, CopyOfOwnerClass->GetOuter(), REN_DoNotDirty | REN_DontCreateRedirectors | REN_ForceNoResetLoaders);
 			CopyOfOwnerClass->ClassDefaultObject = OldCDO;
 		}
@@ -1525,6 +1523,280 @@ UClass* FBlueprintCompileReinstancer::MoveCDOToNewClass(UClass* OwnerClass, cons
 	return CopyOfOwnerClass;
 }
 
+static void ReplaceObjectHelper(UObject*& OldObject, UClass* OldClass, UObject*& NewUObject, UClass* NewClass, TMap<UObject*, UObject*>& OldToNewInstanceMap, TMap<UObject*, FName>& OldToNewNameMap, int32 OldObjIndex, TArray<UObject*>& ObjectsToReplace, TArray<UObject*>& PotentialEditorsForRefreshing, TSet<AActor*>& OwnersToRerunConstructionScript, TFunctionRef<TArray<USceneComponent*>&(USceneComponent*)> GetAttachChildrenArray, bool bIsComponent)
+{
+	// If the old object was spawned from an archetype (i.e. not the CDO), we must use the new version of that archetype as the template object when constructing the new instance.
+	UObject* OldArchetype = OldObject->GetArchetype();
+	UObject* NewArchetype = OldToNewInstanceMap.FindRef(OldArchetype);
+
+	bool bArchetypeReinstanced = (OldArchetype == OldClass->GetDefaultObject()) || (NewArchetype != nullptr);
+	// if we don't have a updated archetype to spawn from, we need to update/reinstance it
+	while (!bArchetypeReinstanced)
+	{
+		int32 ArchetypeIndex = ObjectsToReplace.Find(OldArchetype);
+		if (ArchetypeIndex != INDEX_NONE)
+		{
+			if (ensure(ArchetypeIndex > OldObjIndex))
+			{
+				// if this object has an archetype, but it hasn't been 
+				// reinstanced yet (but is queued to) then we need to swap out 
+				// the two, and reinstance the archetype first
+				ObjectsToReplace.Swap(ArchetypeIndex, OldObjIndex);
+				OldObject = ObjectsToReplace[OldObjIndex];
+				check(OldObject == OldArchetype);
+
+				OldArchetype = OldObject->GetArchetype();
+				NewArchetype = OldToNewInstanceMap.FindRef(OldArchetype);
+				bArchetypeReinstanced = (OldArchetype == OldClass->GetDefaultObject()) || (NewArchetype != nullptr);
+			}
+			else
+			{
+				break;
+			}
+		}
+		else
+		{
+			break;
+		}
+	}
+	// Check that either this was an instance of the class directly, or we found a new archetype for it
+	ensureMsgf(bArchetypeReinstanced, TEXT("Reinstancing non-actor (%s); failed to resolve archetype object - property values may be lost."), *OldObject->GetPathName());
+
+	EObjectFlags OldFlags = OldObject->GetFlags();
+
+	FName OldName(OldObject->GetFName());
+
+	// If the old object is in this table, we've already renamed it away in a previous iteration. Don't rename it again!
+	if (!OldToNewNameMap.Contains(OldObject))
+	{
+		// If we're reinstancing a component template, we also need to rename any inherited templates that are found to be based on it, in order to preserve archetype paths.
+		if (bIsComponent && OldObject->HasAllFlags(RF_ArchetypeObject) && OldObject->GetOuter()->IsA<UBlueprintGeneratedClass>())
+		{
+			// Gather all component templates from the current archetype to the farthest antecedent inherited template(s).
+			TArray<UObject*> OldArchetypeObjects;
+			FArchetypeReinstanceHelper::GetArchetypeObjects(OldObject, OldArchetypeObjects, RF_InheritableComponentTemplate);
+
+			// Find a unique object name that does not conflict with anything in the scope of all outers in the template chain.
+			const FString OldArchetypeName = FArchetypeReinstanceHelper::FindUniqueArchetypeObjectName(OldArchetypeObjects).ToString();
+
+			for (UObject* OldArchetypeObject : OldArchetypeObjects)
+			{
+				OldToNewNameMap.Add(OldArchetypeObject, OldName);
+				OldArchetypeObject->Rename(*OldArchetypeName, OldArchetypeObject->GetOuter(), REN_DoNotDirty | REN_DontCreateRedirectors | (GBlueprintUseCompilationManager ? REN_ForceNoResetLoaders : 0));
+			}
+		}
+		else
+		{
+			OldObject->Rename(nullptr, OldObject->GetOuter(), REN_DoNotDirty | REN_DontCreateRedirectors | (GBlueprintUseCompilationManager ? REN_ForceNoResetLoaders : 0));
+		}
+	}
+						
+	{
+		// We may have already renamed this object to temp space if it was an inherited archetype in a previous iteration; check for that here.
+		FName NewName = OldToNewNameMap.FindRef(OldObject);
+		if (NewName == NAME_None)
+		{
+			// Otherwise, just use the old object's current name.
+			NewName = OldName;
+		}
+
+		FMakeClassSpawnableOnScope TemporarilySpawnable(NewClass);
+		NewUObject = NewObject<UObject>(OldObject->GetOuter(), NewClass, NewName, RF_NoFlags, NewArchetype);
+	}
+
+	check(NewUObject != nullptr);
+
+	const EObjectFlags FlagMask = RF_Public | RF_ArchetypeObject | RF_Transactional | RF_Transient | RF_TextExportTransient | RF_InheritableComponentTemplate | RF_Standalone; //TODO: what about RF_RootSet?
+	NewUObject->SetFlags(OldFlags & FlagMask);
+
+	InstancedPropertyUtils::FInstancedPropertyMap InstancedPropertyMap;
+	InstancedPropertyUtils::FArchiveInstancedSubObjCollector  InstancedSubObjCollector(OldObject, InstancedPropertyMap);
+	UEditorEngine::CopyPropertiesForUnrelatedObjects(OldObject, NewUObject);
+	InstancedPropertyUtils::FArchiveInsertInstancedSubObjects InstancedSubObjSpawner(NewUObject, InstancedPropertyMap);
+
+	if (UAnimInstance* AnimTree = Cast<UAnimInstance>(NewUObject))
+	{
+		// Initialising the anim instance isn't enough to correctly set up the skeletal mesh again in a
+		// paused world, need to initialise the skeletal mesh component that contains the anim instance.
+		if (USkeletalMeshComponent* SkelComponent = Cast<USkeletalMeshComponent>(AnimTree->GetOuter()))
+		{
+			SkelComponent->InitAnim(true);
+		}
+	}
+
+	UWorld* RegisteredWorld = nullptr;
+	bool bWasRegistered = false;
+	if (bIsComponent)
+	{
+		UActorComponent* OldComponent = CastChecked<UActorComponent>(OldObject);
+		if (OldComponent->IsRegistered())
+		{
+			bWasRegistered = true;
+			RegisteredWorld = OldComponent->GetWorld();
+			OldComponent->UnregisterComponent();
+		}
+	}
+
+	OldObject->RemoveFromRoot();
+	OldObject->MarkPendingKill();
+
+	OldToNewInstanceMap.Add(OldObject, NewUObject);
+
+	if (bIsComponent)
+	{
+		UActorComponent* Component = CastChecked<UActorComponent>(NewUObject);
+		AActor* OwningActor = Component->GetOwner();
+		if (OwningActor)
+		{
+			OwningActor->ResetOwnedComponents();
+
+			// Check to see if they have an editor that potentially needs to be refreshed
+			if (OwningActor->GetClass()->ClassGeneratedBy)
+			{
+				PotentialEditorsForRefreshing.AddUnique(OwningActor->GetClass()->ClassGeneratedBy);
+			}
+
+			// we need to keep track of actor instances that need 
+			// their construction scripts re-ran (since we've just 
+			// replaced a component they own)
+
+			// Skipping CDOs as CSs are not allowed for them.
+			if (!OwningActor->HasAnyFlags(RF_ClassDefaultObject))
+			{
+				OwnersToRerunConstructionScript.Add(OwningActor);
+			}
+		}
+
+		if (bWasRegistered)
+		{
+			if (RegisteredWorld && OwningActor == nullptr)
+			{
+				// Thumbnail components are added to a World without an actor, so we must special case their
+				// REINST to register them with the world again.
+				// The old thumbnail component is GC'd and will ensure if all it's attachments are not released
+				// @TODO: This special case can breakdown if the nature of thumbnail components changes and could
+				// use a cleanup later.
+				if (OldObject->GetOutermost() == GetTransientPackage())
+				{
+					if (USceneComponent* SceneComponent = Cast<USceneComponent>(OldObject))
+					{
+						GetAttachChildrenArray(SceneComponent).Empty();
+						SceneComponent->SetupAttachment(nullptr);
+					}
+				}
+
+				Component->RegisterComponentWithWorld(RegisteredWorld);
+			}
+			else
+			{
+				Component->RegisterComponent();
+			}
+		}
+	}
+}
+
+static void ReplaceActorHelper(UObject* OldObject, UClass* OldClass, UObject*& NewUObject, UClass* NewClass, TMap<UObject*, UObject*>& OldToNewInstanceMap, TMap<UClass*, UClass*>& InOldToNewClassMap, AActor* OldActor, TMap<FStringAssetReference, UObject*>& ReinstancedObjectsWeakReferenceMap, TMap<UObject*, FActorAttachmentData>& ActorAttachmentData, TArray<FActorReplacementHelper>& ReplacementActors, bool bPreserveRootComponent, bool& bSelectionChanged)
+{
+	FVector  Location = FVector::ZeroVector;
+	FRotator Rotation = FRotator::ZeroRotator;
+	if (USceneComponent* OldRootComponent = OldActor->GetRootComponent())
+	{
+		// We need to make sure that the GetComponentTransform() transform is up to date, but we don't want to run any initialization logic
+		// so we silence the update, cache it off, revert the change (so no events are raised), and then directly update the transform
+		// with the value calculated in ConditionalUpdateComponentToWorld:
+		FScopedMovementUpdate SilenceMovement(OldRootComponent);
+
+		OldRootComponent->ConditionalUpdateComponentToWorld();
+		FTransform OldComponentToWorld = OldRootComponent->GetComponentTransform();
+		SilenceMovement.RevertMove();
+
+		OldRootComponent->SetComponentToWorld(OldComponentToWorld);
+		Location = OldActor->GetActorLocation();
+		Rotation = OldActor->GetActorRotation();
+	}
+
+	// If this actor was spawned from an Archetype, we spawn the new actor from the new version of that archetype
+	UObject* OldArchetype = OldActor->GetArchetype();
+	UWorld*  World = OldActor->GetWorld();
+	AActor*  NewArchetype = Cast<AActor>(OldToNewInstanceMap.FindRef(OldArchetype));
+	// Check that either this was an instance of the class directly, or we found a new archetype for it
+	check(OldArchetype == OldClass->GetDefaultObject() || NewArchetype);
+
+	// Spawn the new actor instance, in the same level as the original, but deferring running the construction script until we have transferred modified properties
+	ULevel*  ActorLevel = OldActor->GetLevel();
+	UClass** MappedClass = InOldToNewClassMap.Find(OldActor->GetClass());
+	UClass*  SpawnClass = MappedClass ? *MappedClass : NewClass;
+
+	FActorSpawnParameters SpawnInfo;
+	SpawnInfo.OverrideLevel = ActorLevel;
+	SpawnInfo.Template = NewArchetype;
+	SpawnInfo.SpawnCollisionHandlingOverride = ESpawnActorCollisionHandlingMethod::AlwaysSpawn;
+	SpawnInfo.bDeferConstruction = true;
+	SpawnInfo.Name = OldActor->GetFName();
+
+	OldActor->UObject::Rename(nullptr, OldObject->GetOuter(), REN_DoNotDirty | REN_DontCreateRedirectors | (GBlueprintUseCompilationManager ? REN_ForceNoResetLoaders : 0));
+
+	AActor* NewActor = nullptr;
+	{
+		FMakeClassSpawnableOnScope TemporarilySpawnable(SpawnClass);
+		NewActor = World->SpawnActor(SpawnClass, &Location, &Rotation, SpawnInfo);
+	}
+
+	if (OldActor->CurrentTransactionAnnotation.IsValid())
+	{
+		NewActor->CurrentTransactionAnnotation = OldActor->CurrentTransactionAnnotation;
+	}
+
+	check(NewActor != nullptr);
+	NewUObject = NewActor;
+	// store the new actor for the second pass (NOTE: this detaches 
+	// OldActor from all child/parent attachments)
+	//
+	// running the NewActor's construction-script is saved for that 
+	// second pass (because the construction-script may reference 
+	// another instance that hasn't been replaced yet).
+	FActorAttachmentData& CurrentAttachmentData = ActorAttachmentData.FindChecked(OldObject);
+	ReplacementActors.Add(FActorReplacementHelper(NewActor, OldActor, MoveTemp(CurrentAttachmentData)));
+	ActorAttachmentData.Remove(OldObject);
+
+	ReinstancedObjectsWeakReferenceMap.Add(OldObject, NewUObject);
+
+	OldActor->DestroyConstructedComponents(); // don't want to serialize components from the old actor
+												// Unregister native components so we don't copy any sub-components they generate for themselves (like UCameraComponent does)
+	OldActor->UnregisterAllComponents();
+
+	// Unregister any native components, might have cached state based on properties we are going to overwrite
+	NewActor->UnregisterAllComponents();
+
+	UEngine::FCopyPropertiesForUnrelatedObjectsParams Params;
+	Params.bPreserveRootComponent = bPreserveRootComponent;
+	UEngine::CopyPropertiesForUnrelatedObjects(OldActor, NewActor, Params);
+
+	// reset properties/streams
+	NewActor->ResetPropertiesForConstruction();
+	// register native components
+	NewActor->RegisterAllComponents();
+
+	// 
+	// clean up the old actor (unselect it, remove it from the world, etc.)...
+
+	if (OldActor->IsSelected())
+	{
+		if(GEditor)
+		{
+			GEditor->SelectActor(OldActor, /*bInSelected =*/false, /*bNotify =*/false);
+		}
+		bSelectionChanged = true;
+	}
+	if (GEditor && GEditor->Layers.IsValid())
+	{
+		GEditor->Layers->DisassociateActorFromLayers(OldActor);
+	}
+
+	World->EditorDestroyActor(OldActor, /*bShouldModifyLevel =*/true);
+	OldToNewInstanceMap.Add(OldActor, NewActor);
+}
+
 void FBlueprintCompileReinstancer::ReplaceInstancesOfClass_Inner(TMap<UClass*, UClass*>& InOldToNewClassMap, UObject* InOriginalCDO, TSet<UObject*>* ObjectsThatShouldUseOldStuff, bool bClassObjectReplaced, bool bPreserveRootComponent)
 {
 	// If there is an original CDO, we are only reinstancing a single class
@@ -1576,14 +1848,29 @@ void FBlueprintCompileReinstancer::ReplaceInstancesOfClass_Inner(TMap<UClass*, U
 		OnObjectsReplacedHandle = GEditor->OnObjectsReplaced().AddRaw(&ObjectRemappingHelper, &FObjectRemappingHelper::OnObjectsReplaced);
 	}
 
+	auto UpdateObjectBeingDebugged = [](UObject* InOldObject, UObject* InNewObject)
 	{
+		if (UBlueprint* OldObjBlueprint = Cast<UBlueprint>(InOldObject->GetClass()->ClassGeneratedBy))
+		{
+			if (OldObjBlueprint->GetObjectBeingDebugged() == InOldObject)
+			{
+				OldObjBlueprint->SetObjectBeingDebugged(InNewObject);
+			}
+		}
+	};
+
+	{
+		BP_SCOPED_COMPILER_EVENT_STAT(EKismetReinstancerStats_ReplaceInstancesOfClass);
 		if(GEditor && GEditor->GetSelectedActors())
 		{
 			SelectedActors = GEditor->GetSelectedActors();
 			SelectedActors->BeginBatchSelectOperation();
 			SelectedActors->Modify();
 		}
-		
+
+		// WARNING: for (TPair<UClass*, UClass*> OldToNewClass : InOldToNewClassMap) duplicated below 
+		// to handle reconstructing actors which need to be reinstanced aftertheir owned components 
+		// have been updated:
 		for (TPair<UClass*, UClass*> OldToNewClass : InOldToNewClassMap)
 		{
 			UClass* OldClass = OldToNewClass.Key;
@@ -1594,10 +1881,53 @@ void FBlueprintCompileReinstancer::ReplaceInstancesOfClass_Inner(TMap<UClass*, U
 #else
 			check(OldClass != NewClass);
 #endif
+			{
+				const bool bIncludeDerivedClasses = false;
+				ObjectsToReplace.Empty();
+				GetObjectsOfClass(OldClass, ObjectsToReplace, bIncludeDerivedClasses);
+				// Then fix 'real' (non archetype) instances of the class
+				for (int32 OldObjIndex = 0; OldObjIndex < ObjectsToReplace.Num(); ++OldObjIndex)
+				{
+					UObject* OldObject = ObjectsToReplace[OldObjIndex];
+					AActor* OldActor = Cast<AActor>(OldObject);
+
+					// Skip archetype instances, EXCEPT for component templates and child actor templates
+					const bool bIsComponent = NewClass->IsChildOf(UActorComponent::StaticClass());
+					const bool bIsChildActorTemplate = OldActor && OldActor->GetOuter()->IsA<UChildActorComponent>();
+					if (OldObject->IsPendingKill() || 
+						(!bIsComponent && !bIsChildActorTemplate && OldObject->IsTemplate()))
+					{
+						continue;
+					}
 
-			{ BP_SCOPED_COMPILER_EVENT_STAT(EKismetReinstancerStats_ReplaceInstancesOfClass);
+					// WARNING: This loop only handles non-actor objects, actor objects are handled below:
+					if (OldActor == nullptr)
+					{
+						UObject* NewUObject = nullptr;
+						ReplaceObjectHelper(OldObject, OldClass, NewUObject, NewClass, OldToNewInstanceMap, OldToNewNameMap, OldObjIndex, ObjectsToReplace, PotentialEditorsForRefreshing, OwnersToRerunConstructionScript, &FDirectAttachChildrenAccessor::Get, bIsComponent);
+						UpdateObjectBeingDebugged(OldObject, NewUObject);
+						
+						if (bLogConversions)
+						{
+							UE_LOG(LogBlueprint, Log, TEXT("Converted instance '%s' to '%s'"), *GetPathNameSafe(OldObject), *GetPathNameSafe(NewUObject));
+						}
+					}
+				}
+			}
+		}
+
+		// WARNING: for (TPair<UClass*, UClass*> OldToNewClass : InOldToNewClassMap) duplicated above 
+		// this loop only handles actors - which need to be reconstructed *after* their owned components 
+		// have been reinstanced:
+		for (TPair<UClass*, UClass*> OldToNewClass : InOldToNewClassMap)
+		{
+			UClass* OldClass = OldToNewClass.Key;
+			UClass* NewClass = OldToNewClass.Value;
+			check(OldClass && NewClass);
 
+			{
 				const bool bIncludeDerivedClasses = false;
+				ObjectsToReplace.Empty();
 				GetObjectsOfClass(OldClass, ObjectsToReplace, bIncludeDerivedClasses);
 
 				// store old attachment data before we mess with components, etc:
@@ -1618,331 +1948,35 @@ void FBlueprintCompileReinstancer::ReplaceInstancesOfClass_Inner(TMap<UClass*, U
 				for (int32 OldObjIndex = 0; OldObjIndex < ObjectsToReplace.Num(); ++OldObjIndex)
 				{
 					UObject* OldObject = ObjectsToReplace[OldObjIndex];
-
-					if (OldObject->IsPendingKill())
-					{
-						continue;
-					}
-
 					AActor* OldActor = Cast<AActor>(OldObject);
 
-					// Skip archetype instances, EXCEPT for component templates and child actor templates
-					const bool bIsComponent = NewClass->IsChildOf(UActorComponent::StaticClass());
+					// Skip archetype instances, EXCEPT for child actor templates
 					const bool bIsChildActorTemplate = OldActor && OldActor->GetOuter()->IsA<UChildActorComponent>();
-					if (!bIsComponent && !bIsChildActorTemplate && OldObject->IsTemplate())
+					if (OldObject->IsPendingKill() || (!bIsChildActorTemplate && OldObject->IsTemplate()))
 					{
 						continue;
 					}
-
-					UBlueprint* DebuggingBlueprint = nullptr;
-					// ObjectsToReplace is not in any garunteed order; so below there may be a 
-					// scenario where one object in ObjectsToReplace needs to come before another; 
-					// in that scenario we swap the two in ObjectsToReplace and finish the loop 
-					// with a new OldObject - to support that we need to re-cache info about the 
-					// (new) old object before it is destroyed (hence why this is all in a lambda,  
-					// for reuse).
-					auto CacheOldObjectState = [&DebuggingBlueprint](UObject* Replacee)
-					{
-						// clear in case it has already been set, and we're re-executing this
-						DebuggingBlueprint = nullptr;
-
-						if (UBlueprint* OldObjBlueprint = Cast<UBlueprint>(Replacee->GetClass()->ClassGeneratedBy))
-						{
-							if (OldObjBlueprint->GetObjectBeingDebugged() == Replacee)
-							{
-								DebuggingBlueprint = OldObjBlueprint;
-							}
-						}
-					};
-					CacheOldObjectState(OldObject);
-
-					UObject* NewUObject = nullptr;
-					// if the object to replace is an actor...
-					if (OldActor != nullptr && OldActor->GetLevel())
-					{
-						FVector  Location = FVector::ZeroVector;
-						FRotator Rotation = FRotator::ZeroRotator;
-						if (USceneComponent* OldRootComponent = OldActor->GetRootComponent())
-						{
-							// We need to make sure that the GetComponentTransform() transform is up to date, but we don't want to run any initialization logic
-							// so we silence the update, cache it off, revert the change (so no events are raised), and then directly update the transform
-							// with the value calculated in ConditionalUpdateComponentToWorld:
-							FScopedMovementUpdate SilenceMovement(OldRootComponent);
-
-							OldRootComponent->ConditionalUpdateComponentToWorld();
-							FTransform OldComponentToWorld = OldRootComponent->GetComponentTransform();
-							SilenceMovement.RevertMove();
-
-							OldRootComponent->SetComponentToWorld(OldComponentToWorld);
-							Location = OldActor->GetActorLocation();
-							Rotation = OldActor->GetActorRotation();
-						}
-
-						// If this actor was spawned from an Archetype, we spawn the new actor from the new version of that archetype
-						UObject* OldArchetype = OldActor->GetArchetype();
-						UWorld*  World = OldActor->GetWorld();
-						AActor*  NewArchetype = Cast<AActor>(OldToNewInstanceMap.FindRef(OldArchetype));
-						// Check that either this was an instance of the class directly, or we found a new archetype for it
-						check(OldArchetype == OldClass->GetDefaultObject() || NewArchetype);
-
-						// Spawn the new actor instance, in the same level as the original, but deferring running the construction script until we have transferred modified properties
-						ULevel*  ActorLevel = OldActor->GetLevel();
-						UClass** MappedClass = InOldToNewClassMap.Find(OldActor->GetClass());
-						UClass*  SpawnClass = MappedClass ? *MappedClass : NewClass;
-
-						FActorSpawnParameters SpawnInfo;
-						SpawnInfo.OverrideLevel = ActorLevel;
-						SpawnInfo.Template = NewArchetype;
-						SpawnInfo.SpawnCollisionHandlingOverride = ESpawnActorCollisionHandlingMethod::AlwaysSpawn;
-						SpawnInfo.bDeferConstruction = true;
-						SpawnInfo.Name = OldActor->GetFName();
-
-						OldActor->UObject::Rename(nullptr, OldObject->GetOuter(), REN_DoNotDirty | REN_DontCreateRedirectors | (GBlueprintUseCompilationManager ? REN_ForceNoResetLoaders : 0));
-
-						AActor* NewActor = nullptr;
-						{
-							FMakeClassSpawnableOnScope TemporarilySpawnable(SpawnClass);
-							NewActor = World->SpawnActor(SpawnClass, &Location, &Rotation, SpawnInfo);
-						}
-
-						if (OldActor->CurrentTransactionAnnotation.IsValid())
-						{
-							NewActor->CurrentTransactionAnnotation = OldActor->CurrentTransactionAnnotation;
-						}
-
-						check(NewActor != nullptr);
-						NewUObject = NewActor;
-						// store the new actor for the second pass (NOTE: this detaches 
-						// OldActor from all child/parent attachments)
-						//
-						// running the NewActor's construction-script is saved for that 
-						// second pass (because the construction-script may reference 
-						// another instance that hasn't been replaced yet).
-						FActorAttachmentData& CurrentAttachmentData = ActorAttachmentData.FindChecked(OldObject);
-						ReplacementActors.Add(FActorReplacementHelper(NewActor, OldActor, MoveTemp(CurrentAttachmentData)));
-						ActorAttachmentData.Remove(OldObject);
-
-						ReinstancedObjectsWeakReferenceMap.Add(OldObject, NewUObject);
-
-						OldActor->DestroyConstructedComponents(); // don't want to serialize components from the old actor
-																  // Unregister native components so we don't copy any sub-components they generate for themselves (like UCameraComponent does)
-						OldActor->UnregisterAllComponents();
-
-						// Unregister any native components, might have cached state based on properties we are going to overwrite
-						NewActor->UnregisterAllComponents();
-
-						UEngine::FCopyPropertiesForUnrelatedObjectsParams Params;
-						Params.bPreserveRootComponent = bPreserveRootComponent;
-						UEngine::CopyPropertiesForUnrelatedObjects(OldActor, NewActor, Params);
-
-						// reset properties/streams
-						NewActor->ResetPropertiesForConstruction();
-						// register native components
-						NewActor->RegisterAllComponents();
-
-						// 
-						// clean up the old actor (unselect it, remove it from the world, etc.)...
-
-						if (OldActor->IsSelected())
-						{
-							if(GEditor)
-							{
-								GEditor->SelectActor(OldActor, /*bInSelected =*/false, /*bNotify =*/false);
-							}
-							bSelectionChanged = true;
-						}
-						if (GEditor && GEditor->Layers.IsValid())
-						{
-							GEditor->Layers->DisassociateActorFromLayers(OldActor);
-						}
-
-						World->EditorDestroyActor(OldActor, /*bShouldModifyLevel =*/true);
-						OldToNewInstanceMap.Add(OldActor, NewActor);
-					}
-					else
+					
+					// WARNING: This loop only handles actor objects that are in a level, all other objects are
+					// handled above
+					if (OldActor != nullptr)
 					{
-						// If the old object was spawned from an archetype (i.e. not the CDO), we must use the new version of that archetype as the template object when constructing the new instance.
-						UObject* OldArchetype = OldObject->GetArchetype();
-						UObject* NewArchetype = OldToNewInstanceMap.FindRef(OldArchetype);
-
-						bool bArchetypeReinstanced = (OldArchetype == OldClass->GetDefaultObject()) || (NewArchetype != nullptr);
-						// if we don't have a updated archetype to spawn from, we need to update/reinstance it
-						while (!bArchetypeReinstanced)
+						UObject* NewUObject = nullptr;
+						if(OldActor->GetLevel())
 						{
-							int32 ArchetypeIndex = ObjectsToReplace.Find(OldArchetype);
-							if (ArchetypeIndex != INDEX_NONE)
-							{
-								if (ensure(ArchetypeIndex > OldObjIndex))
-								{
-									// if this object has an archetype, but it hasn't been 
-									// reinstanced yet (but is queued to) then we need to swap out 
-									// the two, and reinstance the archetype first
-									ObjectsToReplace.Swap(ArchetypeIndex, OldObjIndex);
-									OldObject = ObjectsToReplace[OldObjIndex];
-									check(OldObject == OldArchetype);
-
-									// update any info that we stored regarding the previous OldObject
-									CacheOldObjectState(OldObject);
-									
-									OldArchetype = OldObject->GetArchetype();
-									NewArchetype = OldToNewInstanceMap.FindRef(OldArchetype);
-									bArchetypeReinstanced = (OldArchetype == OldClass->GetDefaultObject()) || (NewArchetype != nullptr);
-								}
-								else
-								{
-									break;
-								}
-							}
-							else
-							{
-								break;
-							}
+							ReplaceActorHelper(OldObject, OldClass, NewUObject, NewClass, OldToNewInstanceMap, InOldToNewClassMap, OldActor, ReinstancedObjectsWeakReferenceMap, ActorAttachmentData, ReplacementActors, bPreserveRootComponent, bSelectionChanged);
 						}
-						// Check that either this was an instance of the class directly, or we found a new archetype for it
-						ensureMsgf(bArchetypeReinstanced, TEXT("Reinstancing non-actor (%s); failed to resolve archetype object - property values may be lost."), *OldObject->GetPathName());
-
-						EObjectFlags OldFlags = OldObject->GetFlags();
-
-						FName OldName(OldObject->GetFName());
-
-						// If the old object is in this table, we've already renamed it away in a previous iteration. Don't rename it again!
-						if (!OldToNewNameMap.Contains(OldObject))
+						else
 						{
-							// If we're reinstancing a component template, we also need to rename any inherited templates that are found to be based on it, in order to preserve archetype paths.
-							if (bIsComponent && OldObject->HasAllFlags(RF_ArchetypeObject) && OldObject->GetOuter()->IsA<UBlueprintGeneratedClass>())
-							{
-								// Gather all component templates from the current archetype to the farthest antecedent inherited template(s).
-								TArray<UObject*> OldArchetypeObjects;
-								FArchetypeReinstanceHelper::GetArchetypeObjects(OldObject, OldArchetypeObjects, RF_InheritableComponentTemplate);
-
-								// Find a unique object name that does not conflict with anything in the scope of all outers in the template chain.
-								const FString OldArchetypeName = FArchetypeReinstanceHelper::FindUniqueArchetypeObjectName(OldArchetypeObjects).ToString();
-
-								for (UObject* OldArchetypeObject : OldArchetypeObjects)
-								{
-									OldToNewNameMap.Add(OldArchetypeObject, OldName);
-									OldArchetypeObject->Rename(*OldArchetypeName, OldArchetypeObject->GetOuter(), REN_DoNotDirty | REN_DontCreateRedirectors | (GBlueprintUseCompilationManager ? REN_ForceNoResetLoaders : 0));
-								}
-							}
-							else
-							{
-								OldObject->Rename(nullptr, OldObject->GetOuter(), REN_DoNotDirty | REN_DontCreateRedirectors | (GBlueprintUseCompilationManager ? REN_ForceNoResetLoaders : 0));
-							}
+							// Actors that are not in a level cannot be reconstructed, sequencer team decided to reinstance these as normal objects:
+							ReplaceObjectHelper(OldObject, OldClass, NewUObject, NewClass, OldToNewInstanceMap, OldToNewNameMap, OldObjIndex, ObjectsToReplace, PotentialEditorsForRefreshing, OwnersToRerunConstructionScript, &FDirectAttachChildrenAccessor::Get, false);
 						}
+						UpdateObjectBeingDebugged(OldObject, NewUObject);
 						
+						if (bLogConversions)
 						{
-							// We may have already renamed this object to temp space if it was an inherited archetype in a previous iteration; check for that here.
-							FName NewName = OldToNewNameMap.FindRef(OldObject);
-							if (NewName == NAME_None)
-							{
-								// Otherwise, just use the old object's current name.
-								NewName = OldName;
-							}
-
-							FMakeClassSpawnableOnScope TemporarilySpawnable(NewClass);
-							NewUObject = NewObject<UObject>(OldObject->GetOuter(), NewClass, NewName, RF_NoFlags, NewArchetype);
+							UE_LOG(LogBlueprint, Log, TEXT("Converted instance '%s' to '%s'"), *GetPathNameSafe(OldObject), *GetPathNameSafe(NewUObject));
 						}
-
-						check(NewUObject != nullptr);
-
-						const EObjectFlags FlagMask = RF_Public | RF_ArchetypeObject | RF_Transactional | RF_Transient | RF_TextExportTransient | RF_InheritableComponentTemplate | RF_Standalone; //TODO: what about RF_RootSet?
-						NewUObject->SetFlags(OldFlags & FlagMask);
-
-						InstancedPropertyUtils::FInstancedPropertyMap InstancedPropertyMap;
-						InstancedPropertyUtils::FArchiveInstancedSubObjCollector  InstancedSubObjCollector(OldObject, InstancedPropertyMap);
-						UEditorEngine::CopyPropertiesForUnrelatedObjects(OldObject, NewUObject);
-						InstancedPropertyUtils::FArchiveInsertInstancedSubObjects InstancedSubObjSpawner(NewUObject, InstancedPropertyMap);
-
-						if (UAnimInstance* AnimTree = Cast<UAnimInstance>(NewUObject))
-						{
-							// Initialising the anim instance isn't enough to correctly set up the skeletal mesh again in a
-							// paused world, need to initialise the skeletal mesh component that contains the anim instance.
-							if (USkeletalMeshComponent* SkelComponent = Cast<USkeletalMeshComponent>(AnimTree->GetOuter()))
-							{
-								SkelComponent->InitAnim(true);
-							}
-						}
-
-						UWorld* RegisteredWorld = nullptr;
-						bool bWasRegistered = false;
-						if (bIsComponent)
-						{
-							UActorComponent* OldComponent = CastChecked<UActorComponent>(OldObject);
-							if (OldComponent->IsRegistered())
-							{
-								bWasRegistered = true;
-								RegisteredWorld = OldComponent->GetWorld();
-								OldComponent->UnregisterComponent();
-							}
-						}
-
-						OldObject->RemoveFromRoot();
-						OldObject->MarkPendingKill();
-
-						OldToNewInstanceMap.Add(OldObject, NewUObject);
-
-						if (bIsComponent)
-						{
-							UActorComponent* Component = CastChecked<UActorComponent>(NewUObject);
-							AActor* OwningActor = Component->GetOwner();
-							if (OwningActor)
-							{
-								OwningActor->ResetOwnedComponents();
-
-								// Check to see if they have an editor that potentially needs to be refreshed
-								if (OwningActor->GetClass()->ClassGeneratedBy)
-								{
-									PotentialEditorsForRefreshing.AddUnique(OwningActor->GetClass()->ClassGeneratedBy);
-								}
-
-								// we need to keep track of actor instances that need 
-								// their construction scripts re-ran (since we've just 
-								// replaced a component they own)
-
-								// Skipping CDOs as CSs are not allowed for them.
-								if (!OwningActor->HasAnyFlags(RF_ClassDefaultObject))
-								{
-									OwnersToRerunConstructionScript.Add(OwningActor);
-								}
-							}
-
-							if (bWasRegistered)
-							{
-								if (RegisteredWorld && OwningActor == nullptr)
-								{
-									// Thumbnail components are added to a World without an actor, so we must special case their
-									// REINST to register them with the world again.
-									// The old thumbnail component is GC'd and will ensure if all it's attachments are not released
-									// @TODO: This special case can breakdown if the nature of thumbnail components changes and could
-									// use a cleanup later.
-									if (OldObject->GetOutermost() == GetTransientPackage())
-									{
-										if (USceneComponent* SceneComponent = Cast<USceneComponent>(OldObject))
-										{
-											FDirectAttachChildrenAccessor::Get(SceneComponent).Empty();
-											SceneComponent->SetupAttachment(nullptr);
-										}
-									}
-
-									Component->RegisterComponentWithWorld(RegisteredWorld);
-								}
-								else
-								{
-									Component->RegisterComponent();
-								}
-							}
-						}
-					}
-
-					// If this original object came from a blueprint and it was in the selected debug set, change the debugging to the new object.
-					if (DebuggingBlueprint && NewUObject)
-					{
-						DebuggingBlueprint->SetObjectBeingDebugged(NewUObject);
-					}
-
-					if (bLogConversions)
-					{
-						UE_LOG(LogBlueprint, Log, TEXT("Converted instance '%s' to '%s'"), *GetPathNameSafe(OldObject), *GetPathNameSafe(NewUObject));
 					}
 				}
 			}
@@ -2147,7 +2181,7 @@ void FBlueprintCompileReinstancer::CopyPropertiesForUnrelatedObjects(UObject* Ol
 
 	UEngine::FCopyPropertiesForUnrelatedObjectsParams Params;
 	Params.bAggressiveDefaultSubobjectReplacement = false;//true;
-	Params.bDoDelta = false;
+	Params.bDoDelta = !OldObject->HasAnyFlags(RF_ClassDefaultObject);
 	Params.bCopyDeprecatedProperties = true;
 	Params.bSkipCompilerGeneratedDefaults = true;
 	Params.bClearReferences = bClearExternalReferences;
