Static analysis fix.

diff --git a/Engine/Source/Editor/UnrealEd/Private/Commandlets/GatherTextFromSourceCommandlet.cpp b/Engine/Source/Editor/UnrealEd/Private/Commandlets/GatherTextFromSourceCommandlet.cpp
--- a/Engine/Source/Editor/UnrealEd/Private/Commandlets/GatherTextFromSourceCommandlet.cpp
+++ b/Engine/Source/Editor/UnrealEd/Private/Commandlets/GatherTextFromSourceCommandlet.cpp
@@ -1422,7 +1422,7 @@ void UGatherTextFromSourceCommandlet::FCommandMacroDescriptor::TryParse(const FS
 		if (ParseArgsFromMacro(StripCommentsFromToken(Text, Context), Arguments, Context))
 		{
 			// 5 or 6 arguments depending on whether or not there are one or two bindings
-			if (Arguments.Num() != 5 || Arguments.Num() != 6)
+			if (Arguments.Num() != 5 && Arguments.Num() != 6)
 			{
 				UE_LOG(LogGatherTextFromSourceCommandlet, Warning, TEXT("Too many arguments in command %s macro in %s(%d):%s"), *GetToken(), *Context.Filename, Context.LineNumber, *FLocTextHelper::SanitizeLogOutput(Context.LineText));
 			}
