Copying //UE4/Dev-Framework to //UE4/Dev-Main (Source: //UE4/Dev-Framework @ 3716594)

diff --git a/Engine/Source/Editor/MaterialEditor/Private/MaterialEditingLibrary.cpp b/Engine/Source/Editor/MaterialEditor/Private/MaterialEditingLibrary.cpp
--- a/Engine/Source/Editor/MaterialEditor/Private/MaterialEditingLibrary.cpp
+++ b/Engine/Source/Editor/MaterialEditor/Private/MaterialEditingLibrary.cpp
@@ -29,7 +29,7 @@
 DEFINE_LOG_CATEGORY_STATIC(LogMaterialEditingLibrary, Warning, All);
 
 /** Util to find expression  */
-static FExpressionInput* GetExpressionInputByName(UMaterialExpression* Expression, const FString& InputName)
+static FExpressionInput* GetExpressionInputByName(UMaterialExpression* Expression, const FName InputName)
 {
 	check(Expression);
 	FExpressionInput* Result = nullptr;
@@ -37,7 +37,7 @@ static FExpressionInput* GetExpressionInputByName(UMaterialExpression* Expressio
 	TArray<FExpressionInput*> Inputs = Expression->GetInputs();
 
 	// Return first input if no name specified
-	if (InputName.IsEmpty())
+	if (InputName.IsNone())
 	{
 		if (Inputs.Num() > 0)
 		{
@@ -50,7 +50,7 @@ static FExpressionInput* GetExpressionInputByName(UMaterialExpression* Expressio
 		// Get name of each input, see if its the one we want
 		for (int InputIdx = 0; InputIdx < Inputs.Num(); InputIdx++)
 		{
-			FString TestName;
+			FName TestName;
 			if (UMaterialExpressionMaterialFunctionCall* FuncCall = Cast<UMaterialExpressionMaterialFunctionCall>(Expression))
 			{
 				// If a function call, don't want to compare string with type postfix
@@ -58,8 +58,8 @@ static FExpressionInput* GetExpressionInputByName(UMaterialExpression* Expressio
 			}
 			else
 			{
-				TestName = Expression->GetInputName(InputIdx);
-				TestName = UMaterialGraphNode::GetShortenPinName(TestName);
+				const FName ExpressionInputName = Expression->GetInputName(InputIdx);
+				TestName = UMaterialGraphNode::GetShortenPinName(ExpressionInputName);
 			}
 
 			if (TestName == InputName)
@@ -73,7 +73,7 @@ static FExpressionInput* GetExpressionInputByName(UMaterialExpression* Expressio
 	return Result;
 }
 
-static int32 GetExpressionOutputIndexByName(UMaterialExpression* Expression, const FString& OutputName)
+static int32 GetExpressionOutputIndexByName(UMaterialExpression* Expression, const FName OutputName)
 {
 	check(Expression);
 	
@@ -84,7 +84,7 @@ static int32 GetExpressionOutputIndexByName(UMaterialExpression* Expression, con
 		// leave as INDEX_NONE
 	}
 	// Return first output if no name specified
-	else if (OutputName.IsEmpty())
+	else if (OutputName.IsNone())
 	{
 		Result = 0;
 	}
@@ -97,7 +97,7 @@ static int32 GetExpressionOutputIndexByName(UMaterialExpression* Expression, con
 
 			FExpressionOutput& Output = Expression->Outputs[OutIdx];
 			// If output name is no empty - see if it matches
-			if(!Output.OutputName.IsEmpty())
+			if(!Output.OutputName.IsNone())
 			{
 				if (OutputName == Output.OutputName)
 				{
@@ -107,19 +107,19 @@ static int32 GetExpressionOutputIndexByName(UMaterialExpression* Expression, con
 			// if it is empty we look for R/G/B/A
 			else
 			{
-				if (OutputName == TEXT("R") && Output.MaskR && !Output.MaskG && !Output.MaskB && !Output.MaskA)
+				if (Output.MaskR && !Output.MaskG && !Output.MaskB && !Output.MaskA && OutputName == TEXT("R"))
 				{
 					bFoundMatch = true;
 				}
-				else if (OutputName == TEXT("G") && !Output.MaskR && Output.MaskG && !Output.MaskB && !Output.MaskA)
+				else if (!Output.MaskR && Output.MaskG && !Output.MaskB && !Output.MaskA && OutputName == TEXT("G"))
 				{
 					bFoundMatch = true;
 				}
-				else if (OutputName == TEXT("B") && !Output.MaskR && !Output.MaskG && Output.MaskB && !Output.MaskA)
+				else if (!Output.MaskR && !Output.MaskG && Output.MaskB && !Output.MaskA && OutputName == TEXT("B"))
 				{
 					bFoundMatch = true;
 				}
-				else if (OutputName == TEXT("A") && !Output.MaskR && !Output.MaskG && !Output.MaskB && Output.MaskA)
+				else if (!Output.MaskR && !Output.MaskG && !Output.MaskB && Output.MaskA && OutputName == TEXT("A"))
 				{
 					bFoundMatch = true;
 				}
@@ -456,7 +456,7 @@ bool UMaterialEditingLibrary::ConnectMaterialProperty(UMaterialExpression* FromE
 		if (Material)
 		{
 			FExpressionInput* Input = Material->GetExpressionInputForProperty(Property);
-			int32 FromIndex = GetExpressionOutputIndexByName(FromExpression, FromOutputName);
+			int32 FromIndex = GetExpressionOutputIndexByName(FromExpression, *FromOutputName);
 			if (Input && FromIndex != INDEX_NONE)
 			{
 				Input->Connect(FromIndex, FromExpression);
@@ -472,8 +472,8 @@ bool UMaterialEditingLibrary::ConnectMaterialExpressions(UMaterialExpression* Fr
 	bool bResult = false;
 	if (FromExpression && ToExpression)
 	{
-		FExpressionInput* Input = GetExpressionInputByName(ToExpression, ToInputName);
-		int32 FromIndex = GetExpressionOutputIndexByName(FromExpression, FromOutputName);
+		FExpressionInput* Input = GetExpressionInputByName(ToExpression, *ToInputName);
+		int32 FromIndex = GetExpressionOutputIndexByName(FromExpression, *FromOutputName);
 		if (Input && FromIndex != INDEX_NONE)
 		{
 			Input->Connect(FromIndex, FromExpression);
diff --git a/Engine/Source/Editor/MaterialEditor/Private/MaterialEditor.cpp b/Engine/Source/Editor/MaterialEditor/Private/MaterialEditor.cpp
--- a/Engine/Source/Editor/MaterialEditor/Private/MaterialEditor.cpp
+++ b/Engine/Source/Editor/MaterialEditor/Private/MaterialEditor.cpp
@@ -2521,7 +2521,7 @@ UClass* FMaterialEditor::GetOnPromoteToParameterClass(UEdGraphPin* TargetPin)
 
 	if (RootPinNode != nullptr)
 	{
-		EMaterialProperty propertyId = (EMaterialProperty)FCString::Atoi(*TargetPin->PinType.PinSubCategory);
+		EMaterialProperty propertyId = (EMaterialProperty)FCString::Atoi(*TargetPin->PinType.PinSubCategory.ToString());
 
 		switch (propertyId)
 		{
@@ -2549,12 +2549,12 @@ UClass* FMaterialEditor::GetOnPromoteToParameterClass(UEdGraphPin* TargetPin)
 	else if (OtherPinNode)
 	{
 		const TArray<FExpressionInput*> ExpressionInputs = OtherPinNode->MaterialExpression->GetInputs();
-		FString TargetPinName = OtherPinNode->GetShortenPinName(TargetPin->PinName);
+		FName TargetPinName = OtherPinNode->GetShortenPinName(TargetPin->PinName);
 
 		for (int32 Index = 0; Index < ExpressionInputs.Num(); ++Index)
 		{
 			FExpressionInput* Input = ExpressionInputs[Index];
-			FString InputName = OtherPinNode->MaterialExpression->GetInputName(Index);
+			FName InputName = OtherPinNode->MaterialExpression->GetInputName(Index);
 			InputName = OtherPinNode->GetShortenPinName(InputName);
 
 			if (InputName == TargetPinName)
@@ -2609,7 +2609,7 @@ void FMaterialEditor::OnPromoteToParameter()
 
 		if (MaterialNode->MaterialExpression->HasAParameterName())
 		{
-			MaterialNode->MaterialExpression->SetParameterName(FName(*TargetPin->PinName));
+			MaterialNode->MaterialExpression->SetParameterName(TargetPin->PinName);
 			MaterialNode->MaterialExpression->ValidateParameterName();
 		}
 	}
@@ -4106,7 +4106,7 @@ bool FMaterialEditor::CheckExpressionRemovalWarnings(const TArray<UEdGraphNode*>
 					FunctionWarningString += TEXT(", ");
 				}
 				bFirstExpression = false;
-				FunctionWarningString += FunctionInput->InputName;
+				FunctionWarningString += FunctionInput->InputName.ToString();
 			}
 
 			if (FunctionOutput)
@@ -4116,7 +4116,7 @@ bool FMaterialEditor::CheckExpressionRemovalWarnings(const TArray<UEdGraphNode*>
 					FunctionWarningString += TEXT(", ");
 				}
 				bFirstExpression = false;
-				FunctionWarningString += FunctionOutput->OutputName;
+				FunctionWarningString += FunctionOutput->OutputName.ToString();
 			}
 		}
 	}
diff --git a/Engine/Source/Editor/MovieSceneTools/Private/K2Node_GetSequenceBinding.cpp b/Engine/Source/Editor/MovieSceneTools/Private/K2Node_GetSequenceBinding.cpp
--- a/Engine/Source/Editor/MovieSceneTools/Private/K2Node_GetSequenceBinding.cpp
+++ b/Engine/Source/Editor/MovieSceneTools/Private/K2Node_GetSequenceBinding.cpp
@@ -22,8 +22,8 @@
 #include "SComboBox.h"
 #include "Editor.h"
 
-static const FString OutputPinName(TEXT("Output"));
-static const FString SequencePinName(TEXT("Sequence"));
+static const FName OutputPinName(TEXT("Output"));
+static const FName SequencePinName(TEXT("Sequence"));
 
 #define LOCTEXT_NAMESPACE "UK2Node_GetSequenceBinding"
 
@@ -116,12 +116,10 @@ void UK2Node_GetSequenceBinding::ValidateNodeDuringCompilation(FCompilerResultsL
 
 void UK2Node_GetSequenceBinding::AllocateDefaultPins()
 {
-	const UEdGraphSchema_K2* K2Schema = GetDefault<UEdGraphSchema_K2>();
-
-	CreatePin(EGPD_Input, K2Schema->PC_Object, FString(), UMovieSceneSequence::StaticClass(), SequencePinName);
+	CreatePin(EGPD_Input, UEdGraphSchema_K2::PC_Object, UMovieSceneSequence::StaticClass(), SequencePinName);
 
 	// Result pin
-	UEdGraphPin* ResultPin = CreatePin(EGPD_Output, K2Schema->PC_Struct, FString(), FMovieSceneObjectBindingID::StaticStruct(), K2Schema->PN_ReturnValue);
+	UEdGraphPin* ResultPin = CreatePin(EGPD_Output, UEdGraphSchema_K2::PC_Struct, FMovieSceneObjectBindingID::StaticStruct(), UEdGraphSchema_K2::PN_ReturnValue);
 	ResultPin->PinFriendlyName = LOCTEXT("SequenceBindingOutput", "Binding");
 
 	Super::AllocateDefaultPins();
diff --git a/Engine/Source/Editor/MovieSceneTools/Private/K2Node_GetSequenceBindings.cpp b/Engine/Source/Editor/MovieSceneTools/Private/K2Node_GetSequenceBindings.cpp
--- a/Engine/Source/Editor/MovieSceneTools/Private/K2Node_GetSequenceBindings.cpp
+++ b/Engine/Source/Editor/MovieSceneTools/Private/K2Node_GetSequenceBindings.cpp
@@ -138,8 +138,8 @@ void UDEPRECATED_K2Node_GetSequenceBindings::UpdatePins()
 				continue;
 			}
 
-			FString GuidString = Possessable.GetGuid().ToString();
-			UEdGraphPin* NewPin = CreatePin(EGPD_Output, UEdGraphSchema_K2::PC_Struct, FString(), FMovieSceneObjectBindingID::StaticStruct(), GuidString);
+			const FName GuidString = *Possessable.GetGuid().ToString();
+			UEdGraphPin* NewPin = CreatePin(EGPD_Output, UEdGraphSchema_K2::PC_Struct, FMovieSceneObjectBindingID::StaticStruct(), GuidString);
 			NewPin->PinFriendlyName = MovieScene->GetObjectDisplayName(Possessable.GetGuid());
 			NewPin->PersistentGuid = Possessable.GetGuid();
 
diff --git a/Engine/Source/Editor/MovieSceneTools/Private/K2Node_GetSequenceBindings.h b/Engine/Source/Editor/MovieSceneTools/Private/K2Node_GetSequenceBindings.h
--- a/Engine/Source/Editor/MovieSceneTools/Private/K2Node_GetSequenceBindings.h
+++ b/Engine/Source/Editor/MovieSceneTools/Private/K2Node_GetSequenceBindings.h
@@ -20,7 +20,7 @@ struct FGetSequenceBindingGuidMapping
 	GENERATED_BODY()
 
 	UPROPERTY()
-	FString PinName;
+	FName PinName;
 
 	UPROPERTY()
 	FGuid Guid;
diff --git a/Engine/Source/Editor/PhysicsAssetEditor/Private/PhysicsAssetGraph/PhysicsAssetGraphNode.cpp b/Engine/Source/Editor/PhysicsAssetEditor/Private/PhysicsAssetGraph/PhysicsAssetGraphNode.cpp
--- a/Engine/Source/Editor/PhysicsAssetEditor/Private/PhysicsAssetGraph/PhysicsAssetGraphNode.cpp
+++ b/Engine/Source/Editor/PhysicsAssetEditor/Private/PhysicsAssetGraph/PhysicsAssetGraphNode.cpp
@@ -1,4 +1,4 @@
-// Copyright 1998-2016 Epic Games, Inc. All Rights Reserved.
+﻿// Copyright 1998-2016 Epic Games, Inc. All Rights Reserved.
 
 #include "PhysicsAssetGraphNode.h"
 #include "PhysicsAssetGraph.h"
@@ -33,9 +33,9 @@ FText UPhysicsAssetGraphNode::GetNodeTitle(ENodeTitleType::Type TitleType) const
 
 void UPhysicsAssetGraphNode::AllocateDefaultPins()
 {
-	InputPin = CreatePin(EEdGraphPinDirection::EGPD_Input, TEXT(""), TEXT(""), nullptr, TEXT(""));
+	InputPin = CreatePin(EEdGraphPinDirection::EGPD_Input, NAME_None, NAME_None, NAME_None);
 	InputPin->bHidden = true;
-	OutputPin = CreatePin(EEdGraphPinDirection::EGPD_Output, TEXT(""), TEXT(""), nullptr, TEXT(""));
+	OutputPin = CreatePin(EEdGraphPinDirection::EGPD_Output, NAME_None, NAME_None, NAME_None);
 	OutputPin->bHidden = true;
 }
 
diff --git a/Engine/Source/Editor/PhysicsAssetEditor/Private/PhysicsAssetGraph/PhysicsAssetGraphSchema.cpp b/Engine/Source/Editor/PhysicsAssetEditor/Private/PhysicsAssetGraph/PhysicsAssetGraphSchema.cpp
--- a/Engine/Source/Editor/PhysicsAssetEditor/Private/PhysicsAssetGraph/PhysicsAssetGraphSchema.cpp
+++ b/Engine/Source/Editor/PhysicsAssetEditor/Private/PhysicsAssetGraph/PhysicsAssetGraphSchema.cpp
@@ -47,7 +47,7 @@ void UPhysicsAssetGraphSchema::BreakPinLinks(UEdGraphPin& TargetPin, bool bSends
 	// Don't allow breaking any links
 }
 
-void UPhysicsAssetGraphSchema::BreakSinglePinLink(UEdGraphPin* SourcePin, UEdGraphPin* TargetPin)
+void UPhysicsAssetGraphSchema::BreakSinglePinLink(UEdGraphPin* SourcePin, UEdGraphPin* TargetPin) const
 {
 	// Don't allow breaking any links
 }
diff --git a/Engine/Source/Editor/PropertyEditor/Private/UserInterface/PropertyEditor/SPropertyEditorAsset.cpp b/Engine/Source/Editor/PropertyEditor/Private/UserInterface/PropertyEditor/SPropertyEditorAsset.cpp
--- a/Engine/Source/Editor/PropertyEditor/Private/UserInterface/PropertyEditor/SPropertyEditorAsset.cpp
+++ b/Engine/Source/Editor/PropertyEditor/Private/UserInterface/PropertyEditor/SPropertyEditorAsset.cpp
@@ -506,6 +506,12 @@ SPropertyEditorAsset::EActorReferenceState SPropertyEditorAsset::GetActorReferen
 
 		if (Value.Object != nullptr)
 		{
+			// If this is not an actual actor, this is broken
+			if (!Value.Object->IsA(AActor::StaticClass()))
+			{
+				return EActorReferenceState::Error;
+			}
+
 			return EActorReferenceState::Loaded;
 		}
 		else if (Value.ObjectPath.IsNull())
@@ -519,7 +525,7 @@ SPropertyEditorAsset::EActorReferenceState SPropertyEditorAsset::GetActorReferen
 
 			if (MapObjectPath.ResolveObject())
 			{
-				// If the map is valid but the 
+				// If the map is valid but the object is not
 				return EActorReferenceState::Error;
 			}
 
@@ -627,8 +633,16 @@ FText SPropertyEditorAsset::OnGetAssetName() const
 		{
 			if( bIsActor )
 			{
-				AActor* Actor = CastChecked<AActor>(Value.Object);
-				Name = FText::AsCultureInvariant(Actor->GetActorLabel());
+				AActor* Actor = Cast<AActor>(Value.Object);
+
+				if (Actor)
+				{
+					Name = FText::AsCultureInvariant(Actor->GetActorLabel());
+				}
+				else
+				{
+					Name = FText::AsCultureInvariant(Value.Object->GetName());
+				}
 			}
 			else if (UField* AsField = Cast<UField>(Value.Object))
 			{
diff --git a/Engine/Source/Editor/ReferenceViewer/Private/EdGraphNode_Reference.cpp b/Engine/Source/Editor/ReferenceViewer/Private/EdGraphNode_Reference.cpp
--- a/Engine/Source/Editor/ReferenceViewer/Private/EdGraphNode_Reference.cpp
+++ b/Engine/Source/Editor/ReferenceViewer/Private/EdGraphNode_Reference.cpp
@@ -165,8 +165,8 @@ FText UEdGraphNode_Reference::GetTooltipText() const
 
 void UEdGraphNode_Reference::AllocateDefaultPins()
 {
-	ReferencerPin = CreatePin( EEdGraphPinDirection::EGPD_Input, FString(), FString(), nullptr, FString() );
-	DependencyPin = CreatePin( EEdGraphPinDirection::EGPD_Output, FString(), FString(), nullptr, FString() );
+	ReferencerPin = CreatePin( EEdGraphPinDirection::EGPD_Input, NAME_None, NAME_None);
+	DependencyPin = CreatePin( EEdGraphPinDirection::EGPD_Output, NAME_None, NAME_None);
 
 	ReferencerPin->bHidden = true;
 	DependencyPin->bHidden = true;
diff --git a/Engine/Source/Editor/ReferenceViewer/Private/ReferenceViewerSchema.cpp b/Engine/Source/Editor/ReferenceViewer/Private/ReferenceViewerSchema.cpp
--- a/Engine/Source/Editor/ReferenceViewer/Private/ReferenceViewerSchema.cpp
+++ b/Engine/Source/Editor/ReferenceViewer/Private/ReferenceViewerSchema.cpp
@@ -100,7 +100,7 @@ void UReferenceViewerSchema::BreakPinLinks(UEdGraphPin& TargetPin, bool bSendsNo
 	// Don't allow breaking any links
 }
 
-void UReferenceViewerSchema::BreakSinglePinLink(UEdGraphPin* SourcePin, UEdGraphPin* TargetPin)
+void UReferenceViewerSchema::BreakSinglePinLink(UEdGraphPin* SourcePin, UEdGraphPin* TargetPin) const
 {
 	// Don't allow breaking any links
 }
diff --git a/Engine/Source/Editor/StaticMeshEditor/Private/StaticMeshEditor.h b/Engine/Source/Editor/StaticMeshEditor/Private/StaticMeshEditor.h
--- a/Engine/Source/Editor/StaticMeshEditor/Private/StaticMeshEditor.h
+++ b/Engine/Source/Editor/StaticMeshEditor/Private/StaticMeshEditor.h
@@ -169,9 +169,9 @@ public:
 	}
 
 	/** This is a tickable class */
-	virtual bool IsTickable() const final
+	virtual ETickableTickType GetTickableTickType() const final
 	{
-		return true;
+		return ETickableTickType::Always;
 	}
 
 	/** Performs the main 'tick' operation on this class.  The 'tick' step checks to see if there is currently
diff --git a/Engine/Source/Editor/UMGEditor/Private/Nodes/K2Node_CreateDragDropOperation.cpp b/Engine/Source/Editor/UMGEditor/Private/Nodes/K2Node_CreateDragDropOperation.cpp
--- a/Engine/Source/Editor/UMGEditor/Private/Nodes/K2Node_CreateDragDropOperation.cpp
+++ b/Engine/Source/Editor/UMGEditor/Private/Nodes/K2Node_CreateDragDropOperation.cpp
@@ -65,7 +65,7 @@ void UK2Node_CreateDragDropOperation::ExpandNode(class FKismetCompilerContext& C
 	Super::ExpandNode(CompilerContext, SourceGraph);
 
 	static FName Create_FunctionName = GET_FUNCTION_NAME_CHECKED(UWidgetBlueprintLibrary, CreateDragDropOperation);
-	static FString OperationClass_ParamName = FString(TEXT("OperationClass"));
+	static FName OperationClass_ParamName(TEXT("OperationClass"));
 
 	UK2Node_CreateDragDropOperation* CreateOpNode = this;
 	UEdGraphPin* SpawnNodeExec = CreateOpNode->GetExecPin();
diff --git a/Engine/Source/Editor/UMGEditor/Private/Nodes/K2Node_CreateWidget.cpp b/Engine/Source/Editor/UMGEditor/Private/Nodes/K2Node_CreateWidget.cpp
--- a/Engine/Source/Editor/UMGEditor/Private/Nodes/K2Node_CreateWidget.cpp
+++ b/Engine/Source/Editor/UMGEditor/Private/Nodes/K2Node_CreateWidget.cpp
@@ -19,10 +19,10 @@
 
 struct FK2Node_CreateWidgetHelper
 {
-	static FString OwningPlayerPinName;
+	static const FName OwningPlayerPinName;
 };
 
-FString FK2Node_CreateWidgetHelper::OwningPlayerPinName(TEXT("OwningPlayer"));
+const FName FK2Node_CreateWidgetHelper::OwningPlayerPinName(TEXT("OwningPlayer"));
 
 UK2Node_CreateWidget::UK2Node_CreateWidget(const FObjectInitializer& ObjectInitializer)
 	: Super(ObjectInitializer)
@@ -34,18 +34,11 @@ void UK2Node_CreateWidget::AllocateDefaultPins()
 {
 	Super::AllocateDefaultPins();
 
-	const UEdGraphSchema_K2* K2Schema = GetDefault<UEdGraphSchema_K2>();
-
 	// OwningPlayer pin
-	UEdGraphPin* OwningPlayerPin = CreatePin(EGPD_Input, K2Schema->PC_Object, FString(), APlayerController::StaticClass(), FK2Node_CreateWidgetHelper::OwningPlayerPinName);
+	UEdGraphPin* OwningPlayerPin = CreatePin(EGPD_Input, UEdGraphSchema_K2::PC_Object, APlayerController::StaticClass(), FK2Node_CreateWidgetHelper::OwningPlayerPinName);
 	SetPinToolTip(*OwningPlayerPin, LOCTEXT("OwningPlayerPinDescription", "The player that 'owns' the widget."));
 }
 
-FLinearColor UK2Node_CreateWidget::GetNodeTitleColor() const
-{
-	return Super::GetNodeTitleColor();
-}
-
 FText UK2Node_CreateWidget::GetBaseNodeTitle() const
 {
 	return LOCTEXT("CreateWidget_BaseTitle", "Create Widget");
@@ -88,10 +81,10 @@ void UK2Node_CreateWidget::ExpandNode(class FKismetCompilerContext& CompilerCont
 {
 	Super::ExpandNode(CompilerContext, SourceGraph);
 
-	static FName Create_FunctionName = GET_FUNCTION_NAME_CHECKED(UWidgetBlueprintLibrary, Create);
-	static FString WorldContextObject_ParamName = FString(TEXT("WorldContextObject"));
-	static FString WidgetType_ParamName = FString(TEXT("WidgetType"));
-	static FString OwningPlayer_ParamName = FString(TEXT("OwningPlayer"));
+	static const FName Create_FunctionName = GET_FUNCTION_NAME_CHECKED(UWidgetBlueprintLibrary, Create);
+	static const FName WorldContextObject_ParamName(TEXT("WorldContextObject"));
+	static const FName WidgetType_ParamName(TEXT("WidgetType"));
+	static const FName OwningPlayer_ParamName(TEXT("OwningPlayer"));
 
 	UK2Node_CreateWidget* CreateWidgetNode = this;
 	UEdGraphPin* SpawnNodeExec = CreateWidgetNode->GetExecPin();
diff --git a/Engine/Source/Editor/UMGEditor/Private/Nodes/K2Node_CreateWidget.h b/Engine/Source/Editor/UMGEditor/Private/Nodes/K2Node_CreateWidget.h
--- a/Engine/Source/Editor/UMGEditor/Private/Nodes/K2Node_CreateWidget.h
+++ b/Engine/Source/Editor/UMGEditor/Private/Nodes/K2Node_CreateWidget.h
@@ -16,7 +16,6 @@ class UMGEDITOR_API UK2Node_CreateWidget : public UK2Node_ConstructObjectFromCla
 
 	//~ Begin UEdGraphNode Interface.
 	virtual void AllocateDefaultPins() override;
-	virtual FLinearColor GetNodeTitleColor() const override;
 	virtual void ExpandNode(class FKismetCompilerContext& CompilerContext, UEdGraph* SourceGraph) override;
 	//~ End UEdGraphNode Interface.
 
diff --git a/Engine/Source/Editor/UMGEditor/Private/UMGEditorModule.cpp b/Engine/Source/Editor/UMGEditor/Private/UMGEditorModule.cpp
--- a/Engine/Source/Editor/UMGEditor/Private/UMGEditorModule.cpp
+++ b/Engine/Source/Editor/UMGEditor/Private/UMGEditorModule.cpp
@@ -31,6 +31,8 @@
 #include "ISettingsModule.h"
 #include "SequencerSettings.h"
 
+#include "BlueprintEditorModule.h"
+
 #define LOCTEXT_NAMESPACE "UMG"
 
 const FName UMGEditorAppIdentifier = FName(TEXT("UMGEditorApp"));
@@ -75,6 +77,8 @@ public:
 		IAssetTools& AssetTools = FModuleManager::LoadModuleChecked<FAssetToolsModule>("AssetTools").Get();
 		RegisterAssetTypeAction(AssetTools, MakeShareable(new FAssetTypeActions_WidgetBlueprint()));
 
+		FBlueprintEditorModule::RegisterCompilerGetter(UWidgetBlueprint::StaticClass(), FOnGetBlueprintCompiler::CreateStatic(&UWidgetBlueprint::GetCompilerForWidgetBP));
+
 		// Register with the sequencer module that we provide auto-key handlers.
 		ISequencerModule& SequencerModule = FModuleManager::Get().LoadModuleChecked<ISequencerModule>("Sequencer");
 		MarginTrackEditorCreateTrackEditorHandle          = SequencerModule.RegisterPropertyTrackEditor<FMarginTrackEditor>();
@@ -101,6 +105,8 @@ public:
 		}
 		CreatedAssetTypeActions.Empty();
 
+		FBlueprintEditorModule::UnregisterCompilerGetter(UWidgetBlueprint::StaticClass());
+
 		// Unregister sequencer track creation delegates
 		ISequencerModule* SequencerModule = FModuleManager::GetModulePtr<ISequencerModule>( "Sequencer" );
 		if ( SequencerModule != nullptr )
diff --git a/Engine/Source/Editor/UMGEditor/Private/WidgetBlueprint.cpp b/Engine/Source/Editor/UMGEditor/Private/WidgetBlueprint.cpp
--- a/Engine/Source/Editor/UMGEditor/Private/WidgetBlueprint.cpp
+++ b/Engine/Source/Editor/UMGEditor/Private/WidgetBlueprint.cpp
@@ -699,9 +699,9 @@ bool UWidgetBlueprint::ValidateGeneratedClass(const UClass* InClass)
 	return Result;
 }
 
-TSharedPtr<FKismetCompilerContext> UWidgetBlueprint::GetCompilerForWidgetBP(UWidgetBlueprint* BP, FCompilerResultsLog& InMessageLog, const FKismetCompilerOptions& InCompileOptions)
+TSharedPtr<FKismetCompilerContext> UWidgetBlueprint::GetCompilerForWidgetBP(UBlueprint* BP, FCompilerResultsLog& InMessageLog, const FKismetCompilerOptions& InCompileOptions)
 {
-	return TSharedPtr<FKismetCompilerContext>(new FWidgetBlueprintCompiler(BP, InMessageLog, InCompileOptions, nullptr));
+	return TSharedPtr<FKismetCompilerContext>(new FWidgetBlueprintCompiler(CastChecked<UWidgetBlueprint>(BP), InMessageLog, InCompileOptions, nullptr));
 }
 
 void UWidgetBlueprint::GetReparentingRules(TSet< const UClass* >& AllowedChildrenOfClasses, TSet< const UClass* >& DisallowedChildrenOfClasses) const
@@ -795,4 +795,11 @@ void UWidgetBlueprint::ForEachSourceWidgetImpl (TFunctionRef<void(UWidget*)> Fn)
 	);
 }
 
+#if WITH_EDITOR
+void UWidgetBlueprint::LoadModulesRequiredForCompilation()
+{
+	static const FName ModuleName(TEXT("UMGEditor"));
+	FModuleManager::Get().LoadModule(ModuleName);
+}
+#endif // WITH_EDITOR
 #undef LOCTEXT_NAMESPACE 
diff --git a/Engine/Source/Editor/UMGEditor/Private/WidgetBlueprintCompiler.cpp b/Engine/Source/Editor/UMGEditor/Private/WidgetBlueprintCompiler.cpp
--- a/Engine/Source/Editor/UMGEditor/Private/WidgetBlueprintCompiler.cpp
+++ b/Engine/Source/Editor/UMGEditor/Private/WidgetBlueprintCompiler.cpp
@@ -270,7 +270,7 @@ void FWidgetBlueprintCompiler::CreateClassVariablesFromBlueprint()
 			continue;
 		}
 
-		FEdGraphPinType WidgetPinType(Schema->PC_Object, FString(), WidgetClass, EPinContainerType::None, false, FEdGraphTerminalType());
+		FEdGraphPinType WidgetPinType(UEdGraphSchema_K2::PC_Object, NAME_None, WidgetClass, EPinContainerType::None, false, FEdGraphTerminalType());
 		
 		// Always name the variable according to the underlying FName of the widget object
 		UProperty* WidgetProperty = CreateVariable(Widget->GetFName(), WidgetPinType);
@@ -298,7 +298,7 @@ void FWidgetBlueprintCompiler::CreateClassVariablesFromBlueprint()
 	// Add movie scenes variables here
 	for(UWidgetAnimation* Animation : WidgetBP->Animations)
 	{
-		FEdGraphPinType WidgetPinType(Schema->PC_Object, FString(), Animation->GetClass(), EPinContainerType::None, true, FEdGraphTerminalType());
+		FEdGraphPinType WidgetPinType(UEdGraphSchema_K2::PC_Object, NAME_None, Animation->GetClass(), EPinContainerType::None, true, FEdGraphTerminalType());
 		UProperty* AnimationProperty = CreateVariable(Animation->GetFName(), WidgetPinType);
 
 		if ( AnimationProperty != nullptr )
@@ -616,7 +616,7 @@ void FWidgetBlueprintCompiler::VerifyEventReplysAreNotEmpty(FKismetFunctionConte
 	Context.SourceGraph->GetNodesOfClass<UK2Node_FunctionResult>(FunctionResults);
 
 	UScriptStruct* EventReplyStruct = FEventReply::StaticStruct();
-	FEdGraphPinType EventReplyPinType(Schema->PC_Struct, FString(), EventReplyStruct, EPinContainerType::None, /*bIsReference =*/false, /*InValueTerminalType =*/FEdGraphTerminalType());
+	FEdGraphPinType EventReplyPinType(UEdGraphSchema_K2::PC_Struct, NAME_None, EventReplyStruct, EPinContainerType::None, /*bIsReference =*/false, /*InValueTerminalType =*/FEdGraphTerminalType());
 
 	for ( UK2Node_FunctionResult* FunctionResult : FunctionResults )
 	{
diff --git a/Engine/Source/Editor/UMGEditor/Public/WidgetBlueprint.h b/Engine/Source/Editor/UMGEditor/Public/WidgetBlueprint.h
--- a/Engine/Source/Editor/UMGEditor/Public/WidgetBlueprint.h
+++ b/Engine/Source/Editor/UMGEditor/Public/WidgetBlueprint.h
@@ -256,8 +256,12 @@ public:
 
 	static bool ValidateGeneratedClass(const UClass* InClass);
 	
-	static TSharedPtr<FKismetCompilerContext> GetCompilerForWidgetBP(UWidgetBlueprint* BP, FCompilerResultsLog& InMessageLog, const FKismetCompilerOptions& InCompileOptions);
+	static TSharedPtr<FKismetCompilerContext> GetCompilerForWidgetBP(UBlueprint* BP, FCompilerResultsLog& InMessageLog, const FKismetCompilerOptions& InCompileOptions);
 
 private:
 	void ForEachSourceWidgetImpl(TFunctionRef<void(UWidget*)> Fn) const;
+
+#if WITH_EDITOR
+	virtual void LoadModulesRequiredForCompilation() override;
+#endif 
 };
diff --git a/Engine/Source/Editor/UnrealEd/Classes/CookOnTheSide/CookOnTheFlyServer.h b/Engine/Source/Editor/UnrealEd/Classes/CookOnTheSide/CookOnTheFlyServer.h
--- a/Engine/Source/Editor/UnrealEd/Classes/CookOnTheSide/CookOnTheFlyServer.h
+++ b/Engine/Source/Editor/UnrealEd/Classes/CookOnTheSide/CookOnTheFlyServer.h
@@ -1088,8 +1088,6 @@ public:
 		bool bGenerateDependenciesForMaps; 
 		bool bErrorOnEngineContentUse; // this is a flag for dlc, will cause the cooker to error if the dlc references engine content
 		int32 NumProcesses;
-		bool bNativizeAssets;
-		FString NativizedPluginPath; // path expressing the desired .uplugin path (optional); will only be utilized if bNativizeAssets is set
 		FCookByTheBookStartupOptions() :
 			CookOptions(ECookByTheBookOptions::None),
 			DLCName(FString()),
@@ -1097,8 +1095,7 @@ public:
 			bGenerateStreamingInstallManifests(false),
 			bGenerateDependenciesForMaps(false),
 			bErrorOnEngineContentUse(false),
-			NumProcesses(0),
-			bNativizeAssets(false)
+			NumProcesses(0)
 		{ }
 	};
 
diff --git a/Engine/Source/Editor/UnrealEd/Classes/MaterialGraph/MaterialGraphNode.h b/Engine/Source/Editor/UnrealEd/Classes/MaterialGraph/MaterialGraphNode.h
--- a/Engine/Source/Editor/UnrealEd/Classes/MaterialGraph/MaterialGraphNode.h
+++ b/Engine/Source/Editor/UnrealEd/Classes/MaterialGraph/MaterialGraphNode.h
@@ -82,7 +82,7 @@ public:
 	//~ End UMaterialGraphNode_Base Interface
 
 	/** Will return the shorten pin name to use based on long pin name */
-	static FString UNREALED_API GetShortenPinName(const FString& PinName);
+	static FName UNREALED_API GetShortenPinName(const FName PinName);
 
 private:
 	/** Make sure the MaterialExpression is owned by the Material */
diff --git a/Engine/Source/Editor/UnrealEd/Classes/MaterialGraph/MaterialGraphSchema.h b/Engine/Source/Editor/UnrealEd/Classes/MaterialGraph/MaterialGraphSchema.h
--- a/Engine/Source/Editor/UnrealEd/Classes/MaterialGraph/MaterialGraphSchema.h
+++ b/Engine/Source/Editor/UnrealEd/Classes/MaterialGraph/MaterialGraphSchema.h
@@ -125,40 +125,21 @@ class UMaterialGraphSchema : public UEdGraphSchema
 	GENERATED_UCLASS_BODY()
 
 	// Allowable PinType.PinCategory values
-	UPROPERTY()
-	FString PC_Mask;
-
-	UPROPERTY()
-	FString PC_Required;
-
-	UPROPERTY()
-	FString PC_Optional;
-
-	UPROPERTY()
-	FString PC_MaterialInput;
+	UNREALED_API static const FName PC_Mask;
+	UNREALED_API static const FName PC_Required;
+	UNREALED_API static const FName PC_Optional;
+	UNREALED_API static const FName PC_MaterialInput;
 
 	// Common PinType.PinSubCategory values
-	UPROPERTY()
-	FString PSC_Red;
-
-	UPROPERTY()
-	FString PSC_Green;
-
-	UPROPERTY()
-	FString PSC_Blue;
-
-	UPROPERTY()
-	FString PSC_Alpha;
+	UNREALED_API static const FName PSC_Red;
+	UNREALED_API static const FName PSC_Green;
+	UNREALED_API static const FName PSC_Blue;
+	UNREALED_API static const FName PSC_Alpha;
 
 	// Color of certain pins/connections
-	UPROPERTY()
-	FLinearColor ActivePinColor;
-
-	UPROPERTY()
-	FLinearColor InactivePinColor;
-
-	UPROPERTY()
-	FLinearColor AlphaPinColor;
+	UNREALED_API static const FLinearColor ActivePinColor;
+	UNREALED_API static const FLinearColor InactivePinColor;
+	UNREALED_API static const FLinearColor AlphaPinColor;
 
 	/**
 	 *  Add all linked to nodes to this pin to selection
@@ -212,7 +193,7 @@ class UMaterialGraphSchema : public UEdGraphSchema
 	virtual FLinearColor GetPinTypeColor(const FEdGraphPinType& PinType) const override;
 	virtual void BreakNodeLinks(UEdGraphNode& TargetNode) const override;
 	virtual void BreakPinLinks(UEdGraphPin& TargetPin, bool bSendsNodeNotifcation) const override;
-	virtual void BreakSinglePinLink(UEdGraphPin* SourcePin, UEdGraphPin* TargetPin) override;
+	virtual void BreakSinglePinLink(UEdGraphPin* SourcePin, UEdGraphPin* TargetPin) const override;
 	virtual void DroppedAssetsOnGraph(const TArray<struct FAssetData>& Assets, const FVector2D& GraphPosition, UEdGraph* Graph) const override;
 	virtual int32 GetNodeSelectionCount(const UEdGraph* Graph) const override;
 	virtual TSharedPtr<FEdGraphSchemaAction> GetCreateCommentAction() const override;
diff --git a/Engine/Source/Editor/UnrealEd/Classes/Settings/ProjectPackagingSettings.h b/Engine/Source/Editor/UnrealEd/Classes/Settings/ProjectPackagingSettings.h
--- a/Engine/Source/Editor/UnrealEd/Classes/Settings/ProjectPackagingSettings.h
+++ b/Engine/Source/Editor/UnrealEd/Classes/Settings/ProjectPackagingSettings.h
@@ -118,9 +118,9 @@ public:
 	UPROPERTY(config, EditAnywhere, AdvancedDisplay, Category = Blueprints, meta = (DisplayName = "List of Blueprint assets to nativize", RelativeToGameContentDir, LongPackageName))
 	TArray<FFilePath> NativizeBlueprintAssets;
 
-	/** If enabled, a warning will be emitted at build/cook time if nativization is turned on in the Project Settings, but the nativization flag was omitted from the command line. */
-	UPROPERTY(config, EditAnywhere, Category = Blueprints)
-	bool bWarnIfPackagedWithoutNativizationFlag;
+	/** If enabled, the nativized assets code plugin will be added to the Visual Studio solution if it exists when regenerating the game project. Intended primarily to assist with debugging the target platform after cooking with nativization turned on. */
+	UPROPERTY(config, EditAnywhere, AdvancedDisplay, Category = Blueprints)
+	bool bIncludeNativizedAssetsInProjectGeneration;
 
 	/** If enabled, all content will be put into a single .pak file instead of many individual files (default = enabled). */
 	UPROPERTY(config, EditAnywhere, Category=Packaging)
diff --git a/Engine/Source/Editor/UnrealEd/Classes/UserDefinedStructure/UserDefinedStructEditorData.h b/Engine/Source/Editor/UnrealEd/Classes/UserDefinedStructure/UserDefinedStructEditorData.h
--- a/Engine/Source/Editor/UnrealEd/Classes/UserDefinedStructure/UserDefinedStructEditorData.h
+++ b/Engine/Source/Editor/UnrealEd/Classes/UserDefinedStructure/UserDefinedStructEditorData.h
@@ -34,10 +34,10 @@ struct FStructVariableDescription
 
 	// TYPE DATA
 	UPROPERTY()
-	FString Category;
+	FName Category;
 
 	UPROPERTY()
-	FString SubCategory;
+	FName SubCategory;
 
 	UPROPERTY()
 	TSoftObjectPtr<UObject> SubCategoryObject;
diff --git a/Engine/Source/Editor/UnrealEd/Private/Commandlets/CookCommandlet.cpp b/Engine/Source/Editor/UnrealEd/Private/Commandlets/CookCommandlet.cpp
--- a/Engine/Source/Editor/UnrealEd/Private/Commandlets/CookCommandlet.cpp
+++ b/Engine/Source/Editor/UnrealEd/Private/Commandlets/CookCommandlet.cpp
@@ -676,9 +676,6 @@ bool UCookCommandlet::CookByTheBook( const TArray<ITargetPlatform*>& Platforms,
 	FString CreateReleaseVersion;
 	FParse::Value( *Params, TEXT("CreateReleaseVersion="), CreateReleaseVersion);
 
-	FString NativizedPluginPath;
-	FParse::Value(*Params, TEXT("NativizeAssets="), NativizedPluginPath);
-
 	FString OutputDirectoryOverride;
 	FParse::Value( *Params, TEXT("OutputDir="), OutputDirectoryOverride);
 
@@ -853,8 +850,6 @@ bool UCookCommandlet::CookByTheBook( const TArray<ITargetPlatform*>& Platforms,
 	StartupOptions.ChildCookFileName = ChildCookFile;
 	StartupOptions.ChildCookIdentifier = ChildCookIdentifier;
 	StartupOptions.NumProcesses = NumProcesses;
-	StartupOptions.bNativizeAssets = !NativizedPluginPath.IsEmpty() || Switches.Contains(TEXT("NativizeAssets"));
-	Swap( StartupOptions.NativizedPluginPath, NativizedPluginPath );
 
 	COOK_STAT(
 	{
diff --git a/Engine/Source/Editor/UnrealEd/Private/Commandlets/DumpBlueprintsInfoCommandlet.cpp b/Engine/Source/Editor/UnrealEd/Private/Commandlets/DumpBlueprintsInfoCommandlet.cpp
--- a/Engine/Source/Editor/UnrealEd/Private/Commandlets/DumpBlueprintsInfoCommandlet.cpp
+++ b/Engine/Source/Editor/UnrealEd/Private/Commandlets/DumpBlueprintsInfoCommandlet.cpp
@@ -201,7 +201,7 @@ DumpBlueprintsInfo commandlet params: \n\
 		uint32     DumpFlags;
 		UClass*    PaletteFilter;
 		EGraphType GraphFilter;
-		FString    PinType;
+		FName      PinType;
 		UClass*    SelectedObjectType;
 		FString    DiffPath;
 		FString    DiffCommand;
@@ -517,7 +517,9 @@ DumpBlueprintInfoUtils::CommandletOptions::CommandletOptions(TArray<FString> con
 			NewDumpFlags &= ~BPDUMP_PinTypeIsClass;
 
 			FString PinSwitch;
-			Switch.Split(TEXT("="), &PinSwitch, &PinType);
+			FString PinTypeStr;
+			Switch.Split(TEXT("="), &PinSwitch, &PinTypeStr);
+			PinType = *PinTypeStr;
 
 			NewDumpFlags |= DumpBlueprintInfoUtils::BPDUMP_PinContextActions;
 			// implies that we want contextual actions
@@ -528,7 +530,9 @@ DumpBlueprintInfoUtils::CommandletOptions::CommandletOptions(TArray<FString> con
 			NewDumpFlags |= BPDUMP_PinTypeIsClass;
 
 			FString PinSwitch;
-			Switch.Split(TEXT("="), &PinSwitch, &PinType);
+			FString PinTypeStr;
+			Switch.Split(TEXT("="), &PinSwitch, &PinTypeStr);
+			PinType = *PinTypeStr;
 			
 			NewDumpFlags |= BPDUMP_PinContextActions;
 			// implies that we want contextual actions
@@ -1879,16 +1883,15 @@ static bool DumpBlueprintInfoUtils::DumpPinContextActions(uint32 Indent, UEdGrap
 	FGraphContextMenuBuilder ContextMenuBuilder(Graph);
 
 	bool bWroteToFile = false;
-	if (!CommandOptions.PinType.IsEmpty())
+	if (!CommandOptions.PinType.IsNone())
 	{
 		FEdGraphPinType PinType;
 		PinType.PinCategory = CommandOptions.PinType;
 
-		UEdGraphSchema_K2 const* K2Schema = GetDefault<UEdGraphSchema_K2>();
 		bool const bUsePinTypeClass = (CommandOptions.DumpFlags & BPDUMP_PinTypeIsClass) != 0;
 
 		bool bIsValidPinType = true;	
-		if (UClass* TypeClass = FindObject<UClass>(ANY_PACKAGE, *CommandOptions.PinType))
+		if (UClass* TypeClass = FindObject<UClass>(ANY_PACKAGE, *CommandOptions.PinType.ToString()))
 		{
 			bIsValidPinType = UEdGraphSchema_K2::IsAllowableBlueprintVariableType(TypeClass);
 			if (bIsValidPinType)
@@ -1896,23 +1899,23 @@ static bool DumpBlueprintInfoUtils::DumpPinContextActions(uint32 Indent, UEdGrap
 				PinType.PinSubCategoryObject = TypeClass;
 				if (TypeClass->IsChildOf(UInterface::StaticClass()))
 				{
-					PinType.PinCategory = K2Schema->PC_Interface;
+					PinType.PinCategory = UEdGraphSchema_K2::PC_Interface;
 				}
 				else
 				{
-					PinType.PinCategory = K2Schema->PC_Object;
+					PinType.PinCategory = UEdGraphSchema_K2::PC_Object;
 				}
 
 				if (bUsePinTypeClass)
 				{
-					PinType.PinCategory = K2Schema->PC_Class;
+					PinType.PinCategory = UEdGraphSchema_K2::PC_Class;
 				}
 			}
 			
 		}
-		else if (UScriptStruct* StructType = FindObject<UScriptStruct>(ANY_PACKAGE, *CommandOptions.PinType))
+		else if (UScriptStruct* StructType = FindObject<UScriptStruct>(ANY_PACKAGE, *CommandOptions.PinType.ToString()))
 		{
-			PinType.PinCategory = K2Schema->PC_Struct;
+			PinType.PinCategory = UEdGraphSchema_K2::PC_Struct;
 			PinType.PinSubCategoryObject = StructType;
 			bIsValidPinType = UEdGraphSchema_K2::IsAllowableBlueprintVariableType(StructType);
 
@@ -1922,18 +1925,18 @@ static bool DumpBlueprintInfoUtils::DumpPinContextActions(uint32 Indent, UEdGrap
 				bIsValidPinType |= AnimSchema->IsPosePin(PinType);
 			}
 		}
-		else if (!CommandOptions.PinType.Compare("self", ESearchCase::IgnoreCase))
+		else if (CommandOptions.PinType == TEXT("self"))
 		{
-			PinType.PinCategory = K2Schema->PC_Object;
+			PinType.PinCategory = UEdGraphSchema_K2::PC_Object;
 			if (bUsePinTypeClass)
 			{
-				PinType.PinCategory = K2Schema->PC_Class;
+				PinType.PinCategory = UEdGraphSchema_K2::PC_Class;
 			}
-			PinType.PinSubCategory = K2Schema->PSC_Self;
+			PinType.PinSubCategory = UEdGraphSchema_K2::PSC_Self;
 			bIsValidPinType = true;
 		}
-		else if (!CommandOptions.PinType.Compare(K2Schema->PC_Delegate, ESearchCase::IgnoreCase) ||
-		         !CommandOptions.PinType.Compare(K2Schema->PC_MCDelegate, ESearchCase::IgnoreCase))
+		else if ((CommandOptions.PinType != UEdGraphSchema_K2::PC_Delegate) ||
+		         (CommandOptions.PinType != UEdGraphSchema_K2::PC_MCDelegate))
 		{
 			// @TODO: PC_Delegate, PC_MCDelegate
 			bIsValidPinType = false;
@@ -2043,7 +2046,7 @@ static void DumpBlueprintInfoUtils::DumpContextualPinTypeActions(uint32 Indent,
 	FGraphContextMenuBuilder ContextMenuBuilder(Graph);
 
 	UK2Node_Composite* DummyNode = NewObject<UK2Node_Composite>(Graph);
-	UEdGraphPin* DummyPin = DummyNode->CreatePin(EGPD_Input, PinType, DummyNode->GetName());
+	UEdGraphPin* DummyPin = DummyNode->CreatePin(EGPD_Input, PinType, DummyNode->GetFName());
 	ContextMenuBuilder.FromPin = DummyPin;
 
 	UBlueprint* Blueprint = FBlueprintEditorUtils::FindBlueprintForGraph(Graph);
diff --git a/Engine/Source/Editor/UnrealEd/Private/Commandlets/GenerateBlueprintAPICommandlet.cpp b/Engine/Source/Editor/UnrealEd/Private/Commandlets/GenerateBlueprintAPICommandlet.cpp
--- a/Engine/Source/Editor/UnrealEd/Private/Commandlets/GenerateBlueprintAPICommandlet.cpp
+++ b/Engine/Source/Editor/UnrealEd/Private/Commandlets/GenerateBlueprintAPICommandlet.cpp
@@ -910,11 +910,11 @@ static void GenerateBlueprintAPIUtils::DumpActionMenuItem(uint32 Indent, FGraphA
 						ActionEntry += "," + PinDetailsIndentedNewline + PinTooltipFieldLabel + MakeJsonString(PinTooltipStr) + "\"";
 					}
 
-					ActionEntry += "," + PinDetailsIndentedNewline + "\"PinCategory\"          : \"" + MakeJsonString(Pin->PinType.PinCategory) + "\"";
+					ActionEntry += "," + PinDetailsIndentedNewline + "\"PinCategory\"          : \"" + MakeJsonString(Pin->PinType.PinCategory.ToString()) + "\"";
 
-					if (!Pin->PinType.PinSubCategory.IsEmpty())
+					if (!Pin->PinType.PinSubCategory.IsNone())
 					{
-						ActionEntry += "," + PinDetailsIndentedNewline + "\"PinSubCategory\"       : \"" + MakeJsonString(Pin->PinType.PinSubCategory) + "\"";
+						ActionEntry += "," + PinDetailsIndentedNewline + "\"PinSubCategory\"       : \"" + MakeJsonString(Pin->PinType.PinSubCategory.ToString()) + "\"";
 					}
 
 					if (Pin->PinType.PinSubCategoryObject.IsValid())
diff --git a/Engine/Source/Editor/UnrealEd/Private/Commandlets/StabilizeLocalizationKeys.cpp b/Engine/Source/Editor/UnrealEd/Private/Commandlets/StabilizeLocalizationKeys.cpp
--- a/Engine/Source/Editor/UnrealEd/Private/Commandlets/StabilizeLocalizationKeys.cpp
+++ b/Engine/Source/Editor/UnrealEd/Private/Commandlets/StabilizeLocalizationKeys.cpp
@@ -98,7 +98,7 @@ public:
 			{
 				for (FStructVariableDescription& StructVariableDesc : UDSEditorData->VariablesDescriptions)
 				{
-					static const FString TextCategory = TEXT("text"); // Must match UEdGraphSchema_K2::PC_Text
+					static const FName TextCategory = TEXT("text"); // Must match UEdGraphSchema_K2::PC_Text
 					if (StructVariableDesc.Category == TextCategory)
 					{
 						FText StructVariableValue;
diff --git a/Engine/Source/Editor/UnrealEd/Private/CookOnTheFlyServer.cpp b/Engine/Source/Editor/UnrealEd/Private/CookOnTheFlyServer.cpp
--- a/Engine/Source/Editor/UnrealEd/Private/CookOnTheFlyServer.cpp
+++ b/Engine/Source/Editor/UnrealEd/Private/CookOnTheFlyServer.cpp
@@ -4457,30 +4457,32 @@ void UCookOnTheFlyServer::PopulateCookedPackagesFromDisk(const TArray<ITargetPla
 			const FName CookedFile = CookedPaths.Value;
 			const FName UncookedFilename = CookedPaths.Key;
 			const FName* FoundPackageName = GetCachedPackageFilenameToPackageFName(UncookedFilename);
+			bool bShouldKeep = true;
+			const FName SourcePackageName = FoundPackageName ? *FoundPackageName : NAME_None;
 			if ( !FoundPackageName )
 			{
 				// Source file no longer exists
 				++NumPackagesRemoved;
-				continue;
-			}
-			const FName PackageName = *FoundPackageName;
-			bool bShouldKeep = true;
-
-			if (ModifiedPackages.Contains(PackageName))
-			{
-				++NumPackagesFileHashMismatch;
 				bShouldKeep = false;
 			}
-			else if (NewPackages.Contains(PackageName) || RemovedPackages.Contains(PackageName))
-			{
-				++NumPackagesUnableToFindCookedPackageInfo;
-				bShouldKeep = false;
-			}
-			else if (IdenticalUncookedPackages.Contains(PackageName))
+			else
 			{
-				// These are packages which failed to save the first time 
-				// most likely because they are editor only packages
-				bShouldKeep = false;
+				if (ModifiedPackages.Contains(SourcePackageName))
+				{
+					++NumPackagesFileHashMismatch;
+					bShouldKeep = false;
+				}
+				else if (NewPackages.Contains(SourcePackageName) || RemovedPackages.Contains(SourcePackageName))
+				{
+					++NumPackagesUnableToFindCookedPackageInfo;
+					bShouldKeep = false;
+				}
+				else if (IdenticalUncookedPackages.Contains(SourcePackageName))
+				{
+					// These are packages which failed to save the first time 
+					// most likely because they are editor only packages
+					bShouldKeep = false;
+				}
 			}
 				
 			if (CookedFile == NAME_DummyCookedFilename)
@@ -4501,7 +4503,7 @@ void UCookOnTheFlyServer::PopulateCookedPackagesFromDisk(const TArray<ITargetPla
 				TArray<bool> Succeeded;
 				Succeeded.Add(true);
 
-				if (IdenticalCookedPackages.Contains(PackageName))
+				if (IdenticalCookedPackages.Contains(SourcePackageName))
 				{
 					CookedPackages.Add(FFilePlatformCookedPackage(UncookedFilename, MoveTemp(PlatformNames), MoveTemp(Succeeded)));
 					++NumPackagesKept;
@@ -4509,7 +4511,7 @@ void UCookOnTheFlyServer::PopulateCookedPackagesFromDisk(const TArray<ITargetPla
 			}
 			else
 			{
-				if ( IsCookByTheBookMode() ) // cook on the fly will requeue this package when it wants it 
+				if (SourcePackageName != NAME_None && IsCookByTheBookMode()) // cook on the fly will requeue this package when it wants it 
 				{
 					// Force cook the modified file
 					CookRequests.EnqueueUnique(FFilePlatformRequest(UncookedFilename, PlatformNames));
@@ -5411,6 +5413,9 @@ void UCookOnTheFlyServer::CookByTheBookFinished()
 			}
 
 			CodeGenModule.FinalizeManifest();
+
+			// Unload the module as we only need it while cooking. This will also clear the current module's state in order to allow a new cooker pass to function properly.
+			FModuleManager::Get().UnloadModule(CodeGenModule.GetModuleName());
 		}
 
 		check(CookByTheBookOptions->ChildUnsolicitedPackages.Num() == 0);
@@ -5873,7 +5878,7 @@ void UCookOnTheFlyServer::StartCookByTheBook( const FCookByTheBookStartupOptions
 		StartSavingEDLCookInfoForVerification();
 	}
 
-	if (CookByTheBookStartupOptions.bNativizeAssets)
+	if (PackagingSettings->BlueprintNativizationMethod != EProjectPackagingBlueprintNativizationMethod::Disabled)
 	{
 		FNativeCodeGenInitData CodeGenData;
 		for (const ITargetPlatform* Entry : CookByTheBookStartupOptions.TargetPlatforms)
@@ -5883,18 +5888,8 @@ void UCookOnTheFlyServer::StartCookByTheBook( const FCookByTheBookStartupOptions
 			CodeGenData.CodegenTargets.Push(PlatformNativizationDetails);
 		}
 		CodeGenData.ManifestIdentifier = CookByTheBookStartupOptions.ChildCookIdentifier;
-		CodeGenData.DestPluginPath = CookByTheBookStartupOptions.NativizedPluginPath;
 		IBlueprintNativeCodeGenModule::InitializeModule(CodeGenData);
 	}
-	else if(PackagingSettings->bWarnIfPackagedWithoutNativizationFlag && PackagingSettings->BlueprintNativizationMethod != EProjectPackagingBlueprintNativizationMethod::Disabled)
-	{
-		// Warn if we're cooking without the -nativizeAssets flag, when the project settings specify a nativization method.
-		// If the "exclusive" (whitelist) method is set, we only warn if at least one asset has been selected for conversion.
-		if (PackagingSettings->BlueprintNativizationMethod != EProjectPackagingBlueprintNativizationMethod::Exclusive || PackagingSettings->NativizeBlueprintAssets.Num() > 0)
-		{
-			UE_LOG(LogCook, Warning, TEXT("Project is configured for Blueprint nativization, but the conversion flag (-nativizeAssets) has been omitted from the command line. No assets will be converted as a result."));
-		}
-	}
 
 	CookByTheBookOptions->bLeakTest = (CookOptions & ECookByTheBookOptions::LeakTest) != ECookByTheBookOptions::None; // this won't work from the editor this needs to be standalone
 	check(!CookByTheBookOptions->bLeakTest || CurrentCookMode == ECookMode::CookByTheBook);
diff --git a/Engine/Source/Editor/UnrealEd/Private/EdGraphUtilities.cpp b/Engine/Source/Editor/UnrealEd/Private/EdGraphUtilities.cpp
--- a/Engine/Source/Editor/UnrealEd/Private/EdGraphUtilities.cpp
+++ b/Engine/Source/Editor/UnrealEd/Private/EdGraphUtilities.cpp
@@ -426,7 +426,7 @@ void FEdGraphUtilities::CopyCommonState(UEdGraphNode* OldNode, UEdGraphNode* New
 	NewNode->NodeComment = OldNode->NodeComment;
 }
 
-bool FEdGraphUtilities::IsSetParam(const UFunction* Function, const FString& ParameterName)
+bool FEdGraphUtilities::IsSetParam(const UFunction* Function, const FName ParameterName)
 {
 	if (Function == nullptr)
 	{
@@ -444,20 +444,25 @@ bool FEdGraphUtilities::IsSetParam(const UFunction* Function, const FString& Par
 		RawMetaData.ParseIntoArray(SetParamPinGroups, TEXT(","), true);
 	}
 
-	for (FString& Entry : SetParamPinGroups)
+	if (SetParamPinGroups.Num() > 0)
 	{
 		TArray<FString> GroupEntries;
-		Entry.ParseIntoArray(GroupEntries, TEXT("|"), true);
-		if (GroupEntries.Contains(ParameterName))
+		const FString ParameterNameStr = ParameterName.ToString();
+		for (const FString& Entry : SetParamPinGroups)
 		{
-			return true;
+			GroupEntries.Reset();
+			Entry.ParseIntoArray(GroupEntries, TEXT("|"), true);
+			if (GroupEntries.Contains(ParameterNameStr))
+			{
+				return true;
+			}
 		}
 	}
 
 	return false;
 }
 
-bool FEdGraphUtilities::IsMapParam(const UFunction* Function, const FString& ParameterName)
+bool FEdGraphUtilities::IsMapParam(const UFunction* Function, const FName ParameterName)
 {
 	if (Function == nullptr)
 	{
@@ -472,11 +477,14 @@ bool FEdGraphUtilities::IsMapParam(const UFunction* Function, const FString& Par
 		return false;
 	}
 
-	const auto PipeSeparatedStringContains = [ParameterName](const FString& List)
+	TArray<FString> GroupEntries;
+	const FString ParameterNameStr = ParameterName.ToString();
+
+	const auto PipeSeparatedStringContains = [&GroupEntries, &ParameterNameStr](const FString& List)
 	{
-		TArray<FString> GroupEntries;
+		GroupEntries.Reset();
 		List.ParseIntoArray(GroupEntries, TEXT("|"), true);
-		if (GroupEntries.Contains(ParameterName))
+		if (GroupEntries.Contains(ParameterNameStr))
 		{
 			return true;
 		}
@@ -488,7 +496,7 @@ bool FEdGraphUtilities::IsMapParam(const UFunction* Function, const FString& Par
 		|| PipeSeparatedStringContains(MapKeyParamMetaData);
 }
 
-bool FEdGraphUtilities::IsArrayDependentParam(const UFunction* Function, const FString& ParameterName)
+bool FEdGraphUtilities::IsArrayDependentParam(const UFunction* Function, const FName ParameterName)
 {
 	if (Function == nullptr)
 	{
@@ -504,7 +512,7 @@ bool FEdGraphUtilities::IsArrayDependentParam(const UFunction* Function, const F
 	TArray<FString> TypeDependentPinNames;
 	DependentPinMetaData.ParseIntoArray(TypeDependentPinNames, TEXT(","), true);
 
-	return TypeDependentPinNames.Contains(ParameterName);
+	return TypeDependentPinNames.Contains(ParameterName.ToString());
 }
 
 UEdGraphPin* FEdGraphUtilities::FindArrayParamPin(const UFunction* Function, const UEdGraphNode* Node)
@@ -631,17 +639,15 @@ void FWeakGraphPinPtr::operator=(const class UEdGraphPin* Pin)
 
 UEdGraphPin* FWeakGraphPinPtr::Get()
 {
-	UEdGraphNode* Node = NodeObjectPtr.Get();
-	if(Node != NULL)
+	if (UEdGraphNode* Node = NodeObjectPtr.Get())
 	{
 		// If pin is no longer valid or has a different owner, attempt to fix up the reference
 		UEdGraphPin* Pin = PinReference.Get();
-		if(Pin == NULL || Pin->GetOuter() != Node)
+		if (Pin == nullptr || Pin->GetOuter() != Node)
 		{
-			for(auto PinIter = Node->Pins.CreateConstIterator(); PinIter; ++PinIter)
+			for (UEdGraphPin* TestPin : Node->Pins)
 			{
-				UEdGraphPin* TestPin = *PinIter;
-				if(TestPin->PinName.Equals(PinName))
+				if (TestPin->PinName == PinName)
 				{
 					Pin = TestPin;
 					PinReference = Pin;
@@ -653,5 +659,5 @@ UEdGraphPin* FWeakGraphPinPtr::Get()
 		return Pin;
 	}
 
-	return NULL;
+	return nullptr;
 }
diff --git a/Engine/Source/Editor/UnrealEd/Private/GlobalEditorNotification.cpp b/Engine/Source/Editor/UnrealEd/Private/GlobalEditorNotification.cpp
--- a/Engine/Source/Editor/UnrealEd/Private/GlobalEditorNotification.cpp
+++ b/Engine/Source/Editor/UnrealEd/Private/GlobalEditorNotification.cpp
@@ -8,11 +8,6 @@ void FGlobalEditorNotification::Tick(float DeltaTime)
 	TickNotification(DeltaTime);
 }
 
-bool FGlobalEditorNotification::IsTickable() const
-{
-	return true;
-}
-
 TStatId FGlobalEditorNotification::GetStatId() const
 {
 	RETURN_QUICK_DECLARE_CYCLE_STAT(FGlobalEditorNotification, STATGROUP_Tickables);
diff --git a/Engine/Source/Editor/UnrealEd/Private/HierarchicalLOD.cpp b/Engine/Source/Editor/UnrealEd/Private/HierarchicalLOD.cpp
--- a/Engine/Source/Editor/UnrealEd/Private/HierarchicalLOD.cpp
+++ b/Engine/Source/Editor/UnrealEd/Private/HierarchicalLOD.cpp
@@ -58,7 +58,7 @@ void FHierarchicalLODBuilder::Build()
 	bool bVisibleLevelsWarning = false;
 
 	const TArray<ULevel*>& Levels = World->GetLevels();
-	for (const auto& LevelIter : Levels)
+	for (ULevel* LevelIter : Levels)
 	{
 		// Only build clusters for levels that are visible, and throw warning if any are hidden
 		if (LevelIter->bIsVisible)
@@ -87,7 +87,7 @@ void FHierarchicalLODBuilder::PreviewBuild()
 	bool bVisibleLevelsWarning = false;
 
 	const TArray<ULevel*>& Levels = World->GetLevels();
-	for (const auto& LevelIter : Levels)
+	for (ULevel* LevelIter : Levels)
 	{
 		// Only build clusters for levels that are visible
 		if (LevelIter->bIsVisible)
@@ -235,7 +235,7 @@ void FHierarchicalLODBuilder::InitializeClusters(ULevel* InLevel, const int32 LO
 				{
 					// Check whether or not this actor falls within a HierarchicalLODVolume, if so add to the Volume's cluster and exclude from normal process
 					bool bAdded = bVolumesOnly;
-					for (auto& Cluster : HLODVolumeClusters)
+					for (TPair<AHierarchicalLODVolume*, FLODCluster>& Cluster : HLODVolumeClusters)
 					{
 						if (Cluster.Key->EncompassesPoint(Actor->GetActorLocation(), Cluster.Key->bIncludeOverlappingActors ? Actor->GetComponentsBoundingBox().GetSize().Size() : 0.0f, nullptr))
 						{			
@@ -332,14 +332,13 @@ void FHierarchicalLODBuilder::FindMST()
 
 void FHierarchicalLODBuilder::HandleHLODVolumes(ULevel* InLevel)
 {	
-	HLODVolumeClusters.Empty();
+	HLODVolumeClusters.Reset();
 	for (int32 ActorId = 0; ActorId < InLevel->Actors.Num(); ++ActorId)
 	{
-		AActor* Actor = InLevel->Actors[ActorId];
-		if (Actor && Actor->IsA(AHierarchicalLODVolume::StaticClass()))
+		if (AHierarchicalLODVolume* Actor = Cast<AHierarchicalLODVolume>(InLevel->Actors[ActorId]))
 		{
 			// Came across a HLOD volume			
-			FLODCluster& NewCluster = HLODVolumeClusters.Add(CastChecked<AHierarchicalLODVolume>(Actor));
+			FLODCluster& NewCluster = HLODVolumeClusters.Add(Actor);
 
 			FVector Origin, Extent;
 			Actor->GetActorBounds(false, Origin, Extent);
@@ -402,20 +401,26 @@ bool FHierarchicalLODBuilder::ShouldGenerateCluster(AActor* Actor, const bool bP
 	}	
 
 	// for now only consider staticmesh - I don't think skel mesh would work with simplygon merge right now @fixme
-	TArray<UStaticMeshComponent*> Components;
-	Actor->GetComponents<UStaticMeshComponent>(Components);
-	// TODO: support instanced static meshes
-	Components.RemoveAll([](UStaticMeshComponent* Val){ return Val->IsA(UInstancedStaticMeshComponent::StaticClass()); });
+	TInlineComponentArray<UStaticMeshComponent*> Components;
+	for (UActorComponent* Component : Actor->GetComponents())
+	{
+		UStaticMeshComponent* SMComponent = Cast<UStaticMeshComponent>(Component);
+		// TODO: support instanced static meshes
+		if (Component && !Component->IsA<UInstancedStaticMeshComponent>())
+		{
+			Components.Add(SMComponent);
+		}
+	}
 
 	int32 ValidComponentCount = 0;
 	// now make sure you check parent primitive, so that we don't build for the actor that already has built. 
 	if (Components.Num() > 0)
 	{
-		for (auto& ComponentIter : Components)
+		for (UStaticMeshComponent* Component : Components)
 		{			
-			if (ComponentIter->GetLODParentPrimitive())
+			if (Component->GetLODParentPrimitive())
 			{
-				auto ParentActor = CastChecked<ALODActor>(ComponentIter->GetLODParentPrimitive()->GetOwner());
+				ALODActor* ParentActor = CastChecked<ALODActor>(Component->GetLODParentPrimitive()->GetOwner());
 				
 				if (ParentActor && bPreviewBuild)
 				{
@@ -423,13 +428,13 @@ bool FHierarchicalLODBuilder::ShouldGenerateCluster(AActor* Actor, const bool bP
 				}
 			}
 
-			if (ComponentIter->bHiddenInGame)
+			if (Component->bHiddenInGame)
 			{
 				return false;
 			}
 
 			// see if we should generate it
-			if (ComponentIter->ShouldGenerateAutoLOD())
+			if (Component->ShouldGenerateAutoLOD())
 			{
 				++ValidComponentCount;
 				break;
@@ -444,7 +449,7 @@ void FHierarchicalLODBuilder::ClearHLODs()
 {
 	bool bVisibleLevelsWarning = false;
 
-	for (auto& Level : World->GetLevels())
+	for (ULevel* Level : World->GetLevels())
 	{
 		bVisibleLevelsWarning |= !Level->bIsVisible;
 		if (Level->bIsVisible)
@@ -467,7 +472,7 @@ void FHierarchicalLODBuilder::ClearHLODs()
 void FHierarchicalLODBuilder::ClearPreviewBuild()
 {
 	bool bVisibleLevelsWarning = false;
-	for (auto& Level : World->GetLevels())
+	for (ULevel* Level : World->GetLevels())
 	{
 		bVisibleLevelsWarning |= !Level->bIsVisible;
 		if ( Level->bIsVisible )
@@ -491,7 +496,7 @@ void FHierarchicalLODBuilder::BuildMeshesForLODActors(bool bForceAll)
 	bool bVisibleLevelsWarning = false;
 
 	const TArray<ULevel*>& Levels = World->GetLevels();
-	for (const auto& LevelIter : Levels)
+	for (const ULevel* LevelIter : Levels)
 	{
 		// Only meshes for clusters that are in a visible level
 		if (!LevelIter->bIsVisible)
@@ -629,7 +634,7 @@ void FHierarchicalLODBuilder::MergeClustersAndBuildActors(ULevel* InLevel, const
 				// now we have minimum Clusters
 				for (int32 ClusterId = 0; ClusterId < TotalCluster; ++ClusterId)
 				{
-					auto& Cluster = Clusters[ClusterId];
+					FLODCluster& Cluster = Clusters[ClusterId];
 					UE_LOG(LogLODGenerator, Verbose, TEXT("%d. %0.2f {%s}"), ClusterId + 1, Cluster.GetCost(), *Cluster.ToString());
 
 					// progress bar update every percent, if ClustersPerPercent is zero ignore the progress bar as number of iterations is small.
@@ -643,7 +648,7 @@ void FHierarchicalLODBuilder::MergeClustersAndBuildActors(ULevel* InLevel, const
 						for (int32 MergedClusterId = 0; MergedClusterId < ClusterId; ++MergedClusterId)
 						{
 							// compare with previous clusters
-							auto& MergedCluster = Clusters[MergedClusterId];
+							FLODCluster& MergedCluster = Clusters[MergedClusterId];
 							// see if it's valid, if it contains, check the cost
 							if (MergedCluster.IsValid())
 							{
@@ -688,7 +693,7 @@ void FHierarchicalLODBuilder::MergeClustersAndBuildActors(ULevel* InLevel, const
 
 		if (LODIdx == 0)
 		{
-			for (auto& Cluster : HLODVolumeClusters)
+			for (TPair<AHierarchicalLODVolume*, FLODCluster>& Cluster : HLODVolumeClusters)
 			{
 				Clusters.Add(Cluster.Value);
 			}
@@ -699,7 +704,7 @@ void FHierarchicalLODBuilder::MergeClustersAndBuildActors(ULevel* InLevel, const
 			SCOPE_LOG_TIME(TEXT("HLOD_BuildActors"), nullptr);
 			// print data
 			int32 TotalValidCluster = 0;
-			for (auto& Cluster : Clusters)
+			for (FLODCluster& Cluster : Clusters)
 			{
 				if (Cluster.IsValid())
 				{
@@ -710,7 +715,7 @@ void FHierarchicalLODBuilder::MergeClustersAndBuildActors(ULevel* InLevel, const
 			FScopedSlowTask SlowTask(TotalValidCluster, FText::Format(LOCTEXT("HierarchicalLOD_MergeActors", "Merging Actors for LOD {LODIndex} of {LevelName}..."), Arguments));
 			SlowTask.MakeDialog();
 
-			for (auto& Cluster : Clusters)
+			for (FLODCluster& Cluster : Clusters)
 			{
 				if (Cluster.IsValid())
 				{
@@ -724,7 +729,7 @@ void FHierarchicalLODBuilder::MergeClustersAndBuildActors(ULevel* InLevel, const
 							LODLevelLODActors[LODIdx].Add(LODActor);
 						}
 
-						for (auto& RemoveActor : Cluster.Actors)
+						for (AActor* RemoveActor : Cluster.Actors)
 						{
 							ValidStaticMeshActorsInLevel.RemoveSingleSwap(RemoveActor, false);
 						}
diff --git a/Engine/Source/Editor/UnrealEd/Private/Kismet2/BlueprintEditorUtils.cpp b/Engine/Source/Editor/UnrealEd/Private/Kismet2/BlueprintEditorUtils.cpp
--- a/Engine/Source/Editor/UnrealEd/Private/Kismet2/BlueprintEditorUtils.cpp
+++ b/Engine/Source/Editor/UnrealEd/Private/Kismet2/BlueprintEditorUtils.cpp
@@ -46,6 +46,7 @@
 #include "Engine/LevelScriptBlueprint.h"
 #include "BlueprintsObjectVersion.h"
 #include "Kismet2/CompilerResultsLog.h"
+#include "Algo/Transform.h"
 
 #include "Editor/KismetCompiler/Public/KismetCompilerModule.h"
 #include "EdGraphSchema_K2_Actions.h"
@@ -1141,7 +1142,7 @@ struct FRegenerationHelper
 								? FunctionEntry->CustomGeneratedFunctionName 
 								: FunctionEntry->SignatureName;
 							UFunction* ParentFunction = Blueprint->ParentClass ? Blueprint->ParentClass->FindFunctionByName(FunctionName) : nullptr;
-							if (ParentFunction && (Schema->FN_UserConstructionScript != FunctionName))
+							if (ParentFunction && (UEdGraphSchema_K2::FN_UserConstructionScript != FunctionName))
 							{
 								ProcessHierarchy(ParentFunction, Dependencies);
 							}
@@ -3191,10 +3192,9 @@ bool FBlueprintEditorUtils::CanClassGenerateEvents(const UClass* InClass)
 
 UEdGraph* FBlueprintEditorUtils::FindUserConstructionScript(const UBlueprint* Blueprint)
 {
-	const UEdGraphSchema_K2* Schema = GetDefault<UEdGraphSchema_K2>();
 	for (UEdGraph* CurrentGraph : Blueprint->FunctionGraphs)
 	{
-		if( CurrentGraph->GetFName() == Schema->FN_UserConstructionScript )
+		if( CurrentGraph->GetFName() == UEdGraphSchema_K2::FN_UserConstructionScript )
 		{
 			return CurrentGraph;
 		}
@@ -3205,10 +3205,9 @@ UEdGraph* FBlueprintEditorUtils::FindUserConstructionScript(const UBlueprint* Bl
 
 UEdGraph* FBlueprintEditorUtils::FindEventGraph(const UBlueprint* Blueprint)
 {
-	const UEdGraphSchema_K2* Schema = GetDefault<UEdGraphSchema_K2>();
 	for (UEdGraph* CurrentGraph : Blueprint->UbergraphPages)
 	{
-		if( CurrentGraph->GetFName() == Schema->GN_EventGraph )
+		if( CurrentGraph->GetFName() == UEdGraphSchema_K2::GN_EventGraph )
 		{
 			return CurrentGraph;
 		}
@@ -4063,8 +4062,28 @@ UEdGraph* FBlueprintEditorUtils::GetDelegateSignatureGraphByName(UBlueprint* Blu
 	return nullptr;
 }
 
-// Gets a list of pins that should hidden for a given function
 void FBlueprintEditorUtils::GetHiddenPinsForFunction(UEdGraph const* Graph, UFunction const* Function, TSet<FString>& HiddenPins, TSet<FString>* OutInternalPins)
+{
+	TSet<FName> HiddenPinNames;
+	if (OutInternalPins)
+	{
+		TSet<FName> InternalPinNames;
+		GetHiddenPinsForFunction(Graph, Function, HiddenPinNames, &InternalPinNames);
+
+		OutInternalPins->Reserve(HiddenPinNames.Num());
+		Algo::Transform(InternalPinNames, *OutInternalPins, [](const FName InternalPinName) { return InternalPinName.ToString(); });
+	}
+	else
+	{
+		GetHiddenPinsForFunction(Graph, Function, HiddenPinNames);
+	}
+
+	HiddenPins.Reserve(HiddenPinNames.Num());
+	Algo::Transform(HiddenPinNames, HiddenPins, [](const FName HiddenPinName) { return HiddenPinName.ToString(); });
+}
+
+// Gets a list of pins that should hidden for a given function
+void FBlueprintEditorUtils::GetHiddenPinsForFunction(UEdGraph const* Graph, UFunction const* Function, TSet<FName>& HiddenPins, TSet<FName>* OutInternalPins)
 {
 	check(Function != nullptr);
 	TMap<FName, FString>* MetaData = UMetaData::GetMapForObject(Function);	
@@ -4077,27 +4096,20 @@ void FBlueprintEditorUtils::GetHiddenPinsForFunction(UEdGraph const* Graph, UFun
 
 			const FName& Key = It.Key();
 
-			if (Key == NAME_LatentInfo)
-			{
-				HiddenPins.Add(It.Value());
-			}
-			else if (Key == NAME_HidePin)
+			if (Key == NAME_LatentInfo || Key == NAME_HidePin || Key == FBlueprintMetadata::MD_ExpandEnumAsExecs)
 			{
-				HiddenPins.Add(It.Value());
+				HiddenPins.Add(*It.Value());
 			}
 			else if (Key == FBlueprintMetadata::MD_InternalUseParam)
 			{
-				HiddenPins.Add(It.Value());
+				const FName HiddenPinName = *It.Value();
+				HiddenPins.Add(HiddenPinName);
 
 				if (OutInternalPins != nullptr)
 				{
-					OutInternalPins->Add(It.Value());
+					OutInternalPins->Add(HiddenPinName);
 				}
 			}
-			else if(Key == FBlueprintMetadata::MD_ExpandEnumAsExecs)
-			{
-				HiddenPins.Add(It.Value());
-			}
 			else if (Key == FBlueprintMetadata::MD_WorldContext)
 			{
 				const UEdGraphSchema_K2* K2Schema = GetDefault<UEdGraphSchema_K2>();
@@ -4124,7 +4136,7 @@ void FBlueprintEditorUtils::GetHiddenPinsForFunction(UEdGraph const* Graph, UFun
 					// then we can hide this pin (and default it to self)
 					if (bHasIntrinsicWorldContext)
 					{
-						HiddenPins.Add(It.Value());
+						HiddenPins.Add(*It.Value());
 					}
 				}
 			}
@@ -4248,18 +4260,16 @@ bool FBlueprintEditorUtils::AddMemberVariable(UBlueprint* Blueprint, const FName
 
 	Blueprint->Modify();
 
-	const UEdGraphSchema_K2* K2Schema = GetDefault<UEdGraphSchema_K2>();
-
 	// Now create new variable
 	FBPVariableDescription NewVar;
 
 	NewVar.VarName = NewVarName;
 	NewVar.VarGuid = FGuid::NewGuid();
-	NewVar.FriendlyName = FName::NameToDisplayString( NewVarName.ToString(), (NewVarType.PinCategory == K2Schema->PC_Boolean) ? true : false );
+	NewVar.FriendlyName = FName::NameToDisplayString( NewVarName.ToString(), (NewVarType.PinCategory == UEdGraphSchema_K2::PC_Boolean) ? true : false );
 	NewVar.VarType = NewVarType;
 	// default new vars to 'kismet read/write' and 'only editable on owning CDO' 
 	NewVar.PropertyFlags |= (CPF_Edit | CPF_BlueprintVisible | CPF_DisableEditOnInstance);
-	if(NewVarType.PinCategory == K2Schema->PC_MCDelegate)
+	if(NewVarType.PinCategory == UEdGraphSchema_K2::PC_MCDelegate)
 	{
 		NewVar.PropertyFlags |= CPF_BlueprintAssignable | CPF_BlueprintCallable;
 	}
@@ -4443,10 +4453,9 @@ void FBlueprintEditorUtils::RenameMemberVariable(UBlueprint* Blueprint, const FN
 			Blueprint->Modify();
 
 			// Update the name
-			const UEdGraphSchema_K2* K2Schema = GetDefault<UEdGraphSchema_K2>();
 			FBPVariableDescription& Variable = Blueprint->NewVariables[VarIndex];
 			Variable.VarName = NewName;
-			Variable.FriendlyName = FName::NameToDisplayString( NewName.ToString(), (Variable.VarType.PinCategory == K2Schema->PC_Boolean) ? true : false );
+			Variable.FriendlyName = FName::NameToDisplayString( NewName.ToString(), (Variable.VarType.PinCategory == UEdGraphSchema_K2::PC_Boolean) ? true : false );
 
 			// Update any existing references to the old name
 			FBlueprintEditorUtils::ReplaceVariableReferences(Blueprint, OldName, NewName);
@@ -4630,7 +4639,7 @@ void FBlueprintEditorUtils::ChangeMemberVariableType(UBlueprint* Blueprint, cons
 				/** Only change the variable type if type selection is valid, some unloaded Blueprints will turn out to be bad */
 				bool bChangeVariableType = true;
 
-				if ((NewPinType.PinCategory == K2Schema->PC_Object) || (NewPinType.PinCategory == K2Schema->PC_Interface))
+				if ((NewPinType.PinCategory == UEdGraphSchema_K2::PC_Object) || (NewPinType.PinCategory == UEdGraphSchema_K2::PC_Interface))
 				{
 					// if it's a PC_Object, then it should have an associated UClass object
 					if(NewPinType.PinSubCategoryObject.IsValid())
@@ -4825,7 +4834,7 @@ FBPVariableDescription FBlueprintEditorUtils::DuplicateVariableDescription(UBlue
 	// Now create new variable
 	FBPVariableDescription NewVar = InVariableDescription;
 	NewVar.VarName = DuplicatedVariableName;
-	NewVar.FriendlyName = FName::NameToDisplayString( NewVar.VarName.ToString(), NewVar.VarType.PinCategory == GetDefault<UEdGraphSchema_K2>()->PC_Boolean);
+	NewVar.FriendlyName = FName::NameToDisplayString( NewVar.VarName.ToString(), NewVar.VarType.PinCategory == UEdGraphSchema_K2::PC_Boolean);
 	NewVar.VarGuid = FGuid::NewGuid();
 
 	return NewVar;
@@ -4845,8 +4854,6 @@ bool FBlueprintEditorUtils::AddLocalVariable(UBlueprint* Blueprint, UEdGraph* In
 		TargetGraph->GetNodesOfClass(FunctionEntryNodes);
 		check(FunctionEntryNodes.Num());
 
-		const UEdGraphSchema_K2* K2Schema = GetDefault<UEdGraphSchema_K2>();
-
 		// Now create new variable
 		FBPVariableDescription NewVar;
 
@@ -4854,7 +4861,7 @@ bool FBlueprintEditorUtils::AddLocalVariable(UBlueprint* Blueprint, UEdGraph* In
 		NewVar.VarGuid = FGuid::NewGuid();
 		NewVar.VarType = InNewVarType;
 		NewVar.PropertyFlags |= CPF_BlueprintVisible;
-		NewVar.FriendlyName = FName::NameToDisplayString( NewVar.VarName.ToString(), (NewVar.VarType.PinCategory == K2Schema->PC_Boolean) ? true : false );
+		NewVar.FriendlyName = FName::NameToDisplayString( NewVar.VarName.ToString(), (NewVar.VarType.PinCategory == UEdGraphSchema_K2::PC_Boolean) ? true : false );
 		NewVar.Category = UEdGraphSchema_K2::VR_DefaultCategory;
 		NewVar.DefaultValue = DefaultValue;
 
@@ -4986,9 +4993,8 @@ void FBlueprintEditorUtils::RenameLocalVariable(UBlueprint* InBlueprint, const U
 			FunctionEntry->Modify();
 
 			// Update the name
-			const UEdGraphSchema_K2* K2Schema = GetDefault<UEdGraphSchema_K2>();
 			LocalVariable->VarName = InNewName;
-			LocalVariable->FriendlyName = FName::NameToDisplayString( InNewName.ToString(), LocalVariable->VarType.PinCategory == K2Schema->PC_Boolean );
+			LocalVariable->FriendlyName = FName::NameToDisplayString( InNewName.ToString(), LocalVariable->VarType.PinCategory == UEdGraphSchema_K2::PC_Boolean );
 
 			// Update any existing references to the old name
 			RenameVariableReferencesInGraph(InBlueprint, InBlueprint->GeneratedClass, FindScopeGraph(InBlueprint, InScope), InOldName, InNewName);
@@ -5139,7 +5145,7 @@ void FBlueprintEditorUtils::ChangeLocalVariableType(UBlueprint* InBlueprint, con
 				// Mark the Blueprint as structurally modified so we can reconstruct the node successfully
 				FBlueprintEditorUtils::MarkBlueprintAsStructurallyModified(InBlueprint);
 
-				if ((NewPinType.PinCategory == K2Schema->PC_Object) || (NewPinType.PinCategory == K2Schema->PC_Interface))
+				if ((NewPinType.PinCategory == UEdGraphSchema_K2::PC_Object) || (NewPinType.PinCategory == UEdGraphSchema_K2::PC_Interface))
 				{
 					// if it's a PC_Object, then it should have an associated UClass object
 					if(NewPinType.PinSubCategoryObject.IsValid())
@@ -5220,7 +5226,7 @@ void FBlueprintEditorUtils::ReplaceVariableReferences(UBlueprint* Blueprint, con
 bool FBlueprintEditorUtils::IsVariableComponent(const FBPVariableDescription& Variable)
 {
 	// Find the variable in the list
-	if( Variable.VarType.PinCategory == FString(TEXT("object")) )
+	if( Variable.VarType.PinCategory == FName(TEXT("object")) )
 	{
 		const UClass* VarClass = Cast<const UClass>(Variable.VarType.PinSubCategoryObject.Get());
 		return (VarClass && VarClass->HasAnyClassFlags(CLASS_DefaultToInstanced));
@@ -5257,7 +5263,7 @@ bool FBlueprintEditorUtils::IsVariableUsed(const UBlueprint* Blueprint, const FN
 			{
 				if (ClassDefaultsNode->GetInputClass() == Blueprint->SkeletonGeneratedClass)
 				{
-					const UEdGraphPin* VarPin = ClassDefaultsNode->FindPin(Name.ToString());
+					const UEdGraphPin* VarPin = ClassDefaultsNode->FindPin(Name);
 					if (VarPin && VarPin->Direction == EGPD_Output && VarPin->LinkedTo.Num() > 0)
 					{
 						return true;
@@ -6373,8 +6379,6 @@ void FBlueprintEditorUtils::UpdateOutOfDateAnimBlueprints(UBlueprint* InBlueprin
 {
 	if (UAnimBlueprint* AnimBlueprint = Cast<UAnimBlueprint>(InBlueprint))
 	{
-		const UEdGraphSchema_K2* K2Schema = GetDefault<UEdGraphSchema_K2>();
-
 		// Ensure all transition graphs have the correct schema
 		TArray<UAnimStateTransitionNode*> TransitionNodes;
 		GetAllNodesOfClass<UAnimStateTransitionNode>(AnimBlueprint, /*out*/ TransitionNodes);
@@ -6391,9 +6395,9 @@ void FBlueprintEditorUtils::UpdateOutOfDateAnimBlueprints(UBlueprint* InBlueprin
 		if (UAnimBlueprint::FindRootAnimBlueprint(AnimBlueprint) == nullptr)
 		{
 			// Add an anim graph if not present
-			if (FindObject<UEdGraph>(AnimBlueprint, *(K2Schema->GN_AnimGraph.ToString())) == nullptr)
+			if (FindObject<UEdGraph>(AnimBlueprint, *(UEdGraphSchema_K2::GN_AnimGraph.ToString())) == nullptr)
 			{
-				UEdGraph* NewGraph = FBlueprintEditorUtils::CreateNewGraph(AnimBlueprint, K2Schema->GN_AnimGraph, UAnimationGraph::StaticClass(), UAnimationGraphSchema::StaticClass());
+				UEdGraph* NewGraph = FBlueprintEditorUtils::CreateNewGraph(AnimBlueprint, UEdGraphSchema_K2::GN_AnimGraph, UAnimationGraph::StaticClass(), UAnimationGraphSchema::StaticClass());
 				FBlueprintEditorUtils::AddDomainSpecificGraph(AnimBlueprint, NewGraph);
 				AnimBlueprint->LastEditedDocuments.Add(NewGraph);
 				NewGraph->bAllowDeletion = false;
@@ -6405,7 +6409,7 @@ void FBlueprintEditorUtils::UpdateOutOfDateAnimBlueprints(UBlueprint* InBlueprin
 			for (int32 i = 0; i < AnimBlueprint->FunctionGraphs.Num(); ++i)
 			{
 				UEdGraph* FuncGraph = AnimBlueprint->FunctionGraphs[i];
-				if ((FuncGraph != nullptr) && (FuncGraph->GetFName() == K2Schema->GN_AnimGraph))
+				if ((FuncGraph != nullptr) && (FuncGraph->GetFName() == UEdGraphSchema_K2::GN_AnimGraph))
 				{
 					UE_LOG(LogBlueprint, Log, TEXT("!!! Removing AnimGraph from %s, because it has a parent anim blueprint that defines the AnimGraph"), *AnimBlueprint->GetPathName());
 					AnimBlueprint->FunctionGraphs.RemoveAt(i);
@@ -6889,9 +6893,6 @@ bool FBlueprintEditorUtils::RenameTimeline(UBlueprint* Blueprint, const FName Ol
 				TimelineNode->TimelineName = NewName;
 			}
 
-			const FString NewNameStr = NewName.ToString();
-			const FString OldNameStr = OldName.ToString();
-
 			TArray<UK2Node_Variable*> TimelineVarNodes;
 			FBlueprintEditorUtils::GetAllNodesOfClass<UK2Node_Variable>(Blueprint, TimelineVarNodes);
 			for(int32 It = 0; It < TimelineVarNodes.Num(); It++)
@@ -6899,7 +6900,7 @@ bool FBlueprintEditorUtils::RenameTimeline(UBlueprint* Blueprint, const FName Ol
 				UK2Node_Variable* TestNode = TimelineVarNodes[It];
 				if(TestNode && (OldName == TestNode->GetVarName()))
 				{
-					UEdGraphPin* TestPin = TestNode->FindPin(OldNameStr);
+					UEdGraphPin* TestPin = TestNode->FindPin(OldName);
 					if(TestPin && (UTimelineComponent::StaticClass() == TestPin->PinType.PinSubCategoryObject.Get()))
 					{
 						TestNode->Modify();
@@ -6914,7 +6915,7 @@ bool FBlueprintEditorUtils::RenameTimeline(UBlueprint* Blueprint, const FName Ol
 							TestNode->VariableReference.SetExternalMember(NewName, ParentClass);
 						}
 						TestPin->Modify();
-						TestPin->PinName = NewNameStr;
+						TestPin->PinName = NewName;
 					}
 				}
 			}
@@ -8119,12 +8120,10 @@ FName FBlueprintEditorUtils::GenerateUniqueGraphName(UBlueprint* const Blueprint
 
 bool FBlueprintEditorUtils::CheckIfNodeConnectsToSelection(UEdGraphNode* InNode, const TSet<UEdGraphNode*>& InSelectionSet)
 {
-	const UEdGraphSchema_K2* Schema = GetDefault<UEdGraphSchema_K2>();
-
 	for (int32 PinIndex = 0; PinIndex < InNode->Pins.Num(); ++PinIndex)
 	{
 		UEdGraphPin* Pin = InNode->Pins[PinIndex];
-		if(Pin->Direction == EGPD_Input && Pin->PinType.PinCategory != Schema->PC_Exec)
+		if(Pin->Direction == EGPD_Input && Pin->PinType.PinCategory != UEdGraphSchema_K2::PC_Exec)
 		{
 			for (int32 LinkIndex = 0; LinkIndex < Pin->LinkedTo.Num(); ++LinkIndex)
 			{
@@ -8150,7 +8149,6 @@ bool FBlueprintEditorUtils::CheckIfNodeConnectsToSelection(UEdGraphNode* InNode,
 
 bool FBlueprintEditorUtils::CheckIfSelectionIsCycling(const TSet<UEdGraphNode*>& InSelectionSet, FCompilerResultsLog& InMessageLog)
 {
-	const UEdGraphSchema_K2* Schema = GetDefault<UEdGraphSchema_K2>();
 	for (UEdGraphNode* Node : InSelectionSet)
 	{
 		if (Node)
@@ -8158,7 +8156,7 @@ bool FBlueprintEditorUtils::CheckIfSelectionIsCycling(const TSet<UEdGraphNode*>&
 			for (int32 PinIndex = 0; PinIndex < Node->Pins.Num(); ++PinIndex)
 			{
 				UEdGraphPin* Pin = Node->Pins[PinIndex];
-				if(Pin->Direction == EGPD_Input && Pin->PinType.PinCategory != Schema->PC_Exec)
+				if(Pin->Direction == EGPD_Input && Pin->PinType.PinCategory != UEdGraphSchema_K2::PC_Exec)
 				{
 					for (int32 LinkIndex = 0; LinkIndex < Pin->LinkedTo.Num(); ++LinkIndex)
 					{
@@ -8612,12 +8610,11 @@ bool FBlueprintEditorUtils::CheckIfGraphHasLatentFunctions(UEdGraph* InGraph)
 
 void FBlueprintEditorUtils::PostSetupObjectPinType(UBlueprint* InBlueprint, FBPVariableDescription& InOutVarDesc)
 {
-	UEdGraphSchema_K2 const* K2Schema = GetDefault<UEdGraphSchema_K2>();
-	if ((InOutVarDesc.VarType.PinCategory == K2Schema->PC_Object) || (InOutVarDesc.VarType.PinCategory == K2Schema->PC_Interface))
+	if ((InOutVarDesc.VarType.PinCategory == UEdGraphSchema_K2::PC_Object) || (InOutVarDesc.VarType.PinCategory == UEdGraphSchema_K2::PC_Interface))
 	{
-		if (InOutVarDesc.VarType.PinSubCategory == K2Schema->PSC_Self)
+		if (InOutVarDesc.VarType.PinSubCategory == UEdGraphSchema_K2::PSC_Self)
 		{
-			InOutVarDesc.VarType.PinSubCategory.Empty();
+			InOutVarDesc.VarType.PinSubCategory = NAME_None;
 			InOutVarDesc.VarType.PinSubCategoryObject = *InBlueprint->GeneratedClass;
 		}
 		else if (!InOutVarDesc.VarType.PinSubCategoryObject.IsValid())
@@ -8625,7 +8622,7 @@ void FBlueprintEditorUtils::PostSetupObjectPinType(UBlueprint* InBlueprint, FBPV
 			// Fall back to UObject if the given type is not valid. This can happen for example if a variable is removed from
 			// a Blueprint parent class along with the variable's type and the user then attempts to recreate the missing variable
 			// through a stale variable node's context menu in a child Blueprint graph.
-			InOutVarDesc.VarType.PinSubCategory.Empty();
+			InOutVarDesc.VarType.PinSubCategory = NAME_None;
 			InOutVarDesc.VarType.PinSubCategoryObject = UObject::StaticClass();
 		}
 
@@ -8645,19 +8642,17 @@ void FBlueprintEditorUtils::PostSetupObjectPinType(UBlueprint* InBlueprint, FBPV
 
 const FSlateBrush* FBlueprintEditorUtils::GetIconFromPin( const FEdGraphPinType& PinType, bool bIsLarge )
 {
-	const UEdGraphSchema_K2* K2Schema = GetDefault<UEdGraphSchema_K2>();
-	
 	const FSlateBrush* IconBrush = FEditorStyle::GetBrush(TEXT("Kismet.VariableList.TypeIcon"));
 	const UObject* PinSubObject = PinType.PinSubCategoryObject.Get();
-	if( PinType.IsArray() && PinType.PinCategory != K2Schema->PC_Exec )
+	if( PinType.IsArray() && PinType.PinCategory != UEdGraphSchema_K2::PC_Exec )
 	{
 		IconBrush = FEditorStyle::GetBrush(TEXT("Kismet.VariableList.ArrayTypeIcon"));
 	}
-	else if (PinType.IsMap() && PinType.PinCategory != K2Schema->PC_Exec)
+	else if (PinType.IsMap() && PinType.PinCategory != UEdGraphSchema_K2::PC_Exec)
 	{
 		IconBrush = FEditorStyle::GetBrush(TEXT("Kismet.VariableList.MapKeyTypeIcon"));
 	}
-	else if (PinType.IsSet() && PinType.PinCategory != K2Schema->PC_Exec)
+	else if (PinType.IsSet() && PinType.PinCategory != UEdGraphSchema_K2::PC_Exec)
 	{
 		if( bIsLarge )
 		{
@@ -8681,8 +8676,7 @@ const FSlateBrush* FBlueprintEditorUtils::GetIconFromPin( const FEdGraphPinType&
 
 const FSlateBrush* FBlueprintEditorUtils::GetSecondaryIconFromPin(const FEdGraphPinType& PinType)
 {
-	const UEdGraphSchema_K2* K2Schema = GetDefault<UEdGraphSchema_K2>();
-	if (PinType.IsMap() && PinType.PinCategory != K2Schema->PC_Exec)
+	if (PinType.IsMap() && PinType.PinCategory != UEdGraphSchema_K2::PC_Exec)
 	{
 		return FEditorStyle::GetBrush(TEXT("Kismet.VariableList.MapValueTypeIcon"));
 	}
@@ -8691,18 +8685,17 @@ const FSlateBrush* FBlueprintEditorUtils::GetSecondaryIconFromPin(const FEdGraph
 
 bool FBlueprintEditorUtils::HasGetTypeHash(const FEdGraphPinType& PinType)
 {
-	const UEdGraphSchema_K2* K2Schema = GetDefault<UEdGraphSchema_K2>();
-	if(PinType.PinCategory == K2Schema->PC_Boolean)
+	if(PinType.PinCategory == UEdGraphSchema_K2::PC_Boolean)
 	{
 		return false;
 	}
 
-	if (PinType.PinCategory == K2Schema->PC_Text)
+	if (PinType.PinCategory == UEdGraphSchema_K2::PC_Text)
 	{
 		return false;
 	}
 
-	if (PinType.PinCategory != K2Schema->PC_Struct)
+	if (PinType.PinCategory != UEdGraphSchema_K2::PC_Struct)
 	{
 		// even object or class types can be hashed, no reason to investigate further
 		return true;
diff --git a/Engine/Source/Editor/UnrealEd/Private/Kismet2/Kismet2.cpp b/Engine/Source/Editor/UnrealEd/Private/Kismet2/Kismet2.cpp
--- a/Engine/Source/Editor/UnrealEd/Private/Kismet2/Kismet2.cpp
+++ b/Engine/Source/Editor/UnrealEd/Private/Kismet2/Kismet2.cpp
@@ -401,9 +401,7 @@ TMultiMap<void*, FKismetEditorUtilities::FOnBlueprintCreatedData> FKismetEditorU
 /** Create the correct event graphs for this blueprint */
 void FKismetEditorUtilities::CreateDefaultEventGraphs(UBlueprint* Blueprint)
 {
-	const UEdGraphSchema_K2* K2Schema = GetDefault<UEdGraphSchema_K2>();
-
-	UEdGraph* Ubergraph = FBlueprintEditorUtils::CreateNewGraph(Blueprint, K2Schema->GN_EventGraph, UEdGraph::StaticClass(), UEdGraphSchema_K2::StaticClass());
+	UEdGraph* Ubergraph = FBlueprintEditorUtils::CreateNewGraph(Blueprint, UEdGraphSchema_K2::GN_EventGraph, UEdGraph::StaticClass(), UEdGraphSchema_K2::StaticClass());
 	Ubergraph->bAllowDeletion = false; //@TODO: Really, just want to make sure we never drop below 1, not that you cannot delete any particular one!
 	FBlueprintEditorUtils::AddUbergraphPage(Blueprint, Ubergraph);
 
@@ -434,8 +432,6 @@ UBlueprint* FKismetEditorUtilities::CreateBlueprint(UClass* ParentClass, UObject
 	NewBP->bLegacyNeedToPurgeSkelRefs = false;
 	NewBP->GenerateNewGuid();
 
-	const UEdGraphSchema_K2* K2Schema = GetDefault<UEdGraphSchema_K2>();
-
 	// Create SimpleConstructionScript and UserConstructionScript
 	if (FBlueprintEditorUtils::SupportsConstructionScript(NewBP))
 	{ 
@@ -453,7 +449,7 @@ UBlueprint* FKismetEditorUtilities::CreateBlueprint(UClass* ParentClass, UObject
 		NewBP->SimpleConstructionScript->SetFlags(RF_Transactional);
 		NewBP->LastEditedDocuments.Add(NewBP->SimpleConstructionScript);
 
-		UEdGraph* UCSGraph = FBlueprintEditorUtils::CreateNewGraph(NewBP, K2Schema->FN_UserConstructionScript, UEdGraph::StaticClass(), UEdGraphSchema_K2::StaticClass());
+		UEdGraph* UCSGraph = FBlueprintEditorUtils::CreateNewGraph(NewBP, UEdGraphSchema_K2::FN_UserConstructionScript, UEdGraph::StaticClass(), UEdGraphSchema_K2::StaticClass());
 		FBlueprintEditorUtils::AddFunctionGraph(NewBP, UCSGraph, /*bIsUserCreated=*/ false, AActor::StaticClass());
 
 		// If the blueprint is derived from another blueprint, add in a super-call automatically
@@ -463,15 +459,15 @@ UBlueprint* FKismetEditorUtilities::CreateBlueprint(UClass* ParentClass, UObject
 			UK2Node_FunctionEntry* UCSEntry = CastChecked<UK2Node_FunctionEntry>(UCSGraph->Nodes[0]);
 			FGraphNodeCreator<UK2Node_CallParentFunction> FunctionNodeCreator(*UCSGraph);
 			UK2Node_CallParentFunction* ParentFunctionNode = FunctionNodeCreator.CreateNode();
-			ParentFunctionNode->FunctionReference.SetExternalMember(K2Schema->FN_UserConstructionScript, NewBP->ParentClass);
+			ParentFunctionNode->FunctionReference.SetExternalMember(UEdGraphSchema_K2::FN_UserConstructionScript, NewBP->ParentClass);
 			ParentFunctionNode->NodePosX = 200;
 			ParentFunctionNode->NodePosY = 0;
 			ParentFunctionNode->AllocateDefaultPins();
 			FunctionNodeCreator.Finalize();
 
 			// Wire up the new node
-			UEdGraphPin* ExecPin = UCSEntry->FindPin(K2Schema->PN_Then);
-			UEdGraphPin* SuperPin = ParentFunctionNode->FindPin(K2Schema->PN_Execute);
+			UEdGraphPin* ExecPin = UCSEntry->FindPin(UEdGraphSchema_K2::PN_Then);
+			UEdGraphPin* SuperPin = ParentFunctionNode->FindPin(UEdGraphSchema_K2::PN_Execute);
 			ExecPin->MakeLinkTo(SuperPin);
 		}
 
@@ -493,7 +489,7 @@ UBlueprint* FKismetEditorUtilities::CreateBlueprint(UClass* ParentClass, UObject
 		if (RootAnimBP == NULL)
 		{
 			// Only allow an anim graph if there isn't one in a parent blueprint
-			UEdGraph* NewGraph = FBlueprintEditorUtils::CreateNewGraph(AnimBP, K2Schema->GN_AnimGraph, UAnimationGraph::StaticClass(), UAnimationGraphSchema::StaticClass());
+			UEdGraph* NewGraph = FBlueprintEditorUtils::CreateNewGraph(AnimBP, UEdGraphSchema_K2::GN_AnimGraph, UAnimationGraph::StaticClass(), UAnimationGraphSchema::StaticClass());
 			FBlueprintEditorUtils::AddDomainSpecificGraph(NewBP, NewGraph);
 			NewBP->LastEditedDocuments.Add(NewGraph);
 			NewGraph->bAllowDeletion = false;
@@ -667,7 +663,7 @@ UK2Node_Event* FKismetEditorUtilities::AddDefaultEventNode(UBlueprint* InBluepri
 					ParentPin->MakeLinkTo(EventPin);
 				}
 			}
-			ParentFunctionNode->GetExecPin()->MakeLinkTo(NewEventNode->FindPin(Schema->PN_Then));
+			ParentFunctionNode->GetExecPin()->MakeLinkTo(NewEventNode->FindPin(UEdGraphSchema_K2::PN_Then));
 
 			ParentFunctionNode->NodePosX = FunctionFromNode.Node->NodePosX + FunctionFromNode.Node->NodeWidth + 200;
 			ParentFunctionNode->NodePosY = FunctionFromNode.Node->NodePosY;
@@ -1988,15 +1984,13 @@ void FKismetEditorUtilities::UpgradeCosmeticallyStaleBlueprint(UBlueprint* Bluep
 	// Rename the ubergraph page 'StateGraph' to be named 'EventGraph' if possible
 	if (FBlueprintEditorUtils::DoesSupportEventGraphs(Blueprint))
 	{
-		const UEdGraphSchema_K2* K2Schema = GetDefault<UEdGraphSchema_K2>();
-
 		UEdGraph* OldStateGraph = FindObject<UEdGraph>(Blueprint, TEXT("StateGraph"));
-		UObject* CollidingObject = FindObject<UObject>(Blueprint, *(K2Schema->GN_EventGraph.ToString()));
+		UObject* CollidingObject = FindObject<UObject>(Blueprint, *(UEdGraphSchema_K2::GN_EventGraph.ToString()));
 
 		if ((OldStateGraph != NULL) && (CollidingObject == NULL))
 		{
 			check(!OldStateGraph->HasAnyFlags(RF_Public));
-			OldStateGraph->Rename(*(K2Schema->GN_EventGraph.ToString()), OldStateGraph->GetOuter(), REN_DoNotDirty | REN_ForceNoResetLoaders);
+			OldStateGraph->Rename(*(UEdGraphSchema_K2::GN_EventGraph.ToString()), OldStateGraph->GetOuter(), REN_DoNotDirty | REN_ForceNoResetLoaders);
 			Blueprint->Status = BS_Dirty;
 		}
 	}
diff --git a/Engine/Source/Editor/UnrealEd/Private/Kismet2/KismetDebugUtilities.cpp b/Engine/Source/Editor/UnrealEd/Private/Kismet2/KismetDebugUtilities.cpp
--- a/Engine/Source/Editor/UnrealEd/Private/Kismet2/KismetDebugUtilities.cpp
+++ b/Engine/Source/Editor/UnrealEd/Private/Kismet2/KismetDebugUtilities.cpp
@@ -826,12 +826,11 @@ void FKismetDebugUtilities::GetValidBreakpointLocations(const UK2Node_MacroInsta
 	if (!bIsMacroPure && MacroEntryNode)
 	{
 		// Get the execute pin outputs on the entry node
-		const UEdGraphSchema_K2* K2Schema = GetDefault<UEdGraphSchema_K2>();
 		for (auto PinIt = MacroEntryNode->Pins.CreateConstIterator(); PinIt; ++PinIt)
 		{
 			const UEdGraphPin* ExecPin = *PinIt;
 			if (ExecPin && ExecPin->Direction == EGPD_Output
-				&& ExecPin->PinType.PinCategory == K2Schema->PC_Exec)
+				&& ExecPin->PinType.PinCategory == UEdGraphSchema_K2::PC_Exec)
 			{
 				// For each pin linked to each execute pin, collect the node that owns it
 				for (auto LinkedToPinIt = ExecPin->LinkedTo.CreateConstIterator(); LinkedToPinIt; ++LinkedToPinIt)
diff --git a/Engine/Source/Editor/UnrealEd/Private/Kismet2/StructureEditorUtils.cpp b/Engine/Source/Editor/UnrealEd/Private/Kismet2/StructureEditorUtils.cpp
--- a/Engine/Source/Editor/UnrealEd/Private/Kismet2/StructureEditorUtils.cpp
+++ b/Engine/Source/Editor/UnrealEd/Private/Kismet2/StructureEditorUtils.cpp
@@ -49,8 +49,7 @@ UUserDefinedStruct* FStructureEditorUtils::CreateUserDefinedStruct(UObject* InPa
 		Struct->Status = UDSS_Error;
 
 		{
-			const UEdGraphSchema_K2* K2Schema = GetDefault<UEdGraphSchema_K2>();
-			AddVariable(Struct, FEdGraphPinType(K2Schema->PC_Boolean, FString(), nullptr, EPinContainerType::None, false, FEdGraphTerminalType()));
+			AddVariable(Struct, FEdGraphPinType(UEdGraphSchema_K2::PC_Boolean, NAME_None, nullptr, EPinContainerType::None, false, FEdGraphTerminalType()));
 		}
 	}
 
@@ -177,8 +176,7 @@ FStructureEditorUtils::EStructureError FStructureEditorUtils::IsStructureValid(c
 
 bool FStructureEditorUtils::CanHaveAMemberVariableOfType(const UUserDefinedStruct* Struct, const FEdGraphPinType& VarType, FString* OutMsg)
 {
-	const UEdGraphSchema_K2* K2Schema = GetDefault<UEdGraphSchema_K2>();
-	if ((VarType.PinCategory == K2Schema->PC_Struct) && Struct)
+	if ((VarType.PinCategory == UEdGraphSchema_K2::PC_Struct) && Struct)
 	{
 		if (const UScriptStruct* SubCategoryStruct = Cast<const UScriptStruct>(VarType.PinSubCategoryObject.Get()))
 		{
@@ -197,10 +195,10 @@ bool FStructureEditorUtils::CanHaveAMemberVariableOfType(const UUserDefinedStruc
 			return false;
 		}
 	}
-	else if ((VarType.PinCategory == K2Schema->PC_Exec) 
-		|| (VarType.PinCategory == K2Schema->PC_Wildcard)
-		|| (VarType.PinCategory == K2Schema->PC_MCDelegate)
-		|| (VarType.PinCategory == K2Schema->PC_Delegate))
+	else if ((VarType.PinCategory == UEdGraphSchema_K2::PC_Exec) 
+		|| (VarType.PinCategory == UEdGraphSchema_K2::PC_Wildcard)
+		|| (VarType.PinCategory == UEdGraphSchema_K2::PC_MCDelegate)
+		|| (VarType.PinCategory == UEdGraphSchema_K2::PC_Delegate))
 	{
 		if (OutMsg)
 		{
@@ -405,27 +403,27 @@ bool FStructureEditorUtils::ChangeVariableDefaultValue(UUserDefinedStruct* Struc
 {
 	auto ValidateDefaultValue = [](const FStructVariableDescription& VarDesc, const FString& InNewDefaultValue) -> bool
 	{
-		const UEdGraphSchema_K2* K2Schema = GetDefault<UEdGraphSchema_K2>();
 		const FEdGraphPinType PinType = VarDesc.ToPinType();
 
 		bool bResult = false;
 		//TODO: validation for values, that are not passed by string
-		if (PinType.PinCategory == K2Schema->PC_Text)
+		if (PinType.PinCategory == UEdGraphSchema_K2::PC_Text)
 		{
 			bResult = true;
 		}
-		else if ((PinType.PinCategory == K2Schema->PC_Object) 
-			|| (PinType.PinCategory == K2Schema->PC_Interface) 
-			|| (PinType.PinCategory == K2Schema->PC_Class)
-			|| (PinType.PinCategory == K2Schema->PC_SoftClass)
-			|| (PinType.PinCategory == K2Schema->PC_SoftObject))
+		else if ((PinType.PinCategory == UEdGraphSchema_K2::PC_Object) 
+			|| (PinType.PinCategory == UEdGraphSchema_K2::PC_Interface) 
+			|| (PinType.PinCategory == UEdGraphSchema_K2::PC_Class)
+			|| (PinType.PinCategory == UEdGraphSchema_K2::PC_SoftClass)
+			|| (PinType.PinCategory == UEdGraphSchema_K2::PC_SoftObject))
 		{
 			// K2Schema->DefaultValueSimpleValidation finds an object, passed by path, invalid
 			bResult = true;
 		}
 		else
 		{
-			bResult = K2Schema->DefaultValueSimpleValidation(PinType, FString(), InNewDefaultValue, NULL, FText::GetEmpty());
+			const UEdGraphSchema_K2* K2Schema = GetDefault<UEdGraphSchema_K2>();
+			bResult = K2Schema->DefaultValueSimpleValidation(PinType, NAME_None, InNewDefaultValue, nullptr, FText::GetEmpty());
 		}
 		return bResult;
 	};
@@ -589,7 +587,6 @@ void FStructureEditorUtils::RemoveInvalidStructureMemberVariableFromBlueprint(UB
 {
 	if (Blueprint)
 	{
-		const UEdGraphSchema_K2* K2Schema = GetDefault<UEdGraphSchema_K2>();
 		const UScriptStruct* FallbackStruct = GetFallbackStruct();
 
 		FString DislpayList;
@@ -597,7 +594,7 @@ void FStructureEditorUtils::RemoveInvalidStructureMemberVariableFromBlueprint(UB
 		for (int32 VarIndex = 0; VarIndex < Blueprint->NewVariables.Num(); ++VarIndex)
 		{
 			const FBPVariableDescription& Var = Blueprint->NewVariables[VarIndex];
-			if (Var.VarType.PinCategory == K2Schema->PC_Struct)
+			if (Var.VarType.PinCategory == UEdGraphSchema_K2::PC_Struct)
 			{
 				const UScriptStruct* ScriptStruct = Cast<const UScriptStruct>(Var.VarType.PinSubCategoryObject.Get());
 				const bool bInvalidStruct = (NULL == ScriptStruct) || (FallbackStruct == ScriptStruct);
diff --git a/Engine/Source/Editor/UnrealEd/Private/MaterialGraphNode.cpp b/Engine/Source/Editor/UnrealEd/Private/MaterialGraphNode.cpp
--- a/Engine/Source/Editor/UnrealEd/Private/MaterialGraphNode.cpp
+++ b/Engine/Source/Editor/UnrealEd/Private/MaterialGraphNode.cpp
@@ -431,46 +431,67 @@ void UMaterialGraphNode::GetContextMenuActions(const FGraphNodeContextMenuBuilde
 	}
 }
 
-FString UMaterialGraphNode::GetShortenPinName(const FString& PinName)
-{
-	FString InputName = PinName;
+namespace MaterialPinNames
+{
+	static const FName Coordinates(TEXT("Coordinates"));
+	static const FName UVs(TEXT("UVs"));
+	static const FName TextureObject(TEXT("TextureObject"));
+	static const FName Tex(TEXT("Tex"));
+	static const FName Input(TEXT("Input"));
+	static const FName Exponent(TEXT("Exponent"));
+	static const FName Exp(TEXT("Exp"));
+	static const FName AGreaterThanB(TEXT("AGreaterThanB"));
+	static const FName CompactAGreaterThanB(TEXT("A > B"));
+	static const FName AEqualsB(TEXT("AEqualsB"));
+	static const FName CompactAEqualsB(TEXT("A == B"));
+	static const FName ALessThanB(TEXT("ALessThanB"));
+	static const FName CompactALessThanB(TEXT("A < B"));
+	static const FName MipLevel(TEXT("MipLevel"));
+	static const FName Level(TEXT("Level"));
+	static const FName MipBias(TEXT("MipBias"));
+	static const FName Bias(TEXT("Bias"));
+}
+
+FName UMaterialGraphNode::GetShortenPinName(const FName PinName)
+{
+	FName InputName = PinName;
 
 	// Shorten long expression input names.
-	if (!FCString::Stricmp(*PinName, TEXT("Coordinates")))
+	if (PinName == MaterialPinNames::Coordinates)
 	{
-		InputName = TEXT("UVs");
+		InputName = MaterialPinNames::UVs;
 	}
-	else if (!FCString::Stricmp(*PinName, TEXT("TextureObject")))
+	else if (PinName == MaterialPinNames::TextureObject)
 	{
-		InputName = TEXT("Tex");
+		InputName = MaterialPinNames::Tex;
 	}
-	else if (!FCString::Stricmp(*PinName, TEXT("Input")))
+	else if (PinName == MaterialPinNames::Input)
 	{
-		InputName = TEXT("");
+		InputName = NAME_None;
 	}
-	else if (!FCString::Stricmp(*PinName, TEXT("Exponent")))
+	else if (PinName == MaterialPinNames::Exponent)
 	{
-		InputName = TEXT("Exp");
+		InputName = MaterialPinNames::Exp;
 	}
-	else if (!FCString::Stricmp(*PinName, TEXT("AGreaterThanB")))
+	else if (PinName == MaterialPinNames::AGreaterThanB)
 	{
-		InputName = TEXT("A > B");
+		InputName = MaterialPinNames::CompactAGreaterThanB;
 	}
-	else if (!FCString::Stricmp(*PinName, TEXT("AEqualsB")))
+	else if (PinName == MaterialPinNames::AEqualsB)
 	{
-		InputName = TEXT("A == B");
+		InputName = MaterialPinNames::CompactAEqualsB;
 	}
-	else if (!FCString::Stricmp(*PinName, TEXT("ALessThanB")))
+	else if (PinName == MaterialPinNames::ALessThanB)
 	{
-		InputName = TEXT("A < B");
+		InputName = MaterialPinNames::CompactALessThanB;
 	}
-	else if (!FCString::Stricmp(*PinName, TEXT("MipLevel")))
+	else if (PinName == MaterialPinNames::MipLevel)
 	{
-		InputName = TEXT("Level");
+		InputName = MaterialPinNames::Level;
 	}
-	else if (!FCString::Stricmp(*PinName, TEXT("MipBias")))
+	else if (PinName == MaterialPinNames::MipBias)
 	{
-		InputName = TEXT("Bias");
+		InputName = MaterialPinNames::Bias;
 	}
 
 	return InputName;
@@ -483,16 +504,14 @@ void UMaterialGraphNode::CreateInputPins()
 	for (int32 Index = 0; Index < ExpressionInputs.Num() ; ++Index)
 	{
 		FExpressionInput* Input = ExpressionInputs[Index];
-		FString InputName = MaterialExpression->GetInputName(Index);
+		FName InputName = MaterialExpression->GetInputName(Index);
 
 		InputName = GetShortenPinName(InputName);
 
-		const UMaterialGraphSchema* Schema = CastChecked<UMaterialGraphSchema>(GetSchema());
-		FString PinCategory = MaterialExpression->IsInputConnectionRequired(Index) ? Schema->PC_Required : Schema->PC_Optional;
-		FString PinSubCategory = TEXT("");
+		const FName PinCategory = MaterialExpression->IsInputConnectionRequired(Index) ? UMaterialGraphSchema::PC_Required : UMaterialGraphSchema::PC_Optional;
 
-		UEdGraphPin* NewPin = CreatePin(EGPD_Input, PinCategory, PinSubCategory, nullptr, InputName);
-		if (NewPin->PinName.IsEmpty())
+		UEdGraphPin* NewPin = CreatePin(EGPD_Input, PinCategory, InputName);
+		if (NewPin->PinName.IsNone())
 		{
 			// Makes sure pin has a name for lookup purposes but user will never see it
 			NewPin->PinName = CreateUniquePinName(TEXT("Input"));
@@ -505,36 +524,33 @@ void UMaterialGraphNode::CreateOutputPins()
 {
 	TArray<FExpressionOutput>& Outputs = MaterialExpression->GetOutputs();
 
-	for( int32 Index = 0 ; Index < Outputs.Num() ; ++Index )
+	for (const FExpressionOutput& ExpressionOutput : Outputs)
 	{
-		const FExpressionOutput& ExpressionOutput = Outputs[Index];
-		FString PinCategory = TEXT("");
-		FString PinSubCategory = TEXT("");
-		FString OutputName = TEXT("");
+		FName PinCategory;
+		FName PinSubCategory;
+		FName OutputName;
 
-		const UMaterialGraphSchema* Schema = CastChecked<UMaterialGraphSchema>(GetSchema());
-		
 		if (MaterialExpression->bShowMaskColorsOnPin)
 		{
 			if (ExpressionOutput.Mask)
 			{
-				PinCategory = Schema->PC_Mask;
+				PinCategory = UMaterialGraphSchema::PC_Mask;
 
 				if (ExpressionOutput.MaskR && !ExpressionOutput.MaskG && !ExpressionOutput.MaskB && !ExpressionOutput.MaskA)
 				{
-					PinSubCategory = Schema->PSC_Red;
+					PinSubCategory = UMaterialGraphSchema::PSC_Red;
 				}
 				else if (!ExpressionOutput.MaskR &&  ExpressionOutput.MaskG && !ExpressionOutput.MaskB && !ExpressionOutput.MaskA)
 				{
-					PinSubCategory = Schema->PSC_Green;
+					PinSubCategory = UMaterialGraphSchema::PSC_Green;
 				}
 				else if (!ExpressionOutput.MaskR && !ExpressionOutput.MaskG &&  ExpressionOutput.MaskB && !ExpressionOutput.MaskA)
 				{
-					PinSubCategory = Schema->PSC_Blue;
+					PinSubCategory = UMaterialGraphSchema::PSC_Blue;
 				}
 				else if (!ExpressionOutput.MaskR && !ExpressionOutput.MaskG && !ExpressionOutput.MaskB &&  ExpressionOutput.MaskA)
 				{
-					PinSubCategory = Schema->PSC_Alpha;
+					PinSubCategory = UMaterialGraphSchema::PSC_Alpha;
 				}
 			}
 		}
@@ -544,8 +560,8 @@ void UMaterialGraphNode::CreateOutputPins()
 			OutputName = ExpressionOutput.OutputName;
 		}
 
-		UEdGraphPin* NewPin = CreatePin(EGPD_Output, PinCategory, PinSubCategory, nullptr, OutputName);
-		if (NewPin->PinName.IsEmpty())
+		UEdGraphPin* NewPin = CreatePin(EGPD_Output, PinCategory, PinSubCategory, OutputName);
+		if (NewPin->PinName.IsNone())
 		{
 			// Makes sure pin has a name for lookup purposes but user will never see it
 			NewPin->PinName = CreateUniquePinName(TEXT("Output"));
diff --git a/Engine/Source/Editor/UnrealEd/Private/MaterialGraphNode_Knot.cpp b/Engine/Source/Editor/UnrealEd/Private/MaterialGraphNode_Knot.cpp
--- a/Engine/Source/Editor/UnrealEd/Private/MaterialGraphNode_Knot.cpp
+++ b/Engine/Source/Editor/UnrealEd/Private/MaterialGraphNode_Knot.cpp
@@ -22,13 +22,13 @@ UMaterialGraphNode_Knot::UMaterialGraphNode_Knot(const FObjectInitializer& Objec
 
 void UMaterialGraphNode_Knot::AllocateDefaultPins()
 {
-	const FString InputPinName(TEXT("InputPin"));
-	const FString OutputPinName(TEXT("OutputPin"));
+	const FName InputPinName(TEXT("InputPin"));
+	const FName OutputPinName(TEXT("OutputPin"));
 
-	UEdGraphPin* MyInputPin = CreatePin(EGPD_Input, PC_Wildcard, FString(), nullptr, InputPinName);
+	UEdGraphPin* MyInputPin = CreatePin(EGPD_Input, PC_Wildcard, InputPinName);
 	MyInputPin->bDefaultValueIsIgnored = true;
 
-	UEdGraphPin* MyOutputPin = CreatePin(EGPD_Output, PC_Wildcard, FString(), nullptr, OutputPinName);
+	UEdGraphPin* MyOutputPin = CreatePin(EGPD_Output, PC_Wildcard, OutputPinName);
 }
 
 FText UMaterialGraphNode_Knot::GetTooltipText() const
diff --git a/Engine/Source/Editor/UnrealEd/Private/MaterialGraphNode_Root.cpp b/Engine/Source/Editor/UnrealEd/Private/MaterialGraphNode_Root.cpp
--- a/Engine/Source/Editor/UnrealEd/Private/MaterialGraphNode_Root.cpp
+++ b/Engine/Source/Editor/UnrealEd/Private/MaterialGraphNode_Root.cpp
@@ -48,11 +48,10 @@ void UMaterialGraphNode_Root::PostPlacedNewNode()
 void UMaterialGraphNode_Root::CreateInputPins()
 {
 	UMaterialGraph* MaterialGraph = CastChecked<UMaterialGraph>(GetGraph());
-	const UMaterialGraphSchema* Schema = CastChecked<UMaterialGraphSchema>(GetSchema());
 
 	for (const FMaterialInputInfo& MaterialInput : MaterialGraph->MaterialInputs)
 	{
-		UEdGraphPin* InputPin = CreatePin(EGPD_Input, Schema->PC_MaterialInput, FString::Printf(TEXT("%d"), (int32)MaterialInput.GetProperty()), nullptr, MaterialInput.GetName().ToString());
+		UEdGraphPin* InputPin = CreatePin(EGPD_Input, UMaterialGraphSchema::PC_MaterialInput, *FString::Printf(TEXT("%d"), (int32)MaterialInput.GetProperty()), *MaterialInput.GetName().ToString());
 	}
 
 }
