Fix merge issue that was causing a failure for Launch on Device w/ BP nativization enabled.

diff --git a/Engine/Source/Editor/UnrealEd/Private/CookOnTheFlyServer.cpp b/Engine/Source/Editor/UnrealEd/Private/CookOnTheFlyServer.cpp
--- a/Engine/Source/Editor/UnrealEd/Private/CookOnTheFlyServer.cpp
+++ b/Engine/Source/Editor/UnrealEd/Private/CookOnTheFlyServer.cpp
@@ -5878,7 +5878,9 @@ void UCookOnTheFlyServer::StartCookByTheBook( const FCookByTheBookStartupOptions
 		StartSavingEDLCookInfoForVerification();
 	}
 
-	if (PackagingSettings->BlueprintNativizationMethod != EProjectPackagingBlueprintNativizationMethod::Disabled)
+	// Note: Nativization only works with "cook by the book" mode and not from within the current editor process.
+	if (CurrentCookMode == ECookMode::CookByTheBook
+		&& PackagingSettings->BlueprintNativizationMethod != EProjectPackagingBlueprintNativizationMethod::Disabled)
 	{
 		FNativeCodeGenInitData CodeGenData;
 		for (const ITargetPlatform* Entry : CookByTheBookStartupOptions.TargetPlatforms)
