Fix NoPhysX editor
#rnx

diff --git a/Engine/Source/Editor/UnrealEd/Private/EditorEngine.cpp b/Engine/Source/Editor/UnrealEd/Private/EditorEngine.cpp
--- a/Engine/Source/Editor/UnrealEd/Private/EditorEngine.cpp
+++ b/Engine/Source/Editor/UnrealEd/Private/EditorEngine.cpp
@@ -4350,10 +4350,12 @@ void UEditorEngine::CleanupPhysicsSceneThatWasInitializedForSave(UWorld* World)
 	// Make sure we clean up the physics scene here. If we leave too many scenes in memory, undefined behavior occurs when locking a scene for read/write.
 	World->ClearWorldComponents();
 	World->SetPhysicsScene(nullptr);
+#if WITH_PHYSX
 	if (GPhysCommandHandler)
 	{
 		GPhysCommandHandler->Flush();
 	}
+#endif // WITH_PHYSX
 
 	// Update components again in case it was a world without a physics scene but did have rendered components.
 	World->UpdateWorldComponents(true, true);
