Fix limit on number of async precache tasks while cooking, due to counter not being updated correctly.

diff --git a/Engine/Source/Editor/UnrealEd/Private/CookOnTheFlyServer.cpp b/Engine/Source/Editor/UnrealEd/Private/CookOnTheFlyServer.cpp
--- a/Engine/Source/Editor/UnrealEd/Private/CookOnTheFlyServer.cpp
+++ b/Engine/Source/Editor/UnrealEd/Private/CookOnTheFlyServer.cpp
@@ -1867,7 +1867,7 @@ bool UCookOnTheFlyServer::BeginPackageCacheForCookedPlatformData(UPackage* Packa
 				{
 					*Value += 1;
 				}
-				CurrentAsyncCache -= 1;
+				*CurrentAsyncCache -= 1;
 			}
 
 			if (Obj->IsA(UMaterialInterface::StaticClass()))
