Copying //UE4/Dev-Sequencer to //UE4/Dev-Main (Source: //UE4/Dev-Sequencer @ 3946692)

diff --git a/Engine/Source/Editor/UnrealEd/Private/Fbx/FbxMainExport.cpp b/Engine/Source/Editor/UnrealEd/Private/Fbx/FbxMainExport.cpp
--- a/Engine/Source/Editor/UnrealEd/Private/Fbx/FbxMainExport.cpp
+++ b/Engine/Source/Editor/UnrealEd/Private/Fbx/FbxMainExport.cpp
@@ -18,8 +18,7 @@
 #include "Components/StaticMeshComponent.h"
 #include "Components/LightComponent.h"
 #include "Model.h"
-#include "Curves/KeyHandle.h"
-#include "Curves/RichCurve.h"
+#include "Channels/MovieSceneFloatChannel.h"
 #include "Animation/AnimTypes.h"
 #include "Components/SkeletalMeshComponent.h"
 #include "Editor/EditorPerProjectUserSettings.h"
@@ -37,6 +36,7 @@
 #include "Engine/Polys.h"
 #include "Engine/StaticMesh.h"
 #include "Editor.h"
+#include "Channels/MovieSceneChannelProxy.h"
 
 #include "Materials/Material.h"
 #include "Materials/MaterialExpressionConstant.h"
@@ -81,6 +81,7 @@
 #include "Evaluation/MovieScenePlayback.h"
 #include "Evaluation/MovieSceneEvaluationTemplateInstance.h"
 #include "MovieSceneSequence.h"
+#include "MovieSceneTimeHelpers.h"
 
 #if WITH_PHYSX
 #include "DynamicMeshBuilder.h"
@@ -1452,19 +1453,21 @@ FFbxExporter::FLevelSequenceAnimTrackAdapter::FLevelSequenceAnimTrackAdapter( IM
 
 float FFbxExporter::FLevelSequenceAnimTrackAdapter::GetAnimationStart() const
 {
-	return MovieScene->GetPlaybackRange().GetLowerBoundValue();
+	return MovieScene->GetPlaybackRange().GetLowerBoundValue() / MovieScene->GetFrameResolution();
 }
 
 float FFbxExporter::FLevelSequenceAnimTrackAdapter::GetAnimationLength() const
 {
-	return MovieScene->GetPlaybackRange().Size<float>();
+	return FFrameNumber(MovieScene::DiscreteSize(MovieScene->GetPlaybackRange())) / MovieScene->GetFrameResolution();
 }
 
 
 void FFbxExporter::FLevelSequenceAnimTrackAdapter::UpdateAnimation( float Time )
 {
-	FMovieSceneContext Context = FMovieSceneContext(FMovieSceneEvaluationRange(Time), MovieScenePlayer->GetPlaybackStatus()).SetHasJumped(true);
-	return MovieScenePlayer->GetEvaluationTemplate().Evaluate( Context, *MovieScenePlayer );
+	// Convert the time in seconds to a frame time
+	FFrameTime FrameTime = Time * MovieScene->GetFrameResolution();
+	FMovieSceneContext Context = FMovieSceneContext(FMovieSceneEvaluationRange(FrameTime, MovieScene->GetFrameResolution()), MovieScenePlayer->GetPlaybackStatus()).SetHasJumped(true);
+	MovieScenePlayer->GetEvaluationTemplate().Evaluate( Context, *MovieScenePlayer );
 }
 
 
@@ -1474,7 +1477,7 @@ bool FFbxExporter::ExportLevelSequence( UMovieScene* MovieScene, const TArray<FG
 	{
 		return false;
 	}
-			
+
 	for ( const FMovieSceneBinding& MovieSceneBinding : MovieScene->GetBindings() )
 	{
 		// If there are specific bindings to export, export those only
@@ -2261,18 +2264,6 @@ void FFbxExporter::ExportAnimatedFloat(FbxProperty* FbxProperty, FInterpCurveFlo
 	AnimCurve->KeyModifyEnd();
 }
 
-FKeyHandle FindRichCurveKey(FRichCurve& InCurve, float InKeyTime)
-{
-	for ( auto It = InCurve.GetKeyHandleIterator(); It; ++It )
-	{
-		if (IsEquivalent(InCurve.GetKeyTime(It.Key()), InKeyTime))
-		{
-			return It.Key();
-		}
-	}
-	return FKeyHandle();
-}
-
 
 void RichCurveInterpolationToFbxInterpolation(ERichCurveInterpMode InInterpolation, ERichCurveTangentMode InTangentMode, FbxAnimCurveDef::EInterpolationType &OutInterpolation, FbxAnimCurveDef::ETangentMode &OutTangentMode)
 {
@@ -2304,43 +2295,39 @@ void RichCurveInterpolationToFbxInterpolation(ERichCurveInterpMode InInterpolati
 	}
 }
 
-void FFbxExporter::ExportRichCurveToFbxCurve(FbxAnimCurve& InFbxCurve, FRichCurve& InRichCurve, ERichCurveValueMode ValueMode, bool bNegative)
+void FFbxExporter::ExportChannelToFbxCurve(FbxAnimCurve& InFbxCurve, const FMovieSceneFloatChannel& InChannel, FFrameRate FrameResolution, ERichCurveValueMode ValueMode, bool bNegative)
 {
+	const float NegateFactor = bNegative ? -1.f : 1.f;
+
 	InFbxCurve.KeyModifyBegin();
 
-	for ( auto It = InRichCurve.GetKeyHandleIterator(); It; ++It )
+	TArrayView<const FFrameNumber>          Times  = InChannel.GetTimes();
+	TArrayView<const FMovieSceneFloatValue> Values = InChannel.GetValues();
+
+	for (int32 Index = 0; Index < Times.Num(); ++Index)
 	{
-		FKeyHandle KeyHandle = It.Key();
-		float KeyTime = InRichCurve.GetKeyTime(KeyHandle);
-		float Value = ValueMode == ERichCurveValueMode::Fov
-			? DefaultCamera->ComputeFocalLength( InRichCurve.Eval( KeyTime ) )
-			: InRichCurve.Eval( KeyTime );
+		FFrameNumber          KeyTime  = Times[Index];
+		FMovieSceneFloatValue KeyValue = Values[Index];
+
+		const float Value = (ValueMode == ERichCurveValueMode::Fov ? DefaultCamera->ComputeFocalLength( KeyValue.Value ) : KeyValue.Value) * NegateFactor;
 
 		FbxTime FbxTime;
 		FbxAnimCurveKey FbxKey;
-		FbxTime.SetSecondDouble((float)KeyTime);
+		FbxTime.SetSecondDouble(KeyTime / FrameResolution);
 
 		int FbxKeyIndex = InFbxCurve.KeyAdd(FbxTime);
 
 		FbxAnimCurveDef::EInterpolationType Interpolation = FbxAnimCurveDef::eInterpolationCubic;
 		FbxAnimCurveDef::ETangentMode Tangent = FbxAnimCurveDef::eTangentAuto;
 
-		RichCurveInterpolationToFbxInterpolation(InRichCurve.GetKeyInterpMode(KeyHandle), InRichCurve.GetKeyTangentMode(KeyHandle), Interpolation, Tangent);
-
-		if (bNegative)
-		{
-			Value = -Value;
-		}
+		RichCurveInterpolationToFbxInterpolation(KeyValue.InterpMode, KeyValue.TangentMode, Interpolation, Tangent);
 
 		if (Interpolation == FbxAnimCurveDef::eInterpolationCubic)
 		{
-			FRichCurveKey RichCurveKey = InRichCurve.GetKey( KeyHandle );
-
-			FKeyHandle NextKeyHandle = InRichCurve.GetNextKey(KeyHandle);
-			if (InRichCurve.IsKeyHandleValid(NextKeyHandle))
+			if (Index < Times.Num() - 1)
 			{
-				float LeaveTangent = RichCurveKey.LeaveTangent;
-				float NextArriveTangent = InRichCurve.GetKey(NextKeyHandle).ArriveTangent;
+				float LeaveTangent = KeyValue.Tangent.LeaveTangent;
+				float NextArriveTangent = Values[Index+1].Tangent.ArriveTangent;
 
 				if (bNegative)
 				{
@@ -2363,7 +2350,7 @@ void FFbxExporter::ExportRichCurveToFbxCurve(FbxAnimCurve& InFbxCurve, FRichCurv
 	InFbxCurve.KeyModifyEnd();
 }
 
-void FFbxExporter::ExportLevelSequence3DTransformTrack( FbxNode& FbxActor, UMovieScene3DTransformTrack& TransformTrack, AActor* Actor, const TRange<float>& InPlaybackRange )
+void FFbxExporter::ExportLevelSequence3DTransformTrack( FbxNode& FbxActor, UMovieScene3DTransformTrack& TransformTrack, AActor* Actor, const TRange<FFrameNumber>& InPlaybackRange )
 {
 	FbxAnimLayer* BaseLayer = AnimStack->GetMember<FbxAnimLayer>( 0 );
 
@@ -2381,6 +2368,8 @@ void FFbxExporter::ExportLevelSequence3DTransformTrack( FbxNode& FbxActor, UMovi
 		return;
 	}
 
+	FFrameRate FrameResolution = TransformTrack.GetTypedOuter<UMovieScene>()->GetFrameResolution();
+
 	FbxAnimCurveNode* TranslationNode = FbxActor.LclTranslation.GetCurveNode( BaseLayer, true );
 	FbxAnimCurveNode* RotationNode = FbxActor.LclRotation.GetCurveNode( BaseLayer, true );
 	FbxAnimCurveNode* ScaleNode = FbxActor.LclScaling.GetCurveNode( BaseLayer, true );
@@ -2396,26 +2385,28 @@ void FFbxExporter::ExportLevelSequence3DTransformTrack( FbxNode& FbxActor, UMovi
 	FbxAnimCurve* FbxCurveScaleX = FbxActor.LclScaling.GetCurve( BaseLayer, FBXSDK_CURVENODE_COMPONENT_X, true );
 	FbxAnimCurve* FbxCurveScaleY = FbxActor.LclScaling.GetCurve( BaseLayer, FBXSDK_CURVENODE_COMPONENT_Y, true );
 	FbxAnimCurve* FbxCurveScaleZ = FbxActor.LclScaling.GetCurve( BaseLayer, FBXSDK_CURVENODE_COMPONENT_Z, true );
-		
+
+	TArrayView<FMovieSceneFloatChannel*> FloatChannels = TransformSection->GetChannelProxy().GetChannels<FMovieSceneFloatChannel>();
+
 	// Translation
-	ExportRichCurveToFbxCurve(*FbxCurveTransX, TransformSection->GetTranslationCurve( EAxis::X ), ERichCurveValueMode::Default, false);
-	ExportRichCurveToFbxCurve(*FbxCurveTransY, TransformSection->GetTranslationCurve( EAxis::Y ), ERichCurveValueMode::Default, true);
-	ExportRichCurveToFbxCurve(*FbxCurveTransZ, TransformSection->GetTranslationCurve( EAxis::Z ), ERichCurveValueMode::Default, false);
+	ExportChannelToFbxCurve(*FbxCurveTransX, *FloatChannels[0], FrameResolution, ERichCurveValueMode::Default, false);
+	ExportChannelToFbxCurve(*FbxCurveTransY, *FloatChannels[1], FrameResolution, ERichCurveValueMode::Default, true);
+	ExportChannelToFbxCurve(*FbxCurveTransZ, *FloatChannels[2], FrameResolution, ERichCurveValueMode::Default, false);
 
 	// Scale - don't generate scale keys for cameras
 	if (!bIsCameraActor)
 	{
-		ExportRichCurveToFbxCurve(*FbxCurveScaleX, TransformSection->GetScaleCurve( EAxis::X ), ERichCurveValueMode::Default, false);
-		ExportRichCurveToFbxCurve(*FbxCurveScaleY, TransformSection->GetScaleCurve( EAxis::Y ), ERichCurveValueMode::Default, false);
-		ExportRichCurveToFbxCurve(*FbxCurveScaleZ, TransformSection->GetScaleCurve( EAxis::Z ), ERichCurveValueMode::Default, false);
+		ExportChannelToFbxCurve(*FbxCurveScaleX, *FloatChannels[3], FrameResolution, ERichCurveValueMode::Default, false);
+		ExportChannelToFbxCurve(*FbxCurveScaleY, *FloatChannels[4], FrameResolution, ERichCurveValueMode::Default, false);
+		ExportChannelToFbxCurve(*FbxCurveScaleZ, *FloatChannels[5], FrameResolution, ERichCurveValueMode::Default, false);
 	}
 
 	// Rotation - bake rotation for cameras and lights
 	if (!bBakeRotations)
 	{
-		ExportRichCurveToFbxCurve(*FbxCurveRotX, TransformSection->GetRotationCurve( EAxis::X ), ERichCurveValueMode::Default, false);
-		ExportRichCurveToFbxCurve(*FbxCurveRotY, TransformSection->GetRotationCurve( EAxis::Y ), ERichCurveValueMode::Default, true);
-		ExportRichCurveToFbxCurve(*FbxCurveRotZ, TransformSection->GetRotationCurve( EAxis::Z ), ERichCurveValueMode::Default, true);
+		ExportChannelToFbxCurve(*FbxCurveRotX, *FloatChannels[6], FrameResolution, ERichCurveValueMode::Default, false);
+		ExportChannelToFbxCurve(*FbxCurveRotY, *FloatChannels[7], FrameResolution, ERichCurveValueMode::Default, true);
+		ExportChannelToFbxCurve(*FbxCurveRotZ, *FloatChannels[8], FrameResolution, ERichCurveValueMode::Default, true);
 	}
 	else
 	{
@@ -2435,18 +2426,29 @@ void FFbxExporter::ExportLevelSequence3DTransformTrack( FbxNode& FbxActor, UMovi
 		FbxCurveRotY->KeyModifyBegin();
 		FbxCurveRotZ->KeyModifyBegin();
 
-		float InterpLength = InPlaybackRange.GetUpperBoundValue() - InPlaybackRange.GetLowerBoundValue();
-		int NumKeys = InterpLength * BakeTransformsFPS;
+		FFrameNumber PlayStart      = MovieScene::DiscreteInclusiveLower(InPlaybackRange);
+		double       PlayDuration   = FFrameTime(MovieScene::DiscreteSize(InPlaybackRange)) / FrameResolution;
+		int32        NumKeys        = PlayDuration * BakeTransformsFPS;
+
 		for ( int KeyIndex = 0; KeyIndex < NumKeys; KeyIndex++ )
 		{
-			float KeyTime = InPlaybackRange.GetLowerBoundValue() + ( KeyIndex * InterpLength / NumKeys );
+			FFrameTime RelativeKeyTime = ( KeyIndex * PlayDuration / NumKeys ) * FrameResolution;
+			FFrameTime KeyTime         = RelativeKeyTime + PlayStart;
+
+			FVector Trans = FVector::ZeroVector;
+			FloatChannels[0]->Evaluate(KeyTime, Trans.X);
+			FloatChannels[1]->Evaluate(KeyTime, Trans.Y);
+			FloatChannels[2]->Evaluate(KeyTime, Trans.Z);
 
-			FVector Trans;
-			TransformSection->EvalTranslation(KeyTime, Trans);
 			FRotator Rotator;
-			TransformSection->EvalRotation(KeyTime, Rotator);
+			FloatChannels[3]->Evaluate(KeyTime, Rotator.Roll);
+			FloatChannels[4]->Evaluate(KeyTime, Rotator.Pitch);
+			FloatChannels[5]->Evaluate(KeyTime, Rotator.Yaw);
+
 			FVector Scale;
-			TransformSection->EvalScale(KeyTime, Scale);
+			FloatChannels[6]->Evaluate(KeyTime, Scale.X);
+			FloatChannels[7]->Evaluate(KeyTime, Scale.Y);
+			FloatChannels[8]->Evaluate(KeyTime, Scale.Z);
 
 			FTransform RelativeTransform;
 			RelativeTransform.SetTranslation(Trans);
@@ -2460,7 +2462,7 @@ void FFbxExporter::ExportLevelSequence3DTransformTrack( FbxNode& FbxActor, UMovi
 			FbxVector4 KeyScale = Converter.ConvertToFbxScale(RelativeTransform.GetScale3D());
 
 			FbxTime FbxTime;
-			FbxTime.SetSecondDouble((float)KeyTime);
+			FbxTime.SetSecondDouble(KeyTime / FrameResolution);
 
 			FbxCurveRotX->KeySet(FbxCurveRotX->KeyAdd(FbxTime), FbxTime, KeyRot[0]);
 			FbxCurveRotY->KeySet(FbxCurveRotY->KeyAdd(FbxTime), FbxTime, KeyRot[1]);
@@ -2480,7 +2482,8 @@ void FFbxExporter::ExportLevelSequenceFloatTrack( FbxNode& FbxActor, UMovieScene
 		? Cast<UMovieSceneFloatSection>( FloatTrack.GetAllSections()[0] )
 		: nullptr;
 
-	if ( FloatSection == nullptr || FloatSection->GetFloatCurve().GetNumKeys() == 0 )
+	const FMovieSceneFloatChannel* FloatChannel = FloatSection ? FloatSection->GetChannelProxy().GetChannel<FMovieSceneFloatChannel>(0) : nullptr;
+	if ( FloatChannel == nullptr || FloatChannel->GetTimes().Num() == 0 )
 	{
 		return;
 	}
@@ -2521,11 +2524,9 @@ void FFbxExporter::ExportLevelSequenceFloatTrack( FbxNode& FbxActor, UMovieScene
 		Property = FbxActor.FindProperty( "UE_MotionBlur_Amount", false );
 	}
 
-	FRichCurve& FloatCurve = FloatSection->GetFloatCurve();
-
 	if (Property == 0)
 	{
-		CreateAnimatableUserProperty(&FbxActor, FloatCurve.GetDefaultValue(), TCHAR_TO_UTF8(*PropertyName), TCHAR_TO_UTF8(*PropertyName));
+		CreateAnimatableUserProperty(&FbxActor, FloatChannel->GetDefault().Get(MAX_flt), TCHAR_TO_UTF8(*PropertyName), TCHAR_TO_UTF8(*PropertyName));
 
 		Property = FbxActor.FindProperty(TCHAR_TO_UTF8(*PropertyName), false);
 	}
@@ -2542,10 +2543,11 @@ void FFbxExporter::ExportLevelSequenceFloatTrack( FbxNode& FbxActor, UMovieScene
 		return;
 	}
 
-	CurveNode->SetChannelValue<double>( 0U, FloatCurve.GetDefaultValue() );
+	CurveNode->SetChannelValue<double>( 0U, FloatChannel->GetDefault().Get(MAX_flt) );
 	CurveNode->ConnectToChannel( AnimCurve, 0U );
 
-	ExportRichCurveToFbxCurve(*AnimCurve, FloatCurve, IsFoV ? ERichCurveValueMode::Fov : ERichCurveValueMode::Default);
+	FFrameRate FrameResolution = FloatTrack.GetTypedOuter<UMovieScene>()->GetFrameResolution();
+	ExportChannelToFbxCurve(*AnimCurve, *FloatChannel, FrameResolution, IsFoV ? ERichCurveValueMode::Fov : ERichCurveValueMode::Default);
 }
 		
 /**
diff --git a/Engine/Source/Editor/UnrealEd/Private/FbxExporter.h b/Engine/Source/Editor/UnrealEd/Private/FbxExporter.h
--- a/Engine/Source/Editor/UnrealEd/Private/FbxExporter.h
+++ b/Engine/Source/Editor/UnrealEd/Private/FbxExporter.h
@@ -40,7 +40,7 @@ class FColorVertexBuffer;
 class UFbxExportOption;
 struct FAnimControlTrackKey;
 struct FExpressionInput;
-struct FRichCurve;
+struct FMovieSceneFloatChannel;
 
 namespace UnFbx
 {
@@ -414,7 +414,7 @@ private:
 	/**
 	 * Exports a level sequence 3D transform track into the FBX animation stack.
 	 */
-	void ExportLevelSequence3DTransformTrack( FbxNode& FbxActor, UMovieScene3DTransformTrack& TransformTrack, AActor* Actor, const TRange<float>& InPlaybackRange );
+	void ExportLevelSequence3DTransformTrack( FbxNode& FbxActor, UMovieScene3DTransformTrack& TransformTrack, AActor* Actor, const TRange<FFrameNumber>& InPlaybackRange );
 
 	/** 
 	 * Exports a level sequence float track into the FBX animation stack. 
@@ -430,8 +430,8 @@ private:
 		Fov
 	};
 
-	/** Exports an unreal rich curve to an fbx animation curve. */
-	void ExportRichCurveToFbxCurve(FbxAnimCurve& InFbxCurve, FRichCurve& InRichCurve, ERichCurveValueMode ValueMode = ERichCurveValueMode::Default, bool bNegative = false);
+	/** Exports a movie scene channel to an fbx animation curve. */
+	void ExportChannelToFbxCurve(FbxAnimCurve& InFbxCurve, const FMovieSceneFloatChannel& InChannel, FFrameRate FrameResolution, ERichCurveValueMode ValueMode = ERichCurveValueMode::Default, bool bNegative = false);
 
 	/**
 	 * Finds the given actor in the already-exported list of structures
diff --git a/Engine/Source/Editor/UnrealEd/Private/SCurveEditor.cpp b/Engine/Source/Editor/UnrealEd/Private/SCurveEditor.cpp
--- a/Engine/Source/Editor/UnrealEd/Private/SCurveEditor.cpp
+++ b/Engine/Source/Editor/UnrealEd/Private/SCurveEditor.cpp
@@ -27,7 +27,7 @@
 #include "Factories/Factory.h"
 #include "Factories/CurveFactory.h"
 #include "Editor.h"
-#include "RichCurveEditorCommands.h"
+#include "CurveEditorCommands.h"
 #include "CurveEditorSettings.h"
 #include "ScopedTransaction.h"
 #include "Framework/Commands/GenericCommands.h"
@@ -121,158 +121,142 @@ void SCurveEditor::Construct(const FArguments& InArgs)
 	Commands->MapAction(FGenericCommands::Get().Redo,
 		FExecuteAction::CreateSP(this, &SCurveEditor::RedoAction));
 
-	Commands->MapAction(FRichCurveEditorCommands::Get().ZoomToFitHorizontal,
+	Commands->MapAction(FCurveEditorCommands::Get().ZoomToFitHorizontal,
 		FExecuteAction::CreateSP(this, &SCurveEditor::ZoomToFitHorizontal, false));
 
-	Commands->MapAction(FRichCurveEditorCommands::Get().ZoomToFitVertical,
+	Commands->MapAction(FCurveEditorCommands::Get().ZoomToFitVertical,
 		FExecuteAction::CreateSP(this, &SCurveEditor::ZoomToFitVertical, false));
 
-	Commands->MapAction(FRichCurveEditorCommands::Get().ZoomToFit,
+	Commands->MapAction(FCurveEditorCommands::Get().ZoomToFit,
 		FExecuteAction::CreateSP(this, &SCurveEditor::ZoomToFit, false));
 
-	Commands->MapAction(FRichCurveEditorCommands::Get().ZoomToFitAll,
+	Commands->MapAction(FCurveEditorCommands::Get().ZoomToFitAll,
 		FExecuteAction::CreateSP(this, &SCurveEditor::ZoomToFit, true));
 
-	Commands->MapAction(FRichCurveEditorCommands::Get().ToggleInputSnapping,
+	Commands->MapAction(FCurveEditorCommands::Get().ToggleInputSnapping,
 		FExecuteAction::CreateSP(this, &SCurveEditor::ToggleInputSnapping),
 		FCanExecuteAction(),
 		FIsActionChecked::CreateSP(this, &SCurveEditor::IsInputSnappingEnabled));
 
-	Commands->MapAction(FRichCurveEditorCommands::Get().ToggleOutputSnapping,
+	Commands->MapAction(FCurveEditorCommands::Get().ToggleOutputSnapping,
 		FExecuteAction::CreateSP(this, &SCurveEditor::ToggleOutputSnapping),
 		FCanExecuteAction(),
 		FIsActionChecked::CreateSP(this, &SCurveEditor::IsOutputSnappingEnabled));
 
 	// Interpolation
-	Commands->MapAction(FRichCurveEditorCommands::Get().InterpolationConstant,
+	Commands->MapAction(FCurveEditorCommands::Get().InterpolationConstant,
 		FExecuteAction::CreateSP(this, &SCurveEditor::OnSelectInterpolationMode, RCIM_Constant, RCTM_Auto),
 		FCanExecuteAction(),
 		FIsActionChecked::CreateSP(this, &SCurveEditor::IsInterpolationModeSelected, RCIM_Constant, RCTM_Auto));
 
-	Commands->MapAction(FRichCurveEditorCommands::Get().InterpolationLinear,
+	Commands->MapAction(FCurveEditorCommands::Get().InterpolationLinear,
 		FExecuteAction::CreateSP(this, &SCurveEditor::OnSelectInterpolationMode, RCIM_Linear, RCTM_Auto),
 		FCanExecuteAction(),
 		FIsActionChecked::CreateSP(this, &SCurveEditor::IsInterpolationModeSelected, RCIM_Linear, RCTM_Auto));
 
-	Commands->MapAction(FRichCurveEditorCommands::Get().InterpolationCubicAuto,
+	Commands->MapAction(FCurveEditorCommands::Get().InterpolationCubicAuto,
 		FExecuteAction::CreateSP(this, &SCurveEditor::OnSelectInterpolationMode, RCIM_Cubic, RCTM_Auto),
 		FCanExecuteAction(),
 		FIsActionChecked::CreateSP(this, &SCurveEditor::IsInterpolationModeSelected, RCIM_Cubic, RCTM_Auto));
 
-	Commands->MapAction(FRichCurveEditorCommands::Get().InterpolationCubicUser,
+	Commands->MapAction(FCurveEditorCommands::Get().InterpolationCubicUser,
 		FExecuteAction::CreateSP(this, &SCurveEditor::OnSelectInterpolationMode, RCIM_Cubic, RCTM_User),
 		FCanExecuteAction(),
 		FIsActionChecked::CreateSP(this, &SCurveEditor::IsInterpolationModeSelected, RCIM_Cubic, RCTM_User));
 
-	Commands->MapAction(FRichCurveEditorCommands::Get().InterpolationCubicBreak,
+	Commands->MapAction(FCurveEditorCommands::Get().InterpolationCubicBreak,
 		FExecuteAction::CreateSP(this, &SCurveEditor::OnSelectInterpolationMode, RCIM_Cubic, RCTM_Break),
 		FCanExecuteAction(),
 		FIsActionChecked::CreateSP(this, &SCurveEditor::IsInterpolationModeSelected, RCIM_Cubic, RCTM_Break));
 
 	// Tangents
-	Commands->MapAction(FRichCurveEditorCommands::Get().FlattenTangents,
+	Commands->MapAction(FCurveEditorCommands::Get().FlattenTangents,
 		FExecuteAction::CreateSP(this, &SCurveEditor::OnFlattenOrStraightenTangents, true));
 
-	Commands->MapAction(FRichCurveEditorCommands::Get().StraightenTangents,
+	Commands->MapAction(FCurveEditorCommands::Get().StraightenTangents,
 		FExecuteAction::CreateSP(this, &SCurveEditor::OnFlattenOrStraightenTangents, false));
 
 	// Bake and reduce
-	Commands->MapAction(FRichCurveEditorCommands::Get().BakeCurve,
+	Commands->MapAction(FCurveEditorCommands::Get().BakeCurve,
 		FExecuteAction::CreateSP(this, &SCurveEditor::OnBakeCurve));
 
-	Commands->MapAction(FRichCurveEditorCommands::Get().ReduceCurve,
+	Commands->MapAction(FCurveEditorCommands::Get().ReduceCurve,
 		FExecuteAction::CreateSP(this, &SCurveEditor::OnReduceCurve));
 
 	// Pre infinity extrapolation
-	Commands->MapAction(FRichCurveEditorCommands::Get().SetPreInfinityExtrapCycle,
+	Commands->MapAction(FCurveEditorCommands::Get().SetPreInfinityExtrapCycle,
 		FExecuteAction::CreateSP(this, &SCurveEditor::OnSelectPreInfinityExtrap, RCCE_Cycle),
 		FCanExecuteAction(),
 		FIsActionChecked::CreateSP(this, &SCurveEditor::IsPreInfinityExtrapSelected, RCCE_Cycle));
 
-	Commands->MapAction(FRichCurveEditorCommands::Get().SetPreInfinityExtrapCycleWithOffset,
+	Commands->MapAction(FCurveEditorCommands::Get().SetPreInfinityExtrapCycleWithOffset,
 		FExecuteAction::CreateSP(this, &SCurveEditor::OnSelectPreInfinityExtrap, RCCE_CycleWithOffset),
 		FCanExecuteAction(),
 		FIsActionChecked::CreateSP(this, &SCurveEditor::IsPreInfinityExtrapSelected, RCCE_CycleWithOffset));
 
-	Commands->MapAction(FRichCurveEditorCommands::Get().SetPreInfinityExtrapOscillate,
+	Commands->MapAction(FCurveEditorCommands::Get().SetPreInfinityExtrapOscillate,
 		FExecuteAction::CreateSP(this, &SCurveEditor::OnSelectPreInfinityExtrap, RCCE_Oscillate),
 		FCanExecuteAction(),
 		FIsActionChecked::CreateSP(this, &SCurveEditor::IsPreInfinityExtrapSelected, RCCE_Oscillate));
 
-	Commands->MapAction(FRichCurveEditorCommands::Get().SetPreInfinityExtrapLinear,
+	Commands->MapAction(FCurveEditorCommands::Get().SetPreInfinityExtrapLinear,
 		FExecuteAction::CreateSP(this, &SCurveEditor::OnSelectPreInfinityExtrap, RCCE_Linear),
 		FCanExecuteAction(),
 		FIsActionChecked::CreateSP(this, &SCurveEditor::IsPreInfinityExtrapSelected, RCCE_Linear));
 
-	Commands->MapAction(FRichCurveEditorCommands::Get().SetPreInfinityExtrapConstant,
+	Commands->MapAction(FCurveEditorCommands::Get().SetPreInfinityExtrapConstant,
 		FExecuteAction::CreateSP(this, &SCurveEditor::OnSelectPreInfinityExtrap, RCCE_Constant),
 		FCanExecuteAction(),
 		FIsActionChecked::CreateSP(this, &SCurveEditor::IsPreInfinityExtrapSelected, RCCE_Constant));
 
 	// Post infinity extrapolation
-	Commands->MapAction(FRichCurveEditorCommands::Get().SetPostInfinityExtrapCycle,
+	Commands->MapAction(FCurveEditorCommands::Get().SetPostInfinityExtrapCycle,
 		FExecuteAction::CreateSP(this, &SCurveEditor::OnSelectPostInfinityExtrap, RCCE_Cycle),
 		FCanExecuteAction(),
 		FIsActionChecked::CreateSP(this, &SCurveEditor::IsPostInfinityExtrapSelected, RCCE_Cycle));
 
-	Commands->MapAction(FRichCurveEditorCommands::Get().SetPostInfinityExtrapCycleWithOffset,
+	Commands->MapAction(FCurveEditorCommands::Get().SetPostInfinityExtrapCycleWithOffset,
 		FExecuteAction::CreateSP(this, &SCurveEditor::OnSelectPostInfinityExtrap, RCCE_CycleWithOffset),
 		FCanExecuteAction(),
 		FIsActionChecked::CreateSP(this, &SCurveEditor::IsPostInfinityExtrapSelected, RCCE_CycleWithOffset));
 
-	Commands->MapAction(FRichCurveEditorCommands::Get().SetPostInfinityExtrapOscillate,
+	Commands->MapAction(FCurveEditorCommands::Get().SetPostInfinityExtrapOscillate,
 		FExecuteAction::CreateSP(this, &SCurveEditor::OnSelectPostInfinityExtrap, RCCE_Oscillate),
 		FCanExecuteAction(),
 		FIsActionChecked::CreateSP(this, &SCurveEditor::IsPostInfinityExtrapSelected, RCCE_Oscillate));
 
-	Commands->MapAction(FRichCurveEditorCommands::Get().SetPostInfinityExtrapLinear,
+	Commands->MapAction(FCurveEditorCommands::Get().SetPostInfinityExtrapLinear,
 		FExecuteAction::CreateSP(this, &SCurveEditor::OnSelectPostInfinityExtrap, RCCE_Linear),
 		FCanExecuteAction(),
 		FIsActionChecked::CreateSP(this, &SCurveEditor::IsPostInfinityExtrapSelected, RCCE_Linear));
 
-	Commands->MapAction(FRichCurveEditorCommands::Get().SetPostInfinityExtrapConstant,
+	Commands->MapAction(FCurveEditorCommands::Get().SetPostInfinityExtrapConstant,
 		FExecuteAction::CreateSP(this, &SCurveEditor::OnSelectPostInfinityExtrap, RCCE_Constant),
 		FCanExecuteAction(),
 		FIsActionChecked::CreateSP(this, &SCurveEditor::IsPostInfinityExtrapSelected, RCCE_Constant));
 
-	// Curve Visibility
-	Commands->MapAction(FRichCurveEditorCommands::Get().SetAllCurveVisibility,
-		FExecuteAction::CreateLambda( [this]{ Settings->SetCurveVisibility( ECurveEditorCurveVisibility::AllCurves ); } ),
-		FCanExecuteAction::CreateLambda( []{ return true; } ),
-		FIsActionChecked::CreateLambda( [this]{ return Settings->GetCurveVisibility() == ECurveEditorCurveVisibility::AllCurves; } ) );
-
-	Commands->MapAction(FRichCurveEditorCommands::Get().SetSelectedCurveVisibility,
-		FExecuteAction::CreateLambda( [this]{ Settings->SetCurveVisibility( ECurveEditorCurveVisibility::SelectedCurves ); } ),
-		FCanExecuteAction::CreateLambda( []{ return true; } ),
-		FIsActionChecked::CreateLambda( [this]{ return Settings->GetCurveVisibility() == ECurveEditorCurveVisibility::SelectedCurves; } ) );
-
-	Commands->MapAction(FRichCurveEditorCommands::Get().SetAnimatedCurveVisibility,
-		FExecuteAction::CreateLambda( [this]{ Settings->SetCurveVisibility( ECurveEditorCurveVisibility::AnimatedCurves ); } ),
-		FCanExecuteAction::CreateLambda( []{ return true; } ),
-		FIsActionChecked::CreateLambda( [this]{ return Settings->GetCurveVisibility() == ECurveEditorCurveVisibility::AnimatedCurves; } ) );
-
 	// Tangent Visibility
-	Commands->MapAction(FRichCurveEditorCommands::Get().SetAllTangentsVisibility,
+	Commands->MapAction(FCurveEditorCommands::Get().SetAllTangentsVisibility,
 		FExecuteAction::CreateLambda( [this]{ Settings->SetTangentVisibility( ECurveEditorTangentVisibility::AllTangents ); } ),
 		FCanExecuteAction::CreateLambda( []{ return true; } ),
 		FIsActionChecked::CreateLambda( [this]{ return Settings->GetTangentVisibility() == ECurveEditorTangentVisibility::AllTangents; } ) );
 
-	Commands->MapAction(FRichCurveEditorCommands::Get().SetSelectedKeysTangentVisibility,
+	Commands->MapAction(FCurveEditorCommands::Get().SetSelectedKeysTangentVisibility,
 		FExecuteAction::CreateLambda( [this]{ Settings->SetTangentVisibility( ECurveEditorTangentVisibility::SelectedKeys ); } ),
 		FCanExecuteAction::CreateLambda( []{ return true; } ),
 		FIsActionChecked::CreateLambda( [this]{ return Settings->GetTangentVisibility() == ECurveEditorTangentVisibility::SelectedKeys; } ) );
 
-	Commands->MapAction(FRichCurveEditorCommands::Get().SetNoTangentsVisibility,
+	Commands->MapAction(FCurveEditorCommands::Get().SetNoTangentsVisibility,
 		FExecuteAction::CreateLambda( [this]{ Settings->SetTangentVisibility( ECurveEditorTangentVisibility::NoTangents ); } ),
 		FCanExecuteAction::CreateLambda( []{ return true; } ),
 		FIsActionChecked::CreateLambda( [this]{ return Settings->GetTangentVisibility() == ECurveEditorTangentVisibility::NoTangents; } ) );
 
-	Commands->MapAction(FRichCurveEditorCommands::Get().ToggleAutoFrameCurveEditor,
+	Commands->MapAction(FCurveEditorCommands::Get().ToggleAutoFrameCurveEditor,
 		FExecuteAction::CreateLambda( [this]{ Settings->SetAutoFrameCurveEditor( !Settings->GetAutoFrameCurveEditor() ); } ),
 		FCanExecuteAction::CreateLambda( []{ return true; } ),
 		FIsActionChecked::CreateLambda( [this]{ return Settings->GetAutoFrameCurveEditor(); } ) );
 
-	Commands->MapAction(FRichCurveEditorCommands::Get().ToggleShowCurveEditorCurveToolTips,
+	Commands->MapAction(FCurveEditorCommands::Get().ToggleShowCurveEditorCurveToolTips,
 		FExecuteAction::CreateLambda( [this]{ 
 		Settings->SetShowCurveEditorCurveToolTips( !Settings->GetShowCurveEditorCurveToolTips() ); 
 		if (!Settings->GetShowCurveEditorCurveToolTips())
@@ -681,18 +665,18 @@ void SCurveEditor::PushKeyMenu(const FGeometry& InMyGeometry, const FPointerEven
 	FMenuBuilder MenuBuilder(true, Commands.ToSharedRef());
 	MenuBuilder.BeginSection("CurveEditorInterpolation", LOCTEXT("KeyInterpolationMode", "Key Interpolation"));
 	{
-		MenuBuilder.AddMenuEntry(FRichCurveEditorCommands::Get().InterpolationCubicAuto);
-		MenuBuilder.AddMenuEntry(FRichCurveEditorCommands::Get().InterpolationCubicUser);
-		MenuBuilder.AddMenuEntry(FRichCurveEditorCommands::Get().InterpolationCubicBreak);
-		MenuBuilder.AddMenuEntry(FRichCurveEditorCommands::Get().InterpolationLinear);
-		MenuBuilder.AddMenuEntry(FRichCurveEditorCommands::Get().InterpolationConstant);
+		MenuBuilder.AddMenuEntry(FCurveEditorCommands::Get().InterpolationCubicAuto);
+		MenuBuilder.AddMenuEntry(FCurveEditorCommands::Get().InterpolationCubicUser);
+		MenuBuilder.AddMenuEntry(FCurveEditorCommands::Get().InterpolationCubicBreak);
+		MenuBuilder.AddMenuEntry(FCurveEditorCommands::Get().InterpolationLinear);
+		MenuBuilder.AddMenuEntry(FCurveEditorCommands::Get().InterpolationConstant);
 	}
 	MenuBuilder.EndSection(); //CurveEditorInterpolation
 
 	MenuBuilder.BeginSection("CurveEditorTangents", LOCTEXT("Tangents", "Tangents"));
 	{
-		MenuBuilder.AddMenuEntry(FRichCurveEditorCommands::Get().FlattenTangents);
-		MenuBuilder.AddMenuEntry(FRichCurveEditorCommands::Get().StraightenTangents);
+		MenuBuilder.AddMenuEntry(FCurveEditorCommands::Get().FlattenTangents);
+		MenuBuilder.AddMenuEntry(FCurveEditorCommands::Get().StraightenTangents);
 	}
 	MenuBuilder.EndSection(); //CurveEditorTangents
 
@@ -1415,7 +1399,7 @@ void SCurveEditor::AddNewKey(FGeometry InMyGeometry, FVector2D ScreenPosition, T
 				float Input = ScaleInfo.LocalXToInput(LocalClickPos.X);
 				float Output;
 				if (bAddKeysInline)
-				{
+			{
 					Output = SelectedCurve->Eval(Input);
 				}
 				else
@@ -3732,11 +3716,27 @@ void SCurveEditor::EndDragTransaction()
 	if ( TransactionIndex >= 0 )
 	{
 		TArray<FRichCurveEditInfo> ChangedCurveEditInfos;
-		for (auto CurveViewModel : CurveViewModels)
+
+		if (DragState == EDragState::DragKey || DragState == EDragState::FreeDrag)
 		{
-			ChangedCurveEditInfos.Add(CurveViewModel->CurveInfo);
+			for (auto SelectedKey : SelectedKeys)
+			{
+				ChangedCurveEditInfos.Add(GetViewModelForCurve(SelectedKey.Curve)->CurveInfo);
+			}
 		}
-		CurveOwner->OnCurveChanged(ChangedCurveEditInfos);
+		else if (DragState == EDragState::DragTangent)
+		{
+			for (auto SelectedTangent : SelectedTangents)
+			{
+				ChangedCurveEditInfos.Add(GetViewModelForCurve(SelectedTangent.Key.Curve)->CurveInfo);
+			}
+		}
+
+		if (ChangedCurveEditInfos.Num())
+		{
+			CurveOwner->OnCurveChanged(ChangedCurveEditInfos);
+		}
+		
 		GEditor->EndTransaction();
 		TransactionIndex = -1;
 	}
diff --git a/Engine/Source/Editor/UnrealEd/Private/UnrealEdMisc.cpp b/Engine/Source/Editor/UnrealEd/Private/UnrealEdMisc.cpp
--- a/Engine/Source/Editor/UnrealEd/Private/UnrealEdMisc.cpp
+++ b/Engine/Source/Editor/UnrealEd/Private/UnrealEdMisc.cpp
@@ -37,7 +37,7 @@
 #include "Kismet2/DebuggerCommands.h"
 #include "Toolkits/AssetEditorCommonCommands.h"
 #include "SoundCueGraphEditorCommands.h"
-#include "RichCurveEditorCommands.h"
+#include "CurveEditorCommands.h"
 #include "EditorBuildUtils.h"
 #include "Logging/TokenizedMessage.h"
 #include "Logging/MessageLog.h"
@@ -268,7 +268,7 @@ void FUnrealEdMisc::OnInit()
 	FViewportNavigationCommands::Register();
 
 	// Register curve editor commands.
-	FRichCurveEditorCommands::Register();
+	FCurveEditorCommands::Register();
 
 	// Have the User Activity Tracker reject non-editor activities for this run
 	FUserActivityTracking::SetContextFilter(EUserActivityContext::Editor);
diff --git a/Engine/Source/Editor/VREditor/VREditorActions.cpp b/Engine/Source/Editor/VREditor/VREditorActions.cpp
--- a/Engine/Source/Editor/VREditor/VREditorActions.cpp
+++ b/Engine/Source/Editor/VREditor/VREditorActions.cpp
@@ -451,7 +451,7 @@ void FVREditorActionCallbacks::PlayFromBeginning(UVREditorMode* InVRMode)
 	ISequencer* CurrentSequencer = InVRMode->GetCurrentSequencer();
 	if (CurrentSequencer != nullptr)
 	{
-		CurrentSequencer->SetLocalTime(0.0f);
+		CurrentSequencer->SetLocalTime(0);
 		CurrentSequencer->SetPlaybackSpeed(1.f);
 		CurrentSequencer->OnPlay(false);
 	}
diff --git a/Engine/Source/Runtime/AVIWriter/Private/AVIWriter.cpp b/Engine/Source/Runtime/AVIWriter/Private/AVIWriter.cpp
--- a/Engine/Source/Runtime/AVIWriter/Private/AVIWriter.cpp
+++ b/Engine/Source/Runtime/AVIWriter/Private/AVIWriter.cpp
@@ -353,8 +353,10 @@ public:
 			FString Directory = Options.OutputFilename;
 			FString Ext = FPaths::GetExtension(Directory, true);
 
+			int32 FPS = FMath::RoundToInt(double(Options.CaptureFramerateNumerator) / Options.CaptureFramerateDenominator);
+
 			// Keep 3 seconds worth of frames in memory
-			CapturedFrames.Reset(new FCapturedFrames(Directory.LeftChop(Ext.Len()) + TEXT("_tmp"), Options.CaptureFPS * 3));
+			CapturedFrames.Reset(new FCapturedFrames(Directory.LeftChop(Ext.Len()) + TEXT("_tmp"), FPS * 3));
 
 			hr = Control->Run();
 
@@ -485,9 +487,11 @@ public:
 			FString Directory = Options.OutputFilename;
 			FString Ext = FPaths::GetExtension(Directory, true);
 
+			int32 FPS = FMath::RoundToInt(double(Options.CaptureFramerateNumerator) / Options.CaptureFramerateDenominator);
+
 			// Keep 3 seconds worth of frames in memory
-			CapturedFrames.Reset(new FCapturedFrames(Directory.LeftChop(Ext.Len()) + TEXT("_tmp"), Options.CaptureFPS * 3));
-			
+			CapturedFrames.Reset(new FCapturedFrames(Directory.LeftChop(Ext.Len()) + TEXT("_tmp"), FPS * 3));
+
 			bCapturing = true;
 			ThreadTaskFuture = Async<void>(EAsyncExecution::Thread, [this]{	TaskThread(); });
 		}
@@ -542,8 +546,7 @@ public:
 				void* Data = CVPixelBufferGetBaseAddress(PixelBuffer);
 				FMemory::Memcpy(Data, CurrentFrame.FrameData.GetData(), CurrentFrame.FrameData.Num()*sizeof(FColor));
 				CVPixelBufferUnlockBaseAddress(PixelBuffer, 0);
-				
-				CMTime PresentTime = CurrentFrame.FrameIndex > 0 ? CMTimeMake(CurrentFrame.FrameIndex, Options.CaptureFPS) : kCMTimeZero;
+				CMTime PresentTime = CurrentFrame.FrameIndex > 0 ? CMTimeMake(CurrentFrame.FrameIndex * Options.CaptureFramerateDenominator, Options.CaptureFramerateNumerator) : kCMTimeZero;
 				BOOL OK = [AVFPixelBufferAdaptorRef appendPixelBuffer:PixelBuffer withPresentationTime:PresentTime];
 				check(OK);
 				
@@ -788,7 +791,7 @@ void FAVIWriter::Update(double FrameTimeSeconds, TArray<FColor> FrameData)
 {
 	if (bCapturing && FrameData.Num())
 	{
-		double FrameLength = 1.0 / Options.CaptureFPS;
+		double FrameLength = double(Options.CaptureFramerateDenominator) / Options.CaptureFramerateNumerator;
 		double FrameStart = FrameNumber * FrameLength;
 		FCapturedFrame Frame(FrameStart, FrameStart + FrameLength, FrameNumber, MoveTemp(FrameData));
 
diff --git a/Engine/Source/Runtime/AVIWriter/Private/CapturePin.cpp b/Engine/Source/Runtime/AVIWriter/Private/CapturePin.cpp
--- a/Engine/Source/Runtime/AVIWriter/Private/CapturePin.cpp
+++ b/Engine/Source/Runtime/AVIWriter/Private/CapturePin.cpp
@@ -12,7 +12,7 @@ DEFINE_LOG_CATEGORY(LogMovieCapture);
 
 FCapturePin::FCapturePin(HRESULT *phr, CSource *pFilter, const FAVIWriter& InWriter)
         : CSourceStream(NAME("Push Source"), phr, pFilter, L"Capture")
-		, FrameLength(UNITS/InWriter.Options.CaptureFPS)
+		, FrameLength(UNITS / (InWriter.Options.CaptureFramerateNumerator / (double)InWriter.Options.CaptureFramerateDenominator))
 		, Writer(InWriter)
 {
 	ImageWidth  = Writer.GetWidth();
diff --git a/Engine/Source/Runtime/AVIWriter/Public/AVIWriter.h b/Engine/Source/Runtime/AVIWriter/Public/AVIWriter.h
--- a/Engine/Source/Runtime/AVIWriter/Public/AVIWriter.h
+++ b/Engine/Source/Runtime/AVIWriter/Public/AVIWriter.h
@@ -20,7 +20,8 @@ struct FAVIWriterOptions
 {
 	FAVIWriterOptions()
 		: OutputFilename(FPaths::VideoCaptureDir() / TEXT("Capture.avi"))
-		, CaptureFPS(30)
+		, CaptureFramerateNumerator(30)
+		, CaptureFramerateDenominator(1)
 		, bSynchronizeFrames(false)
 		, Width(0)
 		, Height(0)
@@ -29,8 +30,11 @@ struct FAVIWriterOptions
 	/** Output filename */
 	FString OutputFilename;
 
-	/** Constant framerate of the captured video */
-	int32 CaptureFPS;
+	/** The numerator of the captured video, ie (30/1) will capture at 30 frames per second.*/
+	int32 CaptureFramerateNumerator;
+
+	/** The denominator of the captured video, ie (30/1) will capture at 30 frames per second.*/
+	int32 CaptureFramerateDenominator;
 
 	/** Optional compression quality, as a value between 0 and 1 */
 	TOptional<float> CompressionQuality;
diff --git a/Engine/Source/Runtime/AnimGraphRuntime/Private/AnimSequencerInstanceProxy.cpp b/Engine/Source/Runtime/AnimGraphRuntime/Private/AnimSequencerInstanceProxy.cpp
--- a/Engine/Source/Runtime/AnimGraphRuntime/Private/AnimSequencerInstanceProxy.cpp
+++ b/Engine/Source/Runtime/AnimGraphRuntime/Private/AnimSequencerInstanceProxy.cpp
@@ -151,4 +151,8 @@ void FAnimSequencerInstanceProxy::EnsureAnimTrack(UAnimSequenceBase* InAnimSeque
 	{
 		InitAnimTrack(InAnimSequence, SequenceId);
 	}
+	else if (PlayerState->PlayerNode.Sequence != InAnimSequence)
+	{
+		PlayerState->PlayerNode.OverrideAsset(InAnimSequence);
+	}
 }
\ No newline at end of file
diff --git a/Engine/Source/Runtime/CinematicCamera/Public/CineCameraComponent.h b/Engine/Source/Runtime/CinematicCamera/Public/CineCameraComponent.h
--- a/Engine/Source/Runtime/CinematicCamera/Public/CineCameraComponent.h
+++ b/Engine/Source/Runtime/CinematicCamera/Public/CineCameraComponent.h
@@ -199,7 +199,7 @@ struct FCameraFocusSettings
 /**
  * A specialized version of a camera component, geared toward cinematic usage.
  */
-UCLASS(HideCategories = (CameraSettings), HideFunctions = (SetFieldOfView, SetAspectRatio, SetConstraintAspectRatio), Blueprintable, ClassGroup = Camera, meta = (BlueprintSpawnableComponent), Config = Engine)
+UCLASS(HideCategories = (CameraSettings), HideFunctions = (SetFieldOfView, SetAspectRatio), Blueprintable, ClassGroup = Camera, meta = (BlueprintSpawnableComponent), Config = Engine)
 class CINEMATICCAMERA_API UCineCameraComponent : public UCameraComponent
 {
 	GENERATED_BODY()
diff --git a/Engine/Source/Runtime/Core/Private/Math/UnrealMath.cpp b/Engine/Source/Runtime/Core/Private/Math/UnrealMath.cpp
--- a/Engine/Source/Runtime/Core/Private/Math/UnrealMath.cpp
+++ b/Engine/Source/Runtime/Core/Private/Math/UnrealMath.cpp
@@ -8,7 +8,6 @@
 #include "Stats/Stats.h"
 #include "Math/RandomStream.h"
 #include "UObject/PropertyPortFlags.h"
-
 DEFINE_LOG_CATEGORY(LogUnrealMath);
 
 /**
@@ -580,14 +579,14 @@ FRotator FQuat::Rotator() const
 	if (SingularityTest < -SINGULARITY_THRESHOLD)
 	{
 		RotatorFromQuat.Pitch = -90.f;
-		RotatorFromQuat.Yaw = FMath::Atan2(YawY, YawX) * RAD_TO_DEG;
-		RotatorFromQuat.Roll = FRotator::NormalizeAxis(-RotatorFromQuat.Yaw - (2.f * FMath::Atan2(X, W) * RAD_TO_DEG));
+		RotatorFromQuat.Yaw = -2.0f * FMath::Atan2(X, W) * RAD_TO_DEG;
+		RotatorFromQuat.Roll = 0.0f;
 	}
 	else if (SingularityTest > SINGULARITY_THRESHOLD)
 	{
 		RotatorFromQuat.Pitch = 90.f;
-		RotatorFromQuat.Yaw = FMath::Atan2(YawY, YawX) * RAD_TO_DEG;
-		RotatorFromQuat.Roll = FRotator::NormalizeAxis(RotatorFromQuat.Yaw - (2.f * FMath::Atan2(X, W) * RAD_TO_DEG));
+		RotatorFromQuat.Yaw = 2.0f * FMath::Atan2(X, W) * RAD_TO_DEG;
+		RotatorFromQuat.Roll = 0.0f;
 	}
 	else
 	{
diff --git a/Engine/Source/Runtime/Core/Public/Async/Async.h b/Engine/Source/Runtime/Core/Public/Async/Async.h
--- a/Engine/Source/Runtime/Core/Public/Async/Async.h
+++ b/Engine/Source/Runtime/Core/Public/Async/Async.h
@@ -3,18 +3,19 @@
 #pragma once
 
 #include "CoreTypes.h"
-#include "Misc/AssertionMacros.h"
-#include "Containers/UnrealString.h"
-#include "Templates/Function.h"
-#include "HAL/ThreadSafeCounter.h"
-#include "Stats/Stats.h"
+#include "Async/Future.h"
 #include "Async/TaskGraphInterfaces.h"
+#include "Containers/UnrealString.h"
+#include "HAL/PlatformProcess.h"
 #include "HAL/Runnable.h"
-#include "Misc/IQueuedWork.h"
 #include "HAL/RunnableThread.h"
-#include "Misc/QueuedThreadPool.h"
+#include "HAL/ThreadSafeCounter.h"
+#include "Misc/AssertionMacros.h"
 #include "Misc/CoreStats.h"
-#include "Async/Future.h"
+#include "Misc/IQueuedWork.h"
+#include "Misc/QueuedThreadPool.h"
+#include "Stats/Stats.h"
+#include "Templates/Function.h"
 
 /**
  * Enumerates available asynchronous execution methods.
@@ -27,7 +28,7 @@ enum class EAsyncExecution
 	/** Execute in separate thread (for long running tasks). */
 	Thread,
 
-	/** Execute in queued thread pool. */
+	/** Execute in global queued thread pool. */
 	ThreadPool
 };
 
@@ -182,7 +183,7 @@ public:
 
 public:
 
-	// FRunnable interface
+	//~ FRunnable interface
 
 	virtual uint32 Run() override;
 
@@ -221,7 +222,7 @@ public:
 
 public:
 
-	// IQueuedWork interface
+	//~ IQueuedWork interface
 
 	virtual void DoThreadedWork() override
 	{
@@ -262,7 +263,7 @@ struct FAsyncThreadIndex
 
 
 /**
- * Executes a given function asynchronously.
+ * Execute a given function asynchronously.
  *
  * While we still support VS2012, we need to help the compiler deduce the ResultType
  * template argument for Async<T>(), which can be accomplished in two ways:
@@ -316,6 +317,7 @@ TFuture<ResultType> Async(EAsyncExecution Execution, TFunction<ResultType()> Fun
 		break;
 	
 	case EAsyncExecution::Thread:
+		if (FPlatformProcess::SupportsMultithreading())
 		{
 			TPromise<FRunnableThread*> ThreadPromise;
 			TAsyncRunnable<ResultType>* Runnable = new TAsyncRunnable<ResultType>(MoveTemp(Function), MoveTemp(Promise), ThreadPromise.GetFuture());
@@ -327,6 +329,10 @@ TFuture<ResultType> Async(EAsyncExecution Execution, TFunction<ResultType()> Fun
 
 			ThreadPromise.SetValue(RunnableThread);
 		}
+		else
+		{
+			SetPromise(Promise, Function);
+		}
 		break;
 
 	case EAsyncExecution::ThreadPool:
@@ -344,7 +350,29 @@ TFuture<ResultType> Async(EAsyncExecution Execution, TFunction<ResultType()> Fun
 
 
 /**
- * Executes a given function asynchronously using a separate thread.
+ * Execute a given function asynchronously on the specified thread pool.
+ *
+ * @param ResultType The type of the function's return value.
+ * @param ThreadPool The thread pool to execute on.
+ * @param Function The function to execute.
+ * @param CompletionCallback An optional callback function that is executed when the function completed execution.
+ * @result A TFuture object that will receive the return value from the function.
+ */
+template<typename ResultType>
+TFuture<ResultType> AsyncPool(FQueuedThreadPool& ThreadPool, TFunction<ResultType()> Function, TFunction<void()> CompletionCallback = TFunction<void()>())
+{
+	TPromise<ResultType> Promise(MoveTemp(CompletionCallback));
+	TFuture<ResultType> Future = Promise.GetFuture();
+
+	ThreadPool.AddQueuedWork(new TAsyncQueuedWork<ResultType>(MoveTemp(Function), MoveTemp(Promise)));
+
+	return MoveTemp(Future);
+}
+
+
+/**
+ * Execute a given function asynchronously using a separate thread.
+ *
  * @param ResultType The type of the function's return value.
  * @param Function The function to execute.
  * @param StackSize stack space to allocate for the new thread
@@ -358,15 +386,22 @@ TFuture<ResultType> AsyncThread(TFunction<ResultType()> Function, uint32 StackSi
 	TPromise<ResultType> Promise(MoveTemp(CompletionCallback));
 	TFuture<ResultType> Future = Promise.GetFuture();
 
-	TPromise<FRunnableThread*> ThreadPromise;
-	TAsyncRunnable<ResultType>* Runnable = new TAsyncRunnable<ResultType>(MoveTemp(Function), MoveTemp(Promise), ThreadPromise.GetFuture());
+	if (FPlatformProcess::SupportsMultithreading())
+	{
+		TPromise<FRunnableThread*> ThreadPromise;
+		TAsyncRunnable<ResultType>* Runnable = new TAsyncRunnable<ResultType>(MoveTemp(Function), MoveTemp(Promise), ThreadPromise.GetFuture());
 
-	const FString TAsyncThreadName = FString::Printf(TEXT("TAsyncThread %d"), FAsyncThreadIndex::GetNext());
-	FRunnableThread* RunnableThread = FRunnableThread::Create(Runnable, *TAsyncThreadName, StackSize, ThreadPri);
+		const FString TAsyncThreadName = FString::Printf(TEXT("TAsyncThread %d"), FAsyncThreadIndex::GetNext());
+		FRunnableThread* RunnableThread = FRunnableThread::Create(Runnable, *TAsyncThreadName, StackSize, ThreadPri);
 
-	check(RunnableThread != nullptr);
+		check(RunnableThread != nullptr);
 
-	ThreadPromise.SetValue(RunnableThread);
+		ThreadPromise.SetValue(RunnableThread);
+	}
+	else
+	{
+		SetPromise(Promise, Function);
+	}
 
 	return MoveTemp(Future);
 }
diff --git a/Engine/Source/Runtime/Core/Public/Containers/CircularBuffer.h b/Engine/Source/Runtime/Core/Public/Containers/CircularBuffer.h
--- a/Engine/Source/Runtime/Core/Public/Containers/CircularBuffer.h
+++ b/Engine/Source/Runtime/Core/Public/Containers/CircularBuffer.h
@@ -23,7 +23,7 @@ public:
 	 *
 	 * @param Capacity The number of elements that the buffer can store (will be rounded up to the next power of 2).
 	 */
-	TCircularBuffer(uint32 Capacity)
+	explicit TCircularBuffer(uint32 Capacity)
 	{
 		checkSlow(Capacity > 0);
 		checkSlow(Capacity <= 0xffffffffU);
diff --git a/Engine/Source/Runtime/Core/Public/Containers/CircularQueue.h b/Engine/Source/Runtime/Core/Public/Containers/CircularQueue.h
--- a/Engine/Source/Runtime/Core/Public/Containers/CircularQueue.h
+++ b/Engine/Source/Runtime/Core/Public/Containers/CircularQueue.h
@@ -4,18 +4,20 @@
 
 #include "CoreTypes.h"
 #include "Containers/CircularBuffer.h"
+#include "Templates/Atomic.h"
 
 /**
  * Implements a lock-free first-in first-out queue using a circular array.
  *
- * This class is thread safe only in two-thread scenarios, where the first thread
- * always reads and the second thread always writes. The head and tail indices are
- * stored in volatile memory to prevent the compiler from reordering the instructions
- * that are meant to update the indices AFTER items have been inserted or removed.
+ * This class is thread safe only in single-producer single-consumer scenarios.
  *
  * The number of items that can be enqueued is one less than the queue's capacity,
  * because one item will be used for detecting full and empty states.
  *
+ * There is some room for optimization via using fine grained memory fences, but
+ * the implications for all of our target platforms need further analysis, so
+ * we're using the simpler sequentially consistent model for now.
+ *
  * @param ElementType The type of elements held in the queue.
  */
 template<typename ElementType> class TCircularQueue
@@ -27,35 +29,32 @@ public:
 	 *
 	 * @param CapacityPlusOne The number of elements that the queue can hold (will be rounded up to the next power of 2).
 	 */
-	TCircularQueue(uint32 CapacityPlusOne)
+	explicit TCircularQueue(uint32 CapacityPlusOne)
 		: Buffer(CapacityPlusOne)
 		, Head(0)
 		, Tail(0)
 	{ }
 
-	/** Virtual destructor. */
-	virtual ~TCircularQueue() { }
-
 public:
 
 	/**
 	 * Gets the number of elements in the queue.
 	 *
-	 * The result reflects the calling thread's current view. Since no
-	 * locking is used, different threads may return different results.
+	 * Can be called from any thread. The result reflects the calling thread's current
+	 * view. Since no locking is used, different threads may return different results.
 	 *
 	 * @return Number of queued elements.
 	 */
 	uint32 Count() const
 	{
-		int32 Count = Tail - Head;
+		int64 Count = Tail.Load() - Head.Load();
 
 		if (Count < 0)
 		{
 			Count += Buffer.Capacity();
 		}
 
-		return Count;
+		return (uint32)Count;
 	}
 
 	/**
@@ -67,10 +66,12 @@ public:
 	 */
 	bool Dequeue(ElementType& OutElement)
 	{
-		if (Head != Tail)
+		const uint32 CurrentHead = Head.Load();
+
+		if (CurrentHead != Tail.Load())
 		{
-			OutElement = Buffer[Head];
-			Head = Buffer.GetNextIndex(Head);
+			OutElement = Buffer[CurrentHead];
+			Head.Store(Buffer.GetNextIndex(CurrentHead));
 
 			return true;
 		}
@@ -86,7 +87,7 @@ public:
 	 */
 	void Empty()
 	{
-		Head = Tail;
+		Head.Store(Tail.Load());
 	}
 
 	/**
@@ -98,12 +99,13 @@ public:
 	 */
 	bool Enqueue(const ElementType& Element)
 	{
-		uint32 NewTail = Buffer.GetNextIndex(Tail);
+		const uint32 CurrentTail = Tail.Load();
+		uint32 NewTail = Buffer.GetNextIndex(CurrentTail);
 
-		if (NewTail != Head)
+		if (NewTail != Head.Load())
 		{
-			Buffer[Tail] = Element;
-			Tail = NewTail;
+			Buffer[CurrentTail] = Element;
+			Tail.Store(NewTail);
 
 			return true;
 		}
@@ -114,29 +116,29 @@ public:
 	/**
 	 * Checks whether the queue is empty.
 	 *
-	 * The result reflects the calling thread's current view. Since no
-	 * locking is used, different threads may return different results.
+	 * Can be called from any thread. The result reflects the calling thread's current
+	 * view. Since no locking is used, different threads may return different results.
 	 *
 	 * @return true if the queue is empty, false otherwise.
 	 * @see Empty, IsFull
 	 */
 	FORCEINLINE bool IsEmpty() const
 	{
-		return (Head == Tail);
+		return (Head.Load() == Tail.Load());
 	}
 
 	/**
 	 * Checks whether the queue is full.
 	 *
-	 * The result reflects the calling thread's current view. Since no
-	 * locking is used, different threads may return different results.
+	 * Can be called from any thread. The result reflects the calling thread's current
+	 * view. Since no locking is used, different threads may return different results.
 	 *
 	 * @return true if the queue is full, false otherwise.
 	 * @see IsEmpty
 	 */
 	bool IsFull() const
 	{
-		return (Buffer.GetNextIndex(Tail) == Head);
+		return (Buffer.GetNextIndex(Tail.Load()) == Head.Load());
 	}
 
 	/**
@@ -148,9 +150,11 @@ public:
 	 */
 	bool Peek(ElementType& OutItem)
 	{
-		if (Head != Tail)
+		const uint32 CurrentHead = Head.Load();
+
+		if (CurrentHead != Tail.Load())
 		{
-			OutItem = Buffer[Head];
+			OutItem = Buffer[CurrentHead];
 
 			return true;
 		}
@@ -164,8 +168,8 @@ private:
 	TCircularBuffer<ElementType> Buffer;
 
 	/** Holds the index to the first item in the buffer. */
-	MS_ALIGN(PLATFORM_CACHE_LINE_SIZE) volatile uint32 Head GCC_ALIGN(PLATFORM_CACHE_LINE_SIZE);
+	TAtomic<uint32> Head;
 
 	/** Holds the index to the last item in the buffer. */
-	MS_ALIGN(PLATFORM_CACHE_LINE_SIZE) volatile uint32 Tail GCC_ALIGN(PLATFORM_CACHE_LINE_SIZE);
+	TAtomic<uint32> Tail;
 };
diff --git a/Engine/Source/Runtime/Core/Public/Containers/LruCache.h b/Engine/Source/Runtime/Core/Public/Containers/LruCache.h
--- a/Engine/Source/Runtime/Core/Public/Containers/LruCache.h
+++ b/Engine/Source/Runtime/Core/Public/Containers/LruCache.h
@@ -213,7 +213,7 @@ public:
 	template<typename Predicate>
 	FORCEINLINE bool ContainsByPredicate(Predicate Pred) const
 	{
-		for (const auto& Entry : LookupSet)
+		for (const FCacheEntry* Entry : LookupSet)
 		{
 			if (Pred(Entry->Key, Entry->Value))
 			{
@@ -234,7 +234,7 @@ public:
 	{
 		check(InMaxNumElements >= 0);
 
-		for (auto& Entry : LookupSet)
+		for (FCacheEntry* Entry : LookupSet)
 		{
 			delete Entry;
 		}
@@ -258,7 +258,7 @@ public:
 	{
 		TArray<ValueType> Result;
 
-		for (const auto& Entry : LookupSet)
+		for (const FCacheEntry* Entry : LookupSet)
 		{
 			if (Pred(Entry->Key, Entry->Value))
 			{
@@ -319,7 +319,7 @@ public:
 	template<typename Predicate>
 	const ValueType* FindByPredicate(Predicate Pred) const
 	{
-		for (const auto& Entry : LookupSet)
+		for (const FCacheEntry* Entry : LookupSet)
 		{
 			if (Pred(Entry->Key, Entry->Value))
 			{
@@ -338,7 +338,7 @@ public:
 	 */
 	void GetKeys(TArray<KeyType>& OutKeys) const
 	{
-		for (const auto& Entry : LookupSet)
+		for (const FCacheEntry* Entry : LookupSet)
 		{
 			OutKeys.Add(Entry->Key);
 		}
@@ -394,7 +394,7 @@ public:
 	{
 		int32 NumRemoved = 0;
 
-		for (const auto& Entry : LookupSet)
+		for (const FCacheEntry* Entry : LookupSet)
 		{
 			if (Pred(Entry->Key, Entry->Value))
 			{
@@ -539,7 +539,8 @@ public:
 		FORCEINLINE void RemoveCurrentAndIncrement()
 		{
 			check(Cache != nullptr);
-			auto MoreRecentEntry = this->GetCurrentEntry();
+
+			FCacheEntry* MoreRecentEntry = this->GetCurrentEntry();
 			this->Increment();
 			Cache->Remove(MoreRecentEntry);
 		}
diff --git a/Engine/Source/Runtime/Core/Public/Math/Range.h b/Engine/Source/Runtime/Core/Public/Math/Range.h
--- a/Engine/Source/Runtime/Core/Public/Math/Range.h
+++ b/Engine/Source/Runtime/Core/Public/Math/Range.h
@@ -229,6 +229,28 @@ public:
 		return LowerBound;
 	}
 
+	/**
+	 * Assign the new lower bound for this range
+	 *
+	 * @param NewLowerBound The new lower bound to assign
+	 * @see GetLowerBound
+	 */
+	void SetLowerBound(const BoundsType& NewLowerBound)
+	{
+		LowerBound = NewLowerBound;
+	}
+
+	/**
+	 * Assign the new lower bound value for this range. Current lower bound must not be open to call this method.
+	 *
+	 * @param NewLowerBoundValue The new lower bound value to assign
+	 * @see GetLowerBound
+	 */
+	void SetLowerBoundValue(ElementValueOrConstRef NewLowerBoundValue)
+	{
+		LowerBound.SetValue(NewLowerBoundValue);
+	}
+
 	/**
 	 * Get the value of the lower bound.
 	 *
@@ -253,6 +275,28 @@ public:
 		return UpperBound;
 	}
 
+	/**
+	 * Assign the new upper bound for this range
+	 *
+	 * @param NewUpperBound The new upper bound to assign
+	 * @see GetUpperBound
+	 */
+	void SetUpperBound(const BoundsType& NewUpperBound)
+	{
+		UpperBound = NewUpperBound;
+	}
+
+	/**
+	 * Assign the new upper bound value for this range. Current upper bound must not be open to call this method.
+	 *
+	 * @param NewUpperBoundValue The new upper bound value to assign
+	 * @see GetUpperBound
+	 */
+	void SetUpperBoundValue(ElementValueOrConstRef NewUpperBoundValue)
+	{
+		UpperBound.SetValue(NewUpperBoundValue);
+	}
+
 	/**
 	 * Get the value of the upper bound.
 	 *
diff --git a/Engine/Source/Runtime/Core/Public/Math/RangeBound.h b/Engine/Source/Runtime/Core/Public/Math/RangeBound.h
--- a/Engine/Source/Runtime/Core/Public/Math/RangeBound.h
+++ b/Engine/Source/Runtime/Core/Public/Math/RangeBound.h
@@ -100,6 +100,22 @@ public:
 		return Value;
 	}
 
+
+	/**
+	 * Sets the bound's value, maintining the inclusivity of the bound
+	 *
+	 * Use IsClosed() to verify that this bound is closed before calling this method.
+	 *
+	 * @param NewValue New bound value.
+	 * @see IsClosed
+	 */
+	FORCEINLINE_DEBUGGABLE void SetValue(ElementValueOrConstRef NewValue)
+	{
+		check(Type != ERangeBoundTypes::Open);
+
+		Value = NewValue;
+	}
+
 	/**
 	 * Checks whether the bound is closed.
 	 *
diff --git a/Engine/Source/Runtime/Core/Public/Math/Rotator.h b/Engine/Source/Runtime/Core/Public/Math/Rotator.h
--- a/Engine/Source/Runtime/Core/Public/Math/Rotator.h
+++ b/Engine/Source/Runtime/Core/Public/Math/Rotator.h
@@ -297,6 +297,29 @@ public:
 	 */
 	CORE_API void GetWindingAndRemainder( FRotator& Winding, FRotator& Remainder ) const;
 
+	/**
+	* Return the manhattan distance in degrees between this Rotator and the passed in one.
+	* @param Rotator[In] the Rotator we are comparing with.
+	* @return Distance(Manhattan) between the two rotators. 
+	*/
+	float GetManhattanDistance(const FRotator & Rotator) const;
+
+	/**
+	* Return a Rotator that has the same rotation but has different degree values for Yaw, Pitch, and Roll.
+	* This rotator should be within -180,180 range,
+	* @return A Rotator with the same rotation but different degrees.
+	*/
+	FRotator GetEquivalentRotator() const;
+
+	/**
+	* Modify if necessary the passed in rotator to be the closest rotator to it based upon it's equivalent.
+	* This Rotator should be within (-180, 180], usually just constructed from a Matrix or a Quaternion.
+	*
+	* @param MakeClosest[In/Out] the Rotator we want to make closest to us. Should be between 
+	* (-180, 180]. This Rotator may change if we need to use different degree values to make it closer.
+	*/
+	void SetClosestToMe(FRotator& MakeClosest) const;
+
 	/**
 	 * Get a textual representation of the vector.
 	 *
@@ -759,6 +782,24 @@ FORCEINLINE bool FRotator::ContainsNaN() const
 			!FMath::IsFinite(Roll));
 }
 
+FORCEINLINE float FRotator::GetManhattanDistance(const FRotator & Rotator) const
+{
+	return FMath::Abs<float>(Yaw - Rotator.Yaw) + FMath::Abs<float>(Pitch - Rotator.Pitch) + FMath::Abs<float>(Roll - Rotator.Roll);
+}
+
+FORCEINLINE FRotator FRotator::GetEquivalentRotator() const
+{
+	return FRotator(180.0f - Pitch,Yaw + 180.0f, Roll + 180.0f);
+}
+
+FORCEINLINE void FRotator::SetClosestToMe(FRotator& MakeClosest) const
+{
+	FRotator OtherChoice = MakeClosest.GetEquivalentRotator();
+	float FirstDiff = GetManhattanDistance(MakeClosest);
+	float SecondDiff = GetManhattanDistance(OtherChoice);
+	if (SecondDiff < FirstDiff)
+		MakeClosest = OtherChoice;
+}
 
 template<> struct TIsPODType<FRotator> { enum { Value = true }; };
 
diff --git a/Engine/Source/Runtime/Core/Public/Math/UnrealMathUtility.h b/Engine/Source/Runtime/Core/Public/Math/UnrealMathUtility.h
--- a/Engine/Source/Runtime/Core/Public/Math/UnrealMathUtility.h
+++ b/Engine/Source/Runtime/Core/Public/Math/UnrealMathUtility.h
@@ -283,17 +283,25 @@ struct FMath : public FPlatformMath
 
 	/** Divides two integers and rounds up */
 	template <class T>
-	static FORCEINLINE T DivideAndRoundUp(T Dividend,T Divisor)
+	static FORCEINLINE T DivideAndRoundUp(T Dividend, T Divisor)
 	{
 		return (Dividend + Divisor - 1) / Divisor;
 	}
 
+	/** Divides two integers and rounds down */
 	template <class T>
-	static FORCEINLINE T DivideAndRoundDown(T Dividend,T Divisor)
+	static FORCEINLINE T DivideAndRoundDown(T Dividend, T Divisor)
 	{
 		return Dividend / Divisor;
 	}
 
+	/** Divides two integers and rounds to nearest */
+	template <class T>
+	static FORCEINLINE T DivideAndRoundNearest(T Dividend, T Divisor)
+	{
+		return (Dividend > 0) ? 1 + ((Dividend - 1) / Divisor) : (Dividend / Divisor);
+	}
+
 	/**
 	 * Computes the base 2 logarithm of the specified value
 	 *
diff --git a/Engine/Source/Runtime/Core/Public/Misc/Attribute.h b/Engine/Source/Runtime/Core/Public/Misc/Attribute.h
--- a/Engine/Source/Runtime/Core/Public/Misc/Attribute.h
+++ b/Engine/Source/Runtime/Core/Public/Misc/Attribute.h
@@ -340,3 +340,51 @@ private:
 	/** Our attribute's 'getter' delegate */
 	FGetter Getter;
 };
+
+/**
+ * Helper function for creating TAttributes from a non-const member function pointer, accessed through a raw pointer
+ */
+template<typename T, typename SourceType, typename SourceTypeOrBase, typename... PayloadTypes>
+FORCEINLINE TAttribute<T> MakeAttributeRaw(SourceType* InObject, T (SourceTypeOrBase::*InMethod)(PayloadTypes...), typename TDecay<PayloadTypes>::Type... InputPayload)
+{
+	return TAttribute<T>::Create(TAttribute<T>::FGetter::CreateRaw(InObject, InMethod, MoveTemp(InputPayload)...));
+}
+
+/**
+ * Helper function for creating TAttributes from a const member function pointer, accessed through a raw pointer
+ */
+template<typename T, typename SourceType, typename SourceTypeOrBase, typename... PayloadTypes>
+FORCEINLINE TAttribute<T> MakeAttributeRaw(const SourceType* InObject, T (SourceTypeOrBase::*InMethod)(PayloadTypes...) const, typename TDecay<PayloadTypes>::Type...  InputPayload)
+{
+	return TAttribute<T>::Create(TAttribute<T>::FGetter::CreateRaw(InObject, InMethod, MoveTemp(InputPayload)...));
+}
+
+/**
+ * Helper function for creating TAttributes from a non-const member function pointer, accessed through a weak pointer to the shared object
+ */
+template<typename T, typename SourceType, typename SourceTypeOrBase, typename... PayloadTypes>
+FORCEINLINE TAttribute<T> MakeAttributeSP(SourceType* InObject, T (SourceTypeOrBase::*InMethod)(PayloadTypes...), typename TDecay<PayloadTypes>::Type...  InputPayload)
+{
+	return TAttribute<T>::Create(TAttribute<T>::FGetter::CreateSP(InObject, InMethod, MoveTemp(InputPayload)...));
+}
+
+/**
+ * Helper function for creating TAttributes from a const member function pointer, accessed through a weak pointer to the shared object
+ */
+template<typename T, typename SourceType, typename SourceTypeOrBase, typename... PayloadTypes>
+FORCEINLINE TAttribute<T> MakeAttributeSP(const SourceType* InObject, T (SourceTypeOrBase::*InMethod)(PayloadTypes...) const, typename TDecay<PayloadTypes>::Type...  InputPayload)
+{
+	return TAttribute<T>::Create(TAttribute<T>::FGetter::CreateSP(InObject, InMethod, MoveTemp(InputPayload)...));
+}
+
+/**
+ * Helper function for creating TAttributes from a lambda
+ * TAttribute<float> FloatAttribute = MakeAttributeLambda([]{ return 10.f; });
+ */
+template<typename LambdaType, typename... PayloadTypes>
+decltype(auto) MakeAttributeLambda(LambdaType&& InCallable, PayloadTypes&&... InputPayload)
+{
+	typedef decltype(InCallable(DeclVal<PayloadTypes>()...)) T;
+
+	return TAttribute<T>::Create(TAttribute<T>::FGetter::CreateLambda(InCallable, Forward<PayloadTypes>(InputPayload)...));
+}
diff --git a/Engine/Source/Runtime/Core/Public/Stats/Stats2.h b/Engine/Source/Runtime/Core/Public/Stats/Stats2.h
--- a/Engine/Source/Runtime/Core/Public/Stats/Stats2.h
+++ b/Engine/Source/Runtime/Core/Public/Stats/Stats2.h
@@ -2085,6 +2085,7 @@ DECLARE_STATS_GROUP(TEXT("Light Rendering"),STATGROUP_LightRendering, STATCAT_Ad
 DECLARE_STATS_GROUP(TEXT("LoadTime"), STATGROUP_LoadTime, STATCAT_Advanced);
 DECLARE_STATS_GROUP_VERBOSE(TEXT("LoadTimeVerbose"), STATGROUP_LoadTimeVerbose, STATCAT_Advanced);
 DECLARE_STATS_GROUP_VERBOSE(TEXT("MathVerbose"), STATGROUP_MathVerbose, STATCAT_Advanced);
+DECLARE_STATS_GROUP(TEXT("Media"),STATGROUP_Media, STATCAT_Advanced);
 DECLARE_STATS_GROUP(TEXT("Memory Allocator"),STATGROUP_MemoryAllocator, STATCAT_Advanced);
 DECLARE_STATS_GROUP(TEXT("Memory Platform"),STATGROUP_MemoryPlatform, STATCAT_Advanced);
 DECLARE_STATS_GROUP(TEXT("Memory StaticMesh"),STATGROUP_MemoryStaticMesh, STATCAT_Advanced);
diff --git a/Engine/Source/Runtime/Core/Public/Templates/SharedPointer.h b/Engine/Source/Runtime/Core/Public/Templates/SharedPointer.h
--- a/Engine/Source/Runtime/Core/Public/Templates/SharedPointer.h
+++ b/Engine/Source/Runtime/Core/Public/Templates/SharedPointer.h
@@ -522,7 +522,11 @@ struct FMakeReferenceTo<void>
 {
 	typedef void Type;
 };
-
+template<>
+struct FMakeReferenceTo<const void>
+{
+	typedef void Type;
+};
 
 /**
  * TSharedPtr is a non-intrusive reference-counted authoritative object pointer.  This shared pointer
diff --git a/Engine/Source/Runtime/Core/Public/UObject/SequencerObjectVersion.h b/Engine/Source/Runtime/Core/Public/UObject/SequencerObjectVersion.h
--- a/Engine/Source/Runtime/Core/Public/UObject/SequencerObjectVersion.h
+++ b/Engine/Source/Runtime/Core/Public/UObject/SequencerObjectVersion.h
@@ -30,6 +30,9 @@ struct CORE_API FSequencerObjectVersion
 		// When finished now defaults to project default
 		WhenFinishedDefaultsToProjectDefault,
 
+		// When finished now defaults to project default
+		FloatToIntConversion,
+
 		// -----<new versions can be added above this line>-------------------------------------------------
 		VersionPlusOne,
 		LatestVersion = VersionPlusOne - 1
diff --git a/Engine/Source/Runtime/Engine/Classes/Camera/CameraComponent.h b/Engine/Source/Runtime/Engine/Classes/Camera/CameraComponent.h
--- a/Engine/Source/Runtime/Engine/Classes/Camera/CameraComponent.h
+++ b/Engine/Source/Runtime/Engine/Classes/Camera/CameraComponent.h
@@ -53,26 +53,26 @@ class ENGINE_API UCameraComponent : public USceneComponent
 	void SetAspectRatio(float InAspectRatio) { AspectRatio = InAspectRatio; }
 
 	// If bConstrainAspectRatio is true, black bars will be added if the destination view has a different aspect ratio than this camera requested.
-	UPROPERTY(Interp, EditAnywhere, BlueprintReadWrite, Category = CameraSettings)
+	UPROPERTY(Interp, EditAnywhere, BlueprintReadWrite, Category = CameraOptions)
 	uint32 bConstrainAspectRatio : 1;
 	UFUNCTION(BlueprintCallable, Category = Camera)
 	void SetConstraintAspectRatio(bool bInConstrainAspectRatio) { bConstrainAspectRatio = bInConstrainAspectRatio; }
 
 	// If true, account for the field of view angle when computing which level of detail to use for meshes.
-	UPROPERTY(Interp, EditAnywhere, AdvancedDisplay, BlueprintReadWrite, Category = CameraSettings)
+	UPROPERTY(Interp, EditAnywhere, AdvancedDisplay, BlueprintReadWrite, Category = CameraOptions)
 	uint32 bUseFieldOfViewForLOD : 1;
 	UFUNCTION(BlueprintCallable, Category = Camera)
 	void SetUseFieldOfViewForLOD(bool bInUseFieldOfViewForLOD) { bUseFieldOfViewForLOD = bInUseFieldOfViewForLOD; }
 
 	/** True if the camera's orientation and position should be locked to the HMD */
-	UPROPERTY(EditAnywhere, BlueprintReadWrite, Category = CameraSettings)
+	UPROPERTY(EditAnywhere, BlueprintReadWrite, Category = CameraOptions)
 	uint32 bLockToHmd : 1;
 
 	/**
 	 * If this camera component is placed on a pawn, should it use the view/control rotation of the pawn where possible?
 	 * @see APawn::GetViewRotation()
 	 */
-	UPROPERTY(EditAnywhere, BlueprintReadWrite, Category = CameraSettings)
+	UPROPERTY(EditAnywhere, BlueprintReadWrite, Category = CameraOptions)
 	uint32 bUsePawnControlRotation : 1;
 
 protected:
diff --git a/Engine/Source/Runtime/Engine/Classes/Components/SceneComponent.h b/Engine/Source/Runtime/Engine/Classes/Components/SceneComponent.h
--- a/Engine/Source/Runtime/Engine/Classes/Components/SceneComponent.h
+++ b/Engine/Source/Runtime/Engine/Classes/Components/SceneComponent.h
@@ -338,6 +338,26 @@ protected:
 	virtual void PostNetReceive() override;
 	virtual void PostRepNotifies() override;
 
+public:
+
+	/**  Convenience function to get the relative rotation from the passed in world rotation*/
+	/* @param WorldRotation  World rotation that we want to convert to relative to the components parent*/
+	/* @return Returns the relative rotation*/
+	FQuat GetRelativeRotationFromWorld(const FQuat & WorldRotation);
+
+	/**
+	* Set the rotation of the component relative to its parent and force RelativeRotation to be equal to new rotation.
+	* This allows us to set and save Rotators with angles out side the normalized range, Note that doing so may break the 
+	* RotatorCache so use with care.
+	* @param NewRotation		New rotation of the component relative to its parent. We will force RelativeRotation to this value.
+	* @param SweepHitResult	Hit result from any impact if sweep is true.
+	* @param bSweep			Whether we sweep to the destination (currently not supported for rotation).
+	* @param bTeleport			Whether we teleport the physics state (if physics collision is enabled for this object).
+	*							If true, physics velocity for this object is unchanged (so ragdoll parts are not affected by change in location).
+	*							If false, physics velocity is updated based on the change in position (affecting ragdoll parts).
+	*/
+	void SetRelativeRotationExact(FRotator NewRotation, bool bSweep = false, FHitResult* OutSweepHitResult = nullptr, ETeleportType Teleport = ETeleportType::None);
+
 public:
 	/**
 	 * Set the location of the component relative to its parent
diff --git a/Engine/Source/Runtime/Engine/Classes/Curves/KeyHandle.h b/Engine/Source/Runtime/Engine/Classes/Curves/KeyHandle.h
--- a/Engine/Source/Runtime/Engine/Classes/Curves/KeyHandle.h
+++ b/Engine/Source/Runtime/Engine/Classes/Curves/KeyHandle.h
@@ -14,6 +14,8 @@ struct FKeyHandle
 {
 	ENGINE_API FKeyHandle();
 
+	static ENGINE_API FKeyHandle Invalid();
+
 	bool operator ==(const FKeyHandle& Other) const
 	{
 		return Index == Other.Index;
@@ -23,7 +25,12 @@ struct FKeyHandle
 	{
 		return Index != Other.Index;
 	}
-	
+
+	friend bool operator<(FKeyHandle A, FKeyHandle B)
+	{
+		return A.Index < B.Index;
+	}
+
 	friend uint32 GetTypeHash(const FKeyHandle& Handle)
 	{
 		return GetTypeHash(Handle.Index);
@@ -37,6 +44,9 @@ struct FKeyHandle
 
 private:
 
+	/** Private constructor from a specific index - only for use in FKeyHandle::Invalid to avoid allocating new handles unnecessarily */
+	explicit FKeyHandle(uint32 SpecificIndex);
+
 	uint32 Index;
 };
 
diff --git a/Engine/Source/Runtime/Engine/Private/ActorEditor.cpp b/Engine/Source/Runtime/Engine/Private/ActorEditor.cpp
--- a/Engine/Source/Runtime/Engine/Private/ActorEditor.cpp
+++ b/Engine/Source/Runtime/Engine/Private/ActorEditor.cpp
@@ -461,27 +461,37 @@ void AActor::EditorApplyRotation(const FRotator& DeltaRotation, bool bAltDown, b
 {
 	if( RootComponent != NULL )
 	{
-		const FRotator Rot = RootComponent->GetAttachParent() != NULL ? GetActorRotation() : RootComponent->RelativeRotation;
-
+		FRotator Rot = RootComponent->GetAttachParent() != NULL ? GetActorRotation() : RootComponent->RelativeRotation;
 		FRotator ActorRotWind, ActorRotRem;
 		Rot.GetWindingAndRemainder(ActorRotWind, ActorRotRem);
-
 		const FQuat ActorQ = ActorRotRem.Quaternion();
 		const FQuat DeltaQ = DeltaRotation.Quaternion();
-		const FQuat ResultQ = DeltaQ * ActorQ;
-		const FRotator NewActorRotRem = FRotator( ResultQ );
-		FRotator DeltaRot = NewActorRotRem - ActorRotRem;
-		DeltaRot.Normalize();
 
-		if( RootComponent->GetAttachParent() != NULL )
+		FRotator NewActorRotRem;
+		if(RootComponent->GetAttachParent() != NULL )
 		{
-			RootComponent->SetWorldRotation( Rot + DeltaRot );
+			//first we get the new rotation in relative space.
+			const FQuat ResultQ = DeltaQ * ActorQ;
+			NewActorRotRem = FRotator(ResultQ);
+			FRotator DeltaRot = NewActorRotRem - ActorRotRem;
+			FRotator NewRotation = Rot + DeltaRot;
+			FQuat NewRelRotation = NewRotation.Quaternion();
+			NewRelRotation = RootComponent->GetRelativeRotationFromWorld(NewRelRotation);
+			NewActorRotRem = FRotator(NewRelRotation);
+			//now we need to get current relative rotation to find the diff
+			Rot = RootComponent->RelativeRotation;
+			Rot.GetWindingAndRemainder(ActorRotWind, ActorRotRem);
 		}
 		else
 		{
-			// No attachment.  Directly set relative rotation (to support winding)
-			RootComponent->SetRelativeRotation( Rot + DeltaRot );
+			const FQuat ResultQ = DeltaQ * ActorQ;
+			NewActorRotRem = FRotator(ResultQ);
 		}
+
+		ActorRotRem.SetClosestToMe(NewActorRotRem);
+		FRotator DeltaRot = NewActorRotRem - ActorRotRem;
+		DeltaRot.Normalize();
+		RootComponent->SetRelativeRotationExact( Rot + DeltaRot );
 	}
 	else
 	{
@@ -546,7 +556,7 @@ void AActor::EditorApplyMirror(const FVector& MirrorScale, const FVector& PivotL
 
 	if( RootComponent != NULL )
 	{
-		GetRootComponent()->SetRelativeRotation( NewRot.Rotator() );
+		GetRootComponent()->SetRelativeRotationExact( NewRot.Rotator() );
 		FVector Loc = GetActorLocation();
 		Loc -= PivotLocation;
 		Loc *= MirrorScale;
diff --git a/Engine/Source/Runtime/Engine/Private/Animation/AnimInstanceProxy.cpp b/Engine/Source/Runtime/Engine/Private/Animation/AnimInstanceProxy.cpp
--- a/Engine/Source/Runtime/Engine/Private/Animation/AnimInstanceProxy.cpp
+++ b/Engine/Source/Runtime/Engine/Private/Animation/AnimInstanceProxy.cpp
@@ -384,10 +384,13 @@ void FAnimInstanceProxy::PostUpdate(UAnimInstance* InAnimInstance) const
 #endif
 
 	FMessageLog MessageLog("AnimBlueprintLog");
-	const TArray<FLogMessageEntry>& Messages = LoggedMessagesMap.FindChecked("Update");
-	for (const FLogMessageEntry& Message : Messages)
+	const TArray<FLogMessageEntry>* Messages = LoggedMessagesMap.Find("Update");
+	if (ensureMsgf(Messages, TEXT("PreUpdate isn't called. This could potentially cause other issues.")))
 	{
-		MessageLog.Message(Message.Key, Message.Value);
+		for (const FLogMessageEntry& Message : *Messages)
+		{
+			MessageLog.Message(Message.Key, Message.Value);
+		}
 	}
 }
 
diff --git a/Engine/Source/Runtime/Engine/Private/Components/SceneComponent.cpp b/Engine/Source/Runtime/Engine/Private/Components/SceneComponent.cpp
--- a/Engine/Source/Runtime/Engine/Private/Components/SceneComponent.cpp
+++ b/Engine/Source/Runtime/Engine/Private/Components/SceneComponent.cpp
@@ -36,7 +36,6 @@
 
 #define LOCTEXT_NAMESPACE "SceneComponent"
 
-
 namespace SceneComponentStatics
 {
 	static const FName DefaultSceneRootVariableName(TEXT("DefaultSceneRoot"));
@@ -1243,6 +1242,18 @@ void USceneComponent::SetRelativeLocationAndRotation(FVector NewLocation, FRotat
 	}
 }
 
+void USceneComponent::SetRelativeRotationExact(FRotator NewRotation, bool bSweep, FHitResult* OutSweepHitResult, ETeleportType Teleport)
+{
+	if (!NewRotation.Equals(RelativeRotation, SCENECOMPONENT_ROTATOR_TOLERANCE))
+	{
+		// We know the rotations are different, don't bother with the cache.
+		const FQuat NewQuat = NewRotation.Quaternion();
+     	SetRelativeLocationAndRotation(RelativeLocation, NewQuat, bSweep, OutSweepHitResult, Teleport);
+		RelativeRotation = NewRotation;
+		
+	}
+}
+
 void USceneComponent::SetRelativeRotation(FRotator NewRotation, bool bSweep, FHitResult* OutSweepHitResult, ETeleportType Teleport)
 {
 	if (!NewRotation.Equals(RelativeRotation, SCENECOMPONENT_ROTATOR_TOLERANCE))
@@ -1373,7 +1384,7 @@ void USceneComponent::SetWorldLocation(FVector NewLocation, bool bSweep, FHitRes
 	SetRelativeLocation(NewRelLocation, bSweep, OutSweepHitResult, Teleport);
 }
 
-void USceneComponent::SetWorldRotation(const FQuat& NewRotation, bool bSweep, FHitResult* OutSweepHitResult, ETeleportType Teleport)
+FQuat USceneComponent::GetRelativeRotationFromWorld(const FQuat & NewRotation)
 {
 	FQuat NewRelRotation = NewRotation;
 
@@ -1384,7 +1395,7 @@ void USceneComponent::SetWorldRotation(const FQuat& NewRotation, bool bSweep, FH
 		// in order to support mirroring, you'll have to use FTransform.GetRelativeTransform
 		// because negative SCALE should flip the rotation
 		if (FTransform::AnyHasNegativeScale(RelativeScale3D, ParentToWorld.GetScale3D()))
-		{	
+		{
 			FTransform NewTransform = GetComponentTransform();
 			// set new desired rotation
 			NewTransform.SetRotation(NewRotation);
@@ -1397,10 +1408,15 @@ void USceneComponent::SetWorldRotation(const FQuat& NewRotation, bool bSweep, FH
 			const FQuat ParentToWorldQuat = ParentToWorld.GetRotation();
 			// Quat multiplication works reverse way, make sure you do Parent(-1) * World = Local, not World*Parent(-) = Local (the way matrix does)
 			const FQuat NewRelQuat = ParentToWorldQuat.Inverse() * NewRotation;
-			NewRelRotation = NewRelQuat;		
+			NewRelRotation = NewRelQuat;
 		}
 	}
+	return NewRelRotation;
+}
 
+void USceneComponent::SetWorldRotation(const FQuat& NewRotation, bool bSweep, FHitResult* OutSweepHitResult, ETeleportType Teleport)
+{
+	FQuat NewRelRotation = GetRelativeRotationFromWorld(NewRotation);
 	SetRelativeRotation(NewRelRotation, bSweep, OutSweepHitResult, Teleport);
 }
 
@@ -1417,7 +1433,6 @@ void USceneComponent::SetWorldRotation(FRotator NewRotation, bool bSweep, FHitRe
 	}
 }
 
-
 void USceneComponent::SetWorldScale3D(FVector NewScale)
 {
 	FVector NewRelScale = NewScale;
@@ -2684,7 +2699,9 @@ bool USceneComponent::InternalSetWorldLocationAndRotation(FVector NewLocation, c
 	}
 
 	const FRotator NewRelativeRotation = RelativeRotationCache.QuatToRotator_ReadOnly(NewRotationQuat);
-	if (!NewLocation.Equals(RelativeLocation) || !NewRelativeRotation.Equals(RelativeRotation))
+	bool bDiffLocation = !NewLocation.Equals(RelativeLocation);
+	bool bDiffRotation = !NewRelativeRotation.Equals(RelativeRotation);
+	if (bDiffLocation || bDiffRotation)
 	{
 		RelativeLocation = NewLocation;
 
@@ -2695,8 +2712,11 @@ bool USceneComponent::InternalSetWorldLocationAndRotation(FVector NewLocation, c
 		// it is required to generate that same ComponentToWorld otherwise the FComponentInstanceDataCache
 		// might fail to apply to the relevant component. In order to have the exact same transform
 		// we must enforce the quaternion to come from the rotator (as in load)
-		RelativeRotation = NewRelativeRotation;
-		RelativeRotationCache.RotatorToQuat(NewRelativeRotation);
+		if (bDiffRotation)
+		{
+			RelativeRotation = NewRelativeRotation;
+			RelativeRotationCache.RotatorToQuat(NewRelativeRotation);
+		}
 
 #if ENABLE_NAN_DIAGNOSTIC
 		if (RelativeRotation.ContainsNaN())
diff --git a/Engine/Source/Runtime/Engine/Private/Components/SkeletalMeshComponent.cpp b/Engine/Source/Runtime/Engine/Private/Components/SkeletalMeshComponent.cpp
--- a/Engine/Source/Runtime/Engine/Private/Components/SkeletalMeshComponent.cpp
+++ b/Engine/Source/Runtime/Engine/Private/Components/SkeletalMeshComponent.cpp
@@ -624,6 +624,8 @@ void USkeletalMeshComponent::InitAnim(bool bForceReinit)
 		const bool bSkeletonMismatch = AnimSkeleton && (AnimScriptInstance->CurrentSkeleton!=SkeletalMesh->Skeleton);
 		const bool bSkeletonNotCompatible = AnimSkeleton && !bSkeletonMismatch && (AnimSkeleton->IsCompatibleMesh(SkeletalMesh) == false);
 
+		LastPoseTickFrame = 0;
+
 		if (bBlueprintMismatch || bSkeletonMismatch || bSkeletonNotCompatible || bClearAnimInstance)
 		{
 			ClearAnimScriptInstance();
diff --git a/Engine/Source/Runtime/Engine/Private/Curves/KeyHandle.cpp b/Engine/Source/Runtime/Engine/Private/Curves/KeyHandle.cpp
--- a/Engine/Source/Runtime/Engine/Private/Curves/KeyHandle.cpp
+++ b/Engine/Source/Runtime/Engine/Private/Curves/KeyHandle.cpp
@@ -182,7 +182,7 @@ void FKeyHandleLookupTable::MoveHandle(int32 OldIndex, int32 NewIndex)
 	{
 		TOptional<FKeyHandle> Handle = KeyHandles[OldIndex];
 
-		KeyHandles.RemoveAtSwap(OldIndex, 1, false);
+		KeyHandles.RemoveAt(OldIndex, 1, false);
 		KeyHandles.Insert(Handle, NewIndex);
 		if (Handle.IsSet())
 		{
@@ -194,6 +194,13 @@ void FKeyHandleLookupTable::MoveHandle(int32 OldIndex, int32 NewIndex)
 FKeyHandle FKeyHandleLookupTable::AllocateHandle(int32 Index)
 {
 	FKeyHandle NewKeyHandle;
+
+	int32 NumToAdd = Index + 1 - KeyHandles.Num();
+	if (NumToAdd > 0)
+	{
+		KeyHandles.AddDefaulted(NumToAdd);
+	}
+
 	KeyHandles.Insert(NewKeyHandle, Index);
 	KeyHandlesToIndices.Add(NewKeyHandle, Index);
 	return NewKeyHandle;
@@ -202,7 +209,7 @@ FKeyHandle FKeyHandleLookupTable::AllocateHandle(int32 Index)
 void FKeyHandleLookupTable::DeallocateHandle(int32 Index)
 {
 	TOptional<FKeyHandle> KeyHandle = KeyHandles[Index];
-	KeyHandles.RemoveAtSwap(Index, 1, false);
+	KeyHandles.RemoveAt(Index, 1, false);
 	if (KeyHandle.IsSet())
 	{
 		KeyHandlesToIndices.Remove(KeyHandle.GetValue());
diff --git a/Engine/Source/Runtime/Engine/Private/Curves/RichCurve.cpp b/Engine/Source/Runtime/Engine/Private/Curves/RichCurve.cpp
--- a/Engine/Source/Runtime/Engine/Private/Curves/RichCurve.cpp
+++ b/Engine/Source/Runtime/Engine/Private/Curves/RichCurve.cpp
@@ -130,9 +130,18 @@ bool FRichCurveKey::operator!=(const FRichCurveKey& Other) const
 FKeyHandle::FKeyHandle()
 {
 	static uint32 LastKeyHandleIndex = 1;
+	Index = ++LastKeyHandleIndex;
+
 	check(LastKeyHandleIndex != 0); // check in the unlikely event that this overflows
+}
 
-	Index = ++LastKeyHandleIndex;
+FKeyHandle::FKeyHandle(uint32 SpecificIndex)
+	: Index(SpecificIndex)
+{}
+
+FKeyHandle FKeyHandle::Invalid()
+{
+	return FKeyHandle(0);
 }
 
 
@@ -888,11 +897,11 @@ void FRichCurve::RemoveRedundantKeys(float Tolerance, float FirstKeyTime, float
 	{
 		const float CurrentKeyTime = Keys[KeyIndex].Time;
 
-		if (CurrentKeyTime < FirstKeyTime)
+		if (CurrentKeyTime <= FirstKeyTime)
 		{
 			StartKey = KeyIndex;
 		}
-		if (CurrentKeyTime > LastKeyTime)
+		if (CurrentKeyTime >= LastKeyTime)
 		{
 			EndKey = KeyIndex;
 			break;
diff --git a/Engine/Source/Runtime/Engine/Private/GameEngine.cpp b/Engine/Source/Runtime/Engine/Private/GameEngine.cpp
--- a/Engine/Source/Runtime/Engine/Private/GameEngine.cpp
+++ b/Engine/Source/Runtime/Engine/Private/GameEngine.cpp
@@ -177,7 +177,7 @@ UGameEngine::UGameEngine(const FObjectInitializer& ObjectInitializer)
 
 void UGameEngine::CreateGameViewportWidget( UGameViewportClient* GameViewportClient )
 {
-	bool bRenderDirectlyToWindow = !StartupMovieCaptureHandle.IsValid() && GIsDumpingMovie == 0; 
+	bool bRenderDirectlyToWindow = (!StartupMovieCaptureHandle.IsValid() || IMovieSceneCaptureModule::Get().IsStereoAllowed()) && GIsDumpingMovie == 0;
 	const bool bStereoAllowed = bRenderDirectlyToWindow;
 	TSharedRef<SOverlay> ViewportOverlayWidgetRef = SNew( SOverlay );
 
diff --git a/Engine/Source/Runtime/Engine/Private/Slate/SGameLayerManager.cpp b/Engine/Source/Runtime/Engine/Private/Slate/SGameLayerManager.cpp
--- a/Engine/Source/Runtime/Engine/Private/Slate/SGameLayerManager.cpp
+++ b/Engine/Source/Runtime/Engine/Private/Slate/SGameLayerManager.cpp
@@ -256,21 +256,30 @@ bool SGameLayerManager::OnVisualizeTooltip(const TSharedPtr<SWidget>& TooltipCon
 
 float SGameLayerManager::GetGameViewportDPIScale() const
 {
-	if ( const FSceneViewport* Viewport = SceneViewport.Get() )
+	const FSceneViewport* Viewport = SceneViewport.Get();
+
+	if (Viewport == nullptr)
 	{
-		FIntPoint ViewportSize = Viewport->GetSize();
-		const float GameUIScale = GetDefault<UUserInterfaceSettings>(UUserInterfaceSettings::StaticClass())->GetDPIScaleBasedOnSize(ViewportSize);
+		return 1;
+	}
 
-		// Remove the platform DPI scale from the incoming size.  Since the platform DPI is already
-		// attempt to normalize the UI for a high DPI, and the DPI scale curve is based on raw resolution
-		// for what a assumed platform scale of 1, extract that scale the calculated scale, since that will
-		// already be applied by slate.
-		const float FinalUIScale = GameUIScale / Viewport->GetCachedGeometry().Scale;
+	const auto UserInterfaceSettings = GetDefault<UUserInterfaceSettings>(UUserInterfaceSettings::StaticClass());
 
-		return FinalUIScale;
+	if (UserInterfaceSettings == nullptr)
+	{
+		return 1;
 	}
 
-	return 1;
+	FIntPoint ViewportSize = Viewport->GetSize();
+	const float GameUIScale = UserInterfaceSettings->GetDPIScaleBasedOnSize(ViewportSize);
+
+	// Remove the platform DPI scale from the incoming size.  Since the platform DPI is already
+	// attempt to normalize the UI for a high DPI, and the DPI scale curve is based on raw resolution
+	// for what a assumed platform scale of 1, extract that scale the calculated scale, since that will
+	// already be applied by slate.
+	const float FinalUIScale = GameUIScale / Viewport->GetCachedGeometry().Scale;
+
+	return FinalUIScale;
 }
 
 FOptionalSize SGameLayerManager::GetDefaultWindowTitleBarHeight() const
@@ -370,6 +379,11 @@ void SGameLayerManager::RemovePlayerWidgets(ULocalPlayer* LocalPlayer)
 
 void SGameLayerManager::AddOrUpdatePlayerLayers(const FGeometry& AllottedGeometry, UGameViewportClient* ViewportClient, const TArray<ULocalPlayer*>& GamePlayers)
 {
+	if (GamePlayers.Num() == 0)
+	{
+		return;
+	}
+
 	ESplitScreenType::Type SplitType = ViewportClient->GetCurrentSplitscreenConfiguration();
 	TArray<struct FSplitscreenData>& SplitInfo = ViewportClient->SplitscreenInfo;
 
diff --git a/Engine/Source/Runtime/Engine/Public/Curves/KeyFrameManipulator.h b/Engine/Source/Runtime/Engine/Public/Curves/KeyFrameManipulator.h
--- a/Engine/Source/Runtime/Engine/Public/Curves/KeyFrameManipulator.h
+++ b/Engine/Source/Runtime/Engine/Public/Curves/KeyFrameManipulator.h
@@ -38,8 +38,40 @@ public:
 		}
 	}
 
-	TKeyFrameManipulator(const TKeyFrameManipulator&) = default;
-	TKeyFrameManipulator& operator=(const TKeyFrameManipulator&) = default;
+	TKeyFrameManipulator(const TKeyFrameManipulator& Rhs)
+		: KeyTimes(Rhs.KeyTimes)
+		, KeyHandleLUT(Rhs.KeyHandleLUT == &Rhs.TemporaryKeyHandleLUT ? &TemporaryKeyHandleLUT : Rhs.KeyHandleLUT)
+		, TemporaryKeyHandleLUT(Rhs.TemporaryKeyHandleLUT)
+	{}
+
+	TKeyFrameManipulator& operator=(const TKeyFrameManipulator& Rhs)
+	{
+		if (&Rhs != this)
+		{
+			KeyTimes = Rhs.KeyTimes;
+			KeyHandleLUT = Rhs.KeyHandleLUT == &Rhs.TemporaryKeyHandleLUT ? &TemporaryKeyHandleLUT : Rhs.KeyHandleLUT;
+			TemporaryKeyHandleLUT = Rhs.TemporaryKeyHandleLUT;
+		}
+		return *this;
+	}
+
+	TKeyFrameManipulator(TKeyFrameManipulator&& Rhs)
+		: KeyTimes(Rhs.KeyTimes)
+		, KeyHandleLUT(Rhs.KeyHandleLUT == &Rhs.TemporaryKeyHandleLUT ? &TemporaryKeyHandleLUT : Rhs.KeyHandleLUT)
+		, TemporaryKeyHandleLUT(MoveTemp(Rhs.TemporaryKeyHandleLUT))
+	{}
+
+	TKeyFrameManipulator& operator=(TKeyFrameManipulator&& Rhs)
+	{
+		if (&Rhs != this)
+		{
+			KeyTimes = Rhs.KeyTimes;
+			KeyHandleLUT = Rhs.KeyHandleLUT == &Rhs.TemporaryKeyHandleLUT ? &TemporaryKeyHandleLUT : Rhs.KeyHandleLUT;
+			TemporaryKeyHandleLUT = MoveTemp(Rhs.TemporaryKeyHandleLUT);
+		}
+		return *this;
+	}
+
 	virtual ~TKeyFrameManipulator() {}
 
 private:
diff --git a/Engine/Source/Runtime/LevelSequence/Private/LevelSequence.cpp b/Engine/Source/Runtime/LevelSequence/Private/LevelSequence.cpp
--- a/Engine/Source/Runtime/LevelSequence/Private/LevelSequence.cpp
+++ b/Engine/Source/Runtime/LevelSequence/Private/LevelSequence.cpp
@@ -13,10 +13,22 @@
 
 DEFINE_LOG_CATEGORY_STATIC(LogLevelSequence, Log, All);
 
-static TAutoConsoleVariable<int32> CVarFixedFrameIntervalPlayback(
-	TEXT("LevelSequence.DefaultFixedFrameIntervalPlayback"),
+static TAutoConsoleVariable<int32> CVarDefaultEvaluationType(
+	TEXT("LevelSequence.DefaultEvaluationType"),
 	0,
-	TEXT("When non-zero, all newly created level sequences will default to fixed frame interval playback."),
+	TEXT("0: Playback locked to playback frames\n1: Unlocked playback with sub frame interpolation"),
+	ECVF_Default);
+
+static TAutoConsoleVariable<FString> CVarDefaultFrameResolution(
+	TEXT("LevelSequence.DefaultFrameResolution"),
+	TEXT("24000fps"),
+	TEXT("Specifies default a frame resolution for newly created level sequences. Examples: 30 fps, 120/1 (120 fps), 30000/1001 (29.97), 0.01s (10ms)."),
+	ECVF_Default);
+
+static TAutoConsoleVariable<FString> CVarDefaultDisplayRate(
+	TEXT("LevelSequence.DefaultDisplayRate"),
+	TEXT("30fps"),
+	TEXT("Specifies default a display frame rate for newly created level sequences; also defines frame locked frame rate where sequences are set to be frame locked. Examples: 30 fps, 120/1 (120 fps), 30000/1001 (29.97), 0.01s (10ms)."),
 	ECVF_Default);
 
 ULevelSequence::ULevelSequence(const FObjectInitializer& ObjectInitializer)
@@ -31,9 +43,17 @@ void ULevelSequence::Initialize()
 	// @todo sequencer: gmp: fix me
 	MovieScene = NewObject<UMovieScene>(this, NAME_None, RF_Transactional);
 
-	const bool bForceFixedPlayback = CVarFixedFrameIntervalPlayback.GetValueOnGameThread() != 0;
+	const bool bFrameLocked = CVarDefaultEvaluationType.GetValueOnGameThread() != 0;
+
+	MovieScene->SetEvaluationType( bFrameLocked ? EMovieSceneEvaluationType::FrameLocked : EMovieSceneEvaluationType::WithSubFrames );
+
+	FFrameRate FrameResolution(60000, 1);
+	TryParseString(FrameResolution, *CVarDefaultFrameResolution.GetValueOnGameThread());
+	MovieScene->SetFrameResolutionDirectly(FrameResolution);
 
-	MovieScene->SetForceFixedFrameIntervalPlayback( bForceFixedPlayback );
+	FFrameRate PlaybackFrameRate(30, 1);
+	TryParseString(PlaybackFrameRate, *CVarDefaultDisplayRate.GetValueOnGameThread());
+	MovieScene->SetPlaybackFrameRate(PlaybackFrameRate);
 }
 
 UObject* ULevelSequence::MakeSpawnableTemplateFromInstance(UObject& InSourceObject, FName ObjectName)
diff --git a/Engine/Source/Runtime/LevelSequence/Private/LevelSequenceBurnIn.cpp b/Engine/Source/Runtime/LevelSequence/Private/LevelSequenceBurnIn.cpp
--- a/Engine/Source/Runtime/LevelSequence/Private/LevelSequenceBurnIn.cpp
+++ b/Engine/Source/Runtime/LevelSequence/Private/LevelSequenceBurnIn.cpp
@@ -17,7 +17,7 @@ void ULevelSequenceBurnIn::TakeSnapshotsFrom(ALevelSequenceActor& InActor)
 	}
 }
 
-void ULevelSequenceBurnIn::OnSequenceUpdated(const UMovieSceneSequencePlayer& Player, float CurrentTime, float PreviousTime)
+void ULevelSequenceBurnIn::OnSequenceUpdated(const UMovieSceneSequencePlayer& Player, FFrameTime CurrentTime, FFrameTime PreviousTime)
 {
 	static_cast<const ULevelSequencePlayer&>(Player).TakeFrameSnapshot(FrameInformation);
 }
diff --git a/Engine/Source/Runtime/LevelSequence/Private/LevelSequencePlayer.cpp b/Engine/Source/Runtime/LevelSequence/Private/LevelSequencePlayer.cpp
--- a/Engine/Source/Runtime/LevelSequence/Private/LevelSequencePlayer.cpp
+++ b/Engine/Source/Runtime/LevelSequence/Private/LevelSequencePlayer.cpp
@@ -99,7 +99,10 @@ void ULevelSequencePlayer::OnStopped()
 		{
 			for (UActorComponent* Component : Actor->GetComponents())
 			{
-				Component->PrimaryComponentTick.RemovePrerequisite(LevelSequenceActor, LevelSequenceActor->PrimaryActorTick);
+				if (Component)
+				{
+					Component->PrimaryComponentTick.RemovePrerequisite(LevelSequenceActor, LevelSequenceActor->PrimaryActorTick);
+				}
 			}
 
 			Actor->PrimaryActorTick.RemovePrerequisite(LevelSequenceActor, LevelSequenceActor->PrimaryActorTick);
@@ -140,7 +143,11 @@ void ULevelSequencePlayer::UpdateCameraCut(UObject* CameraObject, UObject* Unloc
 	}
 
 	UCameraComponent* CameraComponent = MovieSceneHelpers::CameraComponentFromRuntimeObject(CameraObject);
-	
+	if (CameraComponent && CameraComponent->GetOwner() != CameraObject)
+	{
+		CameraObject = CameraComponent->GetOwner();
+	}
+
 	CachedCameraComponent = CameraComponent;
 
 	if (CameraObject == ViewTarget)
@@ -217,7 +224,10 @@ void ULevelSequencePlayer::NotifyBindingUpdate(const FGuid& InGuid, FMovieSceneS
 		{
 			for (UActorComponent* Component : Actor->GetComponents())
 			{
-				Component->PrimaryComponentTick.AddPrerequisite(LevelSequenceActor, LevelSequenceActor->PrimaryActorTick);
+				if (Component)
+				{
+					Component->PrimaryComponentTick.AddPrerequisite(LevelSequenceActor, LevelSequenceActor->PrimaryActorTick);
+				}
 			}
 
 			Actor->PrimaryActorTick.AddPrerequisite(LevelSequenceActor, LevelSequenceActor->PrimaryActorTick);
@@ -272,18 +282,19 @@ void ULevelSequencePlayer::TakeFrameSnapshot(FLevelSequencePlayerSnapshot& OutSn
 		return;
 	}
 
-	const float StartTimeWithoutWarmupFrames = SnapshotOffsetTime.IsSet() ? StartTime + SnapshotOffsetTime.GetValue() : StartTime;
-
-	// Use the actual last evaluation time as per the play position, which accounts for fixed time step offsetting
-	const float CurrentTime = PlayPosition.GetLastPlayEvalPostition().Get(StartTimeWithoutWarmupFrames + TimeCursorPosition);
+	// In Play Rate Resolution
+	const FFrameTime StartTimeWithoutWarmupFrames = SnapshotOffsetTime.IsSet() ? StartTime + SnapshotOffsetTime.GetValue() : StartTime;
+	const FFrameTime CurrentPlayTime = PlayPosition.GetCurrentPosition();
+	// In Playback Resolution
+	const FFrameTime CurrentSequenceTime		  = ConvertFrameTime(CurrentPlayTime, PlayPosition.GetInputRate(), PlayPosition.GetOutputRate());
 
 	OutSnapshot.Settings = SnapshotSettings;
 
-	OutSnapshot.MasterTime = CurrentTime;
+	OutSnapshot.MasterTime = FQualifiedFrameTime(CurrentPlayTime, PlayPosition.GetInputRate());
 	OutSnapshot.MasterName = Sequence->GetName();
 
 	OutSnapshot.CurrentShotName = OutSnapshot.MasterName;
-	OutSnapshot.CurrentShotLocalTime = CurrentTime;
+	OutSnapshot.CurrentShotLocalTime = FQualifiedFrameTime(CurrentPlayTime, PlayPosition.GetInputRate());
 	OutSnapshot.CameraComponent = CachedCameraComponent.IsValid() ? CachedCameraComponent.Get() : nullptr;
 	OutSnapshot.ShotID = MovieSceneSequenceID::Invalid;
 
@@ -301,8 +312,8 @@ void ULevelSequencePlayer::TakeFrameSnapshot(FLevelSequencePlayerSnapshot& OutSn
 			// It's unfortunate that we have to copy the logic of UMovieSceneCinematicShotTrack::GetRowCompilerRules() to some degree here, but there's no better way atm
 			bool bThisShotIsActive = Section->IsActive();
 
-			TRange<float> SectionRange = Section->GetRange();
-			bThisShotIsActive = bThisShotIsActive && SectionRange.Contains(CurrentTime);
+			TRange<FFrameNumber> SectionRange = Section->GetRange();
+			bThisShotIsActive = bThisShotIsActive && SectionRange.Contains(CurrentSequenceTime.FrameNumber);
 
 			if (bThisShotIsActive && ActiveShot)
 			{
@@ -313,7 +324,7 @@ void ULevelSequencePlayer::TakeFrameSnapshot(FLevelSequencePlayerSnapshot& OutSn
 				else if (Section->GetRowIndex() == ActiveShot->GetRowIndex())
 				{
 					// On the same row - latest start wins
-					bThisShotIsActive = TRangeBound<float>::MaxLower(SectionRange.GetLowerBound(), ActiveShot->GetRange().GetLowerBound()) == SectionRange.GetLowerBound();
+					bThisShotIsActive = TRangeBound<FFrameNumber>::MaxLower(SectionRange.GetLowerBound(), ActiveShot->GetRange().GetLowerBound()) == SectionRange.GetLowerBound();
 				}
 				else
 				{
@@ -330,14 +341,12 @@ void ULevelSequencePlayer::TakeFrameSnapshot(FLevelSequencePlayerSnapshot& OutSn
 		if (ActiveShot)
 		{
 			// Assume that shots with no sequence start at 0.
-			const float ShotLowerBound = ActiveShot->GetSequence() 
-				? ActiveShot->GetSequence()->GetMovieScene()->GetPlaybackRange().GetLowerBoundValue() 
-				: 0;
-			const float ShotOffset = ActiveShot->Parameters.StartOffset + ShotLowerBound - ActiveShot->GetPreRollTime();
-			const float ShotPosition = ShotOffset + (CurrentTime - (ActiveShot->GetStartTime() - ActiveShot->GetPreRollTime())) / ActiveShot->Parameters.TimeScale;
-
-			OutSnapshot.CurrentShotName = ActiveShot->GetShotDisplayName();
-			OutSnapshot.CurrentShotLocalTime = ShotPosition;
+			FMovieSceneSequenceTransform OuterToInnerTransform = ActiveShot->OuterToInnerTransform();
+			UMovieSceneSequence*         InnerSequence         = ActiveShot->GetSequence();
+			FFrameRate                   InnerFrameRate        = InnerSequence ? InnerSequence->GetMovieScene()->GetPlaybackFrameRate() : PlayPosition.GetInputRate();
+
+			OutSnapshot.CurrentShotName      = ActiveShot->GetShotDisplayName();
+			OutSnapshot.CurrentShotLocalTime = FQualifiedFrameTime(CurrentPlayTime * OuterToInnerTransform, InnerFrameRate);
 			OutSnapshot.ShotID = ActiveShot->GetSequenceID();
 		}
 	}
diff --git a/Engine/Source/Runtime/LevelSequence/Private/LevelSequenceSpawnRegister.cpp b/Engine/Source/Runtime/LevelSequence/Private/LevelSequenceSpawnRegister.cpp
--- a/Engine/Source/Runtime/LevelSequence/Private/LevelSequenceSpawnRegister.cpp
+++ b/Engine/Source/Runtime/LevelSequence/Private/LevelSequenceSpawnRegister.cpp
@@ -57,3 +57,19 @@ void FLevelSequenceSpawnRegister::DestroySpawnedObject(UObject& Object)
 
 	checkf(false, TEXT("No valid object spawner found to destroy spawned object of type %s"), *Object.GetClass()->GetName());
 }
+
+#if WITH_EDITOR
+
+bool FLevelSequenceSpawnRegister::CanSpawnObject(UClass* InClass) const
+{
+	for (TSharedRef<IMovieSceneObjectSpawner> MovieSceneObjectSpawner : MovieSceneObjectSpawners)
+	{
+		if (InClass->IsChildOf(MovieSceneObjectSpawner->GetSupportedTemplateType()))
+		{
+			return true;
+		}
+	}
+	return false;
+}
+
+#endif
\ No newline at end of file
diff --git a/Engine/Source/Runtime/LevelSequence/Public/LevelSequenceBurnIn.h b/Engine/Source/Runtime/LevelSequence/Public/LevelSequenceBurnIn.h
--- a/Engine/Source/Runtime/LevelSequence/Public/LevelSequenceBurnIn.h
+++ b/Engine/Source/Runtime/LevelSequence/Public/LevelSequenceBurnIn.h
@@ -36,7 +36,7 @@ public:
 protected:
 
 	/** Called as part of the game tick loop when the sequence has been updated */
-	void OnSequenceUpdated(const UMovieSceneSequencePlayer& Player, float CurrentTime, float PreviousTime);
+	void OnSequenceUpdated(const UMovieSceneSequencePlayer& Player, FFrameTime CurrentTime, FFrameTime PreviousTime);
 
 protected:
 
diff --git a/Engine/Source/Runtime/LevelSequence/Public/LevelSequencePlayer.h b/Engine/Source/Runtime/LevelSequence/Public/LevelSequencePlayer.h
--- a/Engine/Source/Runtime/LevelSequence/Public/LevelSequencePlayer.h
+++ b/Engine/Source/Runtime/LevelSequence/Public/LevelSequencePlayer.h
@@ -12,6 +12,7 @@
 #include "Evaluation/MovieSceneEvaluationTemplateInstance.h"
 #include "Evaluation/PersistentEvaluationData.h"
 #include "MovieSceneSequencePlayer.h"
+#include "QualifiedFrameTime.h"
 #include "LevelSequence.h"
 #include "LevelSequencePlayer.generated.h"
 
@@ -33,10 +34,10 @@ struct FLevelSequenceSnapshotSettings
 	GENERATED_BODY()
 
 	FLevelSequenceSnapshotSettings()
-		: ZeroPadAmount(4), FrameRate(30)
+		: ZeroPadAmount(4), FrameRate(30, 1)
 	{}
 
-	FLevelSequenceSnapshotSettings(int32 InZeroPadAmount, float InFrameRate)
+	FLevelSequenceSnapshotSettings(int32 InZeroPadAmount, FFrameRate InFrameRate)
 		: ZeroPadAmount(InZeroPadAmount), FrameRate(InFrameRate)
 	{}
 
@@ -46,7 +47,7 @@ struct FLevelSequenceSnapshotSettings
 
 	/** Playback framerate */
 	UPROPERTY(VisibleAnywhere, BlueprintReadOnly, Category="General")
-	float FrameRate;
+	FFrameRate FrameRate;
 };
 
 /**
@@ -61,13 +62,13 @@ struct FLevelSequencePlayerSnapshot
 	FString MasterName;
 
 	UPROPERTY(VisibleAnywhere, BlueprintReadOnly, Category="General")
-	float MasterTime;
+	FQualifiedFrameTime MasterTime;
 
 	UPROPERTY(VisibleAnywhere, BlueprintReadOnly, Category="General")
 	FString CurrentShotName;
 
 	UPROPERTY(VisibleAnywhere, BlueprintReadOnly, Category="General")
-	float CurrentShotLocalTime;
+	FQualifiedFrameTime CurrentShotLocalTime;
 
 	UPROPERTY(VisibleAnywhere, BlueprintReadOnly, Category="General")
 	UCameraComponent* CameraComponent;
@@ -160,8 +161,8 @@ public:
 	/** Take a snapshot of the current state of this player */
 	void TakeFrameSnapshot(FLevelSequencePlayerSnapshot& OutSnapshot) const;
 
-	/** Set the offset time for the snapshot */
-	void SetSnapshotOffsetTime(float InTime) {SnapshotOffsetTime = TOptional<float>(InTime); }
+	/** Set the offset time for the snapshot in play rate frames. */
+	void SetSnapshotOffsetFrames(int32 InFrameOffset) { SnapshotOffsetTime = TOptional<int32>(InFrameOffset); }
 
 private:
 
@@ -183,7 +184,7 @@ protected:
 	/** How to take snapshots */
 	FLevelSequenceSnapshotSettings SnapshotSettings;
 
-	TOptional<float> SnapshotOffsetTime;
+	TOptional<int32> SnapshotOffsetTime;
 
 	TWeakObjectPtr<UCameraComponent> CachedCameraComponent;
 
diff --git a/Engine/Source/Runtime/LevelSequence/Public/LevelSequenceSpawnRegister.h b/Engine/Source/Runtime/LevelSequence/Public/LevelSequenceSpawnRegister.h
--- a/Engine/Source/Runtime/LevelSequence/Public/LevelSequenceSpawnRegister.h
+++ b/Engine/Source/Runtime/LevelSequence/Public/LevelSequenceSpawnRegister.h
@@ -4,6 +4,7 @@
 
 #include "CoreMinimal.h"
 #include "MovieSceneSpawnRegister.h"
+#include "UObject/Class.h"
 
 class IMovieScenePlayer;
 class IMovieSceneObjectSpawner;
@@ -19,6 +20,10 @@ protected:
 	virtual UObject* SpawnObject(FMovieSceneSpawnable& Spawnable, FMovieSceneSequenceIDRef TemplateID, IMovieScenePlayer& Player) override;
 	virtual void DestroySpawnedObject(UObject& Object) override;
 
+#if WITH_EDITOR
+	virtual bool CanSpawnObject(UClass* InClass) const override;
+#endif
+
 protected:
 	/** Extension object spawners */
 	TArray<TSharedRef<IMovieSceneObjectSpawner>> MovieSceneObjectSpawners;
diff --git a/Engine/Source/Runtime/MediaAssets/Private/Assets/MediaSoundComponent.cpp b/Engine/Source/Runtime/MediaAssets/Private/Assets/MediaSoundComponent.cpp
--- a/Engine/Source/Runtime/MediaAssets/Private/Assets/MediaSoundComponent.cpp
+++ b/Engine/Source/Runtime/MediaAssets/Private/Assets/MediaSoundComponent.cpp
@@ -14,6 +14,11 @@
 #include "MediaPlayer.h"
 #include "MediaPlayerFacade.h"
 
+#if PLATFORM_HTML5
+	#include "AudioDevice.h"
+	#include "Engine/Engine.h"
+#endif
+
 
 /* Static initialization
  *****************************************************************************/
@@ -69,6 +74,12 @@ bool UMediaSoundComponent::BP_GetAttenuationSettingsToApply(FSoundAttenuationSet
 }
 
 
+UMediaPlayer* UMediaSoundComponent::GetMediaPlayer() const
+{
+	return CurrentPlayer.Get();
+}
+
+
 void UMediaSoundComponent::SetMediaPlayer(UMediaPlayer* NewMediaPlayer)
 {
 	CurrentPlayer = NewMediaPlayer;
@@ -192,6 +203,7 @@ void UMediaSoundComponent::PostInitProperties()
 	if (UMediaSoundComponent::DefaultMediaSoundClassObject == nullptr)
 	{
 		const FSoftObjectPath DefaultMediaSoundClassName = GetDefault<UAudioSettings>()->DefaultMediaSoundClassName;
+
 		if (DefaultMediaSoundClassName.IsValid())
 		{
 			UMediaSoundComponent::DefaultMediaSoundClassObject = LoadObject<USoundClass>(nullptr, *DefaultMediaSoundClassName.ToString());
