Disabling more benign PVS-Studio warnings.

diff --git a/Engine/Plugins/Online/OnlineSubsystemSteam/Source/Private/OnlineSessionAsyncServerSteam.cpp b/Engine/Plugins/Online/OnlineSubsystemSteam/Source/Private/OnlineSessionAsyncServerSteam.cpp
--- a/Engine/Plugins/Online/OnlineSubsystemSteam/Source/Private/OnlineSessionAsyncServerSteam.cpp
+++ b/Engine/Plugins/Online/OnlineSubsystemSteam/Source/Private/OnlineSessionAsyncServerSteam.cpp
@@ -658,7 +658,7 @@ bool FPendingSearchResultSteam::FillSessionFromServerRules()
 	{
 		SessionInfo->HostAddr = HostAddr;
 
-		if (SteamAddrKeysFound == 2)
+		if (SteamAddrKeysFound == 2) //-V547
 		{
 			SessionInfo->SteamP2PAddr = SteamP2PAddr;
 		}
diff --git a/Engine/Source/Developer/BlueprintNativeCodeGen/Private/BlueprintNativeCodeGenModule.cpp b/Engine/Source/Developer/BlueprintNativeCodeGen/Private/BlueprintNativeCodeGenModule.cpp
--- a/Engine/Source/Developer/BlueprintNativeCodeGen/Private/BlueprintNativeCodeGenModule.cpp
+++ b/Engine/Source/Developer/BlueprintNativeCodeGen/Private/BlueprintNativeCodeGenModule.cpp
@@ -500,7 +500,7 @@ void FBlueprintNativeCodeGenModule::GenerateSingleAsset(UField* ForConversion, c
 	{
 		if (!FFileHelper::SaveStringToFile(*CppSource, *ConversionRecord.GeneratedCppPath, ForcedEncoding()))
 		{
-			bSuccess &= false;
+			bSuccess = false;
 			ConversionRecord.GeneratedCppPath.Empty();
 		}
 		CppSource->Empty(CppSource->Len());
@@ -514,7 +514,7 @@ void FBlueprintNativeCodeGenModule::GenerateSingleAsset(UField* ForConversion, c
 	{
 		if (!FFileHelper::SaveStringToFile(*HeaderSource, *ConversionRecord.GeneratedHeaderPath, ForcedEncoding()))
 		{
-			bSuccess &= false;
+			bSuccess = false;
 			ConversionRecord.GeneratedHeaderPath.Empty();
 		}
 		HeaderSource->Empty(HeaderSource->Len());
diff --git a/Engine/Source/Developer/SimplygonMeshReduction/Private/SimplygonMeshReduction.cpp b/Engine/Source/Developer/SimplygonMeshReduction/Private/SimplygonMeshReduction.cpp
--- a/Engine/Source/Developer/SimplygonMeshReduction/Private/SimplygonMeshReduction.cpp
+++ b/Engine/Source/Developer/SimplygonMeshReduction/Private/SimplygonMeshReduction.cpp
@@ -2668,7 +2668,7 @@ private:
 
 		InMappingImageSettings->SetGenerateMappingImage(true);
 		InMappingImageSettings->SetTexCoordLevel(0);
-		if (NumOutputMaterials > 1)
+		if (NumOutputMaterials > 1) //-V547
 		{
 			InMappingImageSettings->SetGenerateTexCoords(true);	//! check if it is possible to avoid this
 		}
diff --git a/Engine/Source/Editor/AdvancedPreviewScene/Private/AdvancedPreviewScene.cpp b/Engine/Source/Editor/AdvancedPreviewScene/Private/AdvancedPreviewScene.cpp
--- a/Engine/Source/Editor/AdvancedPreviewScene/Private/AdvancedPreviewScene.cpp
+++ b/Engine/Source/Editor/AdvancedPreviewScene/Private/AdvancedPreviewScene.cpp
@@ -335,7 +335,7 @@ const bool FAdvancedPreviewScene::HandleInputKey(FViewport* InViewport, int32 Co
 	if (Event == IE_Pressed)
 	{
 		FModifierKeysState KeyState = FSlateApplication::Get().GetModifierKeys();
-		if (UICommandList->ProcessCommandBindings(Key, KeyState, (Event == IE_Repeat)))
+		if (UICommandList->ProcessCommandBindings(Key, KeyState, (Event == IE_Repeat))) //-V547
 		{
 			return true;
 		}
diff --git a/Engine/Source/Editor/Cascade/Private/Cascade.cpp b/Engine/Source/Editor/Cascade/Private/Cascade.cpp
--- a/Engine/Source/Editor/Cascade/Private/Cascade.cpp
+++ b/Engine/Source/Editor/Cascade/Private/Cascade.cpp
@@ -806,7 +806,7 @@ void FCascade::SetSelectedModule(UParticleEmitter* NewSelectedEmitter, UParticle
 		{
 			if (LODLevel != NULL)
 			{
-				if (bReadOnly == false)
+				if (bReadOnly == false) //-V547
 				{
 					if (LODLevel->Level != CurrLODIndex)
 					{
@@ -822,7 +822,7 @@ void FCascade::SetSelectedModule(UParticleEmitter* NewSelectedEmitter, UParticle
 		}
 		else
 		{
-			if (bReadOnly == false)
+			if (bReadOnly == false) //-V547
 			{
 				// Only allowing editing the SelectedEmitter 
 				// properties when at the highest LOD level.
diff --git a/Engine/Source/Editor/Matinee/Private/MatineeTools.cpp b/Engine/Source/Editor/Matinee/Private/MatineeTools.cpp
--- a/Engine/Source/Editor/Matinee/Private/MatineeTools.cpp
+++ b/Engine/Source/Editor/Matinee/Private/MatineeTools.cpp
@@ -3240,7 +3240,7 @@ void FMatinee::UpdatePreviewCamera( UInterpTrack* AssociatedTrack ) const
 	{
 		// If the track selection state has changed, update our camera actor
 		UInterpGroupDirector* DirGroup = CastChecked<UInterpGroupDirector>( DirTrack->GetOuter() );
-		const bool TrackOrGroupSelected = ( DirTrack->IsSelected() | DirGroup->IsSelected() );
+		const bool TrackOrGroupSelected = ( DirTrack->IsSelected() | DirGroup->IsSelected() ); //-V792
 		DirTrack->UpdatePreviewCamera( MatineeActor, TrackOrGroupSelected );
 	}
 }
diff --git a/Engine/Source/Editor/PropertyEditor/Private/UserInterface/PropertyEditor/SPropertyEditorBool.cpp b/Engine/Source/Editor/PropertyEditor/Private/UserInterface/PropertyEditor/SPropertyEditorBool.cpp
--- a/Engine/Source/Editor/PropertyEditor/Private/UserInterface/PropertyEditor/SPropertyEditorBool.cpp
+++ b/Engine/Source/Editor/PropertyEditor/Private/UserInterface/PropertyEditor/SPropertyEditorBool.cpp
@@ -62,7 +62,7 @@ ECheckBoxState SPropertyEditorBool::OnGetCheckState() const
 		{
 			ReturnState = ECheckBoxState::Checked;
 		}
-		else if( Value == false )
+		else
 		{
 			ReturnState = ECheckBoxState::Unchecked;
 		}
diff --git a/Engine/Source/Editor/SequenceRecorder/Private/SequenceRecorderModule.cpp b/Engine/Source/Editor/SequenceRecorder/Private/SequenceRecorderModule.cpp
--- a/Engine/Source/Editor/SequenceRecorder/Private/SequenceRecorderModule.cpp
+++ b/Engine/Source/Editor/SequenceRecorder/Private/SequenceRecorderModule.cpp
@@ -460,7 +460,7 @@ class FSequenceRecorderModule : public ISequenceRecorder, private FSelfRegisteri
 					}
 					return true;
 				}
-				else if(FilterType == EFilterType::Class)
+				else
 				{
 					UClass* FoundClass = FindObject<UClass>(ANY_PACKAGE, *SpecifierStr);
 					if(FoundClass != nullptr)
diff --git a/Engine/Source/Editor/UnrealEd/Private/Commandlets/DiffPackagesCommandlet.cpp b/Engine/Source/Editor/UnrealEd/Private/Commandlets/DiffPackagesCommandlet.cpp
--- a/Engine/Source/Editor/UnrealEd/Private/Commandlets/DiffPackagesCommandlet.cpp
+++ b/Engine/Source/Editor/UnrealEd/Private/Commandlets/DiffPackagesCommandlet.cpp
@@ -1446,7 +1446,7 @@ EObjectDiff UDiffPackagesCommandlet::CompareNativePropertyValues( UObject* ObjA,
 // 				AppendComparisonResultText(PropDiff.DiffText, FString::Printf(TEXT("(%s) %s::%s"), GetDiffTypeText(PropDiff.DiffType,NumPackages), *FullPath, *PropName));
 // 				AppendComparisonResultText(PropDiff.DiffText, FString::Printf(TEXT("     Was: %s"), *PropTextAncestor));
 // 				AppendComparisonResultText(PropDiff.DiffText, FString::Printf(TEXT("     Now: %s"), *PropTextA));
-				if ( NativeDataComparison.DiffType == OD_None )
+				if ( NativeDataComparison.DiffType == OD_None ) //-V547
 				{
 					NativeDataComparison.DiffType = OD_ABSame;
 					AppendComparisonResultText(NativeDataComparison.DiffText, FString::Printf(TEXT("(%s) %s"), GetDiffTypeText(NativeDataComparison.DiffType, NumPackages), *ObjectPathName));
@@ -1472,7 +1472,7 @@ EObjectDiff UDiffPackagesCommandlet::CompareNativePropertyValues( UObject* ObjA,
 // 					AppendComparisonResultText(PropDiff.DiffText, FString::Printf(TEXT("     Was: %s"), *PropTextAncestor));
 // 					AppendComparisonResultText(PropDiff.DiffText, FString::Printf(TEXT("     Now: %s"), *PropTextB));
 
-					if ( NativeDataComparison.DiffType == OD_None )
+					if ( NativeDataComparison.DiffType == OD_None ) //-V547
 					{
 						NativeDataComparison.DiffType = OD_BOnly;
 					}
@@ -1491,7 +1491,7 @@ EObjectDiff UDiffPackagesCommandlet::CompareNativePropertyValues( UObject* ObjA,
 // 					AppendComparisonResultText(PropDiff.DiffText, FString::Printf(TEXT("     Was: %s"), *PropTextAncestor));
 // 					AppendComparisonResultText(PropDiff.DiffText, FString::Printf(TEXT("     Now: %s"), *PropTextA));
 
-					if (NativeDataComparison.DiffType == OD_None )
+					if (NativeDataComparison.DiffType == OD_None ) //-V547
 					{
 						NativeDataComparison.DiffType = OD_AOnly;
 						AppendComparisonResultText(NativeDataComparison.DiffText, FString::Printf(TEXT("(%s) %s"),
diff --git a/Engine/Source/Editor/UnrealEd/Private/CookOnTheFlyServer.cpp b/Engine/Source/Editor/UnrealEd/Private/CookOnTheFlyServer.cpp
--- a/Engine/Source/Editor/UnrealEd/Private/CookOnTheFlyServer.cpp
+++ b/Engine/Source/Editor/UnrealEd/Private/CookOnTheFlyServer.cpp
@@ -3188,7 +3188,7 @@ void UCookOnTheFlyServer::SaveCookedPackage(UPackage* Package, uint32 SaveFlags,
 			if (bCookPackage == true)
 			{
 				bool bPackageFullyLoaded = false;
-				if (bPackageFullyLoaded == false)
+				if (bPackageFullyLoaded == false) //-V547
 				{
 					SCOPE_TIMER(LoadPackage);
 
@@ -5244,7 +5244,7 @@ void UCookOnTheFlyServer::SaveGlobalShaderMapFiles(const TArray<ITargetPlatform*
 
 		RecompileShadersForRemote
 			(RecompileData.PlatformName, 
-			RecompileData.ShaderPlatform == -1 ? SP_NumPlatforms : (EShaderPlatform)RecompileData.ShaderPlatform,
+			RecompileData.ShaderPlatform == -1 ? SP_NumPlatforms : (EShaderPlatform)RecompileData.ShaderPlatform, //-V547
 			OutputDir, 
 			RecompileData.MaterialsToLoad, 
 			RecompileData.SerializedShaderResources, 
diff --git a/Engine/Source/Editor/UnrealEd/Private/Factories/Factory.cpp b/Engine/Source/Editor/UnrealEd/Private/Factories/Factory.cpp
--- a/Engine/Source/Editor/UnrealEd/Private/Factories/Factory.cpp
+++ b/Engine/Source/Editor/UnrealEd/Private/Factories/Factory.cpp
@@ -515,7 +515,7 @@ bool UFactory::ImportUntypedBulkDataFromText(const TCHAR*& Buffer, FUntypedBulkD
 								ParseStr +=2;
 							}
 							Value = FParse::HexDigit(ParseStr[0]) * 16 + FParse::HexDigit(ParseStr[1]);
-							*BulkDataPointer = (uint8)Value;
+							*BulkDataPointer = (uint8)Value; //-V522
 							BulkDataPointer++;
 							ParseStr += 2;
 							ParseStr++;
diff --git a/Engine/Source/Editor/UnrealEd/Private/Fbx/FbxMainExport.cpp b/Engine/Source/Editor/UnrealEd/Private/Fbx/FbxMainExport.cpp
--- a/Engine/Source/Editor/UnrealEd/Private/Fbx/FbxMainExport.cpp
+++ b/Engine/Source/Editor/UnrealEd/Private/Fbx/FbxMainExport.cpp
@@ -2041,7 +2041,7 @@ void FFbxExporter::ExportAnimatedVector(FbxAnimCurve* FbxCurve, const char* Chan
 
 			// Setup tangents for bezier curves. Avoid this for keys created from baking 
 			// transforms since there is no tangent info created for these types of keys. 
-			if( Interpolation == FbxAnimCurveDef::eInterpolationCubic )
+			if( Interpolation == FbxAnimCurveDef::eInterpolationCubic ) //-V547
 			{
 				float OutTangentValue = (CurveIndex == 0) ? Key.LeaveTangent.X : (CurveIndex == 1) ? Key.LeaveTangent.Y : Key.LeaveTangent.Z;
 				float OutTangentX = (KeyIndex < KeyCount - 1) ? (KeyTimes[KeyIndex + 1] - KeyTime) / 3.0f : 0.333f;
diff --git a/Engine/Source/Editor/UnrealEd/Private/FileHelpers.cpp b/Engine/Source/Editor/UnrealEd/Private/FileHelpers.cpp
--- a/Engine/Source/Editor/UnrealEd/Private/FileHelpers.cpp
+++ b/Engine/Source/Editor/UnrealEd/Private/FileHelpers.cpp
@@ -84,20 +84,16 @@ public:
 	{
 		if ( Verbosity == ELogVerbosity::Error || Verbosity == ELogVerbosity::Warning )
 		{
-			EMessageSeverity::Type Severity = EMessageSeverity::Info;
+			EMessageSeverity::Type Severity;
 			if ( Verbosity == ELogVerbosity::Error )
 			{
 				Severity = EMessageSeverity::Error;
 			}
-			else if ( Verbosity == ELogVerbosity::Warning )
+			else
 			{
 				Severity = EMessageSeverity::Warning;
 			}
-			
-			if ( ensure(Severity != EMessageSeverity::Info) )
-			{
-				ErrorMessages.Add(FTokenizedMessage::Create(Severity, FText::FromName(InData)));
-			}
+			ErrorMessages.Add(FTokenizedMessage::Create(Severity, FText::FromName(InData)));
 		}
 	}
 
diff --git a/Engine/Source/Editor/UnrealEd/Private/LODUtilities.cpp b/Engine/Source/Editor/UnrealEd/Private/LODUtilities.cpp
--- a/Engine/Source/Editor/UnrealEd/Private/LODUtilities.cpp
+++ b/Engine/Source/Editor/UnrealEd/Private/LODUtilities.cpp
@@ -303,7 +303,7 @@ void ProjectTargetOnBase(const TArray<FSoftSkinVertex>& BaseVertices, const TArr
 					//We should always have a match
 					if (!bEnsureMsgDone)
 					{
-						ensureMsgf(FoundIndexMatch != INDEX_NONE, TEXT("Reduce LOD vertex should always have a valid matched triangle in the BaseLOD."));
+						ensureMsgf(FoundIndexMatch != INDEX_NONE, TEXT("Reduce LOD vertex should always have a valid matched triangle in the BaseLOD.")); //-V547
 						bEnsureMsgDone = true;
 					}
 					continue;
diff --git a/Engine/Source/Runtime/Engine/Private/Components/CharacterMovementComponent.cpp b/Engine/Source/Runtime/Engine/Private/Components/CharacterMovementComponent.cpp
--- a/Engine/Source/Runtime/Engine/Private/Components/CharacterMovementComponent.cpp
+++ b/Engine/Source/Runtime/Engine/Private/Components/CharacterMovementComponent.cpp
@@ -1380,7 +1380,7 @@ void UCharacterMovementComponent::SimulatedTick(float DeltaSeconds)
 				DrawDebugCoordinateSystem(GetWorld(), CharacterOwner->GetMesh()->GetComponentLocation() + FVector(0,0,1), NewRotation, 50.f, false);
 				DrawDebugLine(GetWorld(), OldLocation, NewLocation, FColor::Red, true, 10.f);
 
-				UE_LOG(LogRootMotion, Log,  TEXT("UCharacterMovementComponent::SimulatedTick DeltaMovement Translation: %s, Rotation: %s, MovementBase: %s"),
+				UE_LOG(LogRootMotion, Log,  TEXT("UCharacterMovementComponent::SimulatedTick DeltaMovement Translation: %s, Rotation: %s, MovementBase: %s"), //-V595
 					*(NewLocation - OldLocation).ToCompactString(), *(NewRotation - OldRotation).GetNormalized().ToCompactString(), *GetNameSafe(CharacterOwner->GetMovementBase()) );
 			}
 #endif // !(UE_BUILD_SHIPPING)
@@ -2241,7 +2241,7 @@ void UCharacterMovementComponent::PerformMovement(float DeltaSeconds)
 				DrawDebugLine(GetWorld(), OldLocation, ResultingLocation, FColor::Red, true, 10.f);
 
 				// Log details.
-				UE_LOG(LogRootMotion, Warning,  TEXT("PerformMovement Resulting DeltaMove Translation: %s, Rotation: %s, MovementBase: %s"),
+				UE_LOG(LogRootMotion, Warning,  TEXT("PerformMovement Resulting DeltaMove Translation: %s, Rotation: %s, MovementBase: %s"), //-V595
 					*(ResultingLocation - OldLocation).ToCompactString(), *(ResultingRotation - OldActorRotation).GetNormalized().ToCompactString(), *GetNameSafe(CharacterOwner->GetMovementBase()) );
 
 				const FVector RMTranslation = RootMotionParams.GetRootMotionTransform().GetTranslation();
diff --git a/Engine/Source/Runtime/Engine/Private/Interpolation.cpp b/Engine/Source/Runtime/Engine/Private/Interpolation.cpp
--- a/Engine/Source/Runtime/Engine/Private/Interpolation.cpp
+++ b/Engine/Source/Runtime/Engine/Private/Interpolation.cpp
@@ -6689,7 +6689,7 @@ void UInterpTrackDirector::PreviewUpdateTrack(float NewPosition, UInterpTrackIns
 	AMatineeActor* MatineeActor = CastChecked<AMatineeActor>( GrInst->GetOuter() );
 
 	// Update the camera actor as the track is scrubbed
-	const bool TrackOrGroupSelected = ( IsSelected() | DirGroup->IsSelected() );
+	const bool TrackOrGroupSelected = ( IsSelected() || DirGroup->IsSelected() ); 
 	if ( UpdatePreviewCamera( MatineeActor, TrackOrGroupSelected ) )
 	{
 		// Refresh the selected group actor (deselect previous actors, otherwise we'll have multiple cameras selected)
diff --git a/Engine/Source/Runtime/Engine/Private/Particles/ParticleSystemRender.cpp b/Engine/Source/Runtime/Engine/Private/Particles/ParticleSystemRender.cpp
--- a/Engine/Source/Runtime/Engine/Private/Particles/ParticleSystemRender.cpp
+++ b/Engine/Source/Runtime/Engine/Private/Particles/ParticleSystemRender.cpp
@@ -5974,7 +5974,7 @@ int32 FDynamicRibbonEmitterData::FillVertexData(struct FAsyncBufferFillData& Dat
 				FVector CurrUp = WorkingUp;
 				FLinearColor CurrColor = PackingParticle->Color;
 
-				FVector PrevPosition = PrevParticle->Location;
+				FVector PrevPosition = PrevParticle->Location; //-V522
 				FVector PrevTangent = PrevTrailPayload->Tangent;
 				FVector PrevUp = PrevWorkingUp;
 				FLinearColor PrevColor = PrevParticle->Color;
diff --git a/Engine/Source/Runtime/Engine/Private/Particles/ParticleTrail2EmitterInstance.cpp b/Engine/Source/Runtime/Engine/Private/Particles/ParticleTrail2EmitterInstance.cpp
--- a/Engine/Source/Runtime/Engine/Private/Particles/ParticleTrail2EmitterInstance.cpp
+++ b/Engine/Source/Runtime/Engine/Private/Particles/ParticleTrail2EmitterInstance.cpp
@@ -3360,7 +3360,7 @@ float FParticleAnimTrailEmitterInstance::Spawn(float DeltaTime)
 	//@todo. Support multiple trails per emitter
 	// Find the start particle of the current trail...
 	int32 StartIndex = -1;
-	if (TrailIdx != INDEX_NONE)
+	if (TrailIdx != INDEX_NONE) //-V547
 	{
 		FBaseParticle* Particle = nullptr;
 		FAnimTrailTypeDataPayload* TrailData = nullptr;
diff --git a/Engine/Source/Runtime/Engine/Private/PhysicsEngine/ConstraintInstance.cpp b/Engine/Source/Runtime/Engine/Private/PhysicsEngine/ConstraintInstance.cpp
--- a/Engine/Source/Runtime/Engine/Private/PhysicsEngine/ConstraintInstance.cpp
+++ b/Engine/Source/Runtime/Engine/Private/PhysicsEngine/ConstraintInstance.cpp
@@ -348,7 +348,7 @@ float ComputeAverageMass_AssumesLocked(const PxRigidActor* PActor1, const PxRigi
 		}
 
 		check(NumDynamic);
-		AverageMass = TotalMass / NumDynamic;
+		AverageMass = TotalMass / NumDynamic; //-V609
 	}
 
 	return AverageMass;
diff --git a/Engine/Source/Runtime/Engine/Private/UnrealEngine.cpp b/Engine/Source/Runtime/Engine/Private/UnrealEngine.cpp
--- a/Engine/Source/Runtime/Engine/Private/UnrealEngine.cpp
+++ b/Engine/Source/Runtime/Engine/Private/UnrealEngine.cpp
@@ -14415,7 +14415,7 @@ static void PakFileTest(const TArray<FString>& Args)
 							}
 							else
 							{
-								while (!CallbackCalled)
+								while (!CallbackCalled) //-V654
 								{
 									FPlatformProcess::SleepNoStats(0);
 								}
diff --git a/Engine/Source/Runtime/GameMenuBuilder/Private/SGameMenuPageWidget.cpp b/Engine/Source/Runtime/GameMenuBuilder/Private/SGameMenuPageWidget.cpp
--- a/Engine/Source/Runtime/GameMenuBuilder/Private/SGameMenuPageWidget.cpp
+++ b/Engine/Source/Runtime/GameMenuBuilder/Private/SGameMenuPageWidget.cpp
@@ -112,7 +112,7 @@ void FMenuPanel::ForcePanelOpen()
 	if (CurrentState != EPanelState::Open)
 	{
 		CurrentState = EPanelState::Open;
-		OnStateChanged.ExecuteIfBound(CurrentState == EPanelState::Open);
+		OnStateChanged.ExecuteIfBound(CurrentState == EPanelState::Open); //-V547
 	}
 }
 
@@ -122,7 +122,7 @@ void FMenuPanel::ForcePanelClosed()
 	if (CurrentState != EPanelState::Closed)
 	{
 		CurrentState = EPanelState::Closed;
-		OnStateChanged.ExecuteIfBound(CurrentState == EPanelState::Open);
+		OnStateChanged.ExecuteIfBound(CurrentState == EPanelState::Open); //-V547
 	}
 }
 
diff --git a/Engine/Source/Runtime/Navmesh/Private/DetourTileCache/DetourTileCacheBuilder.cpp b/Engine/Source/Runtime/Navmesh/Private/DetourTileCache/DetourTileCacheBuilder.cpp
--- a/Engine/Source/Runtime/Navmesh/Private/DetourTileCache/DetourTileCacheBuilder.cpp
+++ b/Engine/Source/Runtime/Navmesh/Private/DetourTileCache/DetourTileCacheBuilder.cpp
@@ -1848,7 +1848,7 @@ static dtStatus removeVertex(dtTileCacheLogContext* ctx, dtTileCachePolyMesh& me
 	
 	// Merge polygons.
 	int maxVertsPerPoly = MAX_VERTS_PER_POLY;
-	if (maxVertsPerPoly > 3)
+	if (maxVertsPerPoly > 3) //-V547
 	{
 		for (;;)
 		{
@@ -2046,7 +2046,7 @@ dtStatus dtBuildTileCachePolyMesh(dtTileCacheAlloc* alloc,
 		
 		// Merge polygons.
 		int maxVertsPerPoly =MAX_VERTS_PER_POLY ;
-		if (maxVertsPerPoly > 3)
+		if (maxVertsPerPoly > 3) //-V547
 		{
 			for(;;)
 			{
diff --git a/Engine/Source/Runtime/NetworkReplayStreaming/LocalFileNetworkReplayStreaming/Private/LocalFileNetworkReplayStreaming.cpp b/Engine/Source/Runtime/NetworkReplayStreaming/LocalFileNetworkReplayStreaming/Private/LocalFileNetworkReplayStreaming.cpp
--- a/Engine/Source/Runtime/NetworkReplayStreaming/LocalFileNetworkReplayStreaming/Private/LocalFileNetworkReplayStreaming.cpp
+++ b/Engine/Source/Runtime/NetworkReplayStreaming/LocalFileNetworkReplayStreaming/Private/LocalFileNetworkReplayStreaming.cpp
@@ -573,7 +573,7 @@ void FLocalFileNetworkReplayStreamer::AddOrUpdateEvent(const FString& Name, cons
 	EventName = CurrentStreamName + TEXT("_") + EventName;
 
 	// see if this event already exists
-	if (EventIndex == INDEX_NONE)
+	if (EventIndex == INDEX_NONE) //-V547
 	{
 		for (int32 i=0; i < ReplayInfo.Events.Num(); ++i)
 		{
diff --git a/Engine/Source/Runtime/Online/XMPP/Private/XmppJingle/XmppMultiUserChatJingle.cpp b/Engine/Source/Runtime/Online/XMPP/Private/XmppJingle/XmppMultiUserChatJingle.cpp
--- a/Engine/Source/Runtime/Online/XMPP/Private/XmppJingle/XmppMultiUserChatJingle.cpp
+++ b/Engine/Source/Runtime/Online/XMPP/Private/XmppJingle/XmppMultiUserChatJingle.cpp
@@ -1862,7 +1862,7 @@ void FXmppMultiUserChatJingle::ChatroomEnteredStatus(
 				UE_LOG(LogXmpp, Log, TEXT("ChatroomEnteredStatus - queueing ChatRoomJoinPublicOpResult for %s"), *RoomId);
 				ResultOpQueue.Enqueue(new FXmppChatRoomJoinPublicOpResult(RoomId, bWasSuccessful, ErrorStr));
 			}
-			else if (LastStatus == FXmppRoomJingle::JoinPrivatePending)
+			else if (LastStatus == FXmppRoomJingle::JoinPrivatePending) //-V547
 			{
 				UE_LOG(LogXmpp, Log, TEXT("ChatroomEnteredStatus - queueing ChatRoomJoinPrivateOpResult result for %s"), *RoomId);
 				ResultOpQueue.Enqueue(new FXmppChatRoomJoinPrivateOpResult(RoomId, bWasSuccessful, ErrorStr));
