Copying //UE4/Dev-Framework to //UE4/Dev-Main (Source: //UE4/Dev-Framework @ 3967517)

diff --git a/Engine/Source/Editor/PropertyEditor/Private/Presentation/PropertyEditor/PropertyEditor.h b/Engine/Source/Editor/PropertyEditor/Private/Presentation/PropertyEditor/PropertyEditor.h
--- a/Engine/Source/Editor/PropertyEditor/Private/Presentation/PropertyEditor/PropertyEditor.h
+++ b/Engine/Source/Editor/PropertyEditor/Private/Presentation/PropertyEditor/PropertyEditor.h
@@ -6,8 +6,6 @@
 #include "PropertyHandle.h"
 #include "Editor/SceneOutliner/Public/SceneOutlinerFwd.h"
 
-class FResetToDefaultOverride;
-
 class FPropertyEditor : public TSharedFromThis< FPropertyEditor >	
 {
 public:
@@ -65,17 +63,9 @@ public:
 	/**	@return The display name to be used for the property */
 	FText GetDisplayName() const;
 
-	/** @return Whether or not resetting this property to its default value is a valid and worthwhile operation */
-	bool IsResetToDefaultAvailable() const;
-
 	/**	@return Whether the property passes the current filter restrictions. If no there are no filter restrictions false will be returned. */
 	bool DoesPassFilterRestrictions() const;
 
-	/**	@return Whether the property's current value differs from the default value. */
-	bool ValueDiffersFromDefault() const;
-
-	FText GetResetToDefaultLabel() const;
-
 	void AddPropertyEditorChild( const TSharedRef<FPropertyEditor>& Child );
 	void RemovePropertyEditorChild( const TSharedRef<FPropertyEditor>& Child );
 	const TArray< TSharedRef< FPropertyEditor > >& GetPropertyEditorChildren() const;
@@ -89,9 +79,6 @@ public:
 	void MakeNewBlueprint();
 	void BrowseTo();
 	void EmptyArray();
-	void Find();
-	void ResetToDefault();
-	void CustomResetToDefault(const FResetToDefaultOverride& OnCustomResetToDefault);
 	void OnGetClassesForAssetPicker( TArray<const UClass*>& OutClasses );
 	void OnAssetSelected( const FAssetData& AssetData );
 	void OnActorSelected( AActor* InActor );
@@ -128,8 +115,6 @@ private:
 	void OnDuplicateItem();
 	void OnBrowseTo();
 	void OnEmptyArray();
-	void OnResetToDefault();
-	void OnCustomResetToDefault(const FResetToDefaultOverride& OnCustomResetToDefault);
 
 	/**
 	 * Returns true if the value of the conditional property matches the value required.  Indicates whether editing or otherwise interacting with this item's
diff --git a/Engine/Source/Editor/PropertyEditor/Private/PropertyHandleImpl.cpp b/Engine/Source/Editor/PropertyEditor/Private/PropertyHandleImpl.cpp
--- a/Engine/Source/Editor/PropertyEditor/Private/PropertyHandleImpl.cpp
+++ b/Engine/Source/Editor/PropertyEditor/Private/PropertyHandleImpl.cpp
@@ -33,11 +33,6 @@ FPropertyValueImpl::FPropertyValueImpl( TSharedPtr<FPropertyNode> InPropertyNode
 	, bInteractiveChangeInProgress( false )
 	, InvalidOperationError( nullptr )
 {
-
-}
-
-FPropertyValueImpl::~FPropertyValueImpl()
-{
 }
 
 void FPropertyValueImpl::EnumerateObjectsToModify( FPropertyNode* InPropertyNode, const EnumerateObjectsToModifyFuncRef& InObjectsToModifyCallback ) const
@@ -1158,9 +1153,15 @@ void FPropertyValueImpl::ResetToDefault()
 	TSharedPtr<FPropertyNode> PropertyNodePin = PropertyNode.Pin();
 	if( PropertyNodePin.IsValid() && !PropertyNodePin->IsEditConst() && PropertyNodePin->GetDiffersFromDefault() )
 	{
-		PropertyNodePin->ResetToDefault( NotifyHook );
+		FScopedTransaction Transaction( NSLOCTEXT("UnrealEd", "PropertyWindowResetToDefault", "Reset to Default") );
+		TSharedPtr<FPropertyNode>& KeyNode = PropertyNodePin->GetPropertyKeyNode();
+		if(KeyNode.IsValid())
+		{
+			FPropertyValueImpl(KeyNode, NotifyHook, PropertyUtilities.Pin())
+				.ImportText(KeyNode->GetDefaultValueAsString(), EPropertyValueSetFlags::DefaultFlags);
+		}
+		ImportText(PropertyNodePin->GetDefaultValueAsString(), EPropertyValueSetFlags::DefaultFlags);
 	}
-
 }
 
 bool FPropertyValueImpl::DiffersFromDefault() const
diff --git a/Engine/Source/Editor/PropertyEditor/Private/PropertyHandleImpl.h b/Engine/Source/Editor/PropertyEditor/Private/PropertyHandleImpl.h
--- a/Engine/Source/Editor/PropertyEditor/Private/PropertyHandleImpl.h
+++ b/Engine/Source/Editor/PropertyEditor/Private/PropertyHandleImpl.h
@@ -45,8 +45,6 @@ public:
 	 */
 	FPropertyValueImpl( TSharedPtr<FPropertyNode> InPropertyNode, FNotifyHook* InNotifyHook, TSharedPtr<IPropertyUtilities> InPropertyUtilities );
 
-	virtual ~FPropertyValueImpl();
-
 	/**
 	 * Sets an object property to point to the new object
 	 * 
@@ -58,7 +56,7 @@ public:
 	 * Sets the value of an object property to the selected object in the content browser
 	 * @return Whether or not the call was successful
 	 */
-	virtual FPropertyAccess::Result OnUseSelected();
+	FPropertyAccess::Result OnUseSelected();
 
 	/**
 	 * Recurse up to the next object node, adding all array indices into a map according to their property name
diff --git a/Engine/Source/Editor/PropertyEditor/Private/PropertyNode.cpp b/Engine/Source/Editor/PropertyEditor/Private/PropertyNode.cpp
--- a/Engine/Source/Editor/PropertyEditor/Private/PropertyNode.cpp
+++ b/Engine/Source/Editor/PropertyEditor/Private/PropertyNode.cpp
@@ -1678,16 +1678,6 @@ FString FPropertyNode::GetDefaultValueAsStringForObject( FPropertyItemValueDataT
 			USetProperty* OuterSetProperty = Cast<USetProperty>(InProperty->GetOuter());
 			UMapProperty* OuterMapProperty = Cast<UMapProperty>(InProperty->GetOuter());
 
-			if ( OuterArrayProperty != NULL )
-			{
-				// make sure we're not trying to compare against an element that doesn't exist
-				if ( ValueTracker.GetPropertyDefaultBaseAddress() != NULL && GetArrayIndex() >= FScriptArrayHelper::Num(ValueTracker.GetPropertyDefaultBaseAddress()) )
-				{
-					bDiffersFromDefaultForObject = true;
-					DefaultValue = NSLOCTEXT("PropertyEditor", "ArrayLongerThanDefault", "Array is longer than the default.").ToString();
-				}
-			}
-
 			// The property is a simple field.  Compare it against the enclosing object's default for that property.
 			if ( !bDiffersFromDefaultForObject)
 			{
@@ -1702,14 +1692,18 @@ FString FPropertyNode::GetDefaultValueAsStringForObject( FPropertyItemValueDataT
 					}
 				}
 
-				if ( ValueTracker.GetPropertyValueAddress() == NULL || ValueTracker.GetPropertyDefaultAddress() == NULL )
+				if ( ValueTracker.GetPropertyDefaultAddress() == NULL )
 				{
-					if ( !OuterSetProperty && !OuterMapProperty )
+					// no default available, fall back on the default value for our primitive:
+					uint8* TempComplexPropAddr = (uint8*)FMemory::Malloc(InProperty->GetSize(), InProperty->GetMinAlignment());
+					InProperty->InitializeValue(TempComplexPropAddr);
+					ON_SCOPE_EXIT
 					{
-					// if either are NULL, we had a dynamic array somewhere in our parent chain and the array doesn't
-					// have enough elements in either the default or the object
-					DefaultValue = NSLOCTEXT("PropertyEditor", "DifferentArrayLength", "Array has different length than the default.").ToString();
-				}
+						InProperty->DestroyValue(TempComplexPropAddr);
+						FMemory::Free(TempComplexPropAddr);
+					};
+					
+					InProperty->ExportText_Direct(DefaultValue, TempComplexPropAddr, TempComplexPropAddr, nullptr, PPF_None);
 				}
 				else if ( GetArrayIndex() == INDEX_NONE && InProperty->ArrayDim > 1 )
 				{
@@ -1784,271 +1778,6 @@ FText FPropertyNode::GetResetToDefaultLabel()
 	return OutLabel;
 }
 
-void FPropertyNode::ResetToDefault( FNotifyHook* InNotifyHook )
-{
-	UProperty* TheProperty = GetProperty();
-	check(TheProperty);
-
-		// Get an iterator for the enclosing objects.
-	FObjectPropertyNode* ObjectNode = FindObjectItemParent();
-
-	if( ObjectNode )
-	{
-		// The property is a simple field.  Compare it against the enclosing object's default for that property.
-		////////////////
-		FScopedTransaction Transaction( NSLOCTEXT("UnrealEd", "PropertyWindowEditProperties", "Edit Properties") );
-
-		// Whether or not we've process prechange already
-		bool bNotifiedPreChange = false;
-
-		// Whether or not an edit inline new was reset as a result of this reset to default
-		bool bEditInlineNewWasReset = false;
-
-		TArray< TMap<FString, int32> > ArrayIndicesPerObject;
-
-		// List of top level objects sent to the PropertyChangedEvent
-		TArray<const UObject*> TopLevelObjects;
-		TopLevelObjects.Reserve(ObjectNode->GetNumObjects());
-
-		for( int32 ObjIndex = 0; ObjIndex < ObjectNode->GetNumObjects(); ++ObjIndex )
-		{
-			UObject* Object = ObjectNode->GetUObject( ObjIndex );
-
-			// special case for UObject class - it has no defaults
-			if( Object && Object != UObject::StaticClass() && Object != UObject::StaticClass()->GetDefaultObject() )
-			{
-				TSharedPtr<FPropertyItemValueDataTrackerSlate> ValueTrackerPtr = GetValueTracker(Object, ObjIndex);
-
-				if( ValueTrackerPtr.IsValid() && ValueTrackerPtr->IsValidTracker() && ValueTrackerPtr->HasDefaultValue() )
-				{
-					FPropertyItemValueDataTrackerSlate& ValueTracker = *ValueTrackerPtr;
-
-					bool bIsGameWorld = false;
-					// If the object we are modifying is in the PIE world, than make the PIE world the active
-					// GWorld.  Assumes all objects managed by this property window belong to the same world.
-					UWorld* OldGWorld = nullptr;
-					if ( GUnrealEd && GUnrealEd->PlayWorld && !GUnrealEd->bIsSimulatingInEditor && Object->IsIn(GUnrealEd->PlayWorld))
-					{
-						OldGWorld = SetPlayInEditorWorld(GUnrealEd->PlayWorld);
-						bIsGameWorld = true;
-					}
-
-					FPropertyNode* ParentPropertyNode = GetParentNode();
-					UProperty* ParentProperty = ParentPropertyNode != nullptr ? ParentPropertyNode->GetProperty() : nullptr;
-
-					// If we're about to modify an element in a set, check to ensure that we're not duplicating a default value
-					if (Cast<USetProperty>(ParentProperty) != nullptr)
-					{
-						FScriptSetHelper SetHelper(Cast<USetProperty>(ParentProperty), ParentPropertyNode->GetValueBaseAddress((uint8*)Object));
-						FDefaultConstructedPropertyElement DefaultElementValue(SetHelper.ElementProp);
-
-						int32 ThisElementIndex = SetHelper.FindElementIndex(TheProperty->ContainerPtrToValuePtr<uint8>(ValueTrackerPtr->GetPropertyValueAddress()));
-						int32 DefaultIndex = SetHelper.FindElementIndex(DefaultElementValue.GetObjAddress());
-
-						if (DefaultIndex != INDEX_NONE && ThisElementIndex != DefaultIndex)
-						{
-							FNotificationInfo ResetToDefaultErrorInfo = LOCTEXT("SetElementResetToDefault_Duplicate", "Cannot reset the element back to its default value because the default already exists in the set");
-							ResetToDefaultErrorInfo.ExpireDuration = 3.0f;
-							FSlateNotificationManager::Get().AddNotification(ResetToDefaultErrorInfo);
-							return;
-						}
-					}
-
-					// If we're about to modify a map, ensure that the default key value is not duplicated
-					if (Cast<UMapProperty>(ParentProperty) != nullptr)
-					{
-						if (PropertyKeyNode.IsValid())
-						{
-							// This is the value node; it should always be reset to default. The key node should be checked separately.
-							PropertyKeyNode->ResetToDefault( InNotifyHook );
-						}
-						else
-						{
-							// Key node, so perform the default check here
-							FScriptMapHelper MapHelper(Cast<UMapProperty>(ParentProperty), ParentPropertyNode->GetValueBaseAddress((uint8*)Object));
-							FDefaultConstructedPropertyElement DefaultKeyValue(MapHelper.KeyProp);
-
-							uint8* PairPtr = MapHelper.GetPairPtr(ArrayIndex);
-							int32 ThisKeyIndex = MapHelper.FindMapIndexWithKey(TheProperty->ContainerPtrToValuePtr<uint8>(ValueTrackerPtr->GetPropertyValueAddress()));
-							int32 DefaultIndex = MapHelper.FindMapIndexWithKey(DefaultKeyValue.GetObjAddress());
-
-							if (DefaultIndex != INDEX_NONE && ThisKeyIndex != DefaultIndex)
-							{
-								FNotificationInfo ResetToDefaultErrorInfo = LOCTEXT("MapKeyResetToDefault_Duplicate", "Cannot reset the key back to its default value because the default already exists in the map");
-								ResetToDefaultErrorInfo.ExpireDuration = 3.0f;
-								FSlateNotificationManager::Get().AddNotification(ResetToDefaultErrorInfo);
-								return;
-							}
-						}
-					}
-
-					if( !bNotifiedPreChange )
-					{
-						// Call preedit change on all the objects
-						NotifyPreChange( GetProperty(), InNotifyHook );
-						bNotifiedPreChange = true;
-					}
-
-					// Cache the value of the property before modifying it.
-					FString PreviousValue;
-					TheProperty->ExportText_Direct(PreviousValue, ValueTracker.GetPropertyValueAddress(), ValueTracker.GetPropertyValueAddress(), nullptr, 0);
-								
-					FString PreviousArrayValue;
-
-					if (ValueTracker.GetPropertyDefaultAddress())
-					{
-						UObject* RootObject = ValueTracker.GetTopLevelObject();
-
-						FPropertyItemComponentCollector ComponentCollector(ValueTracker);
-
-						// dynamic arrays are the only property type that do not support CopySingleValue correctly due to the fact that they cannot
-						// be used in a static array
-
-						if (UArrayProperty* ParentArrayProp = Cast<UArrayProperty>(ParentProperty))
-						{
-							if (ParentArrayProp->Inner == TheProperty)
-							{
-								uint8* Addr = ParentPropertyNode->GetValueBaseAddress((uint8*)Object);
-
-								ParentArrayProp->ExportText_Direct(PreviousArrayValue, Addr, Addr, nullptr, 0);
-							}
-						}
-
-						if (UArrayProperty* ArrayProp = Cast<UArrayProperty>(TheProperty))
-						{
-							TheProperty->CopyCompleteValue(ValueTracker.GetPropertyValueAddress(), ValueTracker.GetPropertyDefaultAddress());
-						}
-						else
-						{
-							if( GetArrayIndex() == INDEX_NONE && TheProperty->ArrayDim > 1 )
-							{
-								TheProperty->CopyCompleteValue(ValueTracker.GetPropertyValueAddress(), ValueTracker.GetPropertyDefaultAddress());
-							}
-							else
-							{
-								TheProperty->CopySingleValue(ValueTracker.GetPropertyValueAddress(), ValueTracker.GetPropertyDefaultAddress());
-							}
-						}
-
-						if( ComponentCollector.Components.Num() > 0 )
-						{
-							TMap<UObject*,UObject*> ReplaceMap;
-							FPropertyItemComponentCollector DefaultComponentCollector(ValueTracker);
-							for ( int32 CompIndex = 0; CompIndex < ComponentCollector.Components.Num(); CompIndex++ )
-							{
-								if (UObject* Component = ComponentCollector.Components[CompIndex])
-								{
-									if ( DefaultComponentCollector.Components.Contains(Component->GetArchetype()) )
-									{
-										ReplaceMap.Add(Component, Component->GetArchetype());
-									}
-									else if( DefaultComponentCollector.Components.IsValidIndex(CompIndex) )
-									{
-										ReplaceMap.Add(Component, DefaultComponentCollector.Components[CompIndex]);
-									}
-								}
-							}
-
-							{ FArchiveReplaceObjectRef<UObject> ReplaceAr(RootObject, ReplaceMap, false, true, true); }
-
-							// The old objects need to be renamed out of the way otherwise the subobject instancing will just find the
-							// same object again and not get a new one
-							for (const TPair<UObject*,UObject*> ReplacedObjPair : ReplaceMap)
-							{
-								ReplacedObjPair.Key->Rename(nullptr, GetTransientPackage(), REN_DontCreateRedirectors);
-							}
-
-							FObjectInstancingGraph InstanceGraph(RootObject);
-
-							TArray<UObject*> Subobjects;
-							FReferenceFinder Collector(
-								Subobjects,	//	InObjectArray
-								RootObject,		//	LimitOuter
-								false,			//	bRequireDirectOuter
-								true,			//	bIgnoreArchetypes
-								true,			//	bSerializeRecursively
-								false			//	bShouldIgnoreTransient
-								);
-							Collector.FindReferences( RootObject );
-
-							for( UObject* SubObj : Subobjects )
-							{
-								InstanceGraph.AddNewInstance(SubObj);
-							}
-
-							RootObject->InstanceSubobjectTemplates(&InstanceGraph);
-
-							{ FArchiveReplaceObjectRef<UObject> ReplaceAr(RootObject, InstanceGraph.GetReplaceMap(), false, true, true); }
-
-						}
-
-						bEditInlineNewWasReset = ComponentCollector.bContainsEditInlineNew;
-
-					}
-					else
-					{
-						TheProperty->ClearValue(ValueTracker.GetPropertyValueAddress());
-					}
-
-					// Cache the value of the property after having modified it.
-					FString ValueAfterImport;
-					TheProperty->ExportText_Direct(ValueAfterImport, ValueTracker.GetPropertyValueAddress(), ValueTracker.GetPropertyValueAddress(), nullptr, 0);
-
-					// If this is an instanced component property we must move the old component to the 
-					// transient package so resetting owned components on the parent doesn't find it
-					UObjectProperty* ObjectProperty = Cast<UObjectProperty>(TheProperty);
-					if (ObjectProperty 
-						&& ObjectProperty->HasAnyPropertyFlags(CPF_InstancedReference) 
-						&& ObjectProperty->PropertyClass->IsChildOf(UActorComponent::StaticClass())
-						&& PreviousValue != ValueAfterImport)
-					{
-						FString ComponentName = PreviousValue;
-						ConstructorHelpers::StripObjectClass(ComponentName);
-						if (UActorComponent* Component = Cast<UActorComponent>(StaticFindObject(UActorComponent::StaticClass(), ANY_PACKAGE, *ComponentName)))
-						{
-							Component->Modify();
-							Component->Rename(nullptr, GetTransientPackage(), REN_DontCreateRedirectors);
-						}
-					}
-
-					if((Object->HasAnyFlags(RF_ClassDefaultObject | RF_ArchetypeObject) ||
-						(Object->HasAnyFlags(RF_DefaultSubObject) && Object->GetOuter()->HasAnyFlags(RF_ClassDefaultObject | RF_ArchetypeObject))) &&
-						!bIsGameWorld)
-					{
-						PropagatePropertyChange(Object, *ValueAfterImport, PreviousArrayValue.IsEmpty() ? PreviousValue : PreviousArrayValue);
-					}
-
-					TopLevelObjects.Add(Object);
-
-					if(OldGWorld)
-					{
-						// restore the original (editor) GWorld
-						RestoreEditorWorld( OldGWorld );
-					}
-
-					ArrayIndicesPerObject.Add(TMap<FString, int32>());
-					FPropertyValueImpl::GenerateArrayIndexMapToObjectNode(ArrayIndicesPerObject[ObjIndex], this);
-				}
-			}
-		}
-
-		if( bNotifiedPreChange )
-		{
-			// Call PostEditchange on all the objects
-			// Assume reset to default, can change topology
-			FPropertyChangedEvent ChangeEvent( TheProperty, EPropertyChangeType::ValueSet, &TopLevelObjects );
-			ChangeEvent.SetArrayIndexPerObject(ArrayIndicesPerObject);
-
-			NotifyPostChange( ChangeEvent, InNotifyHook );
-		}
-
-		if( bEditInlineNewWasReset )
-		{
-			RequestRebuildChildren();
-		}
-	}
-}
-
 bool FPropertyNode::IsReorderable()
 {
 	UProperty* NodeProperty = GetProperty();
diff --git a/Engine/Source/Editor/PropertyEditor/Private/PropertyNode.h b/Engine/Source/Editor/PropertyEditor/Private/PropertyNode.h
--- a/Engine/Source/Editor/PropertyEditor/Private/PropertyNode.h
+++ b/Engine/Source/Editor/PropertyEditor/Private/PropertyNode.h
@@ -501,11 +501,6 @@ public:
 	 */
 	FText GetResetToDefaultLabel();
 
-	/**
-	 * If there is a property, resets it to default.  Otherwise resets the entire parent structure
-	 */
-	void ResetToDefault(FNotifyHook* InNotifyHook);
-
 	/**
 	 * @return If this property node is associated with a property that can be reordered within an array
 	 */
@@ -802,7 +797,11 @@ public:
 	TSharedPtr<FPropertyNode>& GetPropertyKeyNode() { return PropertyKeyNode; }
 
 	const TSharedPtr<FPropertyNode>& GetPropertyKeyNode() const { return PropertyKeyNode; }
-
+	
+	/**
+	* Gets the default value of the property as string.
+	*/
+	FString GetDefaultValueAsString();
 protected:
 
 	TSharedRef<FEditPropertyChain> BuildPropertyChain( UProperty* PropertyAboutToChange );
@@ -842,11 +841,6 @@ protected:
 	bool GetDiffersFromDefaultForObject( FPropertyItemValueDataTrackerSlate& ValueTracker, UProperty* InProperty );
 
 	FString GetDefaultValueAsStringForObject( FPropertyItemValueDataTrackerSlate& ValueTracker, UObject* InObject, UProperty* InProperty );
-	
-	/**
-	* Gets the default value of the property as string.
-	*/
-	FString GetDefaultValueAsString();
 
 	/**
 	 * Helper function to obtain the display name for an enum property
diff --git a/Engine/Source/Editor/PropertyEditor/Private/UserInterface/PropertyEditor/SPropertyEditorAsset.cpp b/Engine/Source/Editor/PropertyEditor/Private/UserInterface/PropertyEditor/SPropertyEditorAsset.cpp
--- a/Engine/Source/Editor/PropertyEditor/Private/UserInterface/PropertyEditor/SPropertyEditorAsset.cpp
+++ b/Engine/Source/Editor/PropertyEditor/Private/UserInterface/PropertyEditor/SPropertyEditorAsset.cpp
@@ -184,33 +184,47 @@ void SPropertyEditorAsset::Construct( const FArguments& InArgs, const TSharedPtr
 	TAttribute<bool> IsEnabledAttribute(this, &SPropertyEditorAsset::CanEdit);
 	TAttribute<FText> TooltipAttribute(this, &SPropertyEditorAsset::OnGetToolTip);
 
-	if (Property && Property->HasAnyPropertyFlags(CPF_EditConst | CPF_DisableEditOnTemplate))
+	if (Property)
 	{
-		// There are some cases where editing an Actor Property is not allowed, such as when it is contained within a struct or a CDO
-		TArray<UObject*> ObjectList;
-		if (PropertyEditor.IsValid())
+		UProperty* PropToConsider = Property;
+		if (UProperty* OuterProperty = Cast<UProperty>(PropToConsider->GetOuter()))
 		{
-			PropertyEditor->GetPropertyHandle()->GetOuterObjects(ObjectList);
+			if (OuterProperty->IsA<UArrayProperty>()
+				|| OuterProperty->IsA<USetProperty>()
+				|| OuterProperty->IsA<UMapProperty>())
+			{
+				PropToConsider = OuterProperty;
+			}
 		}
 
-		// If there is no objects, that means we must have a struct asset managing this property
-		if (ObjectList.Num() == 0)
-		{
-			IsEnabledAttribute.Set(false);
-			TooltipAttribute.Set(LOCTEXT("VariableHasDisableEditOnTemplate", "Editing this value in structure's defaults is not allowed"));
-		}
-		else
+		if (PropToConsider->HasAnyPropertyFlags(CPF_EditConst | CPF_DisableEditOnTemplate))
 		{
-			// Go through all the found objects and see if any are a CDO, we can't set an actor in a CDO default.
-			for (UObject* Obj : ObjectList)
+			// There are some cases where editing an Actor Property is not allowed, such as when it is contained within a struct or a CDO
+			TArray<UObject*> ObjectList;
+			if (PropertyEditor.IsValid())
+			{
+				PropertyEditor->GetPropertyHandle()->GetOuterObjects(ObjectList);
+			}
+
+			// If there is no objects, that means we must have a struct asset managing this property
+			if (ObjectList.Num() == 0)
+			{
+				IsEnabledAttribute.Set(false);
+				TooltipAttribute.Set(LOCTEXT("VariableHasDisableEditOnTemplate", "Editing this value in structure's defaults is not allowed"));
+			}
+			else
 			{
-				if (Obj->IsTemplate())
+				// Go through all the found objects and see if any are a CDO, we can't set an actor in a CDO default.
+				for (UObject* Obj : ObjectList)
 				{
-					IsEnabledAttribute.Set(false);
-					TooltipAttribute.Set(LOCTEXT("VariableHasDisableEditOnTemplateTooltip", "Editing this value in a Class Default Object is not allowed"));
-					break;
-				}
+					if (Obj->IsTemplate())
+					{
+						IsEnabledAttribute.Set(false);
+						TooltipAttribute.Set(LOCTEXT("VariableHasDisableEditOnTemplateTooltip", "Editing this value in a Class Default Object is not allowed"));
+						break;
+					}
 
+				}
 			}
 		}
 	}
diff --git a/Engine/Source/Editor/SourceControlWindows/Private/SSourceControlHistory.cpp b/Engine/Source/Editor/SourceControlWindows/Private/SSourceControlHistory.cpp
--- a/Engine/Source/Editor/SourceControlWindows/Private/SSourceControlHistory.cpp
+++ b/Engine/Source/Editor/SourceControlWindows/Private/SSourceControlHistory.cpp
@@ -1435,7 +1435,35 @@ void FSourceControlWindows::DisplayRevisionHistory( const TArray<FString>& InPac
 	if(SourceControlProvider.Execute(UpdateStatusOperation, PackageFilenames))
 	{
 		TArray< FSourceControlStateRef > SourceControlStates;
-		SourceControlProvider.GetState(PackageFilenames, SourceControlStates, EStateCacheUsage::Use);
+		for(const FString& PackageFilename : PackageFilenames)
+		{
+			TArray<FString> RevisionName;
+			RevisionName.Add(PackageFilename);
+
+			while(RevisionName.Num() != 0)
+			{
+				int32 InitialNum = SourceControlStates.Num();
+				SourceControlProvider.GetState(RevisionName, SourceControlStates, EStateCacheUsage::Use);
+				int32 NewNum = SourceControlStates.Num();
+				ensure(NewNum >= InitialNum);
+				
+				RevisionName.Empty();
+				// check to see if origin of this file is a branch, append the history from the branch point:
+				if(NewNum > InitialNum)
+				{
+					int32 HistorySize = SourceControlStates.Last()->GetHistorySize();
+					if( HistorySize > 0 )
+					{
+						TSharedPtr<class ISourceControlRevision, ESPMode::ThreadSafe> InitialHistory = SourceControlStates.Last()->GetHistoryItem(HistorySize - 1);
+						TSharedPtr<ISourceControlRevision, ESPMode::ThreadSafe> BranchSource = InitialHistory->GetBranchSource();
+						if( BranchSource.IsValid() )
+						{
+							RevisionName.Add(BranchSource->GetFilename());
+						}
+					}
+				}
+			}
+		}
 
 		TSharedRef<SWindow> NewWindow = SNew(SWindow)
 			.Title( NSLOCTEXT("SourceControl.HistoryWindow", "Title", "File History") )
diff --git a/Engine/Source/Editor/StaticMeshEditor/Private/StaticMeshEditorViewportClient.cpp b/Engine/Source/Editor/StaticMeshEditor/Private/StaticMeshEditorViewportClient.cpp
--- a/Engine/Source/Editor/StaticMeshEditor/Private/StaticMeshEditorViewportClient.cpp
+++ b/Engine/Source/Editor/StaticMeshEditor/Private/StaticMeshEditorViewportClient.cpp
@@ -23,7 +23,7 @@
 
 #include "Runtime/Analytics/Analytics/Public/Interfaces/IAnalyticsProvider.h"
 #include "EngineAnalytics.h"
-#include "AI/Navigation/NavCollision.h"
+#include "AI/Navigation/NavCollisionBase.h"
 #include "PhysicsEngine/BodySetup.h"
 
 #include "Engine/AssetUserData.h"
@@ -616,8 +616,8 @@ void FStaticMeshEditorViewportClient::Draw(const FSceneView* View,FPrimitiveDraw
 			}
 		}
 
-		// The simple nav geometry is only used by dynamic obstacles for now
-		if (StaticMesh->NavCollision && StaticMesh->NavCollision->bIsDynamicObstacle)
+		if (StaticMesh->NavCollision 
+			&& StaticMesh->bHasNavigationData)
 		{
 			// Draw the static mesh's body setup (simple collision)
 			FTransform GeomTransform(StaticMeshComponent->GetComponentTransform());
diff --git a/Engine/Source/Editor/UnrealEd/Classes/Settings/EditorExperimentalSettings.h b/Engine/Source/Editor/UnrealEd/Classes/Settings/EditorExperimentalSettings.h
--- a/Engine/Source/Editor/UnrealEd/Classes/Settings/EditorExperimentalSettings.h
+++ b/Engine/Source/Editor/UnrealEd/Classes/Settings/EditorExperimentalSettings.h
@@ -63,6 +63,10 @@ public:
 	UPROPERTY(EditAnywhere, config, Category = Blueprints, meta = (DisplayName = "Find and Replace All References Tool"))
 	bool bEnableFindAndReplaceReferences;
 
+	/** Bookmarks facilitate quick-jump navigation to saved locations in any Blueprint Graph context in the Blueprint Editor. */
+	UPROPERTY(EditAnywhere, config, Category = Blueprints)
+	bool bEnableBlueprintGraphBookmarks;
+
 	/** Should arrows indicating data/execution flow be drawn halfway along wires? */
 	UPROPERTY(/*EditAnywhere - deprecated (moved into UBlueprintEditorSettings), */config/*, Category=Blueprints, meta=(DisplayName="Draw midpoint arrows in Blueprints")*/)
 	bool bDrawMidpointArrowsInBlueprints;
diff --git a/Engine/Source/Editor/UnrealEd/Private/Commandlets/GenerateBlueprintAPICommandlet.cpp b/Engine/Source/Editor/UnrealEd/Private/Commandlets/GenerateBlueprintAPICommandlet.cpp
--- a/Engine/Source/Editor/UnrealEd/Private/Commandlets/GenerateBlueprintAPICommandlet.cpp
+++ b/Engine/Source/Editor/UnrealEd/Private/Commandlets/GenerateBlueprintAPICommandlet.cpp
@@ -58,6 +58,14 @@ GenerateBlueprintAPI commandlet params: \n\
     -class=<Class>      Used to specify the blueprint's parent class, if left  \n\
                         unset then it will go through EVERY blueprint parent   \n\
                         class available.                                       \n\
+\n\
+    -category=<CategoryPath1, CategoryPath2, ...>                              \n\
+                        Used to specify the blueprint's category path, if      \n\
+                        left unset then it will go through EVERY category      \n\
+                        available. A category path can contain a single        \n\
+                        category, or multiple category names concanated by     \n\
+                        \'|\'. Multiple category paths can be supplied by using\n\
+                        \',\' as separators.                                   \n\
 \n\
     -multifile          Used to keep file size down, will split each blueprint \n\
                         into its own file (meaning only one file will be       \n\
@@ -75,7 +83,7 @@ GenerateBlueprintAPI commandlet params: \n\
                         (as it has been time sync in the past). This is not    \n\
                         ideal for diffs though (since times can easily vary).  \n\
 \n\
-    -experimental       Uses an new way of constructing Blueprint action menus \n\
+    -experimental       Uses a new way of constructing Blueprint action menus \n\
                         (that will replace the current system).                \n\
 \n\
     -name=<Filename>    Overrides the default filename. Leave off the extention\n\
@@ -121,6 +129,7 @@ GenerateBlueprintAPI commandlet params: \n\
 		UClass*    PaletteFilter;
 		FString    SaveDir;
 		FString    Filename;
+		TArray<FString> CategoryFilter;
 	};
 
 	/** Static instance of the command switches (so we don't have to pass one along the call stack) */
@@ -222,6 +231,11 @@ GenerateBlueprintAPI commandlet params: \n\
 	 */
 	static void DumpActionList(uint32 Indent, FGraphActionListBuilderBase& ActionList, FArchive* FileOutWriter);
 
+	/**
+	* Get the category information from a given action.
+	*/
+	static FString GetActionCategory(FGraphActionListBuilderBase::ActionGroup const& Action);
+
 	/**
 	 * Generic function that dumps information on a single action (like it's 
 	 * name, category, an associated node if it has one, etc.).
@@ -295,6 +309,34 @@ GenerateBlueprintAPIUtils::CommandletOptions::CommandletOptions(TArray<FString>
 				}
 			}
 		}
+		else if (Switch.StartsWith("category=", ESearchCase::IgnoreCase))
+		{
+			FString ClassSwitch, Filters;
+			Switch.Split(TEXT("="), &ClassSwitch, &Filters, ESearchCase::IgnoreCase);
+
+			// Remove surounding " or ' if exist.
+			Filters = Filters.Replace(TEXT("\""), TEXT(""), ESearchCase::IgnoreCase);
+			Filters = Filters.Replace(TEXT("\'"), TEXT(""), ESearchCase::IgnoreCase);
+			Filters.TrimStartAndEndInline();
+
+			// Parse multiple category paths.
+			while (Filters.Len() > 0)
+			{
+				FString FilterName, OtherFilters;
+				if (Filters.Split(TEXT(","), &FilterName, &OtherFilters, ESearchCase::IgnoreCase))
+				{
+					FilterName.TrimStartAndEndInline();
+					CategoryFilter.Add(FilterName);
+					Filters = OtherFilters;
+				}
+				else
+				{
+					Filters.TrimStartAndEndInline();
+					CategoryFilter.Add(Filters);
+					break;
+				}
+			}
+		}
 		else if (!Switch.Compare("h", ESearchCase::IgnoreCase) || 
 		         !Switch.Compare("?", ESearchCase::IgnoreCase) || 
 				 !Switch.Compare("help", ESearchCase::IgnoreCase))
@@ -726,6 +768,9 @@ static void GenerateBlueprintAPIUtils::DumpPalette(uint32 Indent, UBlueprint* Bl
 static void GenerateBlueprintAPIUtils::DumpActionList(uint32 Indent, FGraphActionListBuilderBase& ActionList, FArchive* FileOutWriter)
 {
 	TArray< FGraphActionListBuilderBase::ActionGroup const* > SortedActions;
+	const TArray<FString>& CategoryFilter = GenerateBlueprintAPIUtils::CommandOptions.CategoryFilter;
+	bool HasCategoryFilter = CategoryFilter.Num() > 0;
+
 	for (int32 ActionIndex = 0; ActionIndex < ActionList.GetNumActions(); ++ActionIndex)
 	{
 		FGraphActionListBuilderBase::ActionGroup& Action = ActionList.GetAction(ActionIndex);
@@ -734,7 +779,22 @@ static void GenerateBlueprintAPIUtils::DumpActionList(uint32 Indent, FGraphActio
 			continue;
 		}
 
-		SortedActions.Add(&Action);
+		if (!HasCategoryFilter)
+		{
+			SortedActions.Add(&Action);
+		}
+		else
+		{
+			const FString& ActionCategory = GenerateBlueprintAPIUtils::GetActionCategory(Action);
+			for (const FString& filter : CategoryFilter)
+			{
+				if (ActionCategory.Contains(filter, ESearchCase::CaseSensitive, ESearchDir::FromStart))
+				{
+					SortedActions.Add(&Action);
+					break;
+				}
+			}
+		}
 	}
 
 	FString const ActionListIndent = BuildIndentString(Indent);
@@ -778,16 +838,15 @@ static void GenerateBlueprintAPIUtils::DumpActionList(uint32 Indent, FGraphActio
 	EndActionListEntry += "\n" + ActionListIndent + "}";
 	FileOutWriter->Serialize(TCHAR_TO_ANSI(*EndActionListEntry), EndActionListEntry.Len());
 }
-
 //------------------------------------------------------------------------------
-static void GenerateBlueprintAPIUtils::DumpActionMenuItem(uint32 Indent, FGraphActionListBuilderBase::ActionGroup const& Action, FGraphActionListBuilderBase& ActionList, FArchive* FileOutWriter)
+static FString GenerateBlueprintAPIUtils::GetActionCategory(FGraphActionListBuilderBase::ActionGroup const& Action)
 {
 	check(Action.Actions.Num() > 0);
 
 	// Get action category info
 	const TArray<FString>& MenuHierarchy = Action.GetCategoryChain();
 
-	FString ActionCategory = TEXT("");
+	FString ActionCategory;
 
 	bool bHasCategory = (MenuHierarchy.Num() > 0);
 	if (bHasCategory)
@@ -797,7 +856,12 @@ static void GenerateBlueprintAPIUtils::DumpActionMenuItem(uint32 Indent, FGraphA
 			ActionCategory += SubCategory + TEXT("|");
 		}
 	}
-
+	return ActionCategory;
+}
+//------------------------------------------------------------------------------
+static void GenerateBlueprintAPIUtils::DumpActionMenuItem(uint32 Indent, FGraphActionListBuilderBase::ActionGroup const& Action, FGraphActionListBuilderBase& ActionList, FArchive* FileOutWriter)
+{
+	const FString& ActionCategory = GenerateBlueprintAPIUtils::GetActionCategory(Action);
 	TArray<FString> Categories;
 	ActionCategory.ParseIntoArray(Categories, TEXT("|"), true);
 
@@ -840,30 +904,30 @@ static void GenerateBlueprintAPIUtils::DumpActionMenuItem(uint32 Indent, FGraphA
 
 		if (Node->ShouldDrawCompact())
 		{
-			ActionEntry += IndentedNewline + "\"CompactTitle\" : \"" + MakeJsonString(Node->GetCompactNodeTitle().ToString()) + "\"";
+			ActionEntry += "," + IndentedNewline + "\"CompactTitle\" : \"" + MakeJsonString(Node->GetCompactNodeTitle().ToString()) + "\"";
 		}
 
 		if (Node->IsNodePure())
 		{
-			ActionEntry += IndentedNewline + "\"NodeType\"     : \"pure\"";
+			ActionEntry += "," + IndentedNewline + "\"NodeType\"     : \"pure\"";
 		}
 		else if (Node->IsA<UK2Node_Event>())
 		{
-			ActionEntry += IndentedNewline + "\"NodeType\"     : \"event\"";
+			ActionEntry += "," + IndentedNewline + "\"NodeType\"     : \"event\"";
 		}
 		else if (Node->IsA<UK2Node_Switch>())
 		{
-			ActionEntry += IndentedNewline + "\"NodeType\"     : \"switch\"";
+			ActionEntry += "," + IndentedNewline + "\"NodeType\"     : \"switch\"";
 		}
 		else
 		{
-			ActionEntry += IndentedNewline + "\"NodeType\"     : \"function\"";
+			ActionEntry += "," + IndentedNewline + "\"NodeType\"     : \"function\"";
 		}
 
 		if (Node->IsA<UK2Node_CommutativeAssociativeBinaryOperator>()
 			|| (Node->IsA<UK2Node_Switch>() && !Node->IsA<UK2Node_SwitchEnum>()))
 		{
-			ActionEntry += IndentedNewline + "\"ShowAddPin\"   : \"true\"";
+			ActionEntry += "," + IndentedNewline + "\"ShowAddPin\"   : \"true\"";
 		}
 
 		if (Node->Pins.Num() > 0)
@@ -898,7 +962,7 @@ static void GenerateBlueprintAPIUtils::DumpActionMenuItem(uint32 Indent, FGraphA
 					{
 						ActionEntry += PinDetailsIndentedNewline + "\"Name\"                 : \"" + MakeJsonString(DisplayName) + "\",";
 					}
-					ActionEntry += PinDetailsIndentedNewline + "\"Direction\"            : \"" + (Pin->Direction == EGPD_Input ? "input" : "output") + "\"";
+					ActionEntry += PinDetailsIndentedNewline + "\"Direction\"            : \"" + (Pin->Direction == EGPD_Input ? "input" : "output") + "\",";
 
 					ActionEntry += PinDetailsIndentedNewline + "\"TypeText\"             : \"" + UEdGraphSchema_K2::TypeToText(Pin->PinType).ToString() + "\"";
 
diff --git a/Engine/Source/Editor/UnrealEd/Private/CookOnTheFlyServer.cpp b/Engine/Source/Editor/UnrealEd/Private/CookOnTheFlyServer.cpp
--- a/Engine/Source/Editor/UnrealEd/Private/CookOnTheFlyServer.cpp
+++ b/Engine/Source/Editor/UnrealEd/Private/CookOnTheFlyServer.cpp
@@ -962,7 +962,7 @@ bool UCookOnTheFlyServer::StartNetworkFileServer( const bool BindAnyPort )
 		FAssetRegistryGenerator* Generator = RegistryGenerators.FindRef(FName(*Platforms[Index]->PlatformName()));
 		if (Generator)
 		{
-			Generator->SaveAssetRegistry(GetSandboxAssetRegistryFilename(), false);
+			Generator->SaveAssetRegistry(GetSandboxAssetRegistryFilename(), true);
 		}
 	}
 
diff --git a/Engine/Source/Editor/UnrealEd/Private/EditorActor.cpp b/Engine/Source/Editor/UnrealEd/Private/EditorActor.cpp
--- a/Engine/Source/Editor/UnrealEd/Private/EditorActor.cpp
+++ b/Engine/Source/Editor/UnrealEd/Private/EditorActor.cpp
@@ -14,7 +14,7 @@
 #include "GameFramework/Actor.h"
 #include "GameFramework/Pawn.h"
 #include "Engine/World.h"
-#include "AI/Navigation/NavigationSystem.h"
+#include "AI/NavigationSystemBase.h"
 #include "Components/LightComponent.h"
 #include "Model.h"
 #include "Exporters/Exporter.h"
@@ -1025,13 +1025,11 @@ bool UUnrealEdEngine::edactDeleteSelected( UWorld* InWorld, bool bVerifyDeletion
 
 	if( LevelsToRebuildNavigation.Num() )
 	{
-		UWorld* World = GetEditorWorldContext().World();
-		UNavigationSystem* NavSys = UNavigationSystem::GetCurrent(World);
-		if (NavSys)
+		for (ULevel* Level : LevelsToRebuildNavigation)
 		{
-			for ( ULevel* Level : LevelsToRebuildNavigation )
+			if (Level)
 			{
-				NavSys->UpdateLevelCollision(Level);
+				FNavigationSystem::UpdateLevelCollision(*Level);
 			}
 		}
 	}
diff --git a/Engine/Source/Editor/UnrealEd/Private/EditorBuildUtils.cpp b/Engine/Source/Editor/UnrealEd/Private/EditorBuildUtils.cpp
--- a/Engine/Source/Editor/UnrealEd/Private/EditorBuildUtils.cpp
+++ b/Engine/Source/Editor/UnrealEd/Private/EditorBuildUtils.cpp
@@ -17,7 +17,7 @@
 #include "SourceControlOperations.h"
 #include "ISourceControlModule.h"
 #include "Materials/MaterialInterface.h"
-#include "AI/Navigation/NavigationSystem.h"
+#include "AI/NavigationSystemBase.h"
 #include "Editor/UnrealEdEngine.h"
 #include "Settings/LevelEditorMiscSettings.h"
 #include "Misc/ConfigCacheIni.h"
@@ -814,8 +814,7 @@ void FEditorBuildUtils::SubmitPackagesForAutomatedBuild( const TSet<UPackage*>&
 
 void FEditorBuildUtils::TriggerNavigationBuilder(UWorld* InWorld, FName Id)
 {
-	if( InWorld->GetWorldSettings()->bEnableNavigationSystem &&
-		InWorld->GetNavigationSystem() )
+	if (InWorld)
 	{
 		if (Id == FBuildOptions::BuildAIPaths ||
 			Id == FBuildOptions::BuildSelectedAIPaths ||
@@ -831,7 +830,7 @@ void FEditorBuildUtils::TriggerNavigationBuilder(UWorld* InWorld, FName Id)
 		}
 
 		// Invoke navmesh generator
-		InWorld->GetNavigationSystem()->Build();
+		FNavigationSystem::Build(*InWorld);
 	}
 }
 
diff --git a/Engine/Source/Editor/UnrealEd/Private/EditorEngine.cpp b/Engine/Source/Editor/UnrealEd/Private/EditorEngine.cpp
--- a/Engine/Source/Editor/UnrealEd/Private/EditorEngine.cpp
+++ b/Engine/Source/Editor/UnrealEd/Private/EditorEngine.cpp
@@ -21,7 +21,7 @@
 #include "PhysicsEngine/BodyInstance.h"
 #include "Components/PrimitiveComponent.h"
 #include "Components/StaticMeshComponent.h"
-#include "AI/Navigation/NavigationSystem.h"
+#include "AI/NavigationSystemBase.h"
 #include "Components/LightComponent.h"
 #include "Tickable.h"
 #include "TickableEditorObject.h"
@@ -625,7 +625,7 @@ void UEditorEngine::InitEditor(IEngineLoop* InEngineLoop)
 
 	// Set navigation system property indicating whether navigation is supposed to rebuild automatically 
 	FWorldContext &EditorContext = GEditor->GetEditorWorldContext();
-	UNavigationSystem::SetNavigationAutoUpdateEnabled(GetDefault<ULevelEditorMiscSettings>()->bNavigationAutoUpdate, EditorContext.World()->GetNavigationSystem());
+	FNavigationSystem::SetNavigationAutoUpdateEnabled(GetDefault<ULevelEditorMiscSettings>()->bNavigationAutoUpdate, EditorContext.World()->GetNavigationSystem());
 
 	// Allocate temporary model.
 	TempModel = NewObject<UModel>();
diff --git a/Engine/Source/Editor/UnrealEd/Private/EditorExporters.cpp b/Engine/Source/Editor/UnrealEd/Private/EditorExporters.cpp
--- a/Engine/Source/Editor/UnrealEd/Private/EditorExporters.cpp
+++ b/Engine/Source/Editor/UnrealEd/Private/EditorExporters.cpp
@@ -20,7 +20,7 @@
 #include "SceneTypes.h"
 #include "RawIndexBuffer.h"
 #include "RenderingThread.h"
-#include "AI/Navigation/NavigationSystem.h"
+#include "AI/NavigationSystemBase.h"
 #include "Model.h"
 #include "Exporters/Exporter.h"
 #include "Exporters/AnimSequenceExporterFBX.h"
@@ -2100,12 +2100,7 @@ void UEditorEngine::RebuildStaticNavigableGeometry(ULevel* Level)
 			}
 		}
 
-		UWorld* World = GetEditorWorldContext().World();
-		UNavigationSystem* NavSys = UNavigationSystem::GetCurrent(World);
-		if (NavSys)
-		{
-			NavSys->UpdateLevelCollision(Level);
-		}
+		FNavigationSystem::UpdateLevelCollision(*Level);
 	}
 }
 
diff --git a/Engine/Source/Editor/UnrealEd/Private/EditorServer.cpp b/Engine/Source/Editor/UnrealEd/Private/EditorServer.cpp
--- a/Engine/Source/Editor/UnrealEd/Private/EditorServer.cpp
+++ b/Engine/Source/Editor/UnrealEd/Private/EditorServer.cpp
@@ -42,7 +42,7 @@
 #include "Materials/MaterialInterface.h"
 #include "Components/PrimitiveComponent.h"
 #include "Components/MeshComponent.h"
-#include "AI/Navigation/NavigationSystem.h"
+#include "AI/NavigationSystemBase.h"
 #include "Components/LightComponent.h"
 #include "Model.h"
 #include "Exporters/Exporter.h"
@@ -145,7 +145,7 @@
 #include "Framework/Notifications/NotificationManager.h"
 #include "Widgets/Notifications/SNotificationList.h"
 #include "Animation/AnimNotifies/AnimNotify.h"
-#include "AI/Navigation/NavLinkRenderingComponent.h"
+#include "NavLinkRenderingComponent.h"
 #include "Analytics/AnalyticsPrivacySettings.h"
 #include "Kismet2/KismetReinstanceUtilities.h"
 #include "AnalyticsEventAttribute.h"
@@ -2232,7 +2232,7 @@ UWorld* UEditorEngine::NewMap()
 	InitBuilderBrush( Context.World() );
 
 	// Let navigation system know we're done creating new world
-	UNavigationSystem::InitializeForWorld(Context.World(), FNavigationSystemRunMode::EditorMode);
+	FNavigationSystem::AddNavigationSystemToWorld(*Context.World(), FNavigationSystemRunMode::EditorMode);
 
 	// Deselect all
 	GEditor->SelectNone( false, true );
@@ -2636,8 +2636,8 @@ bool UEditorEngine::Map_Load(const TCHAR* Str, FOutputDevice& Ar)
 
 					InitializingFeedback.EnterProgressFrame();
 
-					UNavigationSystem::InitializeForWorld(Context.World(), FNavigationSystemRunMode::EditorMode);
 					Context.World()->CreateAISystem();
+					FNavigationSystem::AddNavigationSystemToWorld(*Context.World(), FNavigationSystemRunMode::EditorMode);
 
 					// Assign stationary light channels for previewing
 					ULightComponent::ReassignStationaryLightChannels(Context.World(), false, NULL);
diff --git a/Engine/Source/Editor/UnrealEd/Private/Factories/EditorFactories.cpp b/Engine/Source/Editor/UnrealEd/Private/Factories/EditorFactories.cpp
--- a/Engine/Source/Editor/UnrealEd/Private/Factories/EditorFactories.cpp
+++ b/Engine/Source/Editor/UnrealEd/Private/Factories/EditorFactories.cpp
@@ -228,7 +228,7 @@
 #include "Animation/BlendSpace1D.h"
 #include "Engine/FontFace.h"
 #include "Components/AudioComponent.h"
-#include "AI/Navigation/NavCollision.h"
+#include "AI/Navigation/NavCollisionBase.h"
 #include "Animation/BlendSpace.h"
 #include "Animation/AimOffsetBlendSpace.h"
 #include "Animation/AimOffsetBlendSpace1D.h"
@@ -5245,8 +5245,8 @@ EReimportResult::Type UReimportFbxStaticMeshFactory::Reimport( UObject* Obj )
 			}
 
 			// preserve settings in navcollision subobject
-			UNavCollision* NavCollision = Mesh->NavCollision ? 
-				(UNavCollision*)StaticDuplicateObject(Mesh->NavCollision, GetTransientPackage()) :
+			UNavCollisionBase* NavCollision = Mesh->NavCollision ? 
+				(UNavCollisionBase*)StaticDuplicateObject(Mesh->NavCollision, GetTransientPackage()) :
 				nullptr;
 
 			// preserve extended bound settings
diff --git a/Engine/Source/Editor/UnrealEd/Private/Fbx/ReimportFbxSceneFactory.cpp b/Engine/Source/Editor/UnrealEd/Private/Fbx/ReimportFbxSceneFactory.cpp
--- a/Engine/Source/Editor/UnrealEd/Private/Fbx/ReimportFbxSceneFactory.cpp
+++ b/Engine/Source/Editor/UnrealEd/Private/Fbx/ReimportFbxSceneFactory.cpp
@@ -54,7 +54,7 @@
 
 #include "ObjectTools.h"
 
-#include "AI/Navigation/NavCollision.h"
+#include "AI/Navigation/NavCollisionBase.h"
 
 #define LOCTEXT_NAMESPACE "FBXSceneReImportFactory"
 
@@ -1426,8 +1426,8 @@ EReimportResult::Type UReimportFbxSceneFactory::ReimportStaticMesh(void* VoidFbx
 	}
 
 	// preserve settings in navcollision subobject
-	UNavCollision* NavCollision = Mesh->NavCollision ?
-		(UNavCollision*)StaticDuplicateObject(Mesh->NavCollision, GetTransientPackage()) :
+	UNavCollisionBase* NavCollision = Mesh->NavCollision ?
+		(UNavCollisionBase*)StaticDuplicateObject(Mesh->NavCollision, GetTransientPackage()) :
 		nullptr;
 
 	// preserve extended bound settings
diff --git a/Engine/Source/Editor/UnrealEd/Private/Kismet2/BlueprintEditorUtils.cpp b/Engine/Source/Editor/UnrealEd/Private/Kismet2/BlueprintEditorUtils.cpp
--- a/Engine/Source/Editor/UnrealEd/Private/Kismet2/BlueprintEditorUtils.cpp
+++ b/Engine/Source/Editor/UnrealEd/Private/Kismet2/BlueprintEditorUtils.cpp
@@ -2529,7 +2529,7 @@ void FBlueprintEditorUtils::RemoveGraph(UBlueprint* Blueprint, class UEdGraph* G
 			// Can't just call Remove, the object is wrapped in a struct
 			for(int EditedDocIdx = 0; EditedDocIdx < Blueprint->LastEditedDocuments.Num(); ++EditedDocIdx)
 			{
-				if(Blueprint->LastEditedDocuments[EditedDocIdx].EditedObject == GraphToRemove)
+				if(Blueprint->LastEditedDocuments[EditedDocIdx].EditedObjectPath.ResolveObject() == GraphToRemove)
 				{
 					Blueprint->LastEditedDocuments.RemoveAt(EditedDocIdx);
 					break;
@@ -3873,11 +3873,16 @@ void FBlueprintEditorUtils::SetBlueprintVariableCategory(UBlueprint* Blueprint,
 		UClass* OuterClass = CastChecked<UClass>(TargetProperty->GetOuter());
 		const bool bIsNativeVar = (OuterClass->ClassGeneratedBy == nullptr);
 
-		// If the category does not change, we will not recompile the Blueprint
-		bool bIsCategoryChanged = false;
 		if (!bIsNativeVar)
 		{
-			TargetProperty->SetMetaData(TEXT("Category"), *SetCategory.ToString());
+			if (UTimelineTemplate* Timeline = Blueprint->FindTimelineTemplateByVariableName(TargetProperty->GetFName()))
+			{
+				Timeline->SetMetaData(TEXT("Category"), *SetCategory.ToString());
+			}
+			else
+			{
+				TargetProperty->SetMetaData(TEXT("Category"), *SetCategory.ToString());
+			}
 			const int32 VarIndex = FBlueprintEditorUtils::FindNewVariableIndex(Blueprint, VarName);
 			if (VarIndex != INDEX_NONE)
 			{
@@ -5917,7 +5922,7 @@ void FBlueprintEditorUtils::PurgeNullGraphs(UBlueprint* Blueprint)
 	CleanNullGraphReferencesInArray(Blueprint, Blueprint->DelegateSignatureGraphs);
 	CleanNullGraphReferencesInArray(Blueprint, Blueprint->MacroGraphs);
 
-	Blueprint->LastEditedDocuments.RemoveAll([](const FEditedDocumentInfo& TestDoc) { return TestDoc.EditedObject == nullptr; });
+	Blueprint->LastEditedDocuments.RemoveAll([](const FEditedDocumentInfo& TestDoc) { return TestDoc.EditedObjectPath.ResolveObject() == nullptr; });
 }
 
 struct FConformCallsToParentFunctionUtils
@@ -9013,7 +9018,7 @@ UClass* FBlueprintEditorUtils::GetNativeParent(const UBlueprint* BP)
 	return Ret;
 }
 
-bool FBlueprintEditorUtils::ImplentsGetWorld(const UBlueprint* BP)
+bool FBlueprintEditorUtils::ImplementsGetWorld(const UBlueprint* BP)
 {
 	if(UClass* NativeParent = GetNativeParent(BP))
 	{
diff --git a/Engine/Source/Editor/UnrealEd/Private/Kismet2/KismetDebugUtilities.cpp b/Engine/Source/Editor/UnrealEd/Private/Kismet2/KismetDebugUtilities.cpp
--- a/Engine/Source/Editor/UnrealEd/Private/Kismet2/KismetDebugUtilities.cpp
+++ b/Engine/Source/Editor/UnrealEd/Private/Kismet2/KismetDebugUtilities.cpp
@@ -19,6 +19,7 @@
 #include "Editor/UnrealEdEngine.h"
 #include "Settings/EditorExperimentalSettings.h"
 #include "CallStackViewer.h"
+#include "WatchPointViewer.h"
 #include "Animation/AnimBlueprintGeneratedClass.h"
 #include "UnrealEdGlobals.h"
 #include "Engine/Breakpoint.h"
@@ -705,6 +706,7 @@ void FKismetDebugUtilities::AttemptToBreakExecution(UBlueprint* BlueprintObj, co
 		Data.LastExceptionMessage = Info.GetDescription();
 		FKismetEditorUtilities::BringKismetToFocusAttentionOnObject(NodeStoppedAt);
 		CallStackViewer::UpdateDisplayedCallstack(FBlueprintExceptionTracker::Get().ScriptStack);
+		WatchViewer::UpdateDisplayedWatches(BlueprintObj);
 		FSlateApplication::Get().EnterDebuggingMode();
 	}
 }
@@ -1040,6 +1042,8 @@ void FKismetDebugUtilities::ClearBreakpoints(UBlueprint* Blueprint)
 	Blueprint->MarkPackageDirty();
 }
 
+FKismetDebugUtilities::FOnWatchedPinsListChanged FKismetDebugUtilities::WatchedPinsListChangedEvent;
+
 bool FKismetDebugUtilities::CanWatchPin(const UBlueprint* Blueprint, const UEdGraphPin* Pin)
 {
 	//@TODO: This function belongs in the schema
@@ -1067,6 +1071,7 @@ void FKismetDebugUtilities::RemovePinWatch(UBlueprint* Blueprint, const UEdGraph
 	Blueprint->WatchedPins.Remove(NonConstPin);
 	Blueprint->MarkPackageDirty();
 	Blueprint->PostEditChange();
+	WatchedPinsListChangedEvent.Broadcast(Blueprint);
 }
 
 void FKismetDebugUtilities::TogglePinWatch(UBlueprint* Blueprint, const UEdGraphPin* Pin)
@@ -1083,6 +1088,8 @@ void FKismetDebugUtilities::TogglePinWatch(UBlueprint* Blueprint, const UEdGraph
 		Blueprint->MarkPackageDirty();
 		Blueprint->PostEditChange();
 	}
+
+	WatchedPinsListChangedEvent.Broadcast(Blueprint);
 }
 
 void FKismetDebugUtilities::ClearPinWatches(UBlueprint* Blueprint)
@@ -1090,10 +1097,44 @@ void FKismetDebugUtilities::ClearPinWatches(UBlueprint* Blueprint)
 	Blueprint->WatchedPins.Empty();
 	Blueprint->MarkPackageDirty();
 	Blueprint->PostEditChange();
+
+	WatchedPinsListChangedEvent.Broadcast(Blueprint);
 }
 
 // Gets the watched tooltip for a specified site
 FKismetDebugUtilities::EWatchTextResult FKismetDebugUtilities::GetWatchText(FString& OutWatchText, UBlueprint* Blueprint, UObject* ActiveObject, const UEdGraphPin* WatchPin)
+{
+	UProperty* PropertyToDebug = nullptr;
+	void* DataPtr = nullptr;
+	void* DeltaPtr = nullptr;
+	UObject* ParentObj = nullptr;
+	FKismetDebugUtilities::EWatchTextResult Result = FindDebuggingData(Blueprint, ActiveObject, WatchPin, PropertyToDebug, DataPtr, DeltaPtr, ParentObj);
+
+	if (Result == FKismetDebugUtilities::EWatchTextResult::EWTR_Valid)
+	{
+		PropertyToDebug->ExportText_InContainer(/*ArrayElement=*/ 0, /*inout*/ OutWatchText, DataPtr, DeltaPtr, /*Parent=*/ ParentObj, PPF_PropertyWindow | PPF_BlueprintDebugView);
+	}
+
+	return Result;
+}
+
+FKismetDebugUtilities::EWatchTextResult FKismetDebugUtilities::GetDebugInfo(FDebugInfo& OutDebugInfo, UBlueprint* Blueprint, UObject* ActiveObject, const UEdGraphPin* WatchPin)
+{
+	UProperty* PropertyToDebug = nullptr;
+	void* DataPtr = nullptr;
+	void* DeltaPtr = nullptr;
+	UObject* ParentObj = nullptr;
+	FKismetDebugUtilities::EWatchTextResult Result = FindDebuggingData(Blueprint, ActiveObject, WatchPin, PropertyToDebug, DataPtr, DeltaPtr, ParentObj);
+
+	if (Result == FKismetDebugUtilities::EWatchTextResult::EWTR_Valid)
+	{
+		GetDebugInfo_InContainer(0, OutDebugInfo, PropertyToDebug, DataPtr);
+	}
+
+	return Result;
+}
+
+FKismetDebugUtilities::EWatchTextResult FKismetDebugUtilities::FindDebuggingData(UBlueprint* Blueprint, UObject* ActiveObject, const UEdGraphPin* WatchPin, UProperty*& OutProperty, void*& OutData, void*& OutDelta, UObject*& OutParent)
 {
 	FKismetDebugUtilitiesData& Data = FKismetDebugUtilitiesData::Get();
 
@@ -1161,7 +1202,7 @@ FKismetDebugUtilities::EWatchTextResult FKismetDebugUtilities::GetWatchText(FStr
 #if USE_UBER_GRAPH_PERSISTENT_FRAME
 			// Try find the propertybase in the persistent ubergraph frame
 			UFunction* OuterFunction = Cast<UFunction>(Property->GetOuter());
-			if(!PropertyBase && OuterFunction)
+			if (!PropertyBase && OuterFunction)
 			{
 				UBlueprintGeneratedClass* BPGC = Cast<UBlueprintGeneratedClass>(Blueprint->GeneratedClass);
 				if (BPGC && ActiveObject->IsA(BPGC))
@@ -1170,15 +1211,15 @@ FKismetDebugUtilities::EWatchTextResult FKismetDebugUtilities::GetWatchText(FStr
 				}
 			}
 #endif // USE_UBER_GRAPH_PERSISTENT_FRAME
-			
+
 			// see if our WatchPin is on a animation node & if so try to get its property info
 			UAnimBlueprintGeneratedClass* AnimBlueprintGeneratedClass = Cast<UAnimBlueprintGeneratedClass>(Blueprint->GeneratedClass);
-			if(!PropertyBase && AnimBlueprintGeneratedClass)
+			if (!PropertyBase && AnimBlueprintGeneratedClass)
 			{
 				// are we linked to an anim graph node?
 				UProperty* LinkedProperty = Property;
 				const UAnimGraphNode_Base* Node = Cast<UAnimGraphNode_Base>(WatchPin->GetOuter());
-				if(Node == nullptr && WatchPin->LinkedTo.Num() > 0)
+				if (Node == nullptr && WatchPin->LinkedTo.Num() > 0)
 				{
 					const UEdGraphPin* LinkedPin = WatchPin->LinkedTo[0];
 					// When we change Node we *must* change Property, so it's still a sub-element of that.
@@ -1186,7 +1227,7 @@ FKismetDebugUtilities::EWatchTextResult FKismetDebugUtilities::GetWatchText(FStr
 					Node = Cast<UAnimGraphNode_Base>(LinkedPin->GetOuter());
 				}
 
-				if(Node && LinkedProperty)
+				if (Node && LinkedProperty)
 				{
 					UStructProperty* NodeStructProperty = Cast<UStructProperty>(FKismetDebugUtilities::FindClassPropertyForNode(Blueprint, Node));
 					if (NodeStructProperty)
@@ -1196,7 +1237,10 @@ FKismetDebugUtilities::EWatchTextResult FKismetDebugUtilities::GetWatchText(FStr
 							if (NodeProperty == NodeStructProperty)
 							{
 								void* NodePtr = NodeProperty->ContainerPtrToValuePtr<void>(ActiveObject);
-								LinkedProperty->ExportText_InContainer(/*ArrayElement=*/ 0, /*inout*/ OutWatchText, NodePtr, NodePtr, /*Parent=*/ ActiveObject, PPF_PropertyWindow|PPF_BlueprintDebugView);
+								OutProperty = LinkedProperty;
+								OutData = NodePtr;
+								OutDelta = NodePtr;
+								OutParent = ActiveObject;
 								return EWTR_Valid;
 							}
 						}
@@ -1207,7 +1251,10 @@ FKismetDebugUtilities::EWatchTextResult FKismetDebugUtilities::GetWatchText(FStr
 			// Now either print out the variable value, or that it was out-of-scope
 			if (PropertyBase != nullptr)
 			{
-				Property->ExportText_InContainer(/*ArrayElement=*/ 0, /*inout*/ OutWatchText, PropertyBase, PropertyBase, /*Parent=*/ ActiveObject, PPF_PropertyWindow|PPF_BlueprintDebugView);
+				OutProperty = Property;
+				OutData = PropertyBase;
+				OutDelta = PropertyBase;
+				OutParent = ActiveObject;
 				return EWTR_Valid;
 			}
 			else
@@ -1226,6 +1273,255 @@ FKismetDebugUtilities::EWatchTextResult FKismetDebugUtilities::GetWatchText(FStr
 	}
 }
 
+void FKismetDebugUtilities::GetDebugInfo_InContainer(int32 Index, FDebugInfo& DebugInfo, UProperty* Property, const void* Data)
+{
+	GetDebugInfoInternal(DebugInfo, Property, Property->ContainerPtrToValuePtr<void>(Data, Index));
+}
+
+void FKismetDebugUtilities::GetDebugInfoInternal(FDebugInfo& DebugInfo, UProperty* Property, const void* PropertyValue)
+{
+	if (Property == nullptr)
+	{
+		return;
+	}
+
+	DebugInfo.Type = UEdGraphSchema_K2::TypeToText(Property);
+	DebugInfo.DisplayName = Property->GetDisplayNameText();
+
+	UByteProperty* ByteProperty = Cast<UByteProperty>(Property);
+	if (ByteProperty)
+	{
+		UEnum* Enum = ByteProperty->GetIntPropertyEnum();
+		if (Enum)
+		{
+			if (Enum->IsValidEnumValue(*(const uint8*)PropertyValue))
+			{
+				DebugInfo.Value = Enum->GetDisplayNameTextByValue(*(const uint8*)PropertyValue);
+			}
+			else
+			{
+				DebugInfo.Value = FText::FromString(TEXT("(INVALID)"));
+			}
+
+			return;
+		}
+
+		// if there is no Enum we need to fall through and treat this as a UNumericProperty
+	}
+
+	UNumericProperty* NumericProperty = Cast<UNumericProperty>(Property);
+	if (NumericProperty)
+	{
+		DebugInfo.Value = FText::FromString(NumericProperty->GetNumericPropertyValueToString(PropertyValue));
+		return;
+	}
+
+	UBoolProperty* BoolProperty = Cast<UBoolProperty>(Property);
+	if (BoolProperty)
+	{
+		DebugInfo.Value = BoolProperty->GetPropertyValue(PropertyValue) ? GTrue : GFalse;
+		return;
+	}
+
+	UNameProperty* NameProperty = Cast<UNameProperty>(Property);
+	if (NameProperty)
+	{
+		DebugInfo.Value = FText::FromName(*(FName*)PropertyValue);
+		return;
+	}
+
+	UTextProperty* TextProperty = Cast<UTextProperty>(Property);
+	if (TextProperty)
+	{
+		DebugInfo.Value = TextProperty->GetPropertyValue(PropertyValue);
+		return;
+	}
+
+	UStrProperty* StringProperty = Cast<UStrProperty>(Property);
+	if (StringProperty)
+	{
+		DebugInfo.Value = FText::FromString(StringProperty->GetPropertyValue(PropertyValue));
+		return;
+	}
+
+	UArrayProperty* ArrayProperty = Cast<UArrayProperty>(Property);
+	if (ArrayProperty)
+	{
+		checkSlow(ArrayProperty->Inner);
+
+		FScriptArrayHelper ArrayHelper(ArrayProperty, PropertyValue);
+
+		DebugInfo.Value = FText::Format(LOCTEXT("ArraySize", "Num={0}"), FText::AsNumber(ArrayHelper.Num()));
+
+		for (int32 i = 0; i < ArrayHelper.Num(); i++)
+		{
+			FDebugInfo ArrayDebugInfo;
+
+			uint8* PropData = ArrayHelper.GetRawPtr(i);
+			GetDebugInfoInternal(ArrayDebugInfo, ArrayProperty->Inner, PropData);
+			// overwrite the display name with the array index for the current element
+			ArrayDebugInfo.DisplayName = FText::Format(LOCTEXT("ArrayIndexName", "[{0}]"), FText::AsNumber(i));
+			DebugInfo.Children.Add(ArrayDebugInfo);
+		}
+
+		return;
+	}
+
+	UStructProperty* StructProperty = Cast<UStructProperty>(Property);
+	if (StructProperty)
+	{
+		FString WatchText;
+		StructProperty->ExportTextItem(WatchText, PropertyValue, PropertyValue, nullptr, PPF_PropertyWindow | PPF_BlueprintDebugView, nullptr);
+		DebugInfo.Value = FText::FromString(WatchText);
+
+		for (TFieldIterator<UProperty> It(StructProperty->Struct); It; ++It)
+		{
+			FDebugInfo StructDebugInfo;
+			GetDebugInfoInternal(StructDebugInfo, *It, It->ContainerPtrToValuePtr<void>(PropertyValue, 0));
+
+			DebugInfo.Children.Add(StructDebugInfo);
+		}
+
+		return;
+	}
+
+	UEnumProperty* EnumProperty = Cast<UEnumProperty>(Property);
+	if (EnumProperty)
+	{
+		UNumericProperty* LocalUnderlyingProp = EnumProperty->GetUnderlyingProperty();
+		UEnum* Enum = EnumProperty->GetEnum();
+
+		int64 Value = LocalUnderlyingProp->GetSignedIntPropertyValue(PropertyValue);
+
+		// if the value is the max value (the autogenerated *_MAX value), export as "INVALID", unless we're exporting text for copy/paste (for copy/paste,
+		// the property text value must actually match an entry in the enum's names array)
+		if (Enum)
+		{
+			if (Enum->IsValidEnumValue(Value))
+			{
+				DebugInfo.Value = Enum->GetDisplayNameTextByValue(Value);
+			}
+			else
+			{
+				DebugInfo.Value = LOCTEXT("Invalid", "(INVALID)");
+			}
+		}
+		else
+		{
+			DebugInfo.Value = FText::AsNumber(Value);
+		}
+
+		return;
+	}
+
+	UMapProperty* MapProperty = Cast<UMapProperty>(Property);
+	if (MapProperty)
+	{
+		FScriptMapHelper MapHelper(MapProperty, PropertyValue);
+		DebugInfo.Value = FText::Format(LOCTEXT("MapSize", "Num={0}"), FText::AsNumber(MapHelper.Num()));
+		uint8* PropData = MapHelper.GetPairPtr(0);
+
+		int32 Index = 0;
+		for (int32 Count = MapHelper.Num(); Count; PropData += MapProperty->MapLayout.SetLayout.Size, ++Index)
+		{
+			if (MapHelper.IsValidIndex(Index))
+			{
+				FDebugInfo ChildInfo;
+
+				GetDebugInfoInternal(ChildInfo, MapProperty->ValueProp, PropData + MapProperty->MapLayout.ValueOffset);
+
+				// use the info from the ValueProp and then overwrite the name with the KeyProp data
+				FString NameStr = TEXT("[");
+				MapProperty->KeyProp->ExportTextItem(NameStr, PropData, nullptr, nullptr, PPF_PropertyWindow | PPF_BlueprintDebugView | PPF_Delimited, nullptr);
+				NameStr += TEXT("] ");
+
+				ChildInfo.DisplayName = FText::FromString(NameStr);
+
+				DebugInfo.Children.Add(ChildInfo);
+
+				--Count;
+			}
+		}
+
+		return;
+	}
+
+	USetProperty* SetProperty = Cast<USetProperty>(Property);
+	if (SetProperty)
+	{
+		FScriptSetHelper SetHelper(SetProperty, PropertyValue);
+		DebugInfo.Value = FText::Format(LOCTEXT("SetSize", "Num={0}"), FText::AsNumber(SetHelper.Num()));
+		uint8* PropData = SetHelper.GetElementPtr(0);
+
+		int32 Index = 0;
+		for (int32 Count = SetHelper.Num(); Count; PropData += SetProperty->SetLayout.Size, ++Index)
+		{
+			if (SetHelper.IsValidIndex(Index))
+			{
+				FDebugInfo ChildInfo;
+				GetDebugInfoInternal(ChildInfo, SetProperty->ElementProp, PropData);
+
+				// members of sets don't have their own names
+				ChildInfo.DisplayName = FText::GetEmpty();
+
+				DebugInfo.Children.Add(ChildInfo);
+
+				--Count;
+			}
+		}
+
+		return;
+	}
+
+	UObjectPropertyBase* ObjectPropertyBase = Cast<UObjectPropertyBase>(Property);
+	if (ObjectPropertyBase)
+	{
+		UObject* Obj = ObjectPropertyBase->GetObjectPropertyValue(PropertyValue);
+		if (Obj != nullptr)
+		{
+			DebugInfo.Value = FText::FromString(Obj->GetFullName());
+		}
+		else
+		{
+			DebugInfo.Value = FText::FromString(TEXT("None"));
+		}
+
+		return;
+	}
+
+	UDelegateProperty* DelegateProperty = Cast<UDelegateProperty>(Property);
+	if (DelegateProperty)
+	{
+		if (DelegateProperty->SignatureFunction)
+		{
+			DebugInfo.Value = DelegateProperty->SignatureFunction->GetDisplayNameText();
+		}
+		else
+		{
+			DebugInfo.Value = LOCTEXT("NoFunc", "(No bound function)");
+		}
+
+		return;
+	}
+
+	UMulticastDelegateProperty* MulticastDelegateProperty = Cast<UMulticastDelegateProperty>(Property);
+	if (MulticastDelegateProperty)
+	{
+		if (MulticastDelegateProperty->SignatureFunction)
+		{
+			DebugInfo.Value = MulticastDelegateProperty->SignatureFunction->GetDisplayNameText();
+		}
+		else
+		{
+			DebugInfo.Value = LOCTEXT("NoFunc", "(No bound function)");
+		}
+
+		return;
+	}
+
+	ensure(false);
+}
+
 FText FKismetDebugUtilities::GetAndClearLastExceptionMessage()
 {
 	FKismetDebugUtilitiesData& Data = FKismetDebugUtilitiesData::Get();
diff --git a/Engine/Source/Editor/UnrealEd/Private/NavigationBuildingNotification.cpp b/Engine/Source/Editor/UnrealEd/Private/NavigationBuildingNotification.cpp
--- a/Engine/Source/Editor/UnrealEd/Private/NavigationBuildingNotification.cpp
+++ b/Engine/Source/Editor/UnrealEd/Private/NavigationBuildingNotification.cpp
@@ -1,7 +1,7 @@
 // Copyright 1998-2018 Epic Games, Inc. All Rights Reserved.
 
 #include "NavigationBuildingNotification.h"
-#include "AI/Navigation/NavigationSystem.h"
+#include "NavigationSystem.h"
 #include "Editor/EditorEngine.h"
 #include "Settings/LevelEditorMiscSettings.h"
 #include "EngineGlobals.h"
@@ -111,9 +111,10 @@ FText FNavigationBuildingNotificationImpl::GetNotificationText() const
 	if (EEngine)
 	{
 		FWorldContext &EditorContext = EEngine->GetEditorWorldContext();
-		if (EditorContext.World() && EditorContext.World()->GetNavigationSystem())
+		UNavigationSystemV1* NavSys = FNavigationSystem::GetCurrent<UNavigationSystemV1>(EditorContext.World());
+		if (NavSys)
 		{
-			RemainingTasks = EditorContext.World()->GetNavigationSystem()->GetNumRemainingBuildTasks();
+			RemainingTasks = NavSys->GetNumRemainingBuildTasks();
 		}
 	}
 		
@@ -136,9 +137,10 @@ void FNavigationBuildingNotificationImpl::Tick(float DeltaTime)
 		const bool bUserRequestedBuild = FEditorBuildUtils::IsBuildingNavigationFromUserRequest();
 		FWorldContext &EditorContext = EEngine->GetEditorWorldContext();
 		
-		const bool bBuildInProgress = EditorContext.World() != NULL && EditorContext.World()->GetNavigationSystem() != NULL 
-			&& EditorContext.World()->GetNavigationSystem()->IsNavigationBuildInProgress( GetDefault<ULevelEditorMiscSettings>()->bNavigationAutoUpdate ? true : false) == true
-			&& EditorContext.World()->GetNavigationSystem()->GetNumRemainingBuildTasks() > 0;
+		UNavigationSystemV1* NavSys = FNavigationSystem::GetCurrent<UNavigationSystemV1>(EditorContext.World());
+		const bool bBuildInProgress = NavSys != nullptr
+			&& NavSys->IsNavigationBuildInProgress( GetDefault<ULevelEditorMiscSettings>()->bNavigationAutoUpdate ? true : false) == true
+			&& NavSys->GetNumRemainingBuildTasks() > 0;
 
 		if (!bPreviouslyDetectedBuild && bBuildInProgress)
 		{
diff --git a/Engine/Source/Editor/UnrealEd/Private/PlayLevel.cpp b/Engine/Source/Editor/UnrealEd/Private/PlayLevel.cpp
--- a/Engine/Source/Editor/UnrealEd/Private/PlayLevel.cpp
+++ b/Engine/Source/Editor/UnrealEd/Private/PlayLevel.cpp
@@ -38,7 +38,7 @@
 #include "Engine/GameInstance.h"
 #include "Engine/World.h"
 #include "Settings/LevelEditorPlaySettings.h"
-#include "AI/Navigation/NavigationSystem.h"
+#include "AI/NavigationSystemBase.h"
 #include "Editor/EditorEngine.h"
 #include "Editor/UnrealEdEngine.h"
 #include "Settings/ProjectPackagingSettings.h"
@@ -370,10 +370,7 @@ void UEditorEngine::EndPlayMap()
 	}
 
 	// Lose the EditorWorld pointer (this is only maintained while PIEing)
-	if (EditorWorld->GetNavigationSystem())
-	{
-		EditorWorld->GetNavigationSystem()->OnPIEEnd();
-	}
+	FNavigationSystem::OnPIEEnd(*EditorWorld);
 
 	FGameDelegates::Get().GetEndPlayMapDelegate().Broadcast();
 
@@ -2255,10 +2252,7 @@ void UEditorEngine::PlayInEditor( UWorld* InWorld, bool bInSimulateInEditor )
 	FEditorDelegates::BeginPIE.Broadcast(bInSimulateInEditor);
 
 	// let navigation know PIE starts so it can avoid any blueprint creation/deletion/instantiation affect editor map's navmesh changes
-	if (InWorld->GetNavigationSystem())
-	{
-		InWorld->GetNavigationSystem()->OnPIEStart();
-	}
+	FNavigationSystem::OnPIEStart(*InWorld);
 
 	ULevelEditorPlaySettings* PlayInSettings = GetMutableDefault<ULevelEditorPlaySettings>();
 	check(PlayInSettings);
@@ -2282,11 +2276,7 @@ void UEditorEngine::PlayInEditor( UWorld* InWorld, bool bInSimulateInEditor )
 		if ( !bContinuePIE )
 		{
 			FEditorDelegates::EndPIE.Broadcast(bInSimulateInEditor);
-			if (InWorld->GetNavigationSystem())
-			{
-				InWorld->GetNavigationSystem()->OnPIEEnd();
-			}
-
+			FNavigationSystem::OnPIEEnd(*InWorld);
 			return;
 		}
 		else
@@ -2905,10 +2895,7 @@ UGameInstance* UEditorEngine::CreatePIEGameInstance(int32 InPIEInstance, bool bI
 
 		FEditorDelegates::EndPIE.Broadcast(bInSimulateInEditor);
 
-		if (EditorWorld->GetNavigationSystem())
-		{
-			EditorWorld->GetNavigationSystem()->OnPIEEnd();
-		}
+		FNavigationSystem::OnPIEEnd(*EditorWorld);
 
 		return nullptr;
 	}
diff --git a/Engine/Source/Editor/UnrealEd/Private/SKeySelector.cpp b/Engine/Source/Editor/UnrealEd/Private/SKeySelector.cpp
--- a/Engine/Source/Editor/UnrealEd/Private/SKeySelector.cpp
+++ b/Engine/Source/Editor/UnrealEd/Private/SKeySelector.cpp
@@ -110,7 +110,7 @@ void SKeySelector::Construct(const FArguments& InArgs)
 	// if we allow NoClear, add a "None" option to be able to clear out a binding
 	if (InArgs._AllowClear)
 	{
-		new (KeyTreeRoot) FKeyTreeItem(MakeShareable(new FKeyTreeInfo(FText(), MakeShareable(new FKey(NAME_None)))));
+		new (KeyTreeRoot) FKeyTreeItem(MakeShareable(new FKeyTreeInfo(FText(), MakeShareable(new FKey(EKeys::Invalid)))));
 	}
 
 	TreeViewWidth = InArgs._TreeViewWidth;
diff --git a/Engine/Source/Editor/UnrealEd/Private/Settings/SettingsClasses.cpp b/Engine/Source/Editor/UnrealEd/Private/Settings/SettingsClasses.cpp
--- a/Engine/Source/Editor/UnrealEd/Private/Settings/SettingsClasses.cpp
+++ b/Engine/Source/Editor/UnrealEd/Private/Settings/SettingsClasses.cpp
@@ -7,7 +7,7 @@
 #include "Misc/PackageName.h"
 #include "InputCoreTypes.h"
 #include "Classes/EditorStyleSettings.h"
-#include "AI/Navigation/NavigationSystem.h"
+#include "AI/NavigationSystemBase.h"
 #include "Model.h"
 #include "ISourceControlModule.h"
 #include "Settings/ContentBrowserSettings.h"
@@ -356,7 +356,7 @@ void ULevelEditorMiscSettings::PostEditChangeProperty( struct FPropertyChangedEv
 	if (Name == FName(TEXT("bNavigationAutoUpdate")))
 	{
 		FWorldContext &EditorContext = GEditor->GetEditorWorldContext();
-		UNavigationSystem::SetNavigationAutoUpdateEnabled(bNavigationAutoUpdate, EditorContext.World()->GetNavigationSystem());
+		FNavigationSystem::SetNavigationAutoUpdateEnabled(bNavigationAutoUpdate, EditorContext.World()->GetNavigationSystem());
 	}
 
 	if (!FUnrealEdMisc::Get().IsDeletePreferences())
diff --git a/Engine/Source/Editor/UnrealEd/Private/UnrealEdPrivatePCH.h b/Engine/Source/Editor/UnrealEd/Private/UnrealEdPrivatePCH.h
--- a/Engine/Source/Editor/UnrealEd/Private/UnrealEdPrivatePCH.h
+++ b/Engine/Source/Editor/UnrealEd/Private/UnrealEdPrivatePCH.h
@@ -484,7 +484,7 @@
 #include "Engine/Blueprint.h"
 #include "Engine/BlueprintCore.h"
 #include "AI/Navigation/NavigationTypes.h"
-#include "AI/Navigation/NavFilters/NavigationQueryFilter.h"
+#include "NavFilters/NavigationQueryFilter.h"
 #include "Engine/LatentActionManager.h"
 #include "Engine/World.h"
 #include "WorldCollision.h"
@@ -608,9 +608,8 @@
 #include "Engine/FontImportOptions.h"
 #include "Animation/AnimBlueprint.h"
 #include "EngineModule.h"
-#include "AI/Navigation/NavigationSystem.h"
+#include "AI/NavigationSystemBase.h"
 #include "GenericOctreePublic.h"
-#include "AI/Navigation/NavigationData.h"
 #include "CanvasTypes.h"
 #include "Sound/SoundBase.h"
 #include "Components/LightComponent.h"
diff --git a/Engine/Source/Editor/UnrealEd/Public/GraphEditor.h b/Engine/Source/Editor/UnrealEd/Public/GraphEditor.h
--- a/Engine/Source/Editor/UnrealEd/Public/GraphEditor.h
+++ b/Engine/Source/Editor/UnrealEd/Public/GraphEditor.h
@@ -174,12 +174,12 @@ public:
 		}
 	}
 
-	/* Set new viewer location */
-	virtual void SetViewLocation(const FVector2D& Location, float ZoomAmount)
+	/* Set new viewer location and optionally set the current bookmark */
+	virtual void SetViewLocation(const FVector2D& Location, float ZoomAmount, const FGuid& BookmarkId = FGuid())
 	{
 		if (Implementation.IsValid())
 		{
-			Implementation->SetViewLocation(Location, ZoomAmount);
+			Implementation->SetViewLocation(Location, ZoomAmount, BookmarkId);
 		}
 	}
 
@@ -197,6 +197,19 @@ public:
 		}
 	}
 
+	/**
+	 * Gets the current graph view bookmark
+	 *
+	 * @param OutBookmarkId		Will have the current bookmark ID
+	 */
+	virtual void GetViewBookmark(FGuid& OutBookmarkId)
+	{
+		if (Implementation.IsValid())
+		{
+			Implementation->GetViewBookmark(OutBookmarkId);
+		}
+	}
+
 	/** Check if node title is visible with optional flag to ensure it is */
 	virtual bool IsNodeTitleVisible(const class UEdGraphNode* Node, bool bRequestRename)
 	{
diff --git a/Engine/Source/Editor/UnrealEd/Public/Kismet2/BlueprintEditorUtils.h b/Engine/Source/Editor/UnrealEd/Public/Kismet2/BlueprintEditorUtils.h
--- a/Engine/Source/Editor/UnrealEd/Public/Kismet2/BlueprintEditorUtils.h
+++ b/Engine/Source/Editor/UnrealEd/Public/Kismet2/BlueprintEditorUtils.h
@@ -1566,7 +1566,10 @@ public:
 	/**
 	 * Returns true if this BP is currently based on a type that returns true for the UObject::ImplementsGetWorld() call:
 	 */
-	static bool ImplentsGetWorld(const UBlueprint* BP);
+	static bool ImplementsGetWorld(const UBlueprint* BP);
+
+	DEPRECATED(4.20, "Use ImplementsGetWorld")
+	static bool ImplentsGetWorld(const UBlueprint* BP) { return ImplementsGetWorld(BP); }
 };
 
 struct UNREALED_API FBlueprintDuplicationScopeFlags
diff --git a/Engine/Source/Editor/UnrealEd/Public/Kismet2/KismetDebugUtilities.h b/Engine/Source/Editor/UnrealEd/Public/Kismet2/KismetDebugUtilities.h
--- a/Engine/Source/Editor/UnrealEd/Public/Kismet2/KismetDebugUtilities.h
+++ b/Engine/Source/Editor/UnrealEd/Public/Kismet2/KismetDebugUtilities.h
@@ -24,6 +24,17 @@ struct FKismetTraceSample
 	double ObservationTime;
 };
 
+//////////////////////////////////////////////////////////////////////////
+// FDebugInfo
+
+struct FDebugInfo
+{
+	FText DisplayName;
+	FText Value;
+	FText Type;
+	TArray<FDebugInfo> Children;
+};
+
 //////////////////////////////////////////////////////////////////////////
 // FObjectsBeingDebuggedIterator
 
@@ -105,6 +116,9 @@ private:
 class UNREALED_API FKismetDebugUtilities
 {
 public:
+	// Delegate for when pins are added or removed from the watchlist
+	DECLARE_MULTICAST_DELEGATE_OneParam(FOnWatchedPinsListChanged, class UBlueprint*);
+
 	static void OnScriptException(const UObject* ActiveObject, const FFrame& StackFrame, const FBlueprintExceptionInfo& Info);
 
 	/** Returns the current instruction; if a PIE/SIE session is started but paused; otherwise NULL */
@@ -195,6 +209,9 @@ public:
 	/** Deletes all breakpoints in this blueprint */
 	static void ClearBreakpoints(UBlueprint* Blueprint);
 
+	// Notifies listeners when a watched pin is added or removed
+	static FOnWatchedPinsListChanged WatchedPinsListChangedEvent;
+
 	static bool CanWatchPin(const UBlueprint* Blueprint, const UEdGraphPin* Pin);
 	static bool IsPinBeingWatched(const UBlueprint* Blueprint, const UEdGraphPin* Pin);
 	static void TogglePinWatch(UBlueprint* Blueprint, const UEdGraphPin* Pin);
@@ -219,6 +236,9 @@ public:
 	// Gets the watched tooltip for a specified site
 	static EWatchTextResult GetWatchText(FString& OutWatchText, UBlueprint* Blueprint, UObject* ActiveObject, const UEdGraphPin* WatchPin);
 
+	// Gets the debug info for a specified site
+	static EWatchTextResult GetDebugInfo(FDebugInfo& OutDebugInfo, UBlueprint* Blueprint, UObject* ActiveObject, const UEdGraphPin* WatchPin);
+
 	//@TODO: Pretty lame way to handle this messaging, ideally the entire Info object gets pushed into the editor when intraframe debugging is triggered!
 	// This doesn't work properly if there is more than one blueprint editor open at once either (one will consume it, the others will be left in the cold)
 	static FText GetAndClearLastExceptionMessage();
@@ -226,6 +246,13 @@ protected:
 	static void CheckBreakConditions(UEdGraphNode* NodeStoppedAt, bool bHitBreakpoint, int32 BreakpointOffset, bool& InOutBreakExecution);
 	static void AttemptToBreakExecution(UBlueprint* BlueprintObj, const UObject* ActiveObject, const FFrame& StackFrame, const FBlueprintExceptionInfo& Info, UEdGraphNode* NodeStoppedAt, int32 DebugOpcodeOffset);
 
+	static void GetDebugInfo_InContainer(int32 Index, FDebugInfo& DebugInfo, UProperty* Property, const void* Data);
+	static void GetDebugInfoInternal(FDebugInfo& DebugInfo, UProperty* Property, const void* PropertyValue);
+
+	// helper function for converting between blueprint and debuggable data
+	// output params are only valid if the return result is EWatchTextResult::EWTR_Valid
+	static EWatchTextResult FindDebuggingData(UBlueprint* Blueprint, UObject* ActiveObject, const UEdGraphPin* WatchPin, UProperty*& OutProperty, void*& OutData, void*& OutDelta, UObject*& OutParent);
+
 private:
 	FKismetDebugUtilities() {}
 };
diff --git a/Engine/Source/Editor/UnrealEd/Public/UnrealEdSharedPCH.h b/Engine/Source/Editor/UnrealEd/Public/UnrealEdSharedPCH.h
--- a/Engine/Source/Editor/UnrealEd/Public/UnrealEdSharedPCH.h
+++ b/Engine/Source/Editor/UnrealEd/Public/UnrealEdSharedPCH.h
@@ -471,7 +471,7 @@
 #include "Engine/Level.h"
 #include "GameFramework/Actor.h"
 #include "EngineDefines.h"
-#include "AI/Navigation/NavFilters/NavigationQueryFilter.h"
+#include "NavFilters/NavigationQueryFilter.h"
 #include "AI/Navigation/NavigationTypes.h"
 #include "Engine/LatentActionManager.h"
 #include "Engine/Blueprint.h"
diff --git a/Engine/Source/Editor/UnrealEd/UnrealEd.Build.cs b/Engine/Source/Editor/UnrealEd/UnrealEd.Build.cs
--- a/Engine/Source/Editor/UnrealEd/UnrealEd.Build.cs
+++ b/Engine/Source/Editor/UnrealEd/UnrealEd.Build.cs
@@ -59,8 +59,9 @@ public class UnrealEd : ModuleRules
 				"ViewportInteraction",
 				"VREditor",
 				"Persona",
-                "PhysicsAssetEditor",
-                "ClothingSystemEditorInterface",
+				"PhysicsAssetEditor",
+				"ClothingSystemEditorInterface",
+				"NavigationSystem",
 			}
 		);
 
@@ -82,7 +83,6 @@ public class UnrealEd : ModuleRules
 				"EditorStyle",
 				"SourceControl",
 				"UnrealEdMessages",
-				"AIModule",
 				"GameplayDebugger",
 				"BlueprintGraph",
 				"Http",
@@ -93,6 +93,7 @@ public class UnrealEd : ModuleRules
 				"AudioEditor",
 				"NetworkFileSystem",
 				"UMG",
+				"NavigationSystem",
 			}
 		);
 
@@ -184,7 +185,7 @@ public class UnrealEd : ModuleRules
 				"LevelEditor",
 				"PackagesDialog",
 				"Persona",
-                "PhysicsAssetEditor",
+				"PhysicsAssetEditor",
 				"ProjectLauncher",
 				"DeviceManager",
 				"SettingsEditor",
@@ -231,15 +232,15 @@ public class UnrealEd : ModuleRules
 		{
 			DynamicallyLoadedModuleNames.Add("AndroidPlatformEditor");
 		}
-
-		CircularlyReferencedDependentModules.AddRange(
+        
+        CircularlyReferencedDependentModules.AddRange(
 			new string[]
 			{
 				"GraphEditor",
 				"Kismet",
 				"AudioEditor",
 				"ViewportInteraction",
-				"VREditor"
+				"VREditor",
 			}
 		);
 
@@ -255,6 +256,9 @@ public class UnrealEd : ModuleRules
 				"MeshUtilities",
 				"AssetTools",
 				"KismetCompiler",
+				"NavigationSystem",
+				"GameplayTasks",
+				"AIModule",
 			}
 			);
 
diff --git a/Engine/Source/Programs/IncludeTool/IncludeTool/Rules.cs b/Engine/Source/Programs/IncludeTool/IncludeTool/Rules.cs
--- a/Engine/Source/Programs/IncludeTool/IncludeTool/Rules.cs
+++ b/Engine/Source/Programs/IncludeTool/IncludeTool/Rules.cs
@@ -553,6 +553,10 @@ namespace IncludeTool
 			{
 				return true;
 			}
+			if(Markup.Type == PreprocessorMarkupType.Define && Markup.Tokens[0].Text == "DEPRECATED_FORGAME")
+			{
+				return true;
+			}
 			if(Markup.Type == PreprocessorMarkupType.Undef && Markup.Tokens[0].Text == "TEXT")
 			{
 				return true;
diff --git a/Engine/Source/Programs/UnrealBuildTool/Configuration/UEBuildModule.cs b/Engine/Source/Programs/UnrealBuildTool/Configuration/UEBuildModule.cs
--- a/Engine/Source/Programs/UnrealBuildTool/Configuration/UEBuildModule.cs
+++ b/Engine/Source/Programs/UnrealBuildTool/Configuration/UEBuildModule.cs
@@ -513,7 +513,7 @@ namespace UnrealBuildTool
 			bool bWithLegacyPublicIncludePaths
 			)
 		{
-			if (this.Type.IsGameModule())
+			if (!Rules.bTreatAsEngineModule)
 			{
 				Definitions.Add("DEPRECATED_FORGAME=DEPRECATED");
 			}
diff --git a/Engine/Source/Programs/UnrealBuildTool/Configuration/UEBuildModuleCPP.cs b/Engine/Source/Programs/UnrealBuildTool/Configuration/UEBuildModuleCPP.cs
--- a/Engine/Source/Programs/UnrealBuildTool/Configuration/UEBuildModuleCPP.cs
+++ b/Engine/Source/Programs/UnrealBuildTool/Configuration/UEBuildModuleCPP.cs
@@ -221,7 +221,6 @@ namespace UnrealBuildTool
 		/// </summary>
 		static readonly KeyValuePair<string, string>[] WhitelistedCircularDependencies =
 		{
-			new KeyValuePair<string, string>("Engine", "AIModule"),
 			new KeyValuePair<string, string>("Engine", "Landscape"),
 			new KeyValuePair<string, string>("Engine", "UMG"),
 			new KeyValuePair<string, string>("Engine", "GameplayTags"),
@@ -273,6 +272,7 @@ namespace UnrealBuildTool
             new KeyValuePair<string, string>("MeshPaintMode", "ViewportInteraction"),
             new KeyValuePair<string, string>("MeshPaintMode", "VREditor"),
             new KeyValuePair<string, string>("Sequencer", "ViewportInteraction"),
+            new KeyValuePair<string, string>("NavigationSystem", "UnrealEd"),
         };
 
 
diff --git a/Engine/Source/Programs/UnrealHeaderTool/Private/HeaderParser.cpp b/Engine/Source/Programs/UnrealHeaderTool/Private/HeaderParser.cpp
--- a/Engine/Source/Programs/UnrealHeaderTool/Private/HeaderParser.cpp
+++ b/Engine/Source/Programs/UnrealHeaderTool/Private/HeaderParser.cpp
@@ -951,6 +951,18 @@ namespace
 			}
 			break;
 
+			case ECheckedMetadataSpecifier::ExpandEnumAsExecs:
+			{
+				if (UFunction* Function = Cast<UFunction>(Field))
+				{
+					if (FHeaderParser::FindField(Function, *InValue, false) == nullptr)
+					{
+						UE_LOG_ERROR_UHT(TEXT("Function does not have a parameter named '%s'"), *InValue);
+					}
+				}
+			}
+			break;
+
 			case ECheckedMetadataSpecifier::DevelopmentStatus:
 			{
 				const FString EarlyAccessValue(TEXT("EarlyAccess"));
diff --git a/Engine/Source/Programs/UnrealHeaderTool/Private/HeaderParser.h b/Engine/Source/Programs/UnrealHeaderTool/Private/HeaderParser.h
--- a/Engine/Source/Programs/UnrealHeaderTool/Private/HeaderParser.h
+++ b/Engine/Source/Programs/UnrealHeaderTool/Private/HeaderParser.h
@@ -434,6 +434,20 @@ public:
 	static void RequireSpecifierValue(const FPropertySpecifier& Specifier, bool bRequireExactlyOne = false);
 	static FString RequireExactlyOneSpecifierValue(const FPropertySpecifier& Specifier);
 
+	/**
+	* Find a field in the specified context.  Starts with the specified scope, then iterates
+	* through the Outer chain until the field is found.
+	* 
+	* @param	InScope				scope to start searching for the field in 
+	* @param	InIdentifier		name of the field we're searching for
+	* @param	bIncludeParents		whether to allow searching in the scope of a parent struct
+	* @param	FieldClass			class of the field to search for.  used to e.g. search for functions only
+	* @param	Thing				hint text that will be used in the error message if an error is encountered
+	*
+	* @return	a pointer to a UField with a name matching InIdentifier, or NULL if it wasn't found
+	*/
+	static UField* FindField( UStruct* InScope, const TCHAR* InIdentifier, bool bIncludeParents=true, UClass* FieldClass=UField::StaticClass(), const TCHAR* Thing=nullptr );
+
 protected:
 
 	/**
@@ -543,19 +557,6 @@ protected:
 
 	UStruct* GetSuperScope( UStruct* CurrentScope, const FName& SearchName );
 
-	/**
-	 * Find a field in the specified context.  Starts with the specified scope, then iterates
-	 * through the Outer chain until the field is found.
-	 * 
-	 * @param	InScope				scope to start searching for the field in 
-	 * @param	InIdentifier		name of the field we're searching for
-	 * @param	bIncludeParents		whether to allow searching in the scope of a parent struct
-	 * @param	FieldClass			class of the field to search for.  used to e.g. search for functions only
-	 * @param	Thing				hint text that will be used in the error message if an error is encountered
-	 *
-	 * @return	a pointer to a UField with a name matching InIdentifier, or NULL if it wasn't found
-	 */
-	UField* FindField( UStruct* InScope, const TCHAR* InIdentifier, bool bIncludeParents=true, UClass* FieldClass=UField::StaticClass(), const TCHAR* Thing=NULL );
 	void SkipStatements( int32 SubCount, const TCHAR* ErrorTag );
 
 	/**
diff --git a/Engine/Source/Programs/UnrealHeaderTool/Private/Specifiers/CheckedMetadataSpecifiers.def b/Engine/Source/Programs/UnrealHeaderTool/Private/Specifiers/CheckedMetadataSpecifiers.def
--- a/Engine/Source/Programs/UnrealHeaderTool/Private/Specifiers/CheckedMetadataSpecifiers.def
+++ b/Engine/Source/Programs/UnrealHeaderTool/Private/Specifiers/CheckedMetadataSpecifiers.def
@@ -6,6 +6,7 @@ CHECKED_METADATA_SPECIFIER(BlueprintProtected)
 CHECKED_METADATA_SPECIFIER(ClampMax)
 CHECKED_METADATA_SPECIFIER(ClampMin)
 CHECKED_METADATA_SPECIFIER(DevelopmentStatus)
+CHECKED_METADATA_SPECIFIER(ExpandEnumAsExecs)
 CHECKED_METADATA_SPECIFIER(UIMax)
 CHECKED_METADATA_SPECIFIER(UIMin)
 CHECKED_METADATA_SPECIFIER(Units)
diff --git a/Engine/Source/Runtime/AIModule/AIModule.Build.cs b/Engine/Source/Runtime/AIModule/AIModule.Build.cs
--- a/Engine/Source/Runtime/AIModule/AIModule.Build.cs
+++ b/Engine/Source/Runtime/AIModule/AIModule.Build.cs
@@ -2,54 +2,56 @@
 
 namespace UnrealBuildTool.Rules
 {
-	public class AIModule : ModuleRules
-	{
+    public class AIModule : ModuleRules
+    {
         public AIModule(ReadOnlyTargetRules Target) : base(Target)
-		{
-			PublicIncludePaths.AddRange(
-				new string[] {
+        {
+            PublicIncludePaths.AddRange(
+                new string[] {
+                    "Runtime/NavigationSystem/Public",
                     "Runtime/AIModule/Public",
-				}
-				);
+                }
+                );
 
-			PrivateIncludePaths.AddRange(
-				new string[] {
-					"Runtime/AIModule/Private",
+            PrivateIncludePaths.AddRange(
+                new string[] {
+                    "Runtime/AIModule/Private",
                     "Runtime/Engine/Private",
-				}
-				);
+                }
+                );
 
-			PublicDependencyModuleNames.AddRange(
-				new string[] {
-					"Core",
-					"CoreUObject",
-					"Engine",    
-               		"GameplayTags",
-                    "GameplayTasks"
-				}
-				);
+            PublicDependencyModuleNames.AddRange(
+                new string[] {
+                    "Core",
+                    "CoreUObject",
+                    "Engine",
+                    "GameplayTags",
+                    "GameplayTasks",
+                    "NavigationSystem",
+                }
+                );
 
-			PrivateDependencyModuleNames.AddRange(
-				new string[] {
-					"RHI",
+            PrivateDependencyModuleNames.AddRange(
+                new string[] {
+                    "RHI",
                     "RenderCore",
                     "ShaderCore",
-				}
-				);
+                }
+                );
 
-			DynamicallyLoadedModuleNames.AddRange(
-				new string[] {
+            DynamicallyLoadedModuleNames.AddRange(
+                new string[] {
 					// ... add any modules that your module loads dynamically here ...
 				}
-				);
+                );
 
-			if (Target.bBuildEditor == true)
-			{
-				PrivateDependencyModuleNames.Add("UnrealEd");
+            if (Target.bBuildEditor == true)
+            {
+                PrivateDependencyModuleNames.Add("UnrealEd");
 
-                PrivateDependencyModuleNames.Add("AITestSuite");                
+                PrivateDependencyModuleNames.Add("AITestSuite");
                 CircularlyReferencedDependentModules.Add("AITestSuite");
-			}
+            }
 
             if (Target.bCompileRecast)
             {
@@ -64,15 +66,15 @@ namespace UnrealBuildTool.Rules
                 PublicDefinitions.Add("WITH_RECAST=0");
             }
 
-			if (Target.bBuildDeveloperTools || (Target.Configuration != UnrealTargetConfiguration.Shipping && Target.Configuration != UnrealTargetConfiguration.Test))
-			{
-				PrivateDependencyModuleNames.Add("GameplayDebugger");
-				PublicDefinitions.Add("WITH_GAMEPLAY_DEBUGGER=1");
-			}
-			else
-			{
-				PublicDefinitions.Add("WITH_GAMEPLAY_DEBUGGER=0");
-			}
-		}
+            if (Target.bBuildDeveloperTools || (Target.Configuration != UnrealTargetConfiguration.Shipping && Target.Configuration != UnrealTargetConfiguration.Test))
+            {
+                PrivateDependencyModuleNames.Add("GameplayDebugger");
+                PublicDefinitions.Add("WITH_GAMEPLAY_DEBUGGER=1");
+            }
+            else
+            {
+                PublicDefinitions.Add("WITH_GAMEPLAY_DEBUGGER=0");
+            }
+        }
     }
 }
diff --git a/Engine/Source/Runtime/AIModule/Classes/AIController.h b/Engine/Source/Runtime/AIModule/Classes/AIController.h
--- a/Engine/Source/Runtime/AIModule/Classes/AIController.h
+++ b/Engine/Source/Runtime/AIModule/Classes/AIController.h
@@ -7,7 +7,7 @@
 #include "UObject/UObjectGlobals.h"
 #include "Templates/SubclassOf.h"
 #include "EngineDefines.h"
-#include "AI/Navigation/NavFilters/NavigationQueryFilter.h"
+#include "NavFilters/NavigationQueryFilter.h"
 #include "AITypes.h"
 #include "GameplayTaskOwnerInterface.h"
 #include "GameplayTask.h"
@@ -254,10 +254,7 @@ public:
 	/** Updates state of movement block detection. */
 	UFUNCTION(BlueprintCallable, Category = "AI|Navigation")
 	void SetMoveBlockDetection(bool bEnable);
-
-	/** Prepares path finding and path following components. */
-	virtual void InitNavigationControl(UPathFollowingComponent*& PathFollowingComp) override;
-
+	
 	/** Starts executing behavior tree. */
 	UFUNCTION(BlueprintCallable, Category = "AI")
 	virtual bool RunBehaviorTree(UBehaviorTree* BTAsset);
@@ -384,6 +381,7 @@ public:
 	// INavAgentInterface
 	//----------------------------------------------------------------------//
 	virtual bool IsFollowingAPath() const override;
+	virtual IPathFollowingAgentInterface* GetPathFollowingAgent() const override { return PathFollowingComponent; }
 
 	//----------------------------------------------------------------------//
 	// IGenericTeamAgentInterface
diff --git a/Engine/Source/Runtime/AIModule/Classes/AITypes.h b/Engine/Source/Runtime/AIModule/Classes/AITypes.h
--- a/Engine/Source/Runtime/AIModule/Classes/AITypes.h
+++ b/Engine/Source/Runtime/AIModule/Classes/AITypes.h
@@ -6,8 +6,9 @@
 #include "Stats/Stats.h"
 #include "UObject/ObjectMacros.h"
 #include "Templates/SubclassOf.h"
-#include "AI/Navigation/NavFilters/NavigationQueryFilter.h"
+#include "NavFilters/NavigationQueryFilter.h"
 #include "AI/Navigation/NavigationTypes.h"
+#include "NavigationSystemTypes.h"
 #include "GameFramework/Actor.h"
 #include "AITypes.generated.h"
 
diff --git a/Engine/Source/Runtime/AIModule/Classes/Actions/PawnAction_Move.h b/Engine/Source/Runtime/AIModule/Classes/Actions/PawnAction_Move.h
--- a/Engine/Source/Runtime/AIModule/Classes/Actions/PawnAction_Move.h
+++ b/Engine/Source/Runtime/AIModule/Classes/Actions/PawnAction_Move.h
@@ -6,7 +6,7 @@
 #include "UObject/ObjectMacros.h"
 #include "Templates/SubclassOf.h"
 #include "Engine/EngineTypes.h"
-#include "AI/Navigation/NavFilters/NavigationQueryFilter.h"
+#include "NavFilters/NavigationQueryFilter.h"
 #include "Actions/PawnAction.h"
 #include "Navigation/PathFollowingComponent.h"
 #include "PawnAction_Move.generated.h"
diff --git a/Engine/Source/Runtime/AIModule/Classes/Blueprint/AIBlueprintHelperLibrary.h b/Engine/Source/Runtime/AIModule/Classes/Blueprint/AIBlueprintHelperLibrary.h
--- a/Engine/Source/Runtime/AIModule/Classes/Blueprint/AIBlueprintHelperLibrary.h
+++ b/Engine/Source/Runtime/AIModule/Classes/Blueprint/AIBlueprintHelperLibrary.h
@@ -65,4 +65,9 @@ class AIMODULE_API UAIBlueprintHelperLibrary : public UBlueprintFunctionLibrary
 	UFUNCTION(BlueprintPure, Category = "AI", meta = (UnsafeDuringActorConstruction = "true"))
 	static UNavigationPath* GetCurrentPath(AController* Controller);
 
+	UFUNCTION(BlueprintCallable, Category = "AI|Navigation")
+	static void SimpleMoveToActor(AController* Controller, const AActor* Goal);
+
+	UFUNCTION(BlueprintCallable, Category = "AI|Navigation")
+	static void SimpleMoveToLocation(AController* Controller, const FVector& Goal);
 };
diff --git a/Engine/Source/Runtime/AIModule/Classes/DataProviders/AIDataProvider.h b/Engine/Source/Runtime/AIModule/Classes/DataProviders/AIDataProvider.h
--- a/Engine/Source/Runtime/AIModule/Classes/DataProviders/AIDataProvider.h
+++ b/Engine/Source/Runtime/AIModule/Classes/DataProviders/AIDataProvider.h
@@ -52,7 +52,10 @@ public:
 
 	/** return raw data from provider's property */
 	template<typename T>
-	T* GetRawValuePtr() const;
+	T* GetRawValuePtr() const
+	{
+		return CachedProperty ? CachedProperty->ContainerPtrToValuePtr<T>(DataBinding) : nullptr;
+	}
 
 	/** bind data in provider and cache property for faster access */
 	void BindData(const UObject* Owner, int32 RequestId) const;
diff --git a/Engine/Source/Runtime/AIModule/Classes/EnvironmentQuery/Tests/EnvQueryTest_Pathfinding.h b/Engine/Source/Runtime/AIModule/Classes/EnvironmentQuery/Tests/EnvQueryTest_Pathfinding.h
--- a/Engine/Source/Runtime/AIModule/Classes/EnvironmentQuery/Tests/EnvQueryTest_Pathfinding.h
+++ b/Engine/Source/Runtime/AIModule/Classes/EnvironmentQuery/Tests/EnvQueryTest_Pathfinding.h
@@ -5,15 +5,16 @@
 #include "CoreMinimal.h"
 #include "UObject/ObjectMacros.h"
 #include "Templates/SubclassOf.h"
-#include "AI/Navigation/NavFilters/NavigationQueryFilter.h"
+#include "NavFilters/NavigationQueryFilter.h"
 #include "AI/Navigation/NavigationTypes.h"
+#include "NavigationSystemTypes.h"
 #include "EnvironmentQuery/EnvQueryContext.h"
 #include "DataProviders/AIDataProvider.h"
 #include "EnvironmentQuery/EnvQueryTest.h"
 #include "EnvQueryTest_Pathfinding.generated.h"
 
 class ANavigationData;
-class UNavigationSystem;
+class UNavigationSystemV1;
 
 UENUM()
 namespace EEnvTestPathfinding
@@ -64,15 +65,15 @@ class UEnvQueryTest_Pathfinding : public UEnvQueryTest
 
 protected:
 
-	DECLARE_DELEGATE_RetVal_SevenParams(bool, FTestPathSignature, const FVector&, const FVector&, EPathFindingMode::Type, const ANavigationData&, UNavigationSystem&, FSharedConstNavQueryFilter, const UObject*);
-	DECLARE_DELEGATE_RetVal_SevenParams(float, FFindPathSignature, const FVector&, const FVector&, EPathFindingMode::Type, const ANavigationData&, UNavigationSystem&, FSharedConstNavQueryFilter, const UObject*);
+	DECLARE_DELEGATE_RetVal_SevenParams(bool, FTestPathSignature, const FVector&, const FVector&, EPathFindingMode::Type, const ANavigationData&, UNavigationSystemV1&, FSharedConstNavQueryFilter, const UObject*);
+	DECLARE_DELEGATE_RetVal_SevenParams(float, FFindPathSignature, const FVector&, const FVector&, EPathFindingMode::Type, const ANavigationData&, UNavigationSystemV1&, FSharedConstNavQueryFilter, const UObject*);
 
-	bool TestPathFrom(const FVector& ItemPos, const FVector& ContextPos, EPathFindingMode::Type Mode, const ANavigationData& NavData, UNavigationSystem& NavSys, FSharedConstNavQueryFilter NavFilter, const UObject* PathOwner) const;
-	bool TestPathTo(const FVector& ItemPos, const FVector& ContextPos, EPathFindingMode::Type Mode, const ANavigationData& NavData, UNavigationSystem& NavSys, FSharedConstNavQueryFilter NavFilter, const UObject* PathOwner) const;
-	float FindPathCostFrom(const FVector& ItemPos, const FVector& ContextPos, EPathFindingMode::Type Mode, const ANavigationData& NavData, UNavigationSystem& NavSys, FSharedConstNavQueryFilter NavFilter, const UObject* PathOwner) const;
-	float FindPathCostTo(const FVector& ItemPos, const FVector& ContextPos, EPathFindingMode::Type Mode, const ANavigationData& NavData, UNavigationSystem& NavSys, FSharedConstNavQueryFilter NavFilter, const UObject* PathOwner) const;
-	float FindPathLengthFrom(const FVector& ItemPos, const FVector& ContextPos, EPathFindingMode::Type Mode, const ANavigationData& NavData, UNavigationSystem& NavSys, FSharedConstNavQueryFilter NavFilter, const UObject* PathOwner) const;
-	float FindPathLengthTo(const FVector& ItemPos, const FVector& ContextPos, EPathFindingMode::Type Mode, const ANavigationData& NavData, UNavigationSystem& NavSys, FSharedConstNavQueryFilter NavFilter, const UObject* PathOwner) const;
+	bool TestPathFrom(const FVector& ItemPos, const FVector& ContextPos, EPathFindingMode::Type Mode, const ANavigationData& NavData, UNavigationSystemV1& NavSys, FSharedConstNavQueryFilter NavFilter, const UObject* PathOwner) const;
+	bool TestPathTo(const FVector& ItemPos, const FVector& ContextPos, EPathFindingMode::Type Mode, const ANavigationData& NavData, UNavigationSystemV1& NavSys, FSharedConstNavQueryFilter NavFilter, const UObject* PathOwner) const;
+	float FindPathCostFrom(const FVector& ItemPos, const FVector& ContextPos, EPathFindingMode::Type Mode, const ANavigationData& NavData, UNavigationSystemV1& NavSys, FSharedConstNavQueryFilter NavFilter, const UObject* PathOwner) const;
+	float FindPathCostTo(const FVector& ItemPos, const FVector& ContextPos, EPathFindingMode::Type Mode, const ANavigationData& NavData, UNavigationSystemV1& NavSys, FSharedConstNavQueryFilter NavFilter, const UObject* PathOwner) const;
+	float FindPathLengthFrom(const FVector& ItemPos, const FVector& ContextPos, EPathFindingMode::Type Mode, const ANavigationData& NavData, UNavigationSystemV1& NavSys, FSharedConstNavQueryFilter NavFilter, const UObject* PathOwner) const;
+	float FindPathLengthTo(const FVector& ItemPos, const FVector& ContextPos, EPathFindingMode::Type Mode, const ANavigationData& NavData, UNavigationSystemV1& NavSys, FSharedConstNavQueryFilter NavFilter, const UObject* PathOwner) const;
 
-	ANavigationData* FindNavigationData(UNavigationSystem& NavSys, UObject* Owner) const;
+	ANavigationData* FindNavigationData(UNavigationSystemV1& NavSys, UObject* Owner) const;
 };
diff --git a/Engine/Source/Runtime/AIModule/Classes/Navigation/CrowdManager.h b/Engine/Source/Runtime/AIModule/Classes/Navigation/CrowdManager.h
--- a/Engine/Source/Runtime/AIModule/Classes/Navigation/CrowdManager.h
+++ b/Engine/Source/Runtime/AIModule/Classes/Navigation/CrowdManager.h
@@ -2,7 +2,7 @@
 
 #pragma once
 
-#include "AI/Navigation/CrowdManagerBase.h"
+#include "CrowdManagerBase.h"
 #include "CoreMinimal.h"
 #include "Stats/Stats.h"
 #include "UObject/ObjectMacros.h"
@@ -10,6 +10,7 @@
 #include "AI/Navigation/NavigationTypes.h"
 #include "Tickable.h"
 #include "DrawDebugHelpers.h"
+#include "NavFilters/NavigationQueryFilter.h"
 #include "CrowdManager.generated.h"
 
 class ANavigationData;
@@ -39,7 +40,7 @@ class dtObstacleAvoidanceDebugData;
  *  but it's slower than AvoidanceManager solution (RVO, cares only about agents)
  *
  *  All agents will operate on the same navmesh data, which will be picked from
- *  navigation system defaults (UNavigationSystem::SupportedAgents[0])
+ *  navigation system defaults (UNavigationSystemV1::SupportedAgents[0])
  *
  *  To use it, you have to add CrowdFollowingComponent to your agent
  *  (usually: replace class of PathFollowingComponent in AIController by adding 
diff --git a/Engine/Source/Runtime/AIModule/Classes/Navigation/MetaNavMeshPath.h b/Engine/Source/Runtime/AIModule/Classes/Navigation/MetaNavMeshPath.h
--- a/Engine/Source/Runtime/AIModule/Classes/Navigation/MetaNavMeshPath.h
+++ b/Engine/Source/Runtime/AIModule/Classes/Navigation/MetaNavMeshPath.h
@@ -5,7 +5,7 @@
 #include "CoreMinimal.h"
 #include "EngineDefines.h"
 #include "AI/Navigation/NavigationTypes.h"
-#include "AI/Navigation/RecastNavMesh.h"
+#include "NavMesh/RecastNavMesh.h"
 
 class AController;
 class UCanvas;
