Copying //UE4/Dev-Enterprise to //UE4/Dev-Main (Source: //UE4/Dev-Enterprise @ 3972172)

diff --git a/Engine/Source/Editor/UnrealEd/Classes/Editor/Transactor.h b/Engine/Source/Editor/UnrealEd/Classes/Editor/Transactor.h
--- a/Engine/Source/Editor/UnrealEd/Classes/Editor/Transactor.h
+++ b/Engine/Source/Editor/UnrealEd/Classes/Editor/Transactor.h
@@ -11,6 +11,7 @@
 #include "UObject/Object.h"
 #include "Serialization/ArchiveUObject.h"
 #include "Misc/ITransaction.h"
+#include "Serialization/ArchiveSerializedPropertyChain.h"
 #include "Transactor.generated.h"
 
 
@@ -53,10 +54,9 @@ struct FUndoSessionContext
  */
 class UNREALED_API FTransaction : public ITransaction
 {
+	friend class FTransactionObjectEvent;
+
 protected:
-	/** Map type for efficient unique indexing into UObject* arrays */
-	typedef	TMap<UObject*,int32> ObjectMapType;
-	
 	// Record of an object.
 	class UNREALED_API FObjectRecord
 	{
@@ -116,17 +116,122 @@ protected:
 			}
 		};
 
+		struct FSerializedProperty
+		{
+			FSerializedProperty()
+				: DataOffset(INDEX_NONE)
+				, DataSize(0)
+			{
+			}
+
+			static FName BuildSerializedPropertyKey(const FArchiveSerializedPropertyChain& InPropertyChain)
+			{
+				const int32 NumProperties = InPropertyChain.GetNumProperties();
+				check(NumProperties > 0);
+
+				if (NumProperties == 1)
+				{
+					return InPropertyChain.GetPropertyFromRoot(0)->GetFName();
+				}
+				else
+				{
+					FString FullPropertyName;
+					for (int32 PropertyIndex = 0; PropertyIndex < NumProperties; ++PropertyIndex)
+					{
+						if (PropertyIndex > 0)
+						{
+							FullPropertyName += TEXT('.');
+						}
+						FullPropertyName += InPropertyChain.GetPropertyFromRoot(PropertyIndex)->GetName();
+					}
+					return *FullPropertyName;
+				}
+			}
+
+			void AppendSerializedData(const int32 InOffset, const int32 InSize)
+			{
+				if (DataOffset == INDEX_NONE)
+				{
+					DataOffset = InOffset;
+					DataSize = InSize;
+				}
+				else
+				{
+					DataOffset = FMath::Min(DataOffset, InOffset);
+					DataSize += InSize;
+				}
+			}
+
+			/** Offset to the start of this property within the serialized object */
+			int32 DataOffset;
+			/** Size (in bytes) of this property within the serialized object */
+			int32 DataSize;
+		};
+
+		struct FSerializedObject
+		{
+			void SetObject(const UObject* InObject)
+			{
+				ObjectName = InObject->GetFName();
+				ObjectPathName = *InObject->GetPathName();
+				ObjectOuterPathName = InObject->GetOuter() ? FName(*InObject->GetOuter()->GetPathName()) : FName();
+				bIsPendingKill = InObject->IsPendingKill();
+				ObjectAnnotation = InObject->GetTransactionAnnotation();
+			}
+
+			void Reset()
+			{
+				ObjectName = FName();
+				ObjectPathName = FName();
+				ObjectOuterPathName = FName();
+				bIsPendingKill = false;
+				Data.Reset();
+				ReferencedObjects.Reset();
+				ReferencedNames.Reset();
+				SerializedProperties.Reset();
+				SerializedObjectIndices.Reset();
+				ObjectAnnotation.Reset();
+			}
+
+			void Swap(FSerializedObject& Other)
+			{
+				Exchange(ObjectName, Other.ObjectName);
+				Exchange(ObjectPathName, Other.ObjectPathName);
+				Exchange(ObjectOuterPathName, Other.ObjectOuterPathName);
+				Exchange(bIsPendingKill, Other.bIsPendingKill);
+				Exchange(Data, Other.Data);
+				Exchange(ReferencedObjects, Other.ReferencedObjects);
+				Exchange(ReferencedNames, Other.ReferencedNames);
+				Exchange(SerializedProperties, Other.SerializedProperties);
+				Exchange(SerializedObjectIndices, Other.SerializedObjectIndices);
+				Exchange(ObjectAnnotation, Other.ObjectAnnotation);
+			}
+
+			/** The name of the object when it was serialized */
+			FName ObjectName;
+			/** The path name of the object when it was serialized */
+			FName ObjectPathName;
+			/** The outer path name of the object when it was serialized */
+			FName ObjectOuterPathName;
+			/** The pending kill state of the object when it was serialized */
+			bool bIsPendingKill;
+			/** The data stream used to serialize/deserialize record */
+			TArray<uint8> Data;
+			/** External objects referenced in the transaction */
+			TArray<FPersistentObjectRef> ReferencedObjects;
+			/** FNames referenced in the object record */
+			TArray<FName> ReferencedNames;
+			/** Information about the properties that were serialized within this object */
+			TMap<FName, FSerializedProperty> SerializedProperties;
+			/** Information about the object pointer offsets that were serialized within this object (this maps the property name (or None if there was no property) to the ReferencedObjects indices of the property) */
+			TMultiMap<FName, int32> SerializedObjectIndices;
+			/** Annotation data for the object stored externally */
+			TSharedPtr<ITransactionObjectAnnotation> ObjectAnnotation;
+		};
+
 		// Variables.
-		/** The data stream used to serialize/deserialize record */
-		TArray<uint8>		Data;
-		/** External objects referenced in the transaction */
-		TArray<FPersistentObjectRef>	ReferencedObjects;
-		/** FNames referenced in the object record */
-		TArray<FName>		ReferencedNames;
 		/** The object to track */
-		FPersistentObjectRef	Object;
-		/** Annotation data for the object stored externally */
-		TSharedPtr<ITransactionObjectAnnotation> ObjectAnnotation;
+		FPersistentObjectRef Object;
 		/** Array: If an array object, reference to script array */
 		FScriptArray*		Array;
 		/** Array: Offset into the array */
@@ -146,20 +251,20 @@ protected:
 		STRUCT_DTOR			Destructor;
 		/** True if object  has already been restored from data. False otherwise. */
 		bool				bRestored;
+		/** True if object has been finalized and generated diff data */
+		bool				bFinalized;
+		/** True if object has been snapshot before */
+		bool				bSnapshot;
 		/** True if record should serialize data as binary blob (more compact). False to use tagged serialization (more robust) */
 		bool				bWantsBinarySerialization;
-
-		/** Copy of data that will be used when the transaction is flipped */
-		TArray<uint8> FlipData;
-
-		/** Copy of ReferencedObjects that will be used when the transaction is flipped */
-		TArray<FPersistentObjectRef> FlipReferencedObjects;
-
-		/** Copy of ReferencedNames that will be used when the transaction is flipped */
-		TArray<FName> FlipReferencedNames;
-
-		/** Copy of ObjectAnnotation that will be used when the transaction is flipped */
-		TSharedPtr<ITransactionObjectAnnotation> FlipObjectAnnotation;
+		/** The serialized object data */
+		FSerializedObject	SerializedObject;
+		/** The serialized object data that will be used when the transaction is flipped */
+		FSerializedObject	SerializedObjectFlip;
+		/** The serialized object data when it was last snapshot (if bSnapshot is true) */
+		FSerializedObject	SerializedObjectSnapshot;
+		/** Delta change information between the state of the object when the transaction started, and the state of the object when the transaction ended */
+		FTransactionObjectDeltaChange DeltaChange;
 
 		// Constructors.
 		FObjectRecord()
@@ -168,9 +273,13 @@ protected:
 
 		// Functions.
 		void SerializeContents( FArchive& Ar, int32 InOper );
+		void SerializeObject( FArchive& Ar );
 		void Restore( FTransaction* Owner );
 		void Save( FTransaction* Owner );
 		void Load( FTransaction* Owner );
+		void Finalize( FTransaction* Owner );
+		void Snapshot( FTransaction* Owner );
+		static void Diff( FTransaction* Owner, const FSerializedObject& OldSerializedObect, const FSerializedObject& NewSerializedObject, FTransactionObjectDeltaChange& OutDeltaChange );
 
 		/** Used by GC to collect referenced objects. */
 		void AddReferencedObjects( FReferenceCollector& Collector );
@@ -184,16 +293,12 @@ protected:
 		public:
 			FReader(
 				FTransaction* InOwner,
-				const TArray<uint8>& InData,
-				const TArray<FPersistentObjectRef>& InReferencedObjects,
-				const TArray<FName>& InReferencedNames,
+				const FSerializedObject& InSerializedObject,
 				bool bWantBinarySerialization
-				):
-				Owner(InOwner),
-				Data(InData),
-				ReferencedObjects(InReferencedObjects),
-				ReferencedNames(InReferencedNames),
-				Offset(0)
+				)
+				: Owner(InOwner)
+				, SerializedObject(InSerializedObject)
+				, Offset(0)
 			{
 				ArWantBinaryPropertySerialization = bWantBinarySerialization;
 				ArIsLoading = ArIsTransacting = 1;
@@ -207,8 +312,8 @@ protected:
 			{
 				if( Num )
 				{
-					checkSlow(Offset+Num<=Data.Num());
-					FMemory::Memcpy( SerData, &Data[Offset], Num );
+					checkSlow(Offset+Num<=SerializedObject.Data.Num());
+					FMemory::Memcpy( SerData, &SerializedObject.Data[Offset], Num );
 					Offset += Num;
 				}
 			}
@@ -216,7 +321,7 @@ protected:
 			{
 				int32 NameIndex = 0;
 				(FArchive&)*this << NameIndex;
-				N = ReferencedNames[NameIndex];
+				N = SerializedObject.ReferencedNames[NameIndex];
 				return *this;
 			}
 			FArchive& operator<<( class UObject*& Res ) override
@@ -225,7 +330,7 @@ protected:
 				(FArchive&)*this << ObjectIndex;
 				if (ObjectIndex != INDEX_NONE)
 				{
-					Res = ReferencedObjects[ObjectIndex].Get();
+					Res = SerializedObject.ReferencedObjects[ObjectIndex].Get();
 				}
 				else
 				{
@@ -250,9 +355,7 @@ protected:
 				}
 			}
 			FTransaction* Owner;
-			const TArray<uint8>& Data;
-			const TArray<FPersistentObjectRef>& ReferencedObjects;
-			const TArray<FName>& ReferencedNames;
+			const FSerializedObject& SerializedObject;
 			int64 Offset;
 		};
 
@@ -263,19 +366,15 @@ protected:
 		{
 		public:
 			FWriter(
-				TArray<uint8>& InData,
-				TArray<FPersistentObjectRef>& InReferencedObjects,
-				TArray<FName>& InReferencedNames,
+				FSerializedObject& InSerializedObject,
 				bool bWantBinarySerialization
-				):
-				Data(InData),
-				ReferencedObjects(InReferencedObjects),
-				ReferencedNames(InReferencedNames),
-				Offset(0)
+				)
+				: SerializedObject(InSerializedObject)
+				, Offset(0)
 			{
-				for(int32 ObjIndex = 0; ObjIndex < InReferencedObjects.Num(); ++ObjIndex)
+				for(int32 ObjIndex = 0; ObjIndex < SerializedObject.ReferencedObjects.Num(); ++ObjIndex)
 				{
-					ObjectMap.Add(InReferencedObjects[ObjIndex].Get(), ObjIndex);
+					ObjectMap.Add(SerializedObject.ReferencedObjects[ObjIndex].Get(), ObjIndex);
 				}
 
 				ArWantBinaryPropertySerialization = bWantBinarySerialization;
@@ -285,26 +384,68 @@ protected:
 			virtual int64 Tell() override {return Offset;}
 			virtual void Seek( int64 InPos ) override
 			{
-				checkSlow(Offset<=Data.Num());
+				checkSlow(Offset<=SerializedObject.Data.Num());
 				Offset = InPos; 
 			}
 
 		private:
+			struct FCachedPropertyKey
+			{
+			public:
+				FCachedPropertyKey()
+					: LastUpdateCount(0)
+				{
+				}
+
+				FName SyncCache(const FArchiveSerializedPropertyChain* InPropertyChain)
+				{
+					if (InPropertyChain)
+					{
+						const uint32 CurrentUpdateCount = InPropertyChain->GetUpdateCount();
+						if (CurrentUpdateCount != LastUpdateCount)
+						{
+							CachedKey = InPropertyChain->GetNumProperties() > 0 ? FSerializedProperty::BuildSerializedPropertyKey(*InPropertyChain) : FName();
+							LastUpdateCount = CurrentUpdateCount;
+						}
+					}
+					else
+					{
+						CachedKey = FName();
+						LastUpdateCount = 0;
+					}
+
+					return CachedKey;
+				}
+
+			private:
+				FName CachedKey;
+				uint32 LastUpdateCount;
+			};
+
 			void Serialize( void* SerData, int64 Num ) override
 			{
 				if( Num )
 				{
-					int32 DataIndex = ( Offset == Data.Num() )
-						?  Data.AddUninitialized(Num)
+					int32 DataIndex = ( Offset == SerializedObject.Data.Num() )
+						?  SerializedObject.Data.AddUninitialized(Num)
 						:  Offset;
 					
-					FMemory::Memcpy( &Data[DataIndex], SerData, Num );
+					FMemory::Memcpy( &SerializedObject.Data[DataIndex], SerData, Num );
 					Offset+= Num;
+
+					// Track this property offset in the serialized data
+					const FArchiveSerializedPropertyChain* PropertyChain = GetSerializedPropertyChain();
+					if (PropertyChain && PropertyChain->GetNumProperties() > 0)
+					{
+						const FName SerializedTaggedPropertyKey = CachedSerializedTaggedPropertyKey.SyncCache(PropertyChain);
+						FSerializedProperty& SerializedTaggedProperty = SerializedObject.SerializedProperties.FindOrAdd(SerializedTaggedPropertyKey);
+						SerializedTaggedProperty.AppendSerializedData(DataIndex, Num);
+					}
 				}
 			}
 			FArchive& operator<<( class FName& N ) override
 			{
-				int32 NameIndex = ReferencedNames.AddUnique(N);
+				int32 NameIndex = SerializedObject.ReferencedNames.AddUnique(N);
 				return (FArchive&)*this << NameIndex;
 			}
 			FArchive& operator<<( class UObject*& Res ) override
@@ -317,15 +458,22 @@ protected:
 				}
 				else if(Res)
 				{
-					ObjectIndex = ReferencedObjects.Add(FPersistentObjectRef(Res));
+					ObjectIndex = SerializedObject.ReferencedObjects.Add(FPersistentObjectRef(Res));
 					ObjectMap.Add(Res, ObjectIndex);
 				}
+
+				// Track this object offset in the serialized data
+				{
+					const FArchiveSerializedPropertyChain* PropertyChain = GetSerializedPropertyChain();
+					const FName SerializedTaggedPropertyKey = CachedSerializedTaggedPropertyKey.SyncCache(PropertyChain);
+					SerializedObject.SerializedObjectIndices.Add(SerializedTaggedPropertyKey, ObjectIndex);
+				}
+
 				return (FArchive&)*this << ObjectIndex;
 			}
-			TArray<uint8>& Data;
-			ObjectMapType ObjectMap;
-			TArray<FPersistentObjectRef>& ReferencedObjects;
-			TArray<FName>& ReferencedNames;
+			FSerializedObject& SerializedObject;
+			TMap<UObject*, int32> ObjectMap;
+			FCachedPropertyKey CachedSerializedTaggedPropertyKey;
 			int64 Offset;
 		};
 	};
@@ -344,7 +492,7 @@ protected:
 	UObject*	PrimaryObject;
 
 	/** Used to prevent objects from being serialized to a transaction more than once. */
-	ObjectMapType			ObjectMap;
+	TMap<UObject*, int32>	ObjectMap;
 
 	/** If true, on apply flip the direction of iteration over object records. The only client for which this is false is FMatineeTransaction */
 	bool					bFlip;
@@ -353,8 +501,26 @@ protected:
 	/** Count of the number of UModels modified since the last call to FTransaction::Apply */
 	int32					NumModelsModified;
 
+	struct FChangedObjectValue
+	{
+		FChangedObjectValue()
+			: Annotation()
+			, RecordIndex(INDEX_NONE)
+		{
+		}
+
+		FChangedObjectValue(const int32 InRecordIndex, const TSharedPtr<ITransactionObjectAnnotation>& InAnnotation)
+			: Annotation(InAnnotation)
+			, RecordIndex(InRecordIndex)
+		{
+		}
+
+		TSharedPtr<ITransactionObjectAnnotation> Annotation;
+		int32 RecordIndex;
+	};
+
 	/** Objects that will be changed directly by the transaction, empty when not transacting */
-	TMap<UObject*, TSharedPtr<ITransactionObjectAnnotation>> ChangedObjects;
+	TMap<UObject*, FChangedObjectValue> ChangedObjects;
 public:
 	// Constructor.
 	FTransaction(  const TCHAR* InContext=nullptr, const FText& InTitle=FText(), bool bInFlip=false )
@@ -373,12 +539,18 @@ public:
 	virtual void SaveObject( UObject* Object ) override;
 	virtual void SaveArray( UObject* Object, FScriptArray* Array, int32 Index, int32 Count, int32 Oper, int32 ElementSize, STRUCT_DC DefaultConstructor, STRUCT_AR Serializer, STRUCT_DTOR Destructor ) override;
 	virtual void SetPrimaryObject(UObject* InObject) override;
+	virtual void SnapshotObject( UObject* InObject ) override;
 
 	/**
 	 * Enacts the transaction.
 	 */
 	virtual void Apply() override;
 
+	/**
+	 * Finalize the transaction (try and work out what's changed).
+	 */
+	virtual void Finalize() override;
+
 	/** Returns a unique string to serve as a type ID for the FTranscationBase-derived type. */
 	virtual const TCHAR* GetTransactionType() const
 	{
diff --git a/Engine/Source/Editor/UnrealEd/Classes/Exporters/AnimSequenceExporterFBX.h b/Engine/Source/Editor/UnrealEd/Classes/Exporters/AnimSequenceExporterFBX.h
--- a/Engine/Source/Editor/UnrealEd/Classes/Exporters/AnimSequenceExporterFBX.h
+++ b/Engine/Source/Editor/UnrealEd/Classes/Exporters/AnimSequenceExporterFBX.h
@@ -9,10 +9,11 @@
 #include "CoreMinimal.h"
 #include "UObject/ObjectMacros.h"
 #include "Exporters/Exporter.h"
+#include "ExporterFbx.h"
 #include "AnimSequenceExporterFBX.generated.h"
 
 UCLASS()
-class UAnimSequenceExporterFBX : public UExporter
+class UAnimSequenceExporterFBX : public UExporterFBX
 {
 	GENERATED_UCLASS_BODY()
 
diff --git a/Engine/Source/Editor/UnrealEd/Classes/Exporters/FbxExportOption.h b/Engine/Source/Editor/UnrealEd/Classes/Exporters/FbxExportOption.h
--- a/Engine/Source/Editor/UnrealEd/Classes/Exporters/FbxExportOption.h
+++ b/Engine/Source/Editor/UnrealEd/Classes/Exporters/FbxExportOption.h
@@ -11,7 +11,7 @@
 #include "FbxExportOption.generated.h"
 
  // Fbx export compatibility
-UENUM()
+UENUM(BlueprintType)
 enum class EFbxExportCompatibility : uint8
 {
 	FBX_2011,
@@ -28,31 +28,35 @@ class UFbxExportOption : public UObject
 	GENERATED_UCLASS_BODY()
 public:
 	/** This will set the fbx sdk compatibility when exporting to fbx file. The default value is 2013 */
-	UPROPERTY(EditAnywhere, config, Category = Exporter)
+	UPROPERTY(EditAnywhere, BlueprintReadWrite, config, Category = Exporter)
 	EFbxExportCompatibility FbxExportCompatibility;
 
+	/** If enabled, save as ascii instead of binary */
+	UPROPERTY(EditAnywhere, BlueprintReadWrite, AdvancedDisplay, config, Category = Exporter)
+	uint32 bASCII : 1;
+
 	/** If enabled, export with X axis as the front axis instead of default -Y */
-	UPROPERTY(EditAnywhere, AdvancedDisplay, config, Category = Exporter)
+	UPROPERTY(EditAnywhere, BlueprintReadWrite, AdvancedDisplay, config, Category = Exporter)
 	uint32 bForceFrontXAxis : 1;
 
 	/** If enable, export vertex color */
-	UPROPERTY(EditAnywhere, config, category = Mesh)
+	UPROPERTY(EditAnywhere, BlueprintReadWrite, config, category = Mesh)
 	uint32 VertexColor : 1;
 
 	/** If enabled, export the level of detail */
-	UPROPERTY(EditAnywhere, config, category = StaticMesh)
+	UPROPERTY(EditAnywhere, BlueprintReadWrite, config, category = StaticMesh)
 	uint32 LevelOfDetail : 1;
 
 	/** If enabled, export collision */
-	UPROPERTY(EditAnywhere, config, category = StaticMesh)
+	UPROPERTY(EditAnywhere, BlueprintReadWrite, config, category = StaticMesh)
 	uint32 Collision : 1;
 
 	/** If enable, export welded vertices*/
-	UPROPERTY(EditAnywhere, AdvancedDisplay, config, category = StaticMesh)
+	UPROPERTY(EditAnywhere, BlueprintReadWrite, AdvancedDisplay, config, category = StaticMesh)
 	uint32 WeldedVertices : 1;
 
 	/** If enable, Map skeletal actor motion to the root bone of the skeleton. */
-	UPROPERTY(EditAnywhere, config, category = Animation)
+	UPROPERTY(EditAnywhere, BlueprintReadWrite, config, category = Animation)
 	uint32 MapSkeletalMotionToRoot : 1;
 
 	/* Set all the UProperty to the CDO value */
diff --git a/Engine/Source/Editor/UnrealEd/Classes/Exporters/LevelExporterFBX.h b/Engine/Source/Editor/UnrealEd/Classes/Exporters/LevelExporterFBX.h
--- a/Engine/Source/Editor/UnrealEd/Classes/Exporters/LevelExporterFBX.h
+++ b/Engine/Source/Editor/UnrealEd/Classes/Exporters/LevelExporterFBX.h
@@ -10,10 +10,11 @@
 #include "UObject/ObjectMacros.h"
 #include "UObject/UObjectGlobals.h"
 #include "Exporters/Exporter.h"
+#include "ExporterFbx.h"
 #include "LevelExporterFBX.generated.h"
 
 UCLASS()
-class ULevelExporterFBX : public UExporter
+class ULevelExporterFBX : public UExporterFBX
 {
 public:
 	GENERATED_BODY()
diff --git a/Engine/Source/Editor/UnrealEd/Classes/Exporters/SkeletalMeshExporterFBX.h b/Engine/Source/Editor/UnrealEd/Classes/Exporters/SkeletalMeshExporterFBX.h
--- a/Engine/Source/Editor/UnrealEd/Classes/Exporters/SkeletalMeshExporterFBX.h
+++ b/Engine/Source/Editor/UnrealEd/Classes/Exporters/SkeletalMeshExporterFBX.h
@@ -9,10 +9,11 @@
 #include "CoreMinimal.h"
 #include "UObject/ObjectMacros.h"
 #include "Exporters/Exporter.h"
+#include "ExporterFbx.h"
 #include "SkeletalMeshExporterFBX.generated.h"
 
 UCLASS()
-class USkeletalMeshExporterFBX : public UExporter
+class USkeletalMeshExporterFBX : public UExporterFBX
 {
 	GENERATED_UCLASS_BODY()
 
diff --git a/Engine/Source/Editor/UnrealEd/Classes/Exporters/StaticMeshExporterFBX.h b/Engine/Source/Editor/UnrealEd/Classes/Exporters/StaticMeshExporterFBX.h
--- a/Engine/Source/Editor/UnrealEd/Classes/Exporters/StaticMeshExporterFBX.h
+++ b/Engine/Source/Editor/UnrealEd/Classes/Exporters/StaticMeshExporterFBX.h
@@ -9,10 +9,11 @@
 #include "CoreMinimal.h"
 #include "UObject/ObjectMacros.h"
 #include "Exporters/Exporter.h"
+#include "ExporterFbx.h"
 #include "StaticMeshExporterFBX.generated.h"
 
 UCLASS()
-class UStaticMeshExporterFBX : public UExporter
+class UStaticMeshExporterFBX : public UExporterFBX
 {
 	GENERATED_UCLASS_BODY()
 
diff --git a/Engine/Source/Editor/UnrealEd/Classes/Factories/Factory.h b/Engine/Source/Editor/UnrealEd/Classes/Factories/Factory.h
--- a/Engine/Source/Editor/UnrealEd/Classes/Factories/Factory.h
+++ b/Engine/Source/Editor/UnrealEd/Classes/Factories/Factory.h
@@ -169,10 +169,17 @@ public:
 	 */
 	void SetAutomatedAssetImportData(const class UAutomatedAssetImportData* Data);
 
+	/**
+	 * Sets the import task being used with this factory
+	 *
+	 * @param Task	The import task or nullptr if it does not exist
+	 */
+	void SetAssetImportTask(const class UAssetImportTask* Task);
+
 	/**
 	 * @return true if this factory is being used for automated import.  Dialogs and user input should be disabled if this method returns true
 	 */
-	bool IsAutomatedImport() const { return  GIsAutomationTesting || AutomatedImportData; }
+	bool IsAutomatedImport() const;
 public:
 
 	/**
@@ -323,6 +330,9 @@ public:
 	UPROPERTY()
 	const class UAutomatedAssetImportData* AutomatedImportData;
 
+	/** Task for importing file via script interfaces */
+	UPROPERTY()
+	const class UAssetImportTask* AssetImportTask;
 
 protected:
 
diff --git a/Engine/Source/Editor/UnrealEd/Classes/Factories/FbxAssetImportData.h b/Engine/Source/Editor/UnrealEd/Classes/Factories/FbxAssetImportData.h
--- a/Engine/Source/Editor/UnrealEd/Classes/Factories/FbxAssetImportData.h
+++ b/Engine/Source/Editor/UnrealEd/Classes/Factories/FbxAssetImportData.h
@@ -17,25 +17,25 @@ class UNREALED_API UFbxAssetImportData : public UAssetImportData
 {
 	GENERATED_UCLASS_BODY()
 
-	UPROPERTY(EditAnywhere, Category=Transform, meta=(ImportType="StaticMesh|SkeletalMesh|Animation", ImportCategory="Transform"))
+	UPROPERTY(EditAnywhere, BlueprintReadWrite, Category=Transform, meta=(ImportType="StaticMesh|SkeletalMesh|Animation", ImportCategory="Transform"))
 	FVector ImportTranslation;
 
-	UPROPERTY(EditAnywhere, Category=Transform, meta=(ImportType="StaticMesh|SkeletalMesh|Animation", ImportCategory="Transform"))
+	UPROPERTY(EditAnywhere, BlueprintReadWrite, Category=Transform, meta=(ImportType="StaticMesh|SkeletalMesh|Animation", ImportCategory="Transform"))
 	FRotator ImportRotation;
 
-	UPROPERTY(EditAnywhere, Category=Transform, meta=(ImportType="StaticMesh|SkeletalMesh|Animation", ImportCategory="Transform"))
+	UPROPERTY(EditAnywhere, BlueprintReadWrite, Category=Transform, meta=(ImportType="StaticMesh|SkeletalMesh|Animation", ImportCategory="Transform"))
 	float ImportUniformScale;
 
 	/** Whether to convert scene from FBX scene. */
-	UPROPERTY(EditAnywhere, config, Category = Miscellaneous, meta = (ImportType = "StaticMesh|SkeletalMesh|Animation", ImportCategory = "Miscellaneous", ToolTip = "Convert the scene from FBX coordinate system to UE4 coordinate system"))
+	UPROPERTY(EditAnywhere, BlueprintReadWrite, config, Category = Miscellaneous, meta = (ImportType = "StaticMesh|SkeletalMesh|Animation", ImportCategory = "Miscellaneous", ToolTip = "Convert the scene from FBX coordinate system to UE4 coordinate system"))
 	bool bConvertScene;
 
 	/** Whether to force the front axis to be align with X instead of -Y. */
-	UPROPERTY(EditAnywhere, config, Category = Miscellaneous, meta = (editcondition = "bConvertScene", ImportType = "StaticMesh|SkeletalMesh|Animation", ImportCategory = "Miscellaneous", ToolTip = "Convert the scene from FBX coordinate system to UE4 coordinate system with front X axis instead of -Y"))
+	UPROPERTY(EditAnywhere, BlueprintReadWrite, config, Category = Miscellaneous, meta = (editcondition = "bConvertScene", ImportType = "StaticMesh|SkeletalMesh|Animation", ImportCategory = "Miscellaneous", ToolTip = "Convert the scene from FBX coordinate system to UE4 coordinate system with front X axis instead of -Y"))
 	bool bForceFrontXAxis;
 
 	/** Whether to convert the scene from FBX unit to UE4 unit (centimeter). */
-	UPROPERTY(EditAnywhere, config, Category = Miscellaneous, meta = (ImportType = "StaticMesh|SkeletalMesh|Animation", ImportCategory = "Miscellaneous", ToolTip = "Convert the scene from FBX unit to UE4 unit (centimeter)."))
+	UPROPERTY(EditAnywhere, BlueprintReadWrite, config, Category = Miscellaneous, meta = (ImportType = "StaticMesh|SkeletalMesh|Animation", ImportCategory = "Miscellaneous", ToolTip = "Convert the scene from FBX unit to UE4 unit (centimeter)."))
 	bool bConvertSceneUnit;
 
 	/* Use by the reimport factory to answer CanReimport, if true only factory for scene reimport will return true */
diff --git a/Engine/Source/Editor/UnrealEd/Classes/Factories/FbxFactory.h b/Engine/Source/Editor/UnrealEd/Classes/Factories/FbxFactory.h
--- a/Engine/Source/Editor/UnrealEd/Classes/Factories/FbxFactory.h
+++ b/Engine/Source/Editor/UnrealEd/Classes/Factories/FbxFactory.h
@@ -17,6 +17,9 @@ class UNREALED_API UFbxFactory : public UFactory
 	UPROPERTY()
 	class UFbxImportUI* ImportUI;
 
+	/** Prevent garbage collection of original when overriding ImportUI property */
+	UPROPERTY()
+	class UFbxImportUI* OriginalImportUI;
 
 	/**  Set import batch **/
 	void EnableShowOption() { bShowOption = true; }
@@ -31,7 +34,7 @@ class UNREALED_API UFbxFactory : public UFactory
 	//~ Begin UFactory Interface
 	virtual bool DoesSupportClass(UClass * Class) override;
 	virtual UClass* ResolveSupportedClass() override;
-	virtual UObject* FactoryCreateBinary(UClass* InClass, UObject* InParent, FName InName, EObjectFlags Flags, UObject* Context, const TCHAR* Type, const uint8*& Buffer, const uint8* BufferEnd, FFeedbackContext* Warn, bool& bOutOperationCanceled) override;
+	virtual UObject* FactoryCreateFile(UClass* InClass, UObject* InParent, FName InName, EObjectFlags Flags, const FString& Filename, const TCHAR* Parms, FFeedbackContext* Warn, bool& bOutOperationCanceled);
 	virtual bool FactoryCanImport(const FString& Filename) override;
 	virtual IImportSettingsParser* GetImportSettingsParser() override;
 
diff --git a/Engine/Source/Editor/UnrealEd/Classes/Factories/FbxImportUI.h b/Engine/Source/Editor/UnrealEd/Classes/Factories/FbxImportUI.h
--- a/Engine/Source/Editor/UnrealEd/Classes/Factories/FbxImportUI.h
+++ b/Engine/Source/Editor/UnrealEd/Classes/Factories/FbxImportUI.h
@@ -13,7 +13,7 @@
 #include "FbxImportUI.generated.h"
 
 /** Import mesh type */
-UENUM()
+UENUM(BlueprintType)
 enum EFBXImportType
 {
 	/** Select Static Mesh if you'd like to import static mesh. */
@@ -35,122 +35,123 @@ class UFbxImportUI : public UObject, public IImportSettingsParser
 
 public:
 	/** Whether or not the imported file is in OBJ format */
-	UPROPERTY()
+	UPROPERTY(BlueprintReadWrite, Category = Miscellaneous)
 	bool bIsObjImport;
 
 	/** The original detected type of this import */
-	UPROPERTY()
+	UPROPERTY(BlueprintReadWrite, Category = Miscellaneous)
 	TEnumAsByte<enum EFBXImportType> OriginalImportType;
 
 	/** Type of asset to import from the FBX file */
-	UPROPERTY()
+	UPROPERTY(BlueprintReadWrite, Category = Miscellaneous)
 	TEnumAsByte<enum EFBXImportType> MeshTypeToImport;
 
 	/** Use the string in "Name" field as full name of mesh. The option only works when the scene contains one mesh. */
-	UPROPERTY(EditAnywhere, AdvancedDisplay, config, Category=Miscellaneous, meta=(OBJRestrict="true"))
+	UPROPERTY(EditAnywhere, BlueprintReadWrite, AdvancedDisplay, config, Category=Miscellaneous, meta=(OBJRestrict="true"))
 	uint32 bOverrideFullName:1;
 
 	/** Whether to import the incoming FBX as a skeletal object */
-	UPROPERTY(EditAnywhere, Category = Mesh, meta = (ImportType = "StaticMesh|SkeletalMesh", DisplayName="Skeletal Mesh"))
+	UPROPERTY(EditAnywhere, BlueprintReadWrite, Category = Mesh, meta = (ImportType = "StaticMesh|SkeletalMesh", DisplayName="Skeletal Mesh"))
 	bool bImportAsSkeletal;
 	
 	/** Whether to import the incoming FBX as a Subdivision Surface (could be made a combo box together with bImportAsSkeletal) (Experimental, Early work in progress) */
 	/** Whether to import the mesh. Allows animation only import when importing a skeletal mesh. */
-	UPROPERTY(EditAnywhere, Category=Mesh, meta=(ImportType="SkeletalMesh"))
+	UPROPERTY(EditAnywhere, BlueprintReadWrite, Category=Mesh, meta=(ImportType="SkeletalMesh"))
 	bool bImportMesh;
 
 	/** Skeleton to use for imported asset. When importing a mesh, leaving this as "None" will create a new skeleton. When importing an animation this MUST be specified to import the asset. */
-	UPROPERTY(EditAnywhere, Category=Mesh, meta=(ImportType="SkeletalMesh|Animation"))
+	UPROPERTY(EditAnywhere, BlueprintReadWrite, Category=Mesh, meta=(ImportType="SkeletalMesh|Animation"))
 	class USkeleton* Skeleton;
 
 	/** If checked, create new PhysicsAsset if it doesn't have it */
-	UPROPERTY(EditAnywhere, AdvancedDisplay, config, Category=Mesh, meta=(ImportType="SkeletalMesh"))
+	UPROPERTY(EditAnywhere, BlueprintReadWrite, AdvancedDisplay, config, Category=Mesh, meta=(ImportType="SkeletalMesh"))
 	uint32 bCreatePhysicsAsset:1;
 
 	/** If this is set, use this PhysicsAsset. It is possible bCreatePhysicsAsset == false, and PhysicsAsset == NULL. It is possible they do not like to create anything. */
-	UPROPERTY(EditAnywhere, AdvancedDisplay, Category=Mesh, meta=(ImportType="SkeletalMesh", editcondition="!bCreatePhysicsAsset"))
+	UPROPERTY(EditAnywhere, BlueprintReadWrite, AdvancedDisplay, Category=Mesh, meta=(ImportType="SkeletalMesh", editcondition="!bCreatePhysicsAsset"))
 	class UPhysicsAsset* PhysicsAsset;
 
 	/** If checked, the static mesh auto compute LOD distance will be turn on. If unchecked user will be able to specify custom LOD distance for every LOD. */
-	UPROPERTY(EditAnywhere, config, Category = LodSettings, meta = (ImportType = "StaticMesh"))
+	UPROPERTY(EditAnywhere, BlueprintReadWrite, config, Category = LodSettings, meta = (ImportType = "StaticMesh"))
 	uint32 bAutoComputeLodDistances : 1;
 	/** Specify the LOD distance for LOD 0*/
-	UPROPERTY(EditAnywhere, config, Category = LodSettings, meta = (ImportType = "StaticMesh", UIMin = "0.0"))
+	UPROPERTY(EditAnywhere, BlueprintReadWrite, config, Category = LodSettings, meta = (ImportType = "StaticMesh", UIMin = "0.0"))
 	float LodDistance0;
 	/** Specify the LOD distance for LOD 1*/
-	UPROPERTY(EditAnywhere, config, Category = LodSettings, meta = (ImportType = "StaticMesh", UIMin = "0.0"))
+	UPROPERTY(EditAnywhere, BlueprintReadWrite, config, Category = LodSettings, meta = (ImportType = "StaticMesh", UIMin = "0.0"))
 	float LodDistance1;
 	/** Specify the LOD distance for LOD 2*/
-	UPROPERTY(EditAnywhere, config, Category = LodSettings, meta = (ImportType = "StaticMesh", UIMin = "0.0"))
+	UPROPERTY(EditAnywhere, BlueprintReadWrite, config, Category = LodSettings, meta = (ImportType = "StaticMesh", UIMin = "0.0"))
 	float LodDistance2;
 	/** Specify the LOD distance for LOD 3*/
-	UPROPERTY(EditAnywhere, config, Category = LodSettings, meta = (ImportType = "StaticMesh", UIMin = "0.0"))
+	UPROPERTY(EditAnywhere, BlueprintReadWrite, config, Category = LodSettings, meta = (ImportType = "StaticMesh", UIMin = "0.0"))
 	float LodDistance3;
 	/** Specify the LOD distance for LOD 4*/
-	UPROPERTY(EditAnywhere, config, Category = LodSettings, meta = (ImportType = "StaticMesh", UIMin = "0.0"))
+	UPROPERTY(EditAnywhere, BlueprintReadWrite, config, Category = LodSettings, meta = (ImportType = "StaticMesh", UIMin = "0.0"))
 	float LodDistance4;
 	/** Specify the LOD distance for LOD 5*/
-	UPROPERTY(EditAnywhere, config, Category = LodSettings, meta = (ImportType = "StaticMesh", UIMin = "0.0"))
+	UPROPERTY(EditAnywhere, BlueprintReadWrite, config, Category = LodSettings, meta = (ImportType = "StaticMesh", UIMin = "0.0"))
 	float LodDistance5;
 	/** Specify the LOD distance for LOD 6*/
-	UPROPERTY(EditAnywhere, config, Category = LodSettings, meta = (ImportType = "StaticMesh", UIMin = "0.0"))
+	UPROPERTY(EditAnywhere, BlueprintReadWrite, config, Category = LodSettings, meta = (ImportType = "StaticMesh", UIMin = "0.0"))
 	float LodDistance6;
 	/** Specify the LOD distance for LOD 7*/
-	UPROPERTY(EditAnywhere, config, Category = LodSettings, meta = (ImportType = "StaticMesh", UIMin = "0.0"))
+	UPROPERTY(EditAnywhere, BlueprintReadWrite, config, Category = LodSettings, meta = (ImportType = "StaticMesh", UIMin = "0.0"))
 	float LodDistance7;
 
 	/** Set the minimum LOD number. A value of 0 disable the option. */
-	UPROPERTY(EditAnywhere, config, AdvancedDisplay, Category = LodSettings, meta = (ImportType = "StaticMesh", UIMin = "-1"))
+	UPROPERTY(EditAnywhere, BlueprintReadWrite, config, AdvancedDisplay, Category = LodSettings, meta = (ImportType = "StaticMesh", UIMin = "-1"))
 	int32 MinimumLodNumber;
 
 	/** Set the number of LODs. A value of 0 disable the option. */
-	UPROPERTY(EditAnywhere, config, AdvancedDisplay, Category = LodSettings, meta = (ImportType = "StaticMesh", UIMin = "-1"))
+	UPROPERTY(EditAnywhere, BlueprintReadWrite, config, AdvancedDisplay, Category = LodSettings, meta = (ImportType = "StaticMesh", UIMin = "-1"))
 	int32 LodNumber;
 
 	/** True to import animations from the FBX File */
-	UPROPERTY(EditAnywhere, config, Category=Animation, meta=(ImportType="SkeletalMesh|Animation"))
+	UPROPERTY(EditAnywhere, BlueprintReadWrite, config, Category=Animation, meta=(ImportType="SkeletalMesh|Animation"))
 	uint32 bImportAnimations:1;
 
 	/** Override for the name of the animation to import. By default, it will be the name of FBX **/
-	UPROPERTY(EditAnywhere, AdvancedDisplay, Category=Animation, meta=(editcondition="bImportAnimations", ImportType="SkeletalMesh")) 
+	UPROPERTY(EditAnywhere, BlueprintReadWrite, AdvancedDisplay, Category=Animation, meta=(editcondition="bImportAnimations", ImportType="SkeletalMesh"))
 	FString OverrideAnimationName;
 
 	/** Enables importing of 'rigid skeletalmesh' (unskinned, hierarchy-based animation) from this FBX file, no longer shown, used behind the scenes */
-	UPROPERTY()
+	UPROPERTY(BlueprintReadWrite, Category = Miscellaneous)
 	uint32 bImportRigidMesh:1;
 
 	/** Whether to automatically create Unreal materials for materials found in the FBX scene */
-	UPROPERTY(EditAnywhere, config, Category = Material)
+	UPROPERTY(EditAnywhere, BlueprintReadWrite, config, Category = Material)
 	uint32 bImportMaterials:1;
 
 	/** The option works only when option "Import Material" is OFF. If "Import Material" is ON, textures are always imported. */
-	UPROPERTY(EditAnywhere, config, Category=Material)
+	UPROPERTY(EditAnywhere, BlueprintReadWrite, config, Category=Material)
 	uint32 bImportTextures:1;
 
 	/** Import data used when importing static meshes */
-	UPROPERTY(EditAnywhere, Transient, Instanced, Category = Mesh, meta=(ImportType = "StaticMesh"))
+	UPROPERTY(EditAnywhere, BlueprintReadWrite, Transient, Instanced, Category = Mesh, meta=(ImportType = "StaticMesh"))
 	class UFbxStaticMeshImportData* StaticMeshImportData;
 
 	/** Import data used when importing skeletal meshes */
-	UPROPERTY(EditAnywhere, Transient, Instanced, Category=Mesh, meta=(ImportType = "SkeletalMesh"))
+	UPROPERTY(EditAnywhere, BlueprintReadWrite, Transient, Instanced, Category=Mesh, meta=(ImportType = "SkeletalMesh"))
 	class UFbxSkeletalMeshImportData* SkeletalMeshImportData;
 
 	/** Import data used when importing animations */
-	UPROPERTY(EditAnywhere, Transient, Instanced, Category=Animation, meta=(editcondition="bImportAnimations", ImportType = "Animation"))
+	UPROPERTY(EditAnywhere, BlueprintReadWrite, Transient, Instanced, Category=Animation, meta=(editcondition="bImportAnimations", ImportType = "Animation"))
 	class UFbxAnimSequenceImportData* AnimSequenceImportData;
 
 	/** Import data used when importing textures */
-	UPROPERTY(EditAnywhere, Transient, Instanced, Category=Material)
+	UPROPERTY(EditAnywhere, BlueprintReadWrite, Transient, Instanced, Category=Material)
 	class UFbxTextureImportData* TextureImportData;
 
 	/** If true the automated import path should detect the import type.  If false the import type was specified by the user */
-	UPROPERTY()
+	UPROPERTY(BlueprintReadWrite, Category = Miscellaneous)
 	bool bAutomatedImportShouldDetectType;
 
 	/** If true the existing material array will be reset by the incoming fbx file. The matching "material import name" will be restore properly but, the entries that has no match will use the material instance of the existing data at the same index. (Never enable this option if you have gameplay code that use a material slot)*/
-	UPROPERTY(EditAnywhere, config, AdvancedDisplay, Category = Material, meta = (OBJRestrict = "true", ImportType = "Mesh"))
+	UPROPERTY(EditAnywhere, BlueprintReadWrite, config, AdvancedDisplay, Category = Material, meta = (OBJRestrict = "true", ImportType = "Mesh"))
 	uint32 bResetMaterialSlots : 1;
 
+	UFUNCTION(BlueprintCallable, Category = Miscellaneous)
 	void ResetToDefault();
 
 	/** UObject Interface */
diff --git a/Engine/Source/Editor/UnrealEd/Classes/Factories/FbxMeshImportData.h b/Engine/Source/Editor/UnrealEd/Classes/Factories/FbxMeshImportData.h
--- a/Engine/Source/Editor/UnrealEd/Classes/Factories/FbxMeshImportData.h
+++ b/Engine/Source/Editor/UnrealEd/Classes/Factories/FbxMeshImportData.h
@@ -49,23 +49,23 @@ class UFbxMeshImportData : public UFbxAssetImportData
 	GENERATED_UCLASS_BODY()
 
 	/** If this option is true the node absolute transform (transform, offset and pivot) will be apply to the mesh vertices. */
-	UPROPERTY(EditAnywhere, AdvancedDisplay, config, Category = Mesh, meta = (ImportType = "StaticMesh"))
+	UPROPERTY(EditAnywhere, BlueprintReadWrite, AdvancedDisplay, config, Category = Mesh, meta = (ImportType = "StaticMesh"))
 	bool bTransformVertexToAbsolute;
 
 	/** - Experimental - If this option is true the inverse node rotation pivot will be apply to the mesh vertices. The pivot from the DCC will then be the origin of the mesh. Note: "TransformVertexToAbsolute" must be false.*/
-	UPROPERTY(EditAnywhere, AdvancedDisplay, config, Category = Mesh, meta = (EditCondition = "!bTransformVertexToAbsolute", ImportType = "StaticMesh"))
+	UPROPERTY(EditAnywhere, BlueprintReadWrite, AdvancedDisplay, config, Category = Mesh, meta = (EditCondition = "!bTransformVertexToAbsolute", ImportType = "StaticMesh"))
 	bool bBakePivotInVertex;
 
 	/** Enables importing of mesh LODs from FBX LOD groups, if present in the FBX file */
-	UPROPERTY(EditAnywhere, AdvancedDisplay, config, Category= Mesh, meta=(OBJRestrict="true", ImportType="Mesh", ToolTip="If enabled, creates LOD models for Unreal meshes from LODs in the import file; If not enabled, only the base mesh from the LOD group is imported"))
+	UPROPERTY(EditAnywhere, BlueprintReadWrite, AdvancedDisplay, config, Category= Mesh, meta=(OBJRestrict="true", ImportType="Mesh", ToolTip="If enabled, creates LOD models for Unreal meshes from LODs in the import file; If not enabled, only the base mesh from the LOD group is imported"))
 	uint32 bImportMeshLODs:1;
 
 	/** Enabling this option will read the tangents(tangent,binormal,normal) from FBX file instead of generating them automatically. */
-	UPROPERTY(EditAnywhere, AdvancedDisplay, config, Category= Mesh, meta=(ImportType="Mesh"))
+	UPROPERTY(EditAnywhere, BlueprintReadWrite, AdvancedDisplay, config, Category= Mesh, meta=(ImportType="Mesh"))
 	TEnumAsByte<enum EFBXNormalImportMethod> NormalImportMethod;
 
 	/** Use the MikkTSpace tangent space generator for generating normals and tangents on the mesh */
-	UPROPERTY(EditAnywhere, config, AdvancedDisplay, Category = Mesh, meta=(ImportType="Mesh"))
+	UPROPERTY(EditAnywhere, BlueprintReadWrite, config, AdvancedDisplay, Category = Mesh, meta=(ImportType="Mesh"))
 	TEnumAsByte<enum EFBXNormalGenerationMethod::Type> NormalGenerationMethod;
 
 	bool CanEditChange( const UProperty* InProperty ) const override;
diff --git a/Engine/Source/Editor/UnrealEd/Classes/Factories/FbxStaticMeshImportData.h b/Engine/Source/Editor/UnrealEd/Classes/Factories/FbxStaticMeshImportData.h
--- a/Engine/Source/Editor/UnrealEd/Classes/Factories/FbxStaticMeshImportData.h
+++ b/Engine/Source/Editor/UnrealEd/Classes/Factories/FbxStaticMeshImportData.h
@@ -33,41 +33,41 @@ class UFbxStaticMeshImportData : public UFbxMeshImportData
 	GENERATED_UCLASS_BODY()
 
 	/** The LODGroup to associate with this mesh when it is imported */
-	UPROPERTY(EditAnywhere, config, AdvancedDisplay, Category=Mesh, meta=(ImportType="StaticMesh"))
+	UPROPERTY(EditAnywhere, BlueprintReadWrite, config, AdvancedDisplay, Category=Mesh, meta=(ImportType="StaticMesh"))
 	FName StaticMeshLODGroup;
 
 	/** Specify how vertex colors should be imported */
-	UPROPERTY(EditAnywhere, AdvancedDisplay, config, Category= Mesh, meta=(OBJRestrict="true", ImportType="StaticMesh"))
+	UPROPERTY(EditAnywhere, BlueprintReadWrite, AdvancedDisplay, config, Category= Mesh, meta=(OBJRestrict="true", ImportType="StaticMesh"))
 	TEnumAsByte<EVertexColorImportOption::Type> VertexColorImportOption;
 
 	/** Specify override color in the case that VertexColorImportOption is set to Override */
-	UPROPERTY(EditAnywhere, AdvancedDisplay, config, Category= Mesh, meta=(OBJRestrict="true", ImportType="StaticMesh"))
+	UPROPERTY(EditAnywhere, BlueprintReadWrite, AdvancedDisplay, config, Category= Mesh, meta=(OBJRestrict="true", ImportType="StaticMesh"))
 	FColor VertexOverrideColor;
 
 	/** Disabling this option will keep degenerate triangles found.  In general you should leave this option on. */
-	UPROPERTY(EditAnywhere, AdvancedDisplay, Category = Mesh, meta = (ImportType = "StaticMesh"))
+	UPROPERTY(EditAnywhere, BlueprintReadWrite, AdvancedDisplay, Category = Mesh, meta = (ImportType = "StaticMesh"))
 	uint32 bRemoveDegenerates:1;
 
 	/** Required for PNT tessellation but can be slow. Recommend disabling for larger meshes. */
-	UPROPERTY(EditAnywhere, AdvancedDisplay, Category = Mesh, meta = (ImportType = "StaticMesh"))
+	UPROPERTY(EditAnywhere, BlueprintReadWrite, AdvancedDisplay, Category = Mesh, meta = (ImportType = "StaticMesh"))
 	uint32 bBuildAdjacencyBuffer:1;
 
-	UPROPERTY(EditAnywhere, AdvancedDisplay, Category = Mesh, meta = (ImportType = "StaticMesh"))
+	UPROPERTY(EditAnywhere, BlueprintReadWrite, AdvancedDisplay, Category = Mesh, meta = (ImportType = "StaticMesh"))
 	uint32 bBuildReversedIndexBuffer:1;
 
-	UPROPERTY(EditAnywhere, config, AdvancedDisplay, Category= Mesh, meta=(ImportType="StaticMesh"))
+	UPROPERTY(EditAnywhere, BlueprintReadWrite, config, AdvancedDisplay, Category= Mesh, meta=(ImportType="StaticMesh"))
 	uint32 bGenerateLightmapUVs:1;
 
 	/** If checked, one convex hull per UCX_ prefixed collision mesh will be generated instead of decomposing into multiple hulls */
-	UPROPERTY(EditAnywhere, AdvancedDisplay, config, Category= Mesh, meta=(OBJRestrict="true", ImportType="StaticMesh"))
+	UPROPERTY(EditAnywhere, BlueprintReadWrite, AdvancedDisplay, config, Category= Mesh, meta=(OBJRestrict="true", ImportType="StaticMesh"))
 	uint32 bOneConvexHullPerUCX:1;
 
 	/** If checked, collision will automatically be generated (ignored if custom collision is imported or used). */
-	UPROPERTY(EditAnywhere, config, Category = Mesh, meta=(OBJRestrict="true", ImportType="StaticMesh"))
+	UPROPERTY(EditAnywhere, BlueprintReadWrite, config, Category = Mesh, meta=(OBJRestrict="true", ImportType="StaticMesh"))
 	uint32 bAutoGenerateCollision : 1;
 
 	/** For static meshes, enabling this option will combine all meshes in the FBX into a single monolithic mesh in Unreal */
-	UPROPERTY(EditAnywhere, AdvancedDisplay, config, Category = Mesh, meta = (ToolTip = "If enabled, combines all meshes into a single mesh", ImportType = "StaticMesh"))
+	UPROPERTY(EditAnywhere, BlueprintReadWrite, AdvancedDisplay, config, Category = Mesh, meta = (ToolTip = "If enabled, combines all meshes into a single mesh", ImportType = "StaticMesh"))
 	uint32 bCombineMeshes : 1;
 
 	/** Gets or creates fbx import data for the specified static mesh */
diff --git a/Engine/Source/Editor/UnrealEd/Private/AssetImportTask.cpp b/Engine/Source/Editor/UnrealEd/Private/AssetImportTask.cpp
new file mode 100644
--- /dev/null
+++ b/Engine/Source/Editor/UnrealEd/Private/AssetImportTask.cpp
@@ -0,0 +1,10 @@
+// Copyright 1998-2018 Epic Games, Inc. All Rights Reserved.
+
+#include "AssetImportTask.h"
+
+DEFINE_LOG_CATEGORY(LogAssetImportTask);
+
+UAssetImportTask::UAssetImportTask()
+{
+
+}
diff --git a/Engine/Source/Editor/UnrealEd/Private/Commandlets/GatherTextFromAssetsCommandlet.cpp b/Engine/Source/Editor/UnrealEd/Private/Commandlets/GatherTextFromAssetsCommandlet.cpp
--- a/Engine/Source/Editor/UnrealEd/Private/Commandlets/GatherTextFromAssetsCommandlet.cpp
+++ b/Engine/Source/Editor/UnrealEd/Private/Commandlets/GatherTextFromAssetsCommandlet.cpp
@@ -447,6 +447,7 @@ int32 UGatherTextFromAssetsCommandlet::Main(const FString& Params)
 		if (ShouldExcludeDerivedClasses)
 		{
 			FirstPassFilter.bRecursiveClasses = true;
+			FirstPassFilter.ClassNames.Add(TEXT("Object"));
 			for (const FString& ExcludeClassName : ExcludeClassNames)
 			{
 				// Note: Can't necessarily validate these class names here, as the class may be a generated blueprint class that hasn't been loaded yet.
diff --git a/Engine/Source/Editor/UnrealEd/Private/Commandlets/ImportAssetsCommandlet.cpp b/Engine/Source/Editor/UnrealEd/Private/Commandlets/ImportAssetsCommandlet.cpp
--- a/Engine/Source/Editor/UnrealEd/Private/Commandlets/ImportAssetsCommandlet.cpp
+++ b/Engine/Source/Editor/UnrealEd/Private/Commandlets/ImportAssetsCommandlet.cpp
@@ -13,6 +13,7 @@
 #include "Misc/FileHelper.h"
 #include "Factories/ImportSettings.h"
 #include "ISourceControlModule.h"
+#include "SourceControlHelpers.h"
 #include "Editor.h"
 #include "FileHelpers.h"
 #include "Misc/FeedbackContext.h"
@@ -252,7 +253,7 @@ bool UImportAssetsCommandlet::ImportAndSave(const TArray<UAutomatedAssetImportDa
 					}
 					else if(PackageSCState->CanCheckout())
 					{
-						const bool bWasCheckedOut = SourceControlHelpers::CheckOutFile(PackageFilename);
+						const bool bWasCheckedOut = SourceControlHelpers::CheckOutOrAddFile(PackageFilename);
 						bShouldAttemptToSave = bWasCheckedOut;
 						if(!bWasCheckedOut)
 						{
@@ -263,7 +264,7 @@ bool UImportAssetsCommandlet::ImportAndSave(const TArray<UAutomatedAssetImportDa
 					else
 					{
 						// package was not checked out by another user and is at the current head revision and could not be checked out
-						// this means it should be added after save because it doesnt exist
+						// this means it should be added after save because it doesn't exist
 						bShouldAttemptToSave = true;
 						bShouldAttemptToAdd = true;
 					}
diff --git a/Engine/Source/Editor/UnrealEd/Private/Editor.cpp b/Engine/Source/Editor/UnrealEd/Private/Editor.cpp
--- a/Engine/Source/Editor/UnrealEd/Private/Editor.cpp
+++ b/Engine/Source/Editor/UnrealEd/Private/Editor.cpp
@@ -363,7 +363,7 @@ bool FReimportManager::Reimport( UObject* Obj, bool bAskForNewFileIfMissing, boo
 	return bSuccess;
 }
 
-void FReimportManager::ValidateAllSourceFileAndReimport(TArray<UObject*> &ToImportObjects)
+void FReimportManager::ValidateAllSourceFileAndReimport(TArray<UObject*> &ToImportObjects, bool bShowNotification)
 {
 	//Copy the array to prevent iteration assert if a reimport factory change the selection
 	TArray<UObject*> CopyOfSelectedAssets;
@@ -456,7 +456,7 @@ void FReimportManager::ValidateAllSourceFileAndReimport(TArray<UObject*> &ToImpo
 		//If user ignore those asset just not add them to CopyOfSelectedAssets
 	}
 
-	FReimportManager::Instance()->ReimportMultiple(CopyOfSelectedAssets, /*bAskForNewFileIfMissing=*/false);
+	FReimportManager::Instance()->ReimportMultiple(CopyOfSelectedAssets, /*bAskForNewFileIfMissing=*/false, bShowNotification);
 }
 
 void FReimportManager::AddReferencedObjects( FReferenceCollector& Collector )
diff --git a/Engine/Source/Editor/UnrealEd/Private/EditorExporters.cpp b/Engine/Source/Editor/UnrealEd/Private/EditorExporters.cpp
--- a/Engine/Source/Editor/UnrealEd/Private/EditorExporters.cpp
+++ b/Engine/Source/Editor/UnrealEd/Private/EditorExporters.cpp
@@ -46,6 +46,7 @@
 #include "Exporters/StaticMeshExporterFBX.h"
 #include "Exporters/StaticMeshExporterOBJ.h"
 #include "Exporters/TextBufferExporterTXT.h"
+#include "Exporters/FbxExportOption.h"
 #include "Engine/StaticMesh.h"
 #include "Sound/SoundWave.h"
 #include "Engine/StaticMeshActor.h"
@@ -69,6 +70,7 @@
 #include "UnrealExporter.h"
 #include "InstancedFoliage.h"
 #include "Engine/Selection.h"
+#include "AssetExportTask.h"
 
 #include "Components/SkeletalMeshComponent.h"
 #include "Camera/CameraComponent.h"
@@ -1591,11 +1593,18 @@ bool ULevelExporterFBX::ExportBinary( UObject* Object, const TCHAR* Type, FArchi
 	GWarn->BeginSlowTask(NSLOCTEXT("UnrealEd", "ExportingLevelToFBX", "Exporting Level To FBX"), true);
 
 	UnFbx::FFbxExporter* Exporter = UnFbx::FFbxExporter::GetInstance();
-	
-	//Show the fbx export dialog options
 	bool ExportCancel = false;
-	bool ExportAll = false;
-	Exporter->FillExportOptions(false, true, UExporter::CurrentFilename, ExportCancel, ExportAll);
+	if (UFbxExportOption* AutomatedExportOptions = GetAutomatedExportOptionsFbx())
+	{
+		Exporter->SetExportOptionsOverride(AutomatedExportOptions);
+	}
+	else
+	{
+		//Show the fbx export dialog options
+		bool ExportAll = false;
+		Exporter->FillExportOptions(false, true, UExporter::CurrentFilename, ExportCancel, ExportAll);
+	}
+
 	if (!ExportCancel)
 	{
 		Exporter->CreateDocument();
@@ -1626,6 +1635,8 @@ bool ULevelExporterFBX::ExportBinary( UObject* Object, const TCHAR* Type, FArchi
 		Exporter->WriteToFile(*UExporter::CurrentFilename);
 	}
 
+	Exporter->SetExportOptionsOverride(nullptr);
+
 	GWarn->EndSlowTask();
 
 	return true;
@@ -1918,6 +1929,19 @@ bool UStaticMeshExporterOBJ::ExportText(const FExportObjectInnerContext* Context
 	return true;
 }
 
+/*------------------------------------------------------------------------------
+UExporterFBX implementation.
+------------------------------------------------------------------------------*/
+UFbxExportOption* UExporterFBX::GetAutomatedExportOptionsFbx()
+{
+	if (ExportTask && ExportTask->bAutomated)
+	{
+		return Cast<UFbxExportOption>(ExportTask->Options);
+	}
+
+	return nullptr;
+}
+
 /*------------------------------------------------------------------------------
 UStaticMeshExporterFBX implementation.
 ------------------------------------------------------------------------------*/
@@ -1938,21 +1962,30 @@ bool UStaticMeshExporterFBX::ExportBinary( UObject* Object, const TCHAR* Type, F
 {
 	UStaticMesh* StaticMesh = CastChecked<UStaticMesh>( Object );
 	UnFbx::FFbxExporter* Exporter = UnFbx::FFbxExporter::GetInstance();
-	//Show the fbx export dialog options
-	bool ExportAll = GetBatchMode() && !GetShowExportOption();
-	bool ExportCancel = false;
-	Exporter->FillExportOptions(GetBatchMode(), GetShowExportOption(), UExporter::CurrentFilename, ExportCancel, ExportAll);
-	if (ExportCancel)
+	if (UFbxExportOption* AutomatedExportOptions = GetAutomatedExportOptionsFbx())
 	{
-		SetCancelBatch(GetBatchMode());
-		//User cancel the FBX export
-		return false;
+		Exporter->SetExportOptionsOverride(AutomatedExportOptions);
+		SetShowExportOption(false);
+	}
+	else
+	{
+		//Show the fbx export dialog options
+		bool ExportAll = GetBatchMode() && !GetShowExportOption();
+		bool ExportCancel = false;
+		Exporter->FillExportOptions(GetBatchMode(), GetShowExportOption(), UExporter::CurrentFilename, ExportCancel, ExportAll);
+		if (ExportCancel)
+		{
+			SetCancelBatch(GetBatchMode());
+			//User cancel the FBX export
+			return false;
+		}
+		SetShowExportOption(!ExportAll);
 	}
-	SetShowExportOption(!ExportAll);
 
 	Exporter->CreateDocument();
 	Exporter->ExportStaticMesh(StaticMesh);
 	Exporter->WriteToFile(*UExporter::CurrentFilename);
+	Exporter->SetExportOptionsOverride(nullptr);
 
 	return true;
 }
@@ -1976,21 +2009,30 @@ bool USkeletalMeshExporterFBX::ExportBinary( UObject* Object, const TCHAR* Type,
 {
 	USkeletalMesh* SkeletalMesh = CastChecked<USkeletalMesh>( Object );
 	UnFbx::FFbxExporter* Exporter = UnFbx::FFbxExporter::GetInstance();
-	//Show the fbx export dialog options
-	bool ExportAll = GetBatchMode() && !GetShowExportOption();
-	bool ExportCancel = false;
-	Exporter->FillExportOptions(GetBatchMode(), GetShowExportOption(), UExporter::CurrentFilename, ExportCancel, ExportAll);
-	if (ExportCancel)
+	if (UFbxExportOption* AutomatedExportOptions = GetAutomatedExportOptionsFbx())
 	{
-		SetCancelBatch(GetBatchMode());
-		//User cancel the FBX export
-		return false;
+		Exporter->SetExportOptionsOverride(AutomatedExportOptions);
+		SetShowExportOption(false);
+	}
+	else
+	{
+		//Show the fbx export dialog options
+		bool ExportAll = GetBatchMode() && !GetShowExportOption();
+		bool ExportCancel = false;
+		Exporter->FillExportOptions(GetBatchMode(), GetShowExportOption(), UExporter::CurrentFilename, ExportCancel, ExportAll);
+		if (ExportCancel)
+		{
+			SetCancelBatch(GetBatchMode());
+			//User cancel the FBX export
+			return false;
+		}
+		SetShowExportOption(!ExportAll);
 	}
-	SetShowExportOption(!ExportAll);
 
 	Exporter->CreateDocument();
 	Exporter->ExportSkeletalMesh(SkeletalMesh);
 	Exporter->WriteToFile(*UExporter::CurrentFilename);
+	Exporter->SetExportOptionsOverride(nullptr);
 
 	return true;
 }
@@ -2018,21 +2060,30 @@ bool UAnimSequenceExporterFBX::ExportBinary( UObject* Object, const TCHAR* Type,
 	if (AnimSkeleton && PreviewMesh)
 	{
 		UnFbx::FFbxExporter* Exporter = UnFbx::FFbxExporter::GetInstance();
-		//Show the fbx export dialog options
-		bool ExportAll = GetBatchMode() && !GetShowExportOption();
-		bool ExportCancel = false;
-		Exporter->FillExportOptions(GetBatchMode(), GetShowExportOption(), UExporter::CurrentFilename, ExportCancel, ExportAll);
-		if (ExportCancel)
+		if (UFbxExportOption* AutomatedExportOptions = GetAutomatedExportOptionsFbx())
 		{
-			SetCancelBatch(GetBatchMode());
-			//User cancel the FBX export
-			return false;
+			Exporter->SetExportOptionsOverride(AutomatedExportOptions);
+			SetShowExportOption(false);
+		}
+		else
+		{
+			//Show the fbx export dialog options
+			bool ExportAll = GetBatchMode() && !GetShowExportOption();
+			bool ExportCancel = false;
+			Exporter->FillExportOptions(GetBatchMode(), GetShowExportOption(), UExporter::CurrentFilename, ExportCancel, ExportAll);
+			if (ExportCancel)
+			{
+				SetCancelBatch(GetBatchMode());
+				//User cancel the FBX export
+				return false;
+			}
+			SetShowExportOption(!ExportAll);
 		}
-		SetShowExportOption(!ExportAll);
 
 		Exporter->CreateDocument();
 		Exporter->ExportAnimSequence(AnimSequence, PreviewMesh, false);
 		Exporter->WriteToFile( *UExporter::CurrentFilename );
+		Exporter->SetExportOptionsOverride(nullptr);
 
 		return true;
 	}
diff --git a/Engine/Source/Editor/UnrealEd/Private/EditorLevelUtils.cpp b/Engine/Source/Editor/UnrealEd/Private/EditorLevelUtils.cpp
--- a/Engine/Source/Editor/UnrealEd/Private/EditorLevelUtils.cpp
+++ b/Engine/Source/Editor/UnrealEd/Private/EditorLevelUtils.cpp
@@ -1083,11 +1083,11 @@ void UEditorLevelUtils::GetWorlds(UWorld* InWorld, TArray<UWorld*>& OutWorlds, b
 				const ULevel* Level = StreamingLevel->GetLoadedLevel();
 
 				// This should always be the case for valid level names as the Editor preloads all packages.
-				if (Level != NULL)
+				if (Level)
 				{
 					// Newer levels have their packages' world as the outer.
 					UWorld* World = Cast<UWorld>(Level->GetOuter());
-					if (World != NULL)
+					if (World)
 					{
 						OutWorlds.AddUnique(World);
 					}
@@ -1097,14 +1097,13 @@ void UEditorLevelUtils::GetWorlds(UWorld* InWorld, TArray<UWorld*>& OutWorlds, b
 	}
 
 	// Levels can be loaded directly without StreamingLevel facilities
-	for (int32 LevelIndex = 0; LevelIndex < InWorld->GetLevels().Num(); ++LevelIndex)
+	for (ULevel* Level : InWorld->GetLevels())
 	{
-		ULevel* Level = InWorld->GetLevel(LevelIndex);
 		if (Level)
 		{
 			// Newer levels have their packages' world as the outer.
 			UWorld* World = Cast<UWorld>(Level->GetOuter());
-			if (World != NULL)
+			if (World)
 			{
 				OutWorlds.AddUnique(World);
 			}
diff --git a/Engine/Source/Editor/UnrealEd/Private/EditorServer.cpp b/Engine/Source/Editor/UnrealEd/Private/EditorServer.cpp
--- a/Engine/Source/Editor/UnrealEd/Private/EditorServer.cpp
+++ b/Engine/Source/Editor/UnrealEd/Private/EditorServer.cpp
@@ -16,6 +16,7 @@
 #include "ProfilingDebugging/ResourceSize.h"
 #include "UObject/Object.h"
 #include "UObject/GarbageCollection.h"
+#include "UObject/GCObjectScopeGuard.h"
 #include "UObject/Class.h"
 #include "UObject/UObjectIterator.h"
 #include "UObject/Package.h"
@@ -155,6 +156,7 @@
 #include "ILauncherPlatform.h"
 #include "LauncherPlatformModule.h"
 #include "HAL/PlatformApplicationMisc.h"
+#include "AssetExportTask.h"
 
 DEFINE_LOG_CATEGORY_STATIC(LogEditorServer, Log, All);
 
@@ -2771,17 +2773,18 @@ void UEditorEngine::ExportMap(UWorld* InWorld, const TCHAR* InFilename, bool bEx
 	const FText LocalizedExportingMap = FText::Format( NSLOCTEXT("UnrealEd", "ExportingMap_F", "Exporting map: {0}..." ), FText::FromString(MapFileName) );
 	GWarn->BeginSlowTask( LocalizedExportingMap, true);
 
-	UExporter::FExportToFileParams Params;
-	Params.Object = InWorld;
-	Params.Exporter = NULL;
-	Params.Filename = InFilename;
-	Params.InSelectedOnly = bExportSelectedActorsOnly;
-	Params.NoReplaceIdentical = false;
-	Params.Prompt = false;
-	Params.bUseFileArchive = false;
-	Params.WriteEmptyFiles = false;
-
-	UExporter::ExportToFileEx(Params);
+	UAssetExportTask* ExportTask = NewObject<UAssetExportTask>();
+	FGCObjectScopeGuard ExportTaskGuard(ExportTask);
+	ExportTask->Object = InWorld;
+	ExportTask->Exporter = NULL;
+	ExportTask->Filename = InFilename;
+	ExportTask->bSelected = bExportSelectedActorsOnly;
+	ExportTask->bReplaceIdentical = true;
+	ExportTask->bPrompt = false;
+	ExportTask->bUseFileArchive = false;
+	ExportTask->bWriteEmptyFiles = false;
+
+	UExporter::RunAssetExportTask(ExportTask);
 
 	GWarn->EndSlowTask();
 }
diff --git a/Engine/Source/Editor/UnrealEd/Private/EditorTransaction.cpp b/Engine/Source/Editor/UnrealEd/Private/EditorTransaction.cpp
--- a/Engine/Source/Editor/UnrealEd/Private/EditorTransaction.cpp
+++ b/Engine/Source/Editor/UnrealEd/Private/EditorTransaction.cpp
@@ -69,6 +69,8 @@ FTransaction::FObjectRecord::FObjectRecord(FTransaction* Owner, UObject* InObjec
 	,	Serializer			( InSerializer )
 	,	Destructor			( InDestructor )
 	,	bRestored			( false )
+	,   bFinalized			( false )
+	,	bSnapshot			( false )
 	,	bWantsBinarySerialization ( true )
 {
 	// Blueprint compile-in-place can alter class layout so use tagged serialization for objects relying on a UBlueprint's Class
@@ -81,23 +83,21 @@ FTransaction::FObjectRecord::FObjectRecord(FTransaction* Owner, UObject* InObjec
 	{
 		bWantsBinarySerialization = false;
 	}
-	ObjectAnnotation = Object->GetTransactionAnnotation();
-	FWriter Writer( Data, ReferencedObjects, ReferencedNames, bWantsBinarySerialization );
+	SerializedObject.SetObject(Object.Get());
+	FWriter Writer( SerializedObject, bWantsBinarySerialization );
 	SerializeContents( Writer, Oper );
 }
 
 void FTransaction::FObjectRecord::SerializeContents( FArchive& Ar, int32 InOper )
 {
-	// Cache to restore at the end
-	bool bWasArIgnoreOuterRef = Ar.ArIgnoreOuterRef;
-
-	if (Object.SubObjectHierarchyID.Num() != 0)
-	{
-		Ar.ArIgnoreOuterRef = true;
-	}
-
 	if( Array )
 	{
+		const bool bWasArIgnoreOuterRef = Ar.ArIgnoreOuterRef;
+		if (Object.SubObjectHierarchyID.Num() != 0)
+		{
+			Ar.ArIgnoreOuterRef = true;
+		}
+
 		//UE_LOG( LogEditorTransaction, Log, TEXT("Array %s %i*%i: %i"), Object ? *Object->GetFullName() : TEXT("Invalid Object"), Index, ElementSize, InOper);
 
 		check((SIZE_T)Array >= (SIZE_T)Object.Get() + sizeof(UObject));
@@ -139,6 +139,8 @@ void FTransaction::FObjectRecord::SerializeContents( FArchive& Ar, int32 InOper
 				Serializer( Ar, (uint8*)Array->GetData() + i*ElementSize );
 			}
 		}
+
+		Ar.ArIgnoreOuterRef = bWasArIgnoreOuterRef;
 	}
 	else
 	{
@@ -147,13 +149,25 @@ void FTransaction::FObjectRecord::SerializeContents( FArchive& Ar, int32 InOper
 		check(ElementSize==0);
 		check(DefaultConstructor==NULL);
 		check(Serializer==NULL);
-		// Once UE-46691 this should probably become an ensure
-		if (UObject* Obj = Object.Get())
+		SerializeObject( Ar );
+	}
+}
+
+void FTransaction::FObjectRecord::SerializeObject( FArchive& Ar )
+{
+	check(!Array);
+
+	UObject* CurrentObject = Object.Get();
+	if (CurrentObject)
+	{
+		const bool bWasArIgnoreOuterRef = Ar.ArIgnoreOuterRef;
+		if (Object.SubObjectHierarchyID.Num() != 0)
 		{
-			Obj->Serialize(Ar);
+			Ar.ArIgnoreOuterRef = true;
 		}
+		CurrentObject->Serialize(Ar);
+		Ar.ArIgnoreOuterRef = bWasArIgnoreOuterRef;
 	}
-	Ar.ArIgnoreOuterRef = bWasArIgnoreOuterRef;
 }
 
 void FTransaction::FObjectRecord::Restore( FTransaction* Owner )
@@ -163,7 +177,7 @@ void FTransaction::FObjectRecord::Restore( FTransaction* Owner )
 	{
 		bRestored = true;
 		check(!Owner->bFlip);
-		FTransaction::FObjectRecord::FReader Reader( Owner, Data, ReferencedObjects, ReferencedNames, bWantsBinarySerialization );
+		FReader Reader( Owner, SerializedObject, bWantsBinarySerialization );
 		SerializeContents( Reader, Oper );
 	}
 }
@@ -174,16 +188,15 @@ void FTransaction::FObjectRecord::Save(FTransaction* Owner)
 	check(Owner->bFlip);
 	if (!bRestored)
 	{
-		FlipData.Empty();
-		FlipReferencedObjects.Empty();
-		FlipReferencedNames.Empty();
-		FlipObjectAnnotation = TSharedPtr<ITransactionObjectAnnotation>();
-		// Once UE-46691 this should probably become an ensure
-		if (UObject* Obj = Object.Get())
+		SerializedObjectFlip.Reset();
+
+		UObject* CurrentObject = Object.Get();
+		if (CurrentObject)
 		{
-			FlipObjectAnnotation = Obj->GetTransactionAnnotation();
+			SerializedObjectFlip.SetObject(CurrentObject);
 		}
-		FWriter Writer(FlipData, FlipReferencedObjects, FlipReferencedNames, bWantsBinarySerialization);
+
+		FWriter Writer(SerializedObjectFlip, bWantsBinarySerialization);
 		SerializeContents(Writer, -Oper);
 	}
 }
@@ -195,16 +208,227 @@ void FTransaction::FObjectRecord::Load(FTransaction* Owner)
 	if (!bRestored)
 	{
 		bRestored = true;
-		FTransaction::FObjectRecord::FReader Reader(Owner, Data, ReferencedObjects, ReferencedNames, bWantsBinarySerialization);
+		FReader Reader(Owner, SerializedObject, bWantsBinarySerialization);
 		SerializeContents(Reader, Oper);
-		Exchange(ObjectAnnotation, FlipObjectAnnotation);
-		Exchange(Data, FlipData);
-		Exchange(ReferencedObjects, FlipReferencedObjects);
-		Exchange(ReferencedNames, FlipReferencedNames);
+		SerializedObject.Swap(SerializedObjectFlip);
 		Oper *= -1;
 	}
 }
 
+void FTransaction::FObjectRecord::Finalize( FTransaction* Owner )
+{
+	if (Array)
+	{
+		// Can only diff objects
+		return;
+	}
+
+	if (!bFinalized)
+	{
+		bFinalized = true;
+
+		UObject* CurrentObject = Object.Get();
+		if (CurrentObject)
+		{
+			// Serialize the object so we can diff it
+			FSerializedObject CurrentSerializedObject;
+			{
+				CurrentSerializedObject.SetObject(CurrentObject);
+				FWriter Writer(CurrentSerializedObject, bWantsBinarySerialization);
+				SerializeObject(Writer);
+			}
+			
+			// Diff against the object state when the transaction started
+			Diff(Owner, SerializedObject, CurrentSerializedObject, DeltaChange);
+
+			// If we have a previous snapshot then we need to consider that part of the diff for the finalized object, as systems may 
+			// have been tracking delta-changes between snapshots and this finalization will need to account for those changes too
+			if (bSnapshot)
+			{
+				Diff(Owner, SerializedObjectSnapshot, CurrentSerializedObject, DeltaChange);
+			}
+		}
+
+		// Clear out any snapshot data now as we won't be getting any more snapshot requests once finalized
+		bSnapshot = false;
+		SerializedObjectSnapshot.Reset();
+	}
+}
+
+void FTransaction::FObjectRecord::Snapshot( FTransaction* Owner )
+{
+	if (Array)
+	{
+		// Can only diff objects
+		return;
+	}
+
+	if (bFinalized)
+	{
+		// Cannot snapshot once finalized
+		return;
+	}
+
+	UObject* CurrentObject = Object.Get();
+	if (CurrentObject)
+	{
+		// Serialize the object so we can diff it
+		FSerializedObject CurrentSerializedObject;
+		{
+			CurrentSerializedObject.SetObject(CurrentObject);
+			FWriter Writer(CurrentSerializedObject, bWantsBinarySerialization);
+			SerializeObject(Writer);
+		}
+
+		// Diff against the correct serialized data depending on whether we already had a snapshot
+		const FSerializedObject& InitialSerializedObject = bSnapshot ? SerializedObjectSnapshot : SerializedObject;
+		FTransactionObjectDeltaChange SnapshotDeltaChange;
+		Diff(Owner, InitialSerializedObject, CurrentSerializedObject, SnapshotDeltaChange);
+
+		// Update the snapshot data for next time
+		bSnapshot = true;
+		SerializedObjectSnapshot.Swap(CurrentSerializedObject);
+
+		// Notify any listeners of this change
+		if (SnapshotDeltaChange.HasChanged())
+		{
+			CurrentObject->PostTransacted(FTransactionObjectEvent(ETransactionObjectEventType::Snapshot, SnapshotDeltaChange, InitialSerializedObject.ObjectAnnotation, InitialSerializedObject.ObjectName, InitialSerializedObject.ObjectPathName, InitialSerializedObject.ObjectOuterPathName));
+		}
+	}
+}
+
+void FTransaction::FObjectRecord::Diff( FTransaction* Owner, const FSerializedObject& OldSerializedObject, const FSerializedObject& NewSerializedObject, FTransactionObjectDeltaChange& OutDeltaChange )
+{
+	auto AreObjectPointersIdentical = [&OldSerializedObject, &NewSerializedObject](const FName InPropertyName)
+	{
+		TArray<int32> OldSerializedObjectIndices;
+		OldSerializedObject.SerializedObjectIndices.MultiFind(InPropertyName, OldSerializedObjectIndices, true);
+
+		TArray<int32> NewSerializedObjectIndices;
+		NewSerializedObject.SerializedObjectIndices.MultiFind(InPropertyName, NewSerializedObjectIndices, true);
+
+		bool bAreObjectPointersIdentical = OldSerializedObjectIndices.Num() == NewSerializedObjectIndices.Num();
+		if (bAreObjectPointersIdentical)
+		{
+			for (int32 ObjIndex = 0; ObjIndex < OldSerializedObjectIndices.Num() && bAreObjectPointersIdentical; ++ObjIndex)
+			{
+				const UObject* OldObjectPtr = OldSerializedObject.ReferencedObjects.IsValidIndex(OldSerializedObjectIndices[ObjIndex]) ? OldSerializedObject.ReferencedObjects[OldSerializedObjectIndices[ObjIndex]].Get() : nullptr;
+				const UObject* NewObjectPtr = NewSerializedObject.ReferencedObjects.IsValidIndex(NewSerializedObjectIndices[ObjIndex]) ? NewSerializedObject.ReferencedObjects[NewSerializedObjectIndices[ObjIndex]].Get() : nullptr;
+				bAreObjectPointersIdentical = OldObjectPtr == NewObjectPtr;
+			}
+		}
+		return bAreObjectPointersIdentical;
+	};
+
+	OutDeltaChange.bHasNameChange |= OldSerializedObject.ObjectName != NewSerializedObject.ObjectName;
+	OutDeltaChange.bHasOuterChange |= OldSerializedObject.ObjectOuterPathName != NewSerializedObject.ObjectOuterPathName;
+	OutDeltaChange.bHasPendingKillChange |= OldSerializedObject.bIsPendingKill != NewSerializedObject.bIsPendingKill;
+
+	// If the two have a different number of properties or object references then something structural was changed so we skip the property diff
+	const bool bHasStructuralChange = OldSerializedObject.SerializedProperties.Num() != NewSerializedObject.SerializedProperties.Num() || OldSerializedObject.SerializedObjectIndices.Num() != NewSerializedObject.SerializedObjectIndices.Num();
+	if (bHasStructuralChange)
+	{
+		OutDeltaChange.bHasNonPropertyChanges = true;
+		return;
+	}
+
+	if (!AreObjectPointersIdentical(NAME_None))
+	{
+		OutDeltaChange.bHasNonPropertyChanges = true;
+	}
+
+	if (OldSerializedObject.SerializedProperties.Num() > 0)
+	{
+		int32 StartOfOldPropertyBlock = INT_MAX;
+		int32 StartOfNewPropertyBlock = INT_MAX;
+		int32 EndOfOldPropertyBlock = -1;
+		int32 EndOfNewPropertyBlock = -1;
+
+		for (const TPair<FName, FSerializedProperty>& NewNamePropertyPair : NewSerializedObject.SerializedProperties)
+		{
+			const FSerializedProperty* OldSerializedProperty = OldSerializedObject.SerializedProperties.Find(NewNamePropertyPair.Key);
+			if (!OldSerializedProperty)
+			{
+				// Missing property, assume something structural changed
+				OutDeltaChange.bHasNonPropertyChanges = true;
+				continue;
+			}
+
+			// Update the tracking for the start/end of the property block within the serialized data
+			StartOfOldPropertyBlock = FMath::Min(StartOfOldPropertyBlock, OldSerializedProperty->DataOffset);
+			StartOfNewPropertyBlock = FMath::Min(StartOfNewPropertyBlock, NewNamePropertyPair.Value.DataOffset);
+			EndOfOldPropertyBlock = FMath::Max(EndOfOldPropertyBlock, OldSerializedProperty->DataOffset + OldSerializedProperty->DataSize);
+			EndOfNewPropertyBlock = FMath::Max(EndOfNewPropertyBlock, NewNamePropertyPair.Value.DataOffset + NewNamePropertyPair.Value.DataSize);
+
+			// Binary compare the serialized data to see if something has changed for this property
+			bool bIsPropertyIdentical = OldSerializedProperty->DataSize == NewNamePropertyPair.Value.DataSize;
+			if (bIsPropertyIdentical)
+			{
+				bIsPropertyIdentical = FMemory::Memcmp(&OldSerializedObject.Data[OldSerializedProperty->DataOffset], &NewSerializedObject.Data[NewNamePropertyPair.Value.DataOffset], NewNamePropertyPair.Value.DataSize) == 0;
+			}
+			if (bIsPropertyIdentical)
+			{
+				bIsPropertyIdentical = AreObjectPointersIdentical(NewNamePropertyPair.Key);
+			}
+
+			if (!bIsPropertyIdentical)
+			{
+				OutDeltaChange.ChangedProperties.AddUnique(NewNamePropertyPair.Key);
+			}
+		}
+
+		// Compare the data before the property block to see if something else in the object has changed
+		if (!OutDeltaChange.bHasNonPropertyChanges)
+		{
+			const int32 OldHeaderSize = StartOfOldPropertyBlock;
+			const int32 CurrentHeaderSize = StartOfNewPropertyBlock;
+
+			bool bIsHeaderIdentical = OldHeaderSize == CurrentHeaderSize;
+			if (bIsHeaderIdentical)
+			{
+				bIsHeaderIdentical = FMemory::Memcmp(&OldSerializedObject.Data[0], &NewSerializedObject.Data[0], CurrentHeaderSize) == 0;
+			}
+
+			if (!bIsHeaderIdentical)
+			{
+				OutDeltaChange.bHasNonPropertyChanges = true;
+			}
+		}
+
+		// Compare the data after the property block to see if something else in the object has changed
+		if (!OutDeltaChange.bHasNonPropertyChanges)
+		{
+			const int32 OldFooterSize = OldSerializedObject.Data.Num() - EndOfOldPropertyBlock;
+			const int32 CurrentFooterSize = NewSerializedObject.Data.Num() - EndOfNewPropertyBlock;
+
+			bool bIsFooterIdentical = OldFooterSize == CurrentFooterSize;
+			if (bIsFooterIdentical)
+			{
+				bIsFooterIdentical = FMemory::Memcmp(&OldSerializedObject.Data[EndOfOldPropertyBlock], &NewSerializedObject.Data[EndOfNewPropertyBlock], CurrentFooterSize) == 0;
+			}
+
+			if (!bIsFooterIdentical)
+			{
+				OutDeltaChange.bHasNonPropertyChanges = true;
+			}
+		}
+	}
+	else
+	{
+		// No properties, so just compare the whole blob
+		bool bIsBlobIdentical = OldSerializedObject.Data.Num() == NewSerializedObject.Data.Num();
+		if (bIsBlobIdentical)
+		{
+			bIsBlobIdentical = FMemory::Memcmp(&OldSerializedObject.Data[0], &NewSerializedObject.Data[0], NewSerializedObject.Data.Num()) == 0;
+		}
+
+		if (!bIsBlobIdentical)
+		{
+			OutDeltaChange.bHasNonPropertyChanges = true;
+		}
+	}
+}
+
 int32 FTransaction::GetRecordCount() const
 {
 	return Records.Num();
@@ -253,7 +477,7 @@ void FTransaction::RemoveRecords( int32 Count /* = 1  */ )
 void FTransaction::DumpObjectMap(FOutputDevice& Ar) const
 {
 	Ar.Logf( TEXT("===== DumpObjectMap %s ==== "), *Title.ToString() );
-	for ( ObjectMapType::TConstIterator It(ObjectMap) ; It ; ++It )
+	for ( auto It = ObjectMap.CreateConstIterator(); It; ++It )
 	{
 		const UObject* CurrentObject	= It.Key();
 		const int32 SaveCount				= It.Value();
@@ -266,9 +490,9 @@ FArchive& operator<<( FArchive& Ar, FTransaction::FObjectRecord& R )
 {
 	FMemMark Mark(FMemStack::Get());
 	Ar << R.Object;
-	Ar << R.Data;
-	Ar << R.ReferencedObjects;
-	Ar << R.ReferencedNames;
+	Ar << R.SerializedObject.Data;
+	Ar << R.SerializedObject.ReferencedObjects;
+	Ar << R.SerializedObject.ReferencedNames;
 	Mark.Pop();
 	return Ar;
 }
@@ -333,16 +557,16 @@ void FTransaction::FObjectRecord::AddReferencedObjects( FReferenceCollector& Col
 	Collector.AddReferencedObject(Obj);
 	Object.Object = Obj;
 
-	for (FPersistentObjectRef& ReferencedObject : ReferencedObjects)
+	for (FPersistentObjectRef& ReferencedObject : SerializedObject.ReferencedObjects)
 	{
 		UObject* RefObj = ReferencedObject.Object;
 		Collector.AddReferencedObject(RefObj);
 		ReferencedObject.Object = RefObj;
 	}
 
-	if (ObjectAnnotation.IsValid())
+	if (SerializedObject.ObjectAnnotation.IsValid())
 	{
-		ObjectAnnotation->AddReferencedObjects(Collector);
+		SerializedObject.ObjectAnnotation->AddReferencedObjects(Collector);
 	}
 }
 
@@ -357,7 +581,7 @@ bool FTransaction::FObjectRecord::ContainsPieObject() const
 		}
 	}
 
-	for (const FPersistentObjectRef& ReferencedObject : ReferencedObjects)
+	for (const FPersistentObjectRef& ReferencedObject : SerializedObject.ReferencedObjects)
 	{
 		const UObject* Obj = ReferencedObject.Object;
 		if( Obj && Obj->GetOutermost()->HasAnyPackageFlags(PKG_PlayInEditor))
@@ -427,6 +651,22 @@ void FTransaction::SetPrimaryObject(UObject* InObject)
 	}
 }
 
+void FTransaction::SnapshotObject( UObject* InObject )
+{
+	if (InObject && ObjectMap.Contains(InObject))
+	{
+		FObjectRecord* FoundObjectRecord = Records.FindByPredicate([InObject](const FObjectRecord& ObjRecord)
+		{
+			return ObjRecord.Object.Get() == InObject;
+		});
+
+		if (FoundObjectRecord)
+		{
+			FoundObjectRecord->Snapshot(this);
+		}
+	}
+}
+
 /**
  * Enacts the transaction.
  */
@@ -444,6 +684,13 @@ void FTransaction::Apply()
 		FObjectRecord& Record = Records[i];
 		Record.bRestored = false;
 
+		// Apply may be called before Finalize in order to revert an object back to its prior state in the case that a transaction is canceled
+		// In this case we still need to generate a diff for the transaction so that we notify correctly
+		if (!Record.bFinalized)
+		{
+			Record.Finalize(this);
+		}
+
 		UObject* Object = Record.Object.Get();
 		if (Object)
 		{
@@ -453,7 +700,7 @@ void FTransaction::Apply()
 				Object->PreEditUndo();
 			}
 
-			ChangedObjects.Add(Object, Record.ObjectAnnotation);
+			ChangedObjects.Add(Object, FChangedObjectValue(i, Record.SerializedObject.ObjectAnnotation));
 		}
 	}
 
@@ -502,7 +749,7 @@ void FTransaction::Apply()
 			LevelsToCommitModelSurface.AddUnique(Level);
 		}
 
-		TSharedPtr<ITransactionObjectAnnotation> ChangedObjectTransactionAnnotation = ChangedObjectIt.Value;
+		TSharedPtr<ITransactionObjectAnnotation> ChangedObjectTransactionAnnotation = ChangedObjectIt.Value.Annotation;
 		if (ChangedObjectTransactionAnnotation.IsValid())
 		{
 			ChangedObject->PostEditUndo(ChangedObjectTransactionAnnotation);
@@ -511,6 +758,14 @@ void FTransaction::Apply()
 		{
 			ChangedObject->PostEditUndo();
 		}
+
+		const FObjectRecord& ChangedObjectRecord = Records[ChangedObjectIt.Value.RecordIndex];
+		const FTransactionObjectDeltaChange& DeltaChange = ChangedObjectRecord.DeltaChange;
+		if (DeltaChange.HasChanged())
+		{
+			const FObjectRecord::FSerializedObject& InitialSerializedObject = ChangedObjectRecord.SerializedObject;
+			ChangedObject->PostTransacted(FTransactionObjectEvent(ETransactionObjectEventType::UndoRedo, DeltaChange, ChangedObjectTransactionAnnotation, InitialSerializedObject.ObjectName, InitialSerializedObject.ObjectPathName, InitialSerializedObject.ObjectOuterPathName));
+		}
 	}
 
 	// Commit model surfaces for unique levels within the transaction
@@ -530,7 +785,48 @@ void FTransaction::Apply()
 		ChangedObject->CheckDefaultSubobjects();
 	}
 
-	ChangedObjects.Empty();
+	ChangedObjects.Reset();
+}
+
+void FTransaction::Finalize()
+{
+	for (int32 i = 0; i < Records.Num(); ++i)
+	{
+		FObjectRecord& ObjectRecord = Records[i];
+		ObjectRecord.Finalize(this);
+
+		UObject* Object = ObjectRecord.Object.Get();
+		if (Object)
+		{
+			if (!ChangedObjects.Contains(Object))
+			{
+				ChangedObjects.Add(Object, FChangedObjectValue(i, ObjectRecord.SerializedObject.ObjectAnnotation));
+			}
+		}
+	}
+
+	// An Actor's components must always be notified before the owning Actor so do a quick sort
+	ChangedObjects.KeySort([](UObject& A, UObject& B)
+	{
+		UActorComponent* BAsComponent = Cast<UActorComponent>(&B);
+		return (BAsComponent ? (BAsComponent->GetOwner() != &A) : true);
+	});
+
+	for (auto ChangedObjectIt : ChangedObjects)
+	{
+		const FObjectRecord& ChangedObjectRecord = Records[ChangedObjectIt.Value.RecordIndex];
+		const FTransactionObjectDeltaChange& DeltaChange = ChangedObjectRecord.DeltaChange;
+		if (DeltaChange.HasChanged())
+		{
+			UObject* ChangedObject = ChangedObjectIt.Key;
+			TSharedPtr<ITransactionObjectAnnotation> ChangedObjectTransactionAnnotation = ChangedObjectIt.Value.Annotation;
+
+			const FObjectRecord::FSerializedObject& InitialSerializedObject = ChangedObjectRecord.SerializedObject;
+			ChangedObject->PostTransacted(FTransactionObjectEvent(ETransactionObjectEventType::Finalized, DeltaChange, ChangedObjectTransactionAnnotation, InitialSerializedObject.ObjectName, InitialSerializedObject.ObjectPathName, InitialSerializedObject.ObjectOuterPathName));
+		}
+	}
+
+	ChangedObjects.Reset();
 }
 
 SIZE_T FTransaction::DataSize() const
@@ -538,7 +834,7 @@ SIZE_T FTransaction::DataSize() const
 	SIZE_T Result=0;
 	for( int32 i=0; i<Records.Num(); i++ )
 	{
-		Result += Records[i].Data.Num();
+		Result += Records[i].SerializedObject.Data.Num();
 	}
 	return Result;
 }
@@ -658,6 +954,10 @@ int32 UTransBuffer::End()
 				static_cast<FTransaction*>(GUndo)->DumpObjectMap( *GLog );
 			}
 #endif
+			if (GUndo)
+			{
+				GUndo->Finalize();
+			}
 			GUndo = nullptr;
 			PreviousUndoCount = INDEX_NONE;
 			RemovedTransactions.Reset();
diff --git a/Engine/Source/Editor/UnrealEd/Private/Factories/Factory.cpp b/Engine/Source/Editor/UnrealEd/Private/Factories/Factory.cpp
--- a/Engine/Source/Editor/UnrealEd/Private/Factories/Factory.cpp
+++ b/Engine/Source/Editor/UnrealEd/Private/Factories/Factory.cpp
@@ -17,6 +17,7 @@
 #include "AssetToolsModule.h"
 #include "EditorClassUtils.h"
 #include "AutomatedAssetImportData.h"
+#include "AssetImportTask.h"
 
 DEFINE_LOG_CATEGORY_STATIC(LogFactory, Log, All);
 
@@ -295,6 +296,10 @@ void UFactory::DisplayOverwriteOptionsDialog(const FText& Message)
 	{
 		OverwriteYesOrNoToAllState =  AutomatedImportData->bReplaceExisting ? EAppReturnType::YesAll : EAppReturnType::NoAll;
 	}
+	else if (AssetImportTask && AssetImportTask->bAutomated)
+	{
+		OverwriteYesOrNoToAllState = AssetImportTask->bReplaceExisting ? EAppReturnType::YesAll : EAppReturnType::NoAll;
+	}
 	else if (OverwriteYesOrNoToAllState != EAppReturnType::YesAll && OverwriteYesOrNoToAllState != EAppReturnType::NoAll)
 	{
 		OverwriteYesOrNoToAllState = FMessageDialog::Open(EAppMsgType::YesNoYesAllNoAllCancel, FText::Format(
@@ -607,3 +612,13 @@ void UFactory::SetAutomatedAssetImportData(const UAutomatedAssetImportData* Data
 {
 	AutomatedImportData = Data;
 }
+
+void UFactory::SetAssetImportTask(const UAssetImportTask* Task)
+{
+	AssetImportTask = Task;
+}
+
+bool UFactory::IsAutomatedImport() const
+{
+	return GIsAutomationTesting || AutomatedImportData || (AssetImportTask && AssetImportTask->bAutomated);
+}
\ No newline at end of file
diff --git a/Engine/Source/Editor/UnrealEd/Private/Fbx/FbxAnimationExport.cpp b/Engine/Source/Editor/UnrealEd/Private/Fbx/FbxAnimationExport.cpp
--- a/Engine/Source/Editor/UnrealEd/Private/Fbx/FbxAnimationExport.cpp
+++ b/Engine/Source/Editor/UnrealEd/Private/Fbx/FbxAnimationExport.cpp
@@ -439,7 +439,7 @@ void FFbxExporter::ExportAnimTrack(IAnimTrackAdapter& AnimTrackAdapter, AActor*
 
 			FTransform BoneTransform = SkeletalMeshComponent->BoneSpaceTransforms[BoneIndex];
 
-			if (ExportOptions->MapSkeletalMotionToRoot && BoneIndex == 0)
+			if (GetExportOptions()->MapSkeletalMotionToRoot && BoneIndex == 0)
 			{
 				BoneTransform = SkeletalMeshComponent->GetSocketTransform(BoneName) * InitialInvParentTransform;
 			}
diff --git a/Engine/Source/Editor/UnrealEd/Private/Fbx/FbxExportOption.cpp b/Engine/Source/Editor/UnrealEd/Private/Fbx/FbxExportOption.cpp
--- a/Engine/Source/Editor/UnrealEd/Private/Fbx/FbxExportOption.cpp
+++ b/Engine/Source/Editor/UnrealEd/Private/Fbx/FbxExportOption.cpp
@@ -15,6 +15,7 @@ UFbxExportOption::UFbxExportOption(const FObjectInitializer& ObjectInitializer)
 	: Super(ObjectInitializer)
 {
 	FbxExportCompatibility = EFbxExportCompatibility::FBX_2013;
+	bASCII = false;
 	bForceFrontXAxis = false;
 	LevelOfDetail = true;
 	Collision = true;
diff --git a/Engine/Source/Editor/UnrealEd/Private/Fbx/FbxFactory.cpp b/Engine/Source/Editor/UnrealEd/Private/Fbx/FbxFactory.cpp
--- a/Engine/Source/Editor/UnrealEd/Private/Fbx/FbxFactory.cpp
+++ b/Engine/Source/Editor/UnrealEd/Private/Fbx/FbxFactory.cpp
@@ -24,6 +24,8 @@
 #include "AssetRegistryModule.h"
 #include "ObjectTools.h"
 #include "JsonObjectConverter.h"
+#include "AssetImportTask.h"
+#include "HAL/FileManager.h"
 
 #define LOCTEXT_NAMESPACE "FBXFactory"
 
@@ -87,6 +89,7 @@ bool UFbxFactory::DetectImportType(const FString& InFilename)
 	if ( ImportType == -1)
 	{
 		FFbxImporter->AddTokenizedErrorMessage(FTokenizedMessage::Create(EMessageSeverity::Warning, LOCTEXT("NoImportTypeDetected", "Can't detect import type. No mesh is found or animation track.")), FFbxErrors::Generic_CannotDetectImportType);
+		FFbxImporter->ReleaseScene();
 		return false;
 	}
 	else if(!IsAutomatedImport() || ImportUI->bAutomatedImportShouldDetectType)
@@ -144,20 +147,29 @@ bool UFbxFactory::ConfigureProperties()
 	return true;
 }
 
-UObject* UFbxFactory::FactoryCreateBinary
+UObject* UFbxFactory::FactoryCreateFile
 (
- UClass*			Class,
- UObject*			InParent,
- FName				Name,
- EObjectFlags		Flags,
- UObject*			Context,
- const TCHAR*		Type,
- const uint8*&		Buffer,
- const uint8*		BufferEnd,
- FFeedbackContext*	Warn,
- bool&				bOutOperationCanceled
+ UClass* Class,
+ UObject* InParent,
+ FName Name,
+ EObjectFlags Flags,
+ const FString& InFilename,
+ const TCHAR* Parms,
+ FFeedbackContext* Warn,
+ bool& bOutOperationCanceled
  )
 {
+	FString FileExtension = FPaths::GetExtension(InFilename);
+	const TCHAR* Type = *FileExtension;
+
+	if (!IFileManager::Get().FileExists(*InFilename))
+	{
+		UE_LOG(LogFbx, Error, TEXT("Failed to load file '%s'"), *InFilename)
+		return nullptr;
+	}
+
+	ParseParms(Parms);
+
 	CA_ASSUME(InParent);
 
 	if( bOperationCanceled )
@@ -201,7 +213,8 @@ UObject* UFbxFactory::FactoryCreateBinary
 				//Set the new fbx source path before starting the re-import
 				FReimportManager::Instance()->UpdateReimportPaths(ObjectToReimport, Filenames);
 				//Do the re-import and exit
-				FReimportManager::Instance()->ValidateAllSourceFileAndReimport(ToReimportObjects);
+				const bool bShowNotification = !(AssetImportTask && AssetImportTask->bAutomated);
+				FReimportManager::Instance()->ValidateAllSourceFileAndReimport(ToReimportObjects, bShowNotification);
 				return ObjectToReimport;
 			}
 		}
@@ -239,6 +252,30 @@ UObject* UFbxFactory::FactoryCreateBinary
 		bIsObjFormat = true;
 	}
 
+	struct FRestoreImportUI
+	{
+		FRestoreImportUI(UFbxFactory* InFbxFactory)
+			: FbxFactory(InFbxFactory) 
+		{
+			ensure(FbxFactory->OriginalImportUI == nullptr);
+			FbxFactory->OriginalImportUI = FbxFactory->ImportUI;
+		}
+
+		~FRestoreImportUI()
+		{
+			FbxFactory->ImportUI = FbxFactory->OriginalImportUI;
+			FbxFactory->OriginalImportUI = nullptr;
+		}
+
+	private:
+		UFbxFactory* FbxFactory;
+	};
+	FRestoreImportUI RestoreImportUI(this);
+	UFbxImportUI* OverrideImportUI = AssetImportTask ? Cast<UFbxImportUI>(AssetImportTask->Options) : nullptr;
+	if (OverrideImportUI)
+	{
+		ImportUI = OverrideImportUI;
+	}
 
 	// Show the import dialog only when not in a "yes to all" state or when automating import
 	bool bIsAutomated = IsAutomatedImport();
@@ -661,6 +698,10 @@ UObject* UFbxFactory::FactoryCreateBinary
 		FbxImporter->ReleaseScene();
 		Warn->EndSlowTask();
 	}
+	else // ImportOptions == NULL
+	{
+		FbxImporter->ReleaseScene();
+	}
 
 	FEditorDelegates::OnAssetPostImport.Broadcast(this, NewObject);
 
diff --git a/Engine/Source/Editor/UnrealEd/Private/Fbx/FbxMainExport.cpp b/Engine/Source/Editor/UnrealEd/Private/Fbx/FbxMainExport.cpp
--- a/Engine/Source/Editor/UnrealEd/Private/Fbx/FbxMainExport.cpp
+++ b/Engine/Source/Editor/UnrealEd/Private/Fbx/FbxMainExport.cpp
@@ -104,9 +104,11 @@ TSharedPtr<FFbxExporter> FFbxExporter::StaticInstance;
 FFbxExporter::FFbxExporter()
 {
 	//We use the FGCObject pattern to keep the fbx export option alive during the editor session
-	ExportOptions = NewObject<UFbxExportOption>();
+	ExportOptionsUI = NewObject<UFbxExportOption>();
 	//Load the option from the user save ini file
-	ExportOptions->LoadOptions();
+	ExportOptionsUI->LoadOptions();
+
+	ExportOptionsOverride = nullptr;
 
 	// Create the SdkManager
 	SdkManager = FbxManager::Create();
@@ -141,15 +143,28 @@ void FFbxExporter::DeleteInstance()
 	StaticInstance.Reset();
 }
 
+void FFbxExporter::AddReferencedObjects(FReferenceCollector& Collector)
+{
+	if (ExportOptionsUI != nullptr)
+	{
+		Collector.AddReferencedObject(ExportOptionsUI);
+	}
+
+	if (ExportOptionsOverride != nullptr)
+	{
+		Collector.AddReferencedObject(ExportOptionsOverride);
+	}
+}
+
 void FFbxExporter::FillExportOptions(bool BatchMode, bool bShowOptionDialog, const FString& FullPath, bool& OutOperationCanceled, bool& bOutExportAll)
 {
 	OutOperationCanceled = false;
 	
 	//Export option should have been set in the constructor
-	check(ExportOptions != nullptr);
+	check(ExportOptionsUI != nullptr);
 	
 	//Load the option from the user save ini file
-	ExportOptions->LoadOptions();
+	ExportOptionsUI->LoadOptions();
 	
 	//Return if we do not show the export options or we are running automation test or we are unattended
 	if (!bShowOptionDialog || GIsAutomationTesting || FApp::IsUnattended())
@@ -177,14 +192,14 @@ void FFbxExporter::FillExportOptions(bool BatchMode, bool bShowOptionDialog, con
 	Window->SetContent
 	(
 		SAssignNew(FbxOptionWindow, SFbxExportOptionsWindow)
-		.ExportOptions(ExportOptions)
+		.ExportOptions(ExportOptionsUI)
 		.WidgetWindow(Window)
 		.FullPath(FText::FromString(FullPath))
 		.BatchMode(BatchMode)
 	);
 
 	FSlateApplication::Get().AddModalWindow(Window, ParentWindow, false);
-	ExportOptions->SaveOptions();
+	ExportOptionsUI->SaveOptions();
 
 	if (FbxOptionWindow->ShouldExport())
 	{
@@ -196,6 +211,15 @@ void FFbxExporter::FillExportOptions(bool BatchMode, bool bShowOptionDialog, con
 	}
 }
 
+void FFbxExporter::SetExportOptionsOverride(UFbxExportOption* OverrideOptions)
+{
+	ExportOptionsOverride = OverrideOptions;
+}
+
+UFbxExportOption* FFbxExporter::GetExportOptions()
+{
+	return ExportOptionsOverride ? ExportOptionsOverride : ExportOptionsUI;
+}
 
 void FFbxExporter::CreateDocument()
 {
@@ -216,7 +240,7 @@ void FFbxExporter::CreateDocument()
 	
 	//FbxScene->GetGlobalSettings().SetOriginalUpAxis(KFbxAxisSystem::Max);
 	FbxAxisSystem::EFrontVector FrontVector = (FbxAxisSystem::EFrontVector)-FbxAxisSystem::eParityOdd;
-	if (ExportOptions->bForceFrontXAxis)
+	if (GetExportOptions()->bForceFrontXAxis)
 		FrontVector = FbxAxisSystem::eParityEven;
 
 	const FbxAxisSystem UnrealZUp(FbxAxisSystem::eZAxis, FrontVector, FbxAxisSystem::eRightHanded);
@@ -248,13 +272,21 @@ void FFbxExporter::WriteToFile(const TCHAR* Filename)
 
 	int32 FileFormat = -1;
 	bool bEmbedMedia = false;
+	bool bASCII = GetExportOptions()->bASCII;
 
 	// Create an exporter.
 	FbxExporter* Exporter = FbxExporter::Create(SdkManager, "");
 
 	// set file format
 	// Write in fall back format if pEmbedMedia is true
-	FileFormat = SdkManager->GetIOPluginRegistry()->GetNativeWriterFormat();
+	if (bASCII)
+	{
+		FileFormat = SdkManager->GetIOPluginRegistry()->FindWriterIDByDescription("FBX ascii (*.fbx)");
+	}
+	else
+	{
+		FileFormat = SdkManager->GetIOPluginRegistry()->GetNativeWriterFormat();
+	}
 
 	// Set the export states. By default, the export states are always set to 
 	// true except for the option eEXPORT_TEXTURE_AS_EMBEDDED. The code below 
@@ -267,10 +299,11 @@ void FFbxExporter::WriteToFile(const TCHAR* Filename)
 	IOS_REF.SetBoolProp(EXP_FBX_GOBO,            true);
 	IOS_REF.SetBoolProp(EXP_FBX_ANIMATION,       true);
 	IOS_REF.SetBoolProp(EXP_FBX_GLOBAL_SETTINGS, true);
+	IOS_REF.SetBoolProp(EXP_ASCIIFBX,            bASCII);
 
 	//Get the compatibility from the editor settings
 	const char* CompatibilitySetting = FBX_2013_00_COMPATIBLE;
-	const EFbxExportCompatibility FbxExportCompatibility = ExportOptions->FbxExportCompatibility;
+	const EFbxExportCompatibility FbxExportCompatibility = GetExportOptions()->FbxExportCompatibility;
 	switch (FbxExportCompatibility)
 	{
 		case EFbxExportCompatibility::FBX_2011:
@@ -824,7 +857,7 @@ void FFbxExporter::ExportStaticMesh(AActor* Actor, UStaticMeshComponent* StaticM
 	FString FbxMeshName = StaticMesh->GetName().Replace(TEXT("-"), TEXT("_"));
 	FColorVertexBuffer* ColorBuffer = NULL;
 
-	if(ExportOptions->LevelOfDetail && StaticMesh->GetNumLODs() > 1)
+	if(GetExportOptions()->LevelOfDetail && StaticMesh->GetNumLODs() > 1)
 	{
 		//Create a fbx LOD Group node
 		FbxNode* FbxActor = ExportActor(Actor, false, NodeNameAdapter);
@@ -1040,7 +1073,7 @@ void FFbxExporter::ExportStaticMesh( UStaticMesh* StaticMesh, const TArray<FStat
 	FbxNode* MeshNode = FbxNode::Create(Scene, TCHAR_TO_UTF8(*MeshName));
 	Scene->GetRootNode()->AddChild(MeshNode);
 
-	if (ExportOptions->LevelOfDetail && StaticMesh->GetNumLODs() > 1)
+	if (GetExportOptions()->LevelOfDetail && StaticMesh->GetNumLODs() > 1)
 	{
 		FString LodGroup_MeshName = MeshName + ("_LodGroup");
 		FbxLODGroup *FbxLodGroupAttribute = FbxLODGroup::Create(Scene, TCHAR_TO_UTF8(*LodGroup_MeshName));
@@ -1509,7 +1542,7 @@ bool FFbxExporter::ExportLevelSequence( UMovieScene* MovieScene, const TArray<FG
 				{
 					USkeletalMeshComponent* SkeletalMeshComp = Cast<USkeletalMeshComponent>( Actor->GetComponentByClass( USkeletalMeshComponent::StaticClass() ) );
 					
-					const bool bSkip3DTransformTrack = SkeletalMeshComp && ExportOptions->MapSkeletalMotionToRoot;
+					const bool bSkip3DTransformTrack = SkeletalMeshComp && GetExportOptions()->MapSkeletalMotionToRoot;
 
 					// Look for the tracks that we currently support
 					for ( UMovieSceneTrack* Track : MovieSceneBinding.GetTracks() )
@@ -3372,7 +3405,7 @@ FbxNode* FFbxExporter::ExportStaticMeshToFbx(const UStaticMesh* StaticMesh, int3
 		TArray<int32> VertRemap;
 		TArray<int32> UniqueVerts;
 
-		if (ExportOptions->WeldedVertices)
+		if (GetExportOptions()->WeldedVertices)
 		{
 			// Weld verts
 			DetermineVertsToWeld(VertRemap, UniqueVerts, RenderMesh);
@@ -3516,7 +3549,7 @@ FbxNode* FFbxExporter::ExportStaticMeshToFbx(const UStaticMesh* StaticMesh, int3
 
 			TArray<int32> UvsRemap;
 			TArray<int32> UniqueUVs;
-			if (ExportOptions->WeldedVertices)
+			if (GetExportOptions()->WeldedVertices)
 			{
 				// Weld UVs
 				DetermineUVsToWeld(UvsRemap, UniqueUVs, RenderMesh.VertexBuffers.StaticMeshVertexBuffer, TexCoordSourceIndex);
@@ -3629,7 +3662,7 @@ FbxNode* FFbxExporter::ExportStaticMeshToFbx(const UStaticMesh* StaticMesh, int3
 		uint32 ColorVertexCount = ColorBufferToUse->GetNumVertices();
 
 		// Only export vertex colors if they exist
-		if (ExportOptions->VertexColor && ColorVertexCount > 0)
+		if (GetExportOptions()->VertexColor && ColorVertexCount > 0)
 		{
 			FbxLayerElementVertexColor* VertexColor = FbxLayerElementVertexColor::Create(Mesh, "");
 			VertexColor->SetMappingMode(FbxLayerElement::eByPolygonVertex);
@@ -3684,7 +3717,7 @@ FbxNode* FFbxExporter::ExportStaticMeshToFbx(const UStaticMesh* StaticMesh, int3
 	}
 
 #if WITH_PHYSX
-	if ((ExportLOD == 0 || ExportLOD == -1) && ExportOptions->Collision)
+	if ((ExportLOD == 0 || ExportLOD == -1) && GetExportOptions()->Collision)
 	{
 		ExportCollisionMesh(StaticMesh, MeshName, FbxActor);
 	}
@@ -3720,7 +3753,7 @@ void FFbxExporter::ExportSplineMeshToFbx(const USplineMeshComponent* SplineMeshC
 	TArray<int32> VertRemap;
 	TArray<int32> UniqueVerts;
 
-	if (ExportOptions->WeldedVertices)
+	if (GetExportOptions()->WeldedVertices)
 	{
 		// Weld verts
 		DetermineVertsToWeld(VertRemap, UniqueVerts, RenderMesh);
@@ -3840,7 +3873,7 @@ void FFbxExporter::ExportSplineMeshToFbx(const USplineMeshComponent* SplineMeshC
 
 		TArray<int32> UvsRemap;
 		TArray<int32> UniqueUVs;
-		if (ExportOptions->WeldedVertices)
+		if (GetExportOptions()->WeldedVertices)
 		{
 			// Weld UVs
 			DetermineUVsToWeld(UvsRemap, UniqueUVs, RenderMesh.VertexBuffers.StaticMeshVertexBuffer, TexCoordSourceIndex);
@@ -3927,7 +3960,7 @@ void FFbxExporter::ExportSplineMeshToFbx(const USplineMeshComponent* SplineMeshC
 	uint32 ColorVertexCount = ColorBufferToUse->GetNumVertices();
 
 	// Only export vertex colors if they exist
-	if (ExportOptions->VertexColor && ColorVertexCount > 0)
+	if (GetExportOptions()->VertexColor && ColorVertexCount > 0)
 	{
 		FbxLayerElementVertexColor* VertexColor = FbxLayerElementVertexColor::Create(Mesh, "");
 		VertexColor->SetMappingMode(FbxLayerElement::eByPolygonVertex);
diff --git a/Engine/Source/Editor/UnrealEd/Private/Fbx/FbxMainImport.cpp b/Engine/Source/Editor/UnrealEd/Private/Fbx/FbxMainImport.cpp
--- a/Engine/Source/Editor/UnrealEd/Private/Fbx/FbxMainImport.cpp
+++ b/Engine/Source/Editor/UnrealEd/Private/Fbx/FbxMainImport.cpp
@@ -506,22 +506,9 @@ int32 FFbxImporter::GetImportType(const FString& InFilename)
 	FString Filename = InFilename;
 
 	// Prioritized in the order of SkeletalMesh > StaticMesh > Animation (only if animation data is found)
-	if (OpenFile(Filename, true))
+	if (OpenFile(Filename))
 	{
-		FbxStatistics Statistics;
-		Importer->GetStatistics(&Statistics); //-V595
-		int32 ItemIndex;
-		FbxString ItemName;
-		int32 ItemCount;
 		bool bHasAnimation = false;
-
-		for ( ItemIndex = 0; ItemIndex < Statistics.GetNbItems(); ItemIndex++ )
-		{
-			Statistics.GetItemPair(ItemIndex, ItemName, ItemCount);
-			const FString NameBuffer(ItemName.Buffer());
-			UE_LOG(LogFbx, Log, TEXT("ItemName: %s, ItemCount : %d"), *NameBuffer, ItemCount);
-		}
-
 		FbxSceneInfo SceneInfo;
 		if (GetSceneInfo(Filename, SceneInfo, true))
 		{
@@ -537,14 +524,6 @@ int32 FFbxImporter::GetImportType(const FString& InFilename)
 			bHasAnimation = SceneInfo.bHasAnimation;
 		}
 
-		if (Importer)
-		{
-			Importer->Destroy();
-		}
-
-		Importer = NULL;
-		CurPhase = NOTSTARTED;
-
 		// In case no Geometry was found, check for animation (FBX can still contain mesh data though)
 		if (bHasAnimation)
 		{
@@ -572,7 +551,7 @@ bool FFbxImporter::GetSceneInfo(FString Filename, FbxSceneInfo& SceneInfo, bool
 	switch (CurPhase)
 	{
 	case NOTSTARTED:
-		if (!OpenFile( Filename, false, bSceneInfo ))
+		if (!OpenFile(Filename))
 		{
 			Result = false;
 			break;
@@ -867,7 +846,7 @@ void FFbxImporter::TraverseHierarchyNodeRecursively(FbxSceneInfo& SceneInfo, Fbx
 	}
 }
 
-bool FFbxImporter::OpenFile(FString Filename, bool bParseStatistics, bool bForSceneInfo )
+bool FFbxImporter::OpenFile(FString Filename)
 {
 	bool Result = true;
 	
@@ -892,11 +871,6 @@ bool FFbxImporter::OpenFile(FString Filename, bool bParseStatistics, bool bForSc
 	FbxManager::GetFileFormatVersion(SDKMajor, SDKMinor, SDKRevision);
 
 	// Initialize the importer by providing a filename.
-	if (bParseStatistics)
-	{
-		Importer->ParseForStatistics(true);
-	}
-	
 	const bool bImportStatus = Importer->Initialize(TCHAR_TO_UTF8(*Filename));
 	
 	FbxCreator = EFbxCreator::Unknow;
@@ -929,29 +903,22 @@ bool FFbxImporter::OpenFile(FString Filename, bool bParseStatistics, bool bForSc
 		return false;
 	}
 
-	// Skip the version check if we are just parsing for information or scene info.
-	if( !bParseStatistics && !bForSceneInfo )
+	// Version out of date warning
+	int32 FileMajor = 0, FileMinor = 0, FileRevision = 0;
+	Importer->GetFileVersion(FileMajor, FileMinor, FileRevision);
+	int32 FileVersion = (FileMajor << 16 | FileMinor << 8 | FileRevision);
+	int32 SDKVersion = (SDKMajor << 16 | SDKMinor << 8 | SDKRevision);
+	if( FileVersion != SDKVersion )
 	{
-		int32 FileMajor,  FileMinor,  FileRevision;
-		Importer->GetFileVersion(FileMajor, FileMinor, FileRevision);
-
-		int32 FileVersion = (FileMajor << 16 | FileMinor << 8 | FileRevision);
-		int32 SDKVersion = (SDKMajor << 16 | SDKMinor << 8 | SDKRevision);
-
-		if( FileVersion != SDKVersion )
-		{
-
-			// Appending the SDK version to the config key causes the warning to automatically reappear even if previously suppressed when the SDK version we use changes. 
-			FString ConfigStr = FString::Printf( TEXT("Warning_OutOfDateFBX_%d"), SDKVersion );
-
-			FString FileVerStr = FString::Printf( TEXT("%d.%d.%d"), FileMajor, FileMinor, FileRevision );
-			FString SDKVerStr  = FString::Printf( TEXT("%d.%d.%d"), SDKMajor, SDKMinor, SDKRevision );
+		// Appending the SDK version to the config key causes the warning to automatically reappear even if previously suppressed when the SDK version we use changes. 
+		FString ConfigStr = FString::Printf( TEXT("Warning_OutOfDateFBX_%d"), SDKVersion );
 
-			const FText WarningText = FText::Format(
-				NSLOCTEXT("UnrealEd", "Warning_OutOfDateFBX", "An out of date FBX has been detected.\nImporting different versions of FBX files than the SDK version can cause undesirable results.\n\nFile Version: {0}\nSDK Version: {1}" ),
-				FText::FromString(FileVerStr), FText::FromString(SDKVerStr) );
+		FString FileVerStr = FString::Printf( TEXT("%d.%d.%d"), FileMajor, FileMinor, FileRevision );
+		FString SDKVerStr  = FString::Printf( TEXT("%d.%d.%d"), SDKMajor, SDKMinor, SDKRevision );
 
-		}
+		const FText WarningText = FText::Format(
+			NSLOCTEXT("UnrealEd", "Warning_OutOfDateFBX", "An out of date FBX has been detected.\nImporting different versions of FBX files than the SDK version can cause undesirable results.\n\nFile Version: {0}\nSDK Version: {1}" ),
+			FText::FromString(FileVerStr), FText::FromString(SDKVerStr) );
 	}
 
 	//Cache the current file hash
@@ -1006,6 +973,12 @@ void FFbxImporter::FixMaterialClashName()
 
 bool FFbxImporter::ImportFile(FString Filename, bool bPreventMaterialNameClash /*=false*/)
 {
+	if (Scene)
+	{
+		UE_LOG(LogFbx, Error, TEXT("FBX Scene already loaded from %s"), *Filename);
+		return false;
+	}
+
 	bool Result = true;
 	
 	bool bStatus;
@@ -1123,7 +1096,7 @@ bool FFbxImporter::ImportFromFile(const FString& Filename, const FString& Type,
 	switch (CurPhase)
 	{
 	case NOTSTARTED:
-		if (!OpenFile(FString(Filename), false))
+		if (!OpenFile(FString(Filename)))
 		{
 			Result = false;
 			break;
diff --git a/Engine/Source/Editor/UnrealEd/Private/Fbx/FbxSkeletalMeshExport.cpp b/Engine/Source/Editor/UnrealEd/Private/Fbx/FbxSkeletalMeshExport.cpp
--- a/Engine/Source/Editor/UnrealEd/Private/Fbx/FbxSkeletalMeshExport.cpp
+++ b/Engine/Source/Editor/UnrealEd/Private/Fbx/FbxSkeletalMeshExport.cpp
@@ -220,7 +220,7 @@ FbxNode* FFbxExporter::CreateMesh(const USkeletalMesh* SkelMesh, const TCHAR* Me
 		}
 	}
 
-	if (ExportOptions->VertexColor)
+	if (GetExportOptions()->VertexColor)
 	{
 		// Create and fill in the vertex color data source.
 		FbxLayerElementVertexColor* VertexColor = FbxLayerElementVertexColor::Create(Mesh, "");
