Fix crash when re-importing skeletal mesh. Skinned component render data resource is now release when re-importing.

Adding this to UE4/Main to unblock several branches that need it

diff --git a/Engine/Source/Editor/UnrealEd/Private/Fbx/FbxSkeletalMeshImport.cpp b/Engine/Source/Editor/UnrealEd/Private/Fbx/FbxSkeletalMeshImport.cpp
--- a/Engine/Source/Editor/UnrealEd/Private/Fbx/FbxSkeletalMeshImport.cpp
+++ b/Engine/Source/Editor/UnrealEd/Private/Fbx/FbxSkeletalMeshImport.cpp
@@ -1474,6 +1474,11 @@ USkeletalMesh* UnFbx::FFbxImporter::ImportSkeletalMesh(FImportSkeletalMeshArgs &
 		return nullptr;
 	}
 
+	// Array of re-import contexts for components using this mesh
+	// Will unregister before import, then re-register afterwards
+	TIndirectArray<FComponentReregisterContext> ComponentContexts;
+	TUniquePtr< FSkinnedMeshComponentRecreateRenderStateContext > SkinnedMeshComponentRecreateRenderStateContext;
+
 	TArray<ClothingAssetUtils::FClothingAssetMeshBinding> ClothingBindings;
 
 	//Backup the data before importing the new one
@@ -1486,6 +1491,10 @@ USkeletalMesh* UnFbx::FFbxImporter::ImportSkeletalMesh(FImportSkeletalMeshArgs &
 			Binding.Asset->UnbindFromSkeletalMesh(ExistingSkelMesh, Binding.LODIndex);
 		}
 
+		
+		//Release the Renderdata resources before reimporting the skeletal mesh
+		SkinnedMeshComponentRecreateRenderStateContext = MakeUnique<FSkinnedMeshComponentRecreateRenderStateContext>(ExistingSkelMesh);
+
 		ExistingSkelMesh->PreEditChange(NULL);
 		//The backup of the skeletal mesh data empty the LOD array in the ImportedResource of the skeletal mesh
 		//If the import fail after this step the editor can crash when updating the bone later since the LODModel will not exist anymore
@@ -1538,10 +1547,6 @@ USkeletalMesh* UnFbx::FFbxImporter::ImportSkeletalMesh(FImportSkeletalMeshArgs &
 	// Pass the number of texture coordinate sets to the LODModel.  Ensure there is at least one UV coord
 	LODModel.NumTexCoords = FMath::Max<uint32>(1, SkelMeshImportDataPtr->NumTexCoords);
 
-	// Array of re-import contexts for components using this mesh
-	// Will unregister before import, then re-register afterwards
-	TIndirectArray<FComponentReregisterContext> ComponentContexts;
-
 	if(ImportSkeletalMeshArgs.bCreateRenderData )
 	{
 		TArray<FVector> LODPoints;
@@ -1802,8 +1807,6 @@ USkeletalMesh* UnFbx::FFbxImporter::ImportSkeletalMesh(FImportSkeletalMeshArgs &
 
 	// ComponentContexts will now go out of scope, causing components to be re-registered
 
-	// ComponentContexts will now go out of scope, causing components to be re-registered
-
 	return SkeletalMesh;
 }
 
