Fix hot reload parsing UHT startup message ("LogInit: Display: RandInit(1527028783)  SRandInit(1527028783)") as an error.

diff --git a/Engine/Source/Editor/UnrealEd/Private/Kismet2/CompilerResultsLog.cpp b/Engine/Source/Editor/UnrealEd/Private/Kismet2/CompilerResultsLog.cpp
--- a/Engine/Source/Editor/UnrealEd/Private/Kismet2/CompilerResultsLog.cpp
+++ b/Engine/Source/Editor/UnrealEd/Private/Kismet2/CompilerResultsLog.cpp
@@ -458,7 +458,8 @@ TArray< TSharedRef<FTokenizedMessage> > FCompilerResultsLog::ParseCompilerLogDum
 		FString FullPath, LineNumberString;
 		if (Line.Split(TEXT(")"), &LeftStr, &RightStr, ESearchCase::CaseSensitive) &&
 			LeftStr.Split(TEXT("("), &FullPath, &LineNumberString, ESearchCase::CaseSensitive) &&
-			LineNumberString.IsNumeric() && (FCString::Strtoi(*LineNumberString, NULL, 10) > 0))
+			LineNumberString.IsNumeric() && (FCString::Strtoi(*LineNumberString, NULL, 10) > 0) &&
+			RightStr.StartsWith(TEXT(": error")))
 		{
 			EMessageSeverity::Type Severity = EMessageSeverity::Error;
 			FString FullPathTrimmed = FullPath;
