Fix uninitialized member in structure internal to the transaction system

diff --git a/Engine/Source/Editor/UnrealEd/Classes/Editor/Transactor.h b/Engine/Source/Editor/UnrealEd/Classes/Editor/Transactor.h
--- a/Engine/Source/Editor/UnrealEd/Classes/Editor/Transactor.h
+++ b/Engine/Source/Editor/UnrealEd/Classes/Editor/Transactor.h
@@ -170,6 +170,11 @@ protected:
 
 		struct FSerializedObject
 		{
+			FSerializedObject()
+				: bIsPendingKill(false)
+			{
+			}
+
 			void SetObject(const UObject* InObject)
 			{
 				ObjectName = InObject->GetFName();
