VWI: Editor extensions now update before ViewportClients and EdModes
- This fixes EdModes not being able to receive ViewportInteractor input/hover events until after they had updated in the frame

diff --git a/Engine/Source/Editor/UnrealEd/Private/EditorEngine.cpp b/Engine/Source/Editor/UnrealEd/Private/EditorEngine.cpp
--- a/Engine/Source/Editor/UnrealEd/Private/EditorEngine.cpp
+++ b/Engine/Source/Editor/UnrealEd/Private/EditorEngine.cpp
@@ -1692,8 +1692,10 @@ void UEditorEngine::Tick( float DeltaSeconds, bool bIdleMode )
 		}
 	}
 
-	// Update viewports.
+	// Updates all the extensions for all the editor worlds
+	EditorWorldExtensionsManager->Tick( DeltaSeconds );
 
+	// Update viewports.
 	for (int32 ViewportIndex = AllViewportClients.Num()-1; ViewportIndex >= 0; ViewportIndex--)
 	{
 		FEditorViewportClient* ViewportClient = AllViewportClients[ ViewportIndex ];
@@ -1708,9 +1710,6 @@ void UEditorEngine::Tick( float DeltaSeconds, bool bIdleMode )
 		}
 	}
 
-	// Updates all the extensions for all the editor worlds
-	EditorWorldExtensionsManager->Tick(DeltaSeconds);
-
 	bool bIsMouseOverAnyLevelViewport = false;
 
 	//Do this check separate to the above loop as the ViewportClient may no longer be valid after we have ticked it
