WIP - Create a UMeshDescription when importing FBX
- There is a switch in the fbx import dialog to be able to import with the old raw data
- The rendering data is not build yet, this will be the next thing

diff --git a/Engine/Config/BaseEditorPerProjectUserSettings.ini b/Engine/Config/BaseEditorPerProjectUserSettings.ini
--- a/Engine/Config/BaseEditorPerProjectUserSettings.ini
+++ b/Engine/Config/BaseEditorPerProjectUserSettings.ini
@@ -593,6 +593,7 @@ BackupIntervalInMinutes=5
 
 
 [/Script/UnrealEd.FbxImportUI]
+bImportEditableMesh=False
 bImportMaterials=True
 bAutoCreateGroups=True
 bInvertNormalMaps=False
diff --git a/Engine/Source/Editor/StaticMeshEditor/Private/StaticMeshEditorViewportClient.cpp b/Engine/Source/Editor/StaticMeshEditor/Private/StaticMeshEditorViewportClient.cpp
--- a/Engine/Source/Editor/StaticMeshEditor/Private/StaticMeshEditorViewportClient.cpp
+++ b/Engine/Source/Editor/StaticMeshEditor/Private/StaticMeshEditorViewportClient.cpp
@@ -443,7 +443,7 @@ void FStaticMeshEditorViewportClient::Draw(const FSceneView* View,FPrimitiveDraw
 
 	TSharedPtr<IStaticMeshEditor> StaticMeshEditor = StaticMeshEditorPtr.Pin();
 
-	if(!StaticMesh->RenderData->LODResources.IsValidIndex(StaticMeshEditor->GetCurrentLODIndex()))
+	if(!StaticMesh->RenderData || !StaticMesh->RenderData->LODResources.IsValidIndex(StaticMeshEditor->GetCurrentLODIndex()))
 	{
 		// Guard against corrupted meshes
 		return;
@@ -778,7 +778,7 @@ void FStaticMeshEditorViewportClient::DrawCanvas( FViewport& InViewport, FSceneV
 	TextItems.Add(SStaticMeshEditorViewport::FOverlayTextItem(
 		FText::Format(NSLOCTEXT("UnrealEd", "UVChannels_F", "UV Channels:  {0}"), FText::AsNumber(StaticMeshEditorPtr.Pin()->GetNumUVChannels(CurrentLODLevel)))));
 
-	if( StaticMesh->RenderData->LODResources.Num() > 0 )
+	if(StaticMesh->RenderData && StaticMesh->RenderData->LODResources.Num() > 0 )
 	{
 		if (StaticMesh->RenderData->LODResources[0].DistanceFieldData != nullptr )
 		{
diff --git a/Engine/Source/Editor/UnrealEd/Classes/Factories/FbxImportUI.h b/Engine/Source/Editor/UnrealEd/Classes/Factories/FbxImportUI.h
--- a/Engine/Source/Editor/UnrealEd/Classes/Factories/FbxImportUI.h
+++ b/Engine/Source/Editor/UnrealEd/Classes/Factories/FbxImportUI.h
@@ -50,6 +50,10 @@ public:
 	UPROPERTY(EditAnywhere, AdvancedDisplay, config, Category=Miscellaneous, meta=(OBJRestrict="true"))
 	uint32 bOverrideFullName:1;
 
+	/** Use the string in "Name" field as full name of mesh. The option only works when the scene contains one mesh. */
+	UPROPERTY(EditAnywhere, AdvancedDisplay, config, Category = Miscellaneous, meta = (OBJRestrict = "true"))
+	uint32 bImportEditableMesh:1;
+
 	/** Whether to import the incoming FBX as a skeletal object */
 	UPROPERTY(EditAnywhere, Category = Mesh, meta = (ImportType = "StaticMesh|SkeletalMesh", DisplayName="Skeletal Mesh"))
 	bool bImportAsSkeletal;
diff --git a/Engine/Source/Editor/UnrealEd/Private/Factories/ActorFactory.cpp b/Engine/Source/Editor/UnrealEd/Private/Factories/ActorFactory.cpp
--- a/Engine/Source/Editor/UnrealEd/Private/Factories/ActorFactory.cpp
+++ b/Engine/Source/Editor/UnrealEd/Private/Factories/ActorFactory.cpp
@@ -354,7 +354,10 @@ void UActorFactoryStaticMesh::PostSpawnActor( UObject* Asset, AActor* NewActor)
 	StaticMeshComponent->UnregisterComponent();
 
 	StaticMeshComponent->SetStaticMesh(StaticMesh);
-	StaticMeshComponent->StaticMeshDerivedDataKey = StaticMesh->RenderData->DerivedDataKey;
+	if (StaticMesh->RenderData)
+	{
+		StaticMeshComponent->StaticMeshDerivedDataKey = StaticMesh->RenderData->DerivedDataKey;
+	}
 
 	// Init Component
 	StaticMeshComponent->RegisterComponent();
diff --git a/Engine/Source/Editor/UnrealEd/Private/Fbx/FbxFactory.cpp b/Engine/Source/Editor/UnrealEd/Private/Fbx/FbxFactory.cpp
--- a/Engine/Source/Editor/UnrealEd/Private/Fbx/FbxFactory.cpp
+++ b/Engine/Source/Editor/UnrealEd/Private/Fbx/FbxFactory.cpp
@@ -109,7 +109,7 @@ UObject* UFbxFactory::ImportANode(void* VoidFbxImporter, TArray<void*> VoidNodes
 	}
 	check(Nodes.Num() > 0 && Nodes[0] != nullptr);
 
-	UObject* NewObject = NULL;
+	UObject* CreatedObject = NULL;
 	FName OutputName = FFbxImporter->MakeNameForMesh(InName.ToString(), Nodes[0]);
 	
 	{
@@ -121,10 +121,10 @@ UObject* UFbxFactory::ImportANode(void* VoidFbxImporter, TArray<void*> VoidNodes
 			return NULL;
 		}
 
-		NewObject = FFbxImporter->ImportStaticMeshAsSingle( InParent, Nodes, OutputName, Flags, ImportUI->StaticMeshImportData, Cast<UStaticMesh>(InMesh), LODIndex );
+		CreatedObject = FFbxImporter->ImportStaticMeshAsSingle( InParent, Nodes, OutputName, Flags, ImportUI->StaticMeshImportData, Cast<UStaticMesh>(InMesh), LODIndex );
 	}
 
-	if (NewObject)
+	if (CreatedObject)
 	{
 		NodeIndex++;
 		FFormatNamedArguments Args;
@@ -133,7 +133,7 @@ UObject* UFbxFactory::ImportANode(void* VoidFbxImporter, TArray<void*> VoidNodes
 		GWarn->StatusUpdate( NodeIndex, Total, FText::Format( NSLOCTEXT("UnrealEd", "Importingf", "Importing ({NodeIndex} of {ArrayLength})"), Args ) );
 	}
 
-	return NewObject;
+	return CreatedObject;
 }
 
 bool UFbxFactory::ConfigureProperties()
@@ -169,7 +169,7 @@ UObject* UFbxFactory::FactoryCreateBinary
 
 	FEditorDelegates::OnAssetPreImport.Broadcast(this, Class, InParent, Name, Type);
 
-	UObject* NewObject = NULL;
+	UObject* CreatedObject = NULL;
 	
 	//Look if its a re-import, in that cazse we must call the re-import factory
 	UObject *ExistingObject = nullptr;
@@ -441,7 +441,7 @@ UObject* UFbxFactory::FactoryCreateBinary
 						FbxImporter->ImportStaticMeshGlobalSockets( NewStaticMesh );
 					}
 
-					NewObject = NewStaticMesh;
+					CreatedObject = NewStaticMesh;
 
 				}
 				else if ( ImportUI->MeshTypeToImport == FBXIT_SkeletalMesh )// skeletal mesh
@@ -523,7 +523,7 @@ UObject* UFbxFactory::FactoryCreateBinary
 								ImportSkeletalMeshArgs.OutData = &OutData;
 
 								USkeletalMesh* NewMesh = FbxImporter->ImportSkeletalMesh( ImportSkeletalMeshArgs );
-								NewObject = NewMesh;
+								CreatedObject = NewMesh;
 
 								if(bOperationCanceled)
 								{
@@ -551,9 +551,9 @@ UObject* UFbxFactory::FactoryCreateBinary
 									SuccessfulLodIndex++;
 								}
 							}
-							else if (NewObject) // the base skeletal mesh is imported successfully
+							else if (CreatedObject) // the base skeletal mesh is imported successfully
 							{
-								USkeletalMesh* BaseSkeletalMesh = Cast<USkeletalMesh>(NewObject);
+								USkeletalMesh* BaseSkeletalMesh = Cast<USkeletalMesh>(CreatedObject);
 								FName LODObjectName = NAME_None;
 								UnFbx::FFbxImporter::FImportSkeletalMeshArgs ImportSkeletalMeshArgs;
 								ImportSkeletalMeshArgs.InParent = BaseSkeletalMesh->GetOutermost();
@@ -581,7 +581,7 @@ UObject* UFbxFactory::FactoryCreateBinary
 							}
 						
 							// import morph target
-							if ( NewObject && ImportUI->SkeletalMeshImportData->bImportMorphTargets && ImportedSuccessfulLodIndex != INDEX_NONE)
+							if (CreatedObject && ImportUI->SkeletalMeshImportData->bImportMorphTargets && ImportedSuccessfulLodIndex != INDEX_NONE)
 							{
 								// Disable material importing when importing morph targets
 								uint32 bImportMaterials = ImportOptions->bImportMaterials;
@@ -589,14 +589,14 @@ UObject* UFbxFactory::FactoryCreateBinary
 								uint32 bImportTextures = ImportOptions->bImportTextures;
 								ImportOptions->bImportTextures = 0;
 
-								FbxImporter->ImportFbxMorphTarget(SkelMeshNodeArray, Cast<USkeletalMesh>(NewObject), InParent, ImportedSuccessfulLodIndex, OutData);
+								FbxImporter->ImportFbxMorphTarget(SkelMeshNodeArray, Cast<USkeletalMesh>(CreatedObject), InParent, ImportedSuccessfulLodIndex, OutData);
 							
 								ImportOptions->bImportMaterials = !!bImportMaterials;
 								ImportOptions->bImportTextures = !!bImportTextures;
 							}
 						}
 					
-						if (NewObject)
+						if (CreatedObject)
 						{
 							NodeIndex++;
 							FFormatNamedArguments Args;
@@ -604,7 +604,7 @@ UObject* UFbxFactory::FactoryCreateBinary
 							Args.Add( TEXT("ArrayLength"), SkelMeshArray.Num() );
 							GWarn->StatusUpdate( NodeIndex, SkelMeshArray.Num(), FText::Format( NSLOCTEXT("UnrealEd", "Importingf", "Importing ({NodeIndex} of {ArrayLength})"), Args ) );
 							
-							USkeletalMesh* SkeletalMesh = Cast<USkeletalMesh>(NewObject);
+							USkeletalMesh* SkeletalMesh = Cast<USkeletalMesh>(CreatedObject);
 							UnFbx::FFbxImporter::UpdateSkeletalMeshImportData(SkeletalMesh, ImportUI->SkeletalMeshImportData, INDEX_NONE, nullptr, nullptr);
 						}
 					}
@@ -626,7 +626,7 @@ UObject* UFbxFactory::FactoryCreateBinary
 					if (ImportOptions->SkeletonForAnimation)
 					{
 						// will return the last animation sequence that were added
-						NewObject = UEditorEngine::ImportFbxAnimation( ImportOptions->SkeletonForAnimation, InParent, ImportUI->AnimSequenceImportData, *Filename, *Name.ToString(), true );
+						CreatedObject = UEditorEngine::ImportFbxAnimation( ImportOptions->SkeletonForAnimation, InParent, ImportUI->AnimSequenceImportData, *Filename, *Name.ToString(), true );
 					}
 				}
 			}
@@ -647,7 +647,7 @@ UObject* UFbxFactory::FactoryCreateBinary
 			}
 		}
 
-		if (NewObject == NULL)
+		if (CreatedObject == NULL)
 		{
 			FbxImporter->AddTokenizedErrorMessage(FTokenizedMessage::Create(EMessageSeverity::Error, LOCTEXT("FailedToImport_NoObject", "Import failed.")), FFbxErrors::Generic_ImportingNewObjectFailed);
 		}
@@ -656,16 +656,16 @@ UObject* UFbxFactory::FactoryCreateBinary
 		Warn->EndSlowTask();
 	}
 
-	FEditorDelegates::OnAssetPostImport.Broadcast(this, NewObject);
+	FEditorDelegates::OnAssetPostImport.Broadcast(this, CreatedObject);
 
-	return NewObject;
+	return CreatedObject;
 }
 
 
 UObject* UFbxFactory::RecursiveImportNode(void* VoidFbxImporter, void* VoidNode, UObject* InParent, FName InName, EObjectFlags Flags, int32& NodeIndex, int32 Total, TArray<UObject*>& OutNewAssets)
 {
 	TArray<void*> TmpVoidArray;
-	UObject* NewObject = NULL;
+	UObject* CreatedObject = NULL;
 	UnFbx::FFbxImporter *FbxImporter = (UnFbx::FFbxImporter *)VoidFbxImporter;
 	FbxNode* Node = (FbxNode*)VoidNode;
 	if (Node->GetNodeAttribute() && Node->GetNodeAttribute()->GetAttributeType() == FbxNodeAttribute::eLODGroup && Node->GetChildCount() > 0 )
@@ -680,17 +680,17 @@ UObject* UFbxFactory::RecursiveImportNode(void* VoidFbxImporter, void* VoidNode,
 			{
 				TmpVoidArray.Add(LodNode);
 			}
-			NewObject = ImportANode(VoidFbxImporter, TmpVoidArray, InParent, InName, Flags, NodeIndex, Total);
+			CreatedObject = ImportANode(VoidFbxImporter, TmpVoidArray, InParent, InName, Flags, NodeIndex, Total);
 		}
 
-		if ( NewObject )
+		if (CreatedObject)
 		{
-			OutNewAssets.AddUnique(NewObject);
+			OutNewAssets.AddUnique(CreatedObject);
 		}
 
 		bool bImportMeshLODs = ImportUI->StaticMeshImportData->bImportMeshLODs;
 
-		if (NewObject && bImportMeshLODs)
+		if (CreatedObject && bImportMeshLODs)
 		{
 			// import LOD meshes
 			for (int32 LODIndex = 1; LODIndex < Node->GetChildCount(); LODIndex++)
@@ -709,14 +709,14 @@ UObject* UFbxFactory::RecursiveImportNode(void* VoidFbxImporter, void* VoidNode,
 					{
 						TmpVoidArray.Add(LodNode);
 					}
-					ImportANode(VoidFbxImporter, TmpVoidArray, InParent, InName, Flags, NodeIndex, Total, NewObject, LODIndex);
+					ImportANode(VoidFbxImporter, TmpVoidArray, InParent, InName, Flags, NodeIndex, Total, CreatedObject, LODIndex);
 				}
 			}
 		}
 		
-		if (NewObject)
+		if (CreatedObject)
 		{
-			UStaticMesh *NewStaticMesh = Cast<UStaticMesh>(NewObject);
+			UStaticMesh *NewStaticMesh = Cast<UStaticMesh>(CreatedObject);
 			if (NewStaticMesh != nullptr)
 			{
 				//Reorder the material
@@ -736,11 +736,11 @@ UObject* UFbxFactory::RecursiveImportNode(void* VoidFbxImporter, void* VoidNode,
 		{
 			TmpVoidArray.Empty();
 			TmpVoidArray.Add(Node);
-			NewObject = ImportANode(VoidFbxImporter, TmpVoidArray, InParent, InName, Flags, NodeIndex, Total);
+			CreatedObject = ImportANode(VoidFbxImporter, TmpVoidArray, InParent, InName, Flags, NodeIndex, Total);
 
-			if ( NewObject )
+			if (CreatedObject)
 			{
-				UStaticMesh *NewStaticMesh = Cast<UStaticMesh>(NewObject);
+				UStaticMesh *NewStaticMesh = Cast<UStaticMesh>(CreatedObject);
 				if (NewStaticMesh != nullptr)
 				{
 					//Reorder the material
@@ -749,7 +749,7 @@ UObject* UFbxFactory::RecursiveImportNode(void* VoidFbxImporter, void* VoidNode,
 					FbxImporter->PostImportStaticMesh(NewStaticMesh, Nodes);
 					FbxImporter->UpdateStaticMeshImportData(NewStaticMesh, nullptr);
 				}
-				OutNewAssets.AddUnique(NewObject);
+				OutNewAssets.AddUnique(CreatedObject);
 			}
 		}
 		
@@ -762,14 +762,14 @@ UObject* UFbxFactory::RecursiveImportNode(void* VoidFbxImporter, void* VoidNode,
 				OutNewAssets.AddUnique(SubObject);
 			}
 
-			if (NewObject==NULL)
+			if (CreatedObject ==NULL)
 			{
-				NewObject = SubObject;
+				CreatedObject = SubObject;
 			}
 		}
 	}
 
-	return NewObject;
+	return CreatedObject;
 }
 
 
diff --git a/Engine/Source/Editor/UnrealEd/Private/Fbx/FbxMainImport.cpp b/Engine/Source/Editor/UnrealEd/Private/Fbx/FbxMainImport.cpp
--- a/Engine/Source/Editor/UnrealEd/Private/Fbx/FbxMainImport.cpp
+++ b/Engine/Source/Editor/UnrealEd/Private/Fbx/FbxMainImport.cpp
@@ -239,6 +239,7 @@ void ApplyImportUIToImportOptions(UFbxImportUI* ImportUI, FBXImportOptions& InOu
 		InOutImportOptions.BaseEmissiveColorName = ImportUI->TextureImportData->BaseEmissiveColorName;
 	}
 	InOutImportOptions.bImportTextures = ImportUI->bImportTextures;
+	InOutImportOptions.bImportEditableMesh = ImportUI->bImportEditableMesh;
 	InOutImportOptions.bUsedAsFullName = ImportUI->bOverrideFullName;
 	InOutImportOptions.bImportAnimations = ImportUI->bImportAnimations;
 	InOutImportOptions.SkeletonForAnimation = ImportUI->Skeleton;
diff --git a/Engine/Source/Editor/UnrealEd/Private/Fbx/FbxStaticMeshImport.cpp b/Engine/Source/Editor/UnrealEd/Private/Fbx/FbxStaticMeshImport.cpp
--- a/Engine/Source/Editor/UnrealEd/Private/Fbx/FbxStaticMeshImport.cpp
+++ b/Engine/Source/Editor/UnrealEd/Private/Fbx/FbxStaticMeshImport.cpp
@@ -20,6 +20,7 @@
 #include "Engine/Polys.h"
 #include "Engine/StaticMeshSocket.h"
 #include "Editor.h"
+#include "Modules/ModuleManager.h"
 #include "RawMesh.h"
 
 #include "StaticMeshResources.h"
@@ -32,6 +33,8 @@
 #include "Interfaces/ITargetPlatformManagerModule.h"
 #include "Misc/FbxErrors.h"
 #include "PhysicsEngine/BodySetup.h"
+#include "MeshDescription.h"
+#include "IMeshDescriptionModule.h"
 
 #define LOCTEXT_NAMESPACE "FbxStaticMeshImport"
 
@@ -254,13 +257,20 @@ struct FFBXUVs
 	int32 UniqueUVCount;
 };
 
+
 bool UnFbx::FFbxImporter::BuildStaticMeshFromGeometry(FbxNode* Node, UStaticMesh* StaticMesh, TArray<FFbxMaterial>& MeshMaterials, int32 LODIndex,FRawMesh& RawMesh,
 	EVertexColorImportOption::Type VertexColorImportOption, const TMap<FVector, FColor>& ExistingVertexColorData, const FColor& VertexOverrideColor)
 {
 	check(StaticMesh->SourceModels.IsValidIndex(LODIndex));
 	FbxMesh* Mesh = Node->GetMesh();
 	FStaticMeshSourceModel& SrcModel = StaticMesh->SourceModels[LODIndex];
-
+	
+	UMeshDescription *MeshDescription = StaticMesh->MeshDescription;
+	if (ImportOptions->bImportEditableMesh && MeshDescription == nullptr)
+	{
+		StaticMesh->MeshDescription = NewObject<UMeshDescription>(StaticMesh, NAME_None, RF_NoFlags);
+		MeshDescription = StaticMesh->MeshDescription;
+	}
 	//remove the bad polygons before getting any data from mesh
 	Mesh->RemoveBadPolygons();
 
@@ -398,18 +408,31 @@ bool UnFbx::FFbxImporter::BuildStaticMeshFromGeometry(FbxNode* Node, UStaticMesh
 	//
 	bool bSmoothingAvailable = false;
 
-	FbxLayerElementSmoothing const* SmoothingInfo = BaseLayer->GetSmoothing();
+	FbxLayerElementSmoothing* SmoothingInfo = BaseLayer->GetSmoothing();
 	FbxLayerElement::EReferenceMode SmoothingReferenceMode(FbxLayerElement::eDirect);
 	FbxLayerElement::EMappingMode SmoothingMappingMode(FbxLayerElement::eByEdge);
 	if (SmoothingInfo)
 	{
+		if(MeshDescription != nullptr)
+		{
+			if (SmoothingInfo->GetMappingMode() == FbxLayerElement::eByPolygon)
+			{
+				//Convert the base layer to edge smoothing
+				GeometryConverter->ComputeEdgeSmoothingFromPolygonSmoothing(Mesh, 0);
+				BaseLayer = Mesh->GetLayer(0);
+				SmoothingInfo = BaseLayer->GetSmoothing();
+			}
 
-		if( SmoothingInfo->GetMappingMode() == FbxLayerElement::eByPolygon )
+			if (SmoothingInfo->GetMappingMode() == FbxLayerElement::eByEdge)
+			{
+				bSmoothingAvailable = true;
+			}
+		}
+		else if (SmoothingInfo->GetMappingMode() == FbxLayerElement::eByPolygon)
 		{
 			bSmoothingAvailable = true;
 		}
 
-
 		SmoothingReferenceMode = SmoothingInfo->GetReferenceMode();
 		SmoothingMappingMode = SmoothingInfo->GetMappingMode();
 	}
@@ -493,281 +516,608 @@ bool UnFbx::FFbxImporter::BuildStaticMeshFromGeometry(FbxNode* Node, UStaticMesh
 	}
 
 	int32 VertexCount = Mesh->GetControlPointsCount();
-	int32 WedgeCount = TriangleCount * 3;
 	bool OddNegativeScale = IsOddNegativeScale(TotalMatrix);
-
-	int32 VertexOffset = RawMesh.VertexPositions.Num();
-	int32 WedgeOffset = RawMesh.WedgeIndices.Num();
-	int32 TriangleOffset = RawMesh.FaceMaterialIndices.Num();
-
-	int32 MaxMaterialIndex = 0;
-
-	// Reserve space for attributes.
-	RawMesh.FaceMaterialIndices.AddZeroed(TriangleCount);
-	RawMesh.FaceSmoothingMasks.AddZeroed(TriangleCount);
-	RawMesh.WedgeIndices.AddZeroed(WedgeCount);
-
-	if (bHasNTBInformation || RawMesh.WedgeTangentX.Num() > 0 || RawMesh.WedgeTangentY.Num() > 0)
-	{
-		RawMesh.WedgeTangentX.AddZeroed(WedgeOffset + WedgeCount - RawMesh.WedgeTangentX.Num());
-		RawMesh.WedgeTangentY.AddZeroed(WedgeOffset + WedgeCount - RawMesh.WedgeTangentY.Num());
-	}
-
-	if (LayerElementNormal || RawMesh.WedgeTangentZ.Num() > 0 )
+	bool bHasNonDegenerateTriangles = false;
+	if(MeshDescription != nullptr)
 	{
-		RawMesh.WedgeTangentZ.AddZeroed(WedgeOffset + WedgeCount - RawMesh.WedgeTangentZ.Num());
-	}
+		int32 VertexOffset = MeshDescription->Vertices().Num();
+		int32 VertexInstanceOffset = MeshDescription->VertexInstances().Num();
+		int32 TriangleOffset = MeshDescription->Polygons().Num();
 
-	if (LayerElementVertexColor || VertexColorImportOption != EVertexColorImportOption::Replace || RawMesh.WedgeColors.Num() )
-	{
-		int32 NumNewColors = WedgeOffset + WedgeCount - RawMesh.WedgeColors.Num();
-		int32 FirstNewColor = RawMesh.WedgeColors.Num();
-		RawMesh.WedgeColors.AddUninitialized(NumNewColors);
-		for (int32 WedgeIndex = FirstNewColor; WedgeIndex < FirstNewColor + NumNewColors; ++WedgeIndex)
+		if (StaticMesh->RenderData == nullptr)
 		{
-			RawMesh.WedgeColors[WedgeIndex] = FColor::White;
+			StaticMesh->CacheDerivedData();
 		}
-	}
-
-	// When importing multiple mesh pieces to the same static mesh.  Ensure each mesh piece has the same number of Uv's
-	int32 ExistingUVCount = 0;
-	for( int32 ExistingUVIndex = 0; ExistingUVIndex < MAX_MESH_TEXTURE_COORDS; ++ExistingUVIndex )
-	{
-		if( RawMesh.WedgeTexCoords[ExistingUVIndex].Num() > 0 )
+		
+		if (StaticMesh->RenderData->LODResources.Num() == 0)
 		{
-			// Mesh already has UVs at this index
-			++ExistingUVCount;
+			FStaticMeshLODResources *DummyStaticMeshLodResources = new FStaticMeshLODResources();
+			StaticMesh->RenderData->LODResources.Add(DummyStaticMeshLodResources);
 		}
-		else
-		{
-			// No more UVs
-			break;
+
+		int32 MaxMaterialIndex = 0;
+		TMap<int32, FPolygonGroupID> PolygonGroupMapping;
+		int32 NumUVs = FBXUVs.UniqueUVCount;
+		// At least one UV set must exist.  
+		NumUVs = FMath::Max(1, NumUVs);
+
+		//Fill the vertex array
+		for (int32 VertexIndex = 0; VertexIndex < VertexCount; ++VertexIndex)
+		{
+			int32 RealVertexIndex = VertexOffset + VertexIndex;
+			FbxVector4 FbxPosition = Mesh->GetControlPoints()[VertexIndex];
+			FbxPosition = TotalMatrix.MultT(FbxPosition);
+			const FVector VertexPosition = Converter.ConvertPos(FbxPosition);
+			
+			FMeshVertex NewVertex;
+			NewVertex.VertexPosition = VertexPosition;
+			NewVertex.CornerSharpness = 0.0f;
+			FVertexID AddedVertexId = MeshDescription->Vertices().Add(NewVertex);
+			if (!MeshDescription->Vertices().IsValid(AddedVertexId) || AddedVertexId.GetValue() != RealVertexIndex)
+			{
+				AddTokenizedErrorMessage(FTokenizedMessage::Create(EMessageSeverity::Error, FText::Format(LOCTEXT("Error_CannotCreateVertex", "Cannot create valid vertex for mesh '{0}'"), FText::FromString(Mesh->GetName()))), FFbxErrors::StaticMesh_BuildError);
+				return false;
+			}
 		}
-	}
-	
-	int32 UVCount = FMath::Max( FBXUVs.UniqueUVCount, ExistingUVCount );
 
-	// At least one UV set must exist.  
-	UVCount = FMath::Max( 1, UVCount );
+		//Polygons
+		for (int32 TriangleIndex = 0; TriangleIndex < TriangleCount; TriangleIndex++)
+		{
+			int32 RealTriangleIndex = TriangleOffset + TriangleIndex;
+			FVertexInstanceID CornerInstanceIDs[3];
+			FVertexID CornerVerticesIDs[3];
+			for (int32 CornerIndex = 0; CornerIndex < 3; CornerIndex++)
+			{
+				int32 VertexInstanceIndex = RealTriangleIndex * 3 + CornerIndex;
+				const FVertexInstanceID VertexInstanceID(VertexInstanceIndex);
+				CornerInstanceIDs[CornerIndex] = VertexInstanceID;
+				int32 ControlPointIndex = Mesh->GetPolygonVertex(TriangleIndex, CornerIndex);
+				const FVertexID VertexID(VertexOffset + ControlPointIndex);
+				FVector VertexPosition = MeshDescription->GetVertex(VertexID).VertexPosition;
+				CornerVerticesIDs[CornerIndex] = VertexID;
 
-	for (int32 UVLayerIndex = 0; UVLayerIndex<UVCount; UVLayerIndex++)
-	{
-		RawMesh.WedgeTexCoords[UVLayerIndex].AddZeroed(WedgeOffset + WedgeCount - RawMesh.WedgeTexCoords[UVLayerIndex].Num());
-	}
+				FMeshVertexInstance NewVertexInstance;
+				NewVertexInstance.VertexID = VertexID;
+				FVertexInstanceID AddedVertexInstanceId = MeshDescription->VertexInstances().Add(NewVertexInstance);
+				if (!MeshDescription->VertexInstances().IsValid(AddedVertexInstanceId) || AddedVertexInstanceId.GetValue() != VertexInstanceIndex)
+				{
+					AddTokenizedErrorMessage(FTokenizedMessage::Create(EMessageSeverity::Error, FText::Format(LOCTEXT("Error_CannotCreateVertexInstance", "Cannot create valid vertex instance for mesh '{0}'"), FText::FromString(Mesh->GetName()))), FFbxErrors::StaticMesh_BuildError);
+					return false;
+				}
 
-	int32 TriangleIndex;
-	TMap<int32,int32> IndexMap;
-	bool bHasNonDegenerateTriangles = false;
+				//UVs attributes
+				for (int32 UVLayerIndex = 0; UVLayerIndex < FBXUVs.UniqueUVCount; UVLayerIndex++)
+				{
+					FVector2D FinalUVVector(0.0f, 0.0f);
+					if (FBXUVs.LayerElementUV[UVLayerIndex] != NULL)
+					{
+						int UVMapIndex = (FBXUVs.UVMappingMode[UVLayerIndex] == FbxLayerElement::eByControlPoint) ? ControlPointIndex : TriangleIndex * 3 + CornerIndex;
+						int32 UVIndex = (FBXUVs.UVReferenceMode[UVLayerIndex] == FbxLayerElement::eDirect) ?
+							UVMapIndex : FBXUVs.LayerElementUV[UVLayerIndex]->GetIndexArray().GetAt(UVMapIndex);
 
-	for( TriangleIndex = 0 ; TriangleIndex < TriangleCount ; TriangleIndex++ )
-	{
-		int32 DestTriangleIndex = TriangleOffset + TriangleIndex;
-		FVector CornerPositions[3];
+						FbxVector2	UVVector = FBXUVs.LayerElementUV[UVLayerIndex]->GetDirectArray().GetAt(UVIndex);
+						FinalUVVector.X = static_cast<float>(UVVector[0]);
+						FinalUVVector.Y = 1.f - static_cast<float>(UVVector[1]);   //flip the Y of UVs for DirectX
+					}
+					MeshDescription->VertexInstanceAttributes().SetAttribute<FVector2D>(VertexInstanceID, UEditableMeshAttribute::VertexTextureCoordinate(), UVLayerIndex, FinalUVVector);
+				}
 
-		for ( int32 CornerIndex=0; CornerIndex<3; CornerIndex++)
-		{
-			// If there are odd number negative scale, invert the vertex order for triangles
-			int32 WedgeIndex = WedgeOffset + TriangleIndex * 3 + (OddNegativeScale ? 2 - CornerIndex : CornerIndex);
+				//Color attribute
+				if (VertexColorImportOption == EVertexColorImportOption::Replace)
+				{
+					if (LayerElementVertexColor)
+					{
+						int32 VertexColorMappingIndex = (VertexColorMappingMode == FbxLayerElement::eByControlPoint) ?
+							Mesh->GetPolygonVertex(TriangleIndex, CornerIndex) : (TriangleIndex * 3 + CornerIndex);
 
-			// Store vertex index and position.
-			int32 ControlPointIndex = Mesh->GetPolygonVertex(TriangleIndex, CornerIndex);
-			int32* ExistingIndex = IndexMap.Find(ControlPointIndex);
-			if (ExistingIndex)
-			{
-				RawMesh.WedgeIndices[WedgeIndex] = *ExistingIndex;
-				CornerPositions[CornerIndex] = RawMesh.VertexPositions[*ExistingIndex];
+						int32 VectorColorIndex = (VertexColorReferenceMode == FbxLayerElement::eDirect) ?
+							VertexColorMappingIndex : LayerElementVertexColor->GetIndexArray().GetAt(VertexColorMappingIndex);
+
+						FbxColor VertexColor = LayerElementVertexColor->GetDirectArray().GetAt(VectorColorIndex);
+
+						FColor VertexInstanceColor(
+							uint8(255.f*VertexColor.mRed),
+							uint8(255.f*VertexColor.mGreen),
+							uint8(255.f*VertexColor.mBlue),
+							uint8(255.f*VertexColor.mAlpha)
+						);
+						MeshDescription->VertexInstanceAttributes().SetAttribute<FVector4>(VertexInstanceID, UEditableMeshAttribute::VertexColor(), 0, FVector4(FLinearColor(VertexInstanceColor)));
+					}
+				}
+				else if (VertexColorImportOption == EVertexColorImportOption::Ignore)
+				{
+					// try to match this triangles current vertex with one that existed in the previous mesh.
+					// This is a find in a tmap which uses a fast hash table lookup.
+					const FColor* PaintedColor = ExistingVertexColorData.Find(VertexPosition);
+					if (PaintedColor)
+					{
+						// A matching color for this vertex was found
+						MeshDescription->VertexInstanceAttributes().SetAttribute<FVector4>(VertexInstanceID, UEditableMeshAttribute::VertexColor(), 0, FVector4(FLinearColor(*PaintedColor)));
+					}
+				}
+				else
+				{
+					// set the triangle's vertex color to a constant override
+					check(VertexColorImportOption == EVertexColorImportOption::Override);
+					MeshDescription->VertexInstanceAttributes().SetAttribute<FVector4>(VertexInstanceID, UEditableMeshAttribute::VertexColor(), 0, FVector4(FLinearColor(VertexOverrideColor)));
+				}
+
+				if (LayerElementNormal)
+				{
+					int TriangleCornerIndex = TriangleIndex * 3 + CornerIndex;
+					//normals may have different reference and mapping mode than tangents and binormals
+					int NormalMapIndex = (NormalMappingMode == FbxLayerElement::eByControlPoint) ?
+						ControlPointIndex : TriangleCornerIndex;
+					int NormalValueIndex = (NormalReferenceMode == FbxLayerElement::eDirect) ?
+						NormalMapIndex : LayerElementNormal->GetIndexArray().GetAt(NormalMapIndex);
+					
+					FbxVector4 TempValue = LayerElementNormal->GetDirectArray().GetAt(NormalValueIndex);
+					TempValue = TotalMatrixForNormal.MultT(TempValue);
+					FVector TangentZ = Converter.ConvertDir(TempValue);
+					MeshDescription->VertexInstanceAttributes().SetAttribute<FVector>(CornerInstanceIDs[CornerIndex], UEditableMeshAttribute::VertexNormal(), 0, TangentZ.GetSafeNormal());
+					//tangents and binormals share the same reference, mapping mode and index array
+					if (bHasNTBInformation)
+					{
+						int TangentMapIndex = (TangentMappingMode == FbxLayerElement::eByControlPoint) ?
+							ControlPointIndex : TriangleCornerIndex;
+						int TangentValueIndex = (TangentReferenceMode == FbxLayerElement::eDirect) ?
+							TangentMapIndex : LayerElementTangent->GetIndexArray().GetAt(TangentMapIndex);
+
+						TempValue = LayerElementTangent->GetDirectArray().GetAt(TangentValueIndex);
+						TempValue = TotalMatrixForNormal.MultT(TempValue);
+						FVector TangentX = Converter.ConvertDir(TempValue);
+						MeshDescription->VertexInstanceAttributes().SetAttribute<FVector>(CornerInstanceIDs[CornerIndex], UEditableMeshAttribute::VertexTangent(), 0, TangentX.GetSafeNormal());
+
+						int BinormalMapIndex = (BinormalMappingMode == FbxLayerElement::eByControlPoint) ?
+							ControlPointIndex : TriangleCornerIndex;
+						int BinormalValueIndex = (BinormalReferenceMode == FbxLayerElement::eDirect) ?
+							BinormalMapIndex : LayerElementBinormal->GetIndexArray().GetAt(BinormalMapIndex);
+
+						TempValue = LayerElementBinormal->GetDirectArray().GetAt(BinormalValueIndex);
+						TempValue = TotalMatrixForNormal.MultT(TempValue);
+						FVector TangentY = -Converter.ConvertDir(TempValue);
+						FVector BiNormal = TangentY.GetSafeNormal();
+						MeshDescription->VertexInstanceAttributes().SetAttribute<float>(CornerInstanceIDs[CornerIndex], UEditableMeshAttribute::VertexBinormalSign(), 0, GetBasisDeterminantSign(TangentX.GetSafeNormal(), BiNormal, TangentZ.GetSafeNormal()) );
+					}
+				}
 			}
-			else
-			{
-				FbxVector4 FbxPosition = Mesh->GetControlPoints()[ControlPointIndex];
-				FbxVector4 FinalPosition = TotalMatrix.MultT(FbxPosition);
-				int32 VertexIndex = RawMesh.VertexPositions.Add(Converter.ConvertPos(FinalPosition));
-				RawMesh.WedgeIndices[WedgeIndex] = VertexIndex;
-				IndexMap.Add(ControlPointIndex, VertexIndex);
-				CornerPositions[CornerIndex] = RawMesh.VertexPositions[VertexIndex];
+			
+			// Check if the triangle just discovered is non-degenerate if we haven't found one yet
+			if (!bHasNonDegenerateTriangles)
+			{
+				float TriangleComparisonThreshold = ImportOptions->bRemoveDegenerates ? THRESH_POINTS_ARE_SAME : 0.0f;
+				FVector VertexPosition[3];
+				VertexPosition[0] = MeshDescription->GetVertex(CornerVerticesIDs[0]).VertexPosition;
+				VertexPosition[1] = MeshDescription->GetVertex(CornerVerticesIDs[1]).VertexPosition;
+				VertexPosition[2] = MeshDescription->GetVertex(CornerVerticesIDs[2]).VertexPosition;
+				if (!( VertexPosition[0].Equals(VertexPosition[1], TriangleComparisonThreshold)
+					|| VertexPosition[0].Equals(VertexPosition[2], TriangleComparisonThreshold)
+					|| VertexPosition[1].Equals(VertexPosition[2], TriangleComparisonThreshold)))
+				{
+					bHasNonDegenerateTriangles = true;
+				}
 			}
 
 			//
-			// normals, tangents and binormals
+			// material index
 			//
-			if (LayerElementNormal)
+			int32 MaterialIndex = 0;
+			if (MaterialCount > 0)
 			{
-				int TriangleCornerIndex = TriangleIndex*3 + CornerIndex;
-				//normals may have different reference and mapping mode than tangents and binormals
-				int NormalMapIndex = (NormalMappingMode == FbxLayerElement::eByControlPoint) ? 
-					ControlPointIndex : TriangleCornerIndex;
-				int NormalValueIndex = (NormalReferenceMode == FbxLayerElement::eDirect) ? 
-					NormalMapIndex : LayerElementNormal->GetIndexArray().GetAt(NormalMapIndex);
-
-				//tangents and binormals share the same reference, mapping mode and index array
-				if (bHasNTBInformation)
+				if (LayerElementMaterial)
 				{
-					int TangentMapIndex = (TangentMappingMode == FbxLayerElement::eByControlPoint) ? 
-						ControlPointIndex : TriangleCornerIndex;
-					int TangentValueIndex = (TangentReferenceMode == FbxLayerElement::eDirect) ? 
-						TangentMapIndex : LayerElementTangent->GetIndexArray().GetAt(TangentMapIndex);
-
-					FbxVector4 TempValue = LayerElementTangent->GetDirectArray().GetAt(TangentValueIndex);
-					TempValue = TotalMatrixForNormal.MultT(TempValue);
-					FVector TangentX = Converter.ConvertDir(TempValue);
-					RawMesh.WedgeTangentX[WedgeIndex] = TangentX.GetSafeNormal();
+					switch (MaterialMappingMode)
+					{
+						// material index is stored in the IndexArray, not the DirectArray (which is irrelevant with 2009.1)
+					case FbxLayerElement::eAllSame:
+					{
+						MaterialIndex = LayerElementMaterial->GetIndexArray().GetAt(0);
+					}
+					break;
+					case FbxLayerElement::eByPolygon:
+					{
+						MaterialIndex = LayerElementMaterial->GetIndexArray().GetAt(TriangleIndex);
+					}
+					break;
+					}
+				}
+			}
 
-					int BinormalMapIndex = (BinormalMappingMode == FbxLayerElement::eByControlPoint) ? 
-						ControlPointIndex : TriangleCornerIndex;
-					int BinormalValueIndex = (BinormalReferenceMode == FbxLayerElement::eDirect) ? 
-						BinormalMapIndex : LayerElementBinormal->GetIndexArray().GetAt(BinormalMapIndex);
+			if (MaterialIndex >= MaterialCount || MaterialIndex < 0)
+			{
+				AddTokenizedErrorMessage(FTokenizedMessage::Create(EMessageSeverity::Warning, LOCTEXT("Error_MaterialIndexInconsistency", "Face material index inconsistency - forcing to 0")), FFbxErrors::Generic_Mesh_MaterialIndexInconsistency);
+				MaterialIndex = 0;
+			}
 
-					TempValue = LayerElementBinormal->GetDirectArray().GetAt(BinormalValueIndex);
-					TempValue = TotalMatrixForNormal.MultT(TempValue);
-					FVector TangentY = -Converter.ConvertDir(TempValue);
-					RawMesh.WedgeTangentY[WedgeIndex] = TangentY.GetSafeNormal();
+			//Create a polygon with the 3 vertex intances Add it to the material group
+			int32 RealMaterialIndex = MaterialIndexOffset + MaterialIndex;
+			if (!PolygonGroupMapping.Contains(RealMaterialIndex))
+			{
+				UMaterialInterface* Material = MeshMaterials.IsValidIndex(RealMaterialIndex) ? MeshMaterials[RealMaterialIndex].Material : UMaterial::GetDefaultMaterial(MD_Surface);
+				FPolygonGroupID ExistingPolygonGroup = FPolygonGroupID::Invalid;
+				for (const FPolygonGroupID PolygonGroupID : MeshDescription->PolygonGroups().GetElementIDs())
+				{
+					if (MeshDescription->PolygonGroups()[PolygonGroupID].MaterialAsset.ResolveObject() == Material)
+					{
+						ExistingPolygonGroup = PolygonGroupID;
+						break;
+					}
 				}
+				if (ExistingPolygonGroup == FPolygonGroupID::Invalid)
+				{
+					FMeshPolygonGroup NewPolygonGroup;
+					NewPolygonGroup.bCastShadow = true;
+					NewPolygonGroup.bEnableCollision = bEnableCollision;
+					NewPolygonGroup.ImportedMaterialSlotName = Material->GetFName();
+					NewPolygonGroup.MaterialSlotName = Material->GetFName();
+					NewPolygonGroup.MaterialAsset = FStringAssetReference(Material);
+					ExistingPolygonGroup = MeshDescription->PolygonGroups().Add(NewPolygonGroup);
+				}
+				PolygonGroupMapping.Add(RealMaterialIndex, ExistingPolygonGroup);
+			}
 
-				FbxVector4 TempValue = LayerElementNormal->GetDirectArray().GetAt(NormalValueIndex);
-				TempValue = TotalMatrixForNormal.MultT(TempValue);
-				FVector TangentZ = Converter.ConvertDir(TempValue);
-				RawMesh.WedgeTangentZ[WedgeIndex] = TangentZ.GetSafeNormal();
+			FPolygonGroupID PolygonGroupID = PolygonGroupMapping[RealMaterialIndex];
+
+			// Insert a polygon into the mesh
+			FMeshPolygon NewPolygon;
+			NewPolygon.PolygonGroupID = PolygonGroupID;
+
+			for (int32 TriangleVertexIndex = 0; TriangleVertexIndex < 3; ++TriangleVertexIndex)
+			{
+				const FVertexInstanceID &VertexInstanceID = CornerInstanceIDs[TriangleVertexIndex];
+				const FVertexID &VertexID = CornerVerticesIDs[TriangleVertexIndex];
+				NewPolygon.PerimeterContour.VertexInstanceIDs.Add(VertexInstanceID);
 			}
+			const FPolygonID NewPolygonID = MeshDescription->Polygons().Add(NewPolygon);
 
-			//
-			// vertex colors
-			//
-			if (VertexColorImportOption == EVertexColorImportOption::Replace)
+			// Set edges hardness
 			{
-				if (LayerElementVertexColor)
+				TArray<FEdgeID> EdgeIDs;
+				TArray<bool> EdgesIsHard;
+				// Add the edges of this triangle
+				for (uint32 TriangleEdgeNumber = 0; TriangleEdgeNumber < 3; ++TriangleEdgeNumber)
 				{
-					int32 VertexColorMappingIndex = (VertexColorMappingMode == FbxLayerElement::eByControlPoint) ?
-						Mesh->GetPolygonVertex(TriangleIndex, CornerIndex) : (TriangleIndex * 3 + CornerIndex);
+					//Find the matching edge ID
+					uint32 CornerIndices[2];
+					CornerIndices[0] = (TriangleEdgeNumber + 0) % 3;
+					CornerIndices[1] = (TriangleEdgeNumber + 1) % 3;
+
+					FVertexID EdgeVertexIDs[2];
+					EdgeVertexIDs[0] = CornerVerticesIDs[CornerIndices[0]];
+					EdgeVertexIDs[1] = CornerVerticesIDs[CornerIndices[1]];
+
+					FEdgeID MatchEdgeId = FEdgeID::Invalid;
+					for (const FEdgeID EdgeID : MeshDescription->Edges().GetElementIDs())
+					{
+						FVertexID OutEdgeVertexID0 = MeshDescription->Edges()[EdgeID].VertexIDs[0];
+						FVertexID OutEdgeVertexID1 = MeshDescription->Edges()[EdgeID].VertexIDs[1];
+						if (EdgeVertexIDs[0] == OutEdgeVertexID0 && EdgeVertexIDs[1] == OutEdgeVertexID1 ||
+							EdgeVertexIDs[1] == OutEdgeVertexID0 && EdgeVertexIDs[0] == OutEdgeVertexID1)
+						{
+							MatchEdgeId = EdgeID;
+							break;
+						}
+					}
+					int32 SmoothMaskValue = 0;
 
-					int32 VectorColorIndex = (VertexColorReferenceMode == FbxLayerElement::eDirect) ?
-					VertexColorMappingIndex : LayerElementVertexColor->GetIndexArray().GetAt(VertexColorMappingIndex);
+					if (MatchEdgeId == FEdgeID::Invalid)
+					{
+						FMeshEdge NewEdge;
+						NewEdge.VertexIDs[0] = EdgeVertexIDs[0];
+						NewEdge.VertexIDs[1] = EdgeVertexIDs[1];
+						NewEdge.bIsHardEdge = false;
+						NewEdge.CreaseSharpness = 0.0f;
+						MatchEdgeId = MeshDescription->Edges().Add(NewEdge);
+					}
 
-					FbxColor VertexColor = LayerElementVertexColor->GetDirectArray().GetAt(VectorColorIndex);
+					FMeshEdge& ExistingEdge = MeshDescription->GetEdge(MatchEdgeId);
+					ExistingEdge.ConnectedPolygons.Add(NewPolygonID);
 
-					RawMesh.WedgeColors[WedgeIndex] = FColor(
-						uint8(255.f*VertexColor.mRed),
-						uint8(255.f*VertexColor.mGreen),
-						uint8(255.f*VertexColor.mBlue),
-						uint8(255.f*VertexColor.mAlpha)
-						);
+					if (!ExistingEdge.bIsHardEdge && bSmoothingAvailable && SmoothingInfo)
+					{
+						if (SmoothingMappingMode == FbxLayerElement::eByEdge)
+						{
+							int32 EdgeIndex = Mesh->GetMeshEdgeIndexForPolygon(TriangleIndex, TriangleEdgeNumber);
+							int lSmoothingIndex = (SmoothingReferenceMode == FbxLayerElement::eDirect) ? EdgeIndex : SmoothingInfo->GetIndexArray().GetAt(EdgeIndex);
+							SmoothMaskValue = SmoothingInfo->GetDirectArray().GetAt(lSmoothingIndex);
+							if (SmoothMaskValue > 0)
+							{
+								//Set the hard edges
+								EdgeIDs.Add(MatchEdgeId);
+								EdgesIsHard.Add(true);
+							}
+						}
+						else
+						{
+							AddTokenizedErrorMessage(FTokenizedMessage::Create(EMessageSeverity::Warning, FText::Format(LOCTEXT("Error_UnsupportedSmoothingGroup", "Unsupported Smoothing group mapping mode on mesh  '{0}'"), FText::FromString(Mesh->GetName()))), FFbxErrors::Generic_Mesh_UnsupportingSmoothingGroup);
+						}
+					}
 				}
-			}
-			else if (VertexColorImportOption == EVertexColorImportOption::Ignore)
-			{
-				// try to match this triangles current vertex with one that existed in the previous mesh.
-				// This is a find in a tmap which uses a fast hash table lookup.
-				FVector Position = RawMesh.VertexPositions[RawMesh.WedgeIndices[WedgeIndex]];
-				const FColor* PaintedColor = ExistingVertexColorData.Find(Position);
-				if (PaintedColor)
+				for (int32 EdgeIndex = 0; EdgeIndex < EdgeIDs.Num(); ++EdgeIndex)
 				{
-					// A matching color for this vertex was found
-					RawMesh.WedgeColors[WedgeIndex] = *PaintedColor;
+					FEdgeID EdgeID = EdgeIDs[EdgeIndex];
+					MeshDescription->GetEdge(EdgeID).bIsHardEdge = EdgesIsHard[EdgeIndex];
 				}
 			}
-			else
-			{
-				// set the triangle's vertex color to a constant override
-				check(VertexColorImportOption == EVertexColorImportOption::Override);
-				RawMesh.WedgeColors[WedgeIndex] = VertexOverrideColor;
-			}
 		}
+		// needed?
+		FBXUVs.Cleanup();
+	}
+	else
+	{
+		int32 WedgeCount = TriangleCount * 3;
+		int32 VertexOffset = RawMesh.VertexPositions.Num();
+		int32 WedgeOffset = RawMesh.WedgeIndices.Num();
+		int32 TriangleOffset = RawMesh.FaceMaterialIndices.Num();
 
-		// Check if the triangle just discovered is non-degenerate if we haven't found one yet
-		if (!bHasNonDegenerateTriangles)
+		int32 MaxMaterialIndex = 0;
+
+		// Reserve space for attributes.
+		RawMesh.FaceMaterialIndices.AddZeroed(TriangleCount);
+		RawMesh.FaceSmoothingMasks.AddZeroed(TriangleCount);
+		RawMesh.WedgeIndices.AddZeroed(WedgeCount);
+
+		if (bHasNTBInformation || RawMesh.WedgeTangentX.Num() > 0 || RawMesh.WedgeTangentY.Num() > 0)
+		{
+			RawMesh.WedgeTangentX.AddZeroed(WedgeOffset + WedgeCount - RawMesh.WedgeTangentX.Num());
+			RawMesh.WedgeTangentY.AddZeroed(WedgeOffset + WedgeCount - RawMesh.WedgeTangentY.Num());
+		}
+
+		if (LayerElementNormal || RawMesh.WedgeTangentZ.Num() > 0 )
 		{
-			float ComparisonThreshold = ImportOptions->bRemoveDegenerates ? THRESH_POINTS_ARE_SAME : 0.0f;
+			RawMesh.WedgeTangentZ.AddZeroed(WedgeOffset + WedgeCount - RawMesh.WedgeTangentZ.Num());
+		}
 
-			if (!(CornerPositions[0].Equals(CornerPositions[1], ComparisonThreshold)
-				  || CornerPositions[0].Equals(CornerPositions[2], ComparisonThreshold)
-				  || CornerPositions[1].Equals(CornerPositions[2], ComparisonThreshold)))
+		if (LayerElementVertexColor || VertexColorImportOption != EVertexColorImportOption::Replace || RawMesh.WedgeColors.Num() )
+		{
+			int32 NumNewColors = WedgeOffset + WedgeCount - RawMesh.WedgeColors.Num();
+			int32 FirstNewColor = RawMesh.WedgeColors.Num();
+			RawMesh.WedgeColors.AddUninitialized(NumNewColors);
+			for (int32 WedgeIndex = FirstNewColor; WedgeIndex < FirstNewColor + NumNewColors; ++WedgeIndex)
 			{
-				bHasNonDegenerateTriangles = true;
+				RawMesh.WedgeColors[WedgeIndex] = FColor::White;
 			}
 		}
 
-		//
-		// smoothing mask
-		//
-		if (bSmoothingAvailable && SmoothingInfo)
+		// When importing multiple mesh pieces to the same static mesh.  Ensure each mesh piece has the same number of Uv's
+		int32 ExistingUVCount = 0;
+		for( int32 ExistingUVIndex = 0; ExistingUVIndex < MAX_MESH_TEXTURE_COORDS; ++ExistingUVIndex )
 		{
-			if (SmoothingMappingMode == FbxLayerElement::eByPolygon)
+			if( RawMesh.WedgeTexCoords[ExistingUVIndex].Num() > 0 )
 			{
-				int lSmoothingIndex = (SmoothingReferenceMode == FbxLayerElement::eDirect) ? TriangleIndex : SmoothingInfo->GetIndexArray().GetAt(TriangleIndex);
-				RawMesh.FaceSmoothingMasks[DestTriangleIndex] = SmoothingInfo->GetDirectArray().GetAt(lSmoothingIndex);
+				// Mesh already has UVs at this index
+				++ExistingUVCount;
 			}
 			else
 			{
-				AddTokenizedErrorMessage(FTokenizedMessage::Create(EMessageSeverity::Warning, FText::Format(LOCTEXT("Error_UnsupportedSmoothingGroup", "Unsupported Smoothing group mapping mode on mesh  '{0}'"), FText::FromString(Mesh->GetName()))), FFbxErrors::Generic_Mesh_UnsupportingSmoothingGroup);
+				// No more UVs
+				break;
 			}
 		}
+	
+		int32 UVCount = FMath::Max( FBXUVs.UniqueUVCount, ExistingUVCount );
 
-		//
-		// uvs
-		//
-		// In FBX file, the same UV may be saved multiple times, i.e., there may be same UV in LayerElementUV
-		// So we don't import the duplicate UVs
-		int32 UVLayerIndex;
-		for (UVLayerIndex = 0; UVLayerIndex<FBXUVs.UniqueUVCount; UVLayerIndex++)
+		// At least one UV set must exist.  
+		UVCount = FMath::Max( 1, UVCount );
+
+		for (int32 UVLayerIndex = 0; UVLayerIndex<UVCount; UVLayerIndex++)
 		{
-			if (FBXUVs.LayerElementUV[UVLayerIndex] != NULL) 
+			RawMesh.WedgeTexCoords[UVLayerIndex].AddZeroed(WedgeOffset + WedgeCount - RawMesh.WedgeTexCoords[UVLayerIndex].Num());
+		}
+
+		int32 TriangleIndex;
+		TMap<int32,int32> IndexMap;
+
+		for( TriangleIndex = 0 ; TriangleIndex < TriangleCount ; TriangleIndex++ )
+		{
+			int32 DestTriangleIndex = TriangleOffset + TriangleIndex;
+			FVector CornerPositions[3];
+
+			for ( int32 CornerIndex=0; CornerIndex<3; CornerIndex++)
 			{
-				for (int32 CornerIndex=0;CornerIndex<3;CornerIndex++)
+				// If there are odd number negative scale, invert the vertex order for triangles
+				int32 WedgeIndex = WedgeOffset + TriangleIndex * 3 + (OddNegativeScale ? 2 - CornerIndex : CornerIndex);
+
+				// Store vertex index and position.
+				int32 ControlPointIndex = Mesh->GetPolygonVertex(TriangleIndex, CornerIndex);
+				int32* ExistingIndex = IndexMap.Find(ControlPointIndex);
+				if (ExistingIndex)
 				{
-					// If there are odd number negative scale, invert the vertex order for triangles
-					int32 WedgeIndex = WedgeOffset + TriangleIndex * 3 + (OddNegativeScale ? 2 - CornerIndex : CornerIndex);
+					RawMesh.WedgeIndices[WedgeIndex] = *ExistingIndex;
+					CornerPositions[CornerIndex] = RawMesh.VertexPositions[*ExistingIndex];
+				}
+				else
+				{
+					FbxVector4 FbxPosition = Mesh->GetControlPoints()[ControlPointIndex];
+					FbxVector4 FinalPosition = TotalMatrix.MultT(FbxPosition);
+					int32 VertexIndex = RawMesh.VertexPositions.Add(Converter.ConvertPos(FinalPosition));
+					RawMesh.WedgeIndices[WedgeIndex] = VertexIndex;
+					IndexMap.Add(ControlPointIndex, VertexIndex);
+					CornerPositions[CornerIndex] = RawMesh.VertexPositions[VertexIndex];
+				}
+
+				//
+				// normals, tangents and binormals
+				//
+				if (LayerElementNormal)
+				{
+					int TriangleCornerIndex = TriangleIndex*3 + CornerIndex;
+					//normals may have different reference and mapping mode than tangents and binormals
+					int NormalMapIndex = (NormalMappingMode == FbxLayerElement::eByControlPoint) ? 
+						ControlPointIndex : TriangleCornerIndex;
+					int NormalValueIndex = (NormalReferenceMode == FbxLayerElement::eDirect) ? 
+						NormalMapIndex : LayerElementNormal->GetIndexArray().GetAt(NormalMapIndex);
+
+					//tangents and binormals share the same reference, mapping mode and index array
+					if (bHasNTBInformation)
+					{
+						int TangentMapIndex = (TangentMappingMode == FbxLayerElement::eByControlPoint) ? 
+							ControlPointIndex : TriangleCornerIndex;
+						int TangentValueIndex = (TangentReferenceMode == FbxLayerElement::eDirect) ? 
+							TangentMapIndex : LayerElementTangent->GetIndexArray().GetAt(TangentMapIndex);
+
+						FbxVector4 TempValue = LayerElementTangent->GetDirectArray().GetAt(TangentValueIndex);
+						TempValue = TotalMatrixForNormal.MultT(TempValue);
+						FVector TangentX = Converter.ConvertDir(TempValue);
+						RawMesh.WedgeTangentX[WedgeIndex] = TangentX.GetSafeNormal();
+
+						int BinormalMapIndex = (BinormalMappingMode == FbxLayerElement::eByControlPoint) ? 
+							ControlPointIndex : TriangleCornerIndex;
+						int BinormalValueIndex = (BinormalReferenceMode == FbxLayerElement::eDirect) ? 
+							BinormalMapIndex : LayerElementBinormal->GetIndexArray().GetAt(BinormalMapIndex);
+
+						TempValue = LayerElementBinormal->GetDirectArray().GetAt(BinormalValueIndex);
+						TempValue = TotalMatrixForNormal.MultT(TempValue);
+						FVector TangentY = -Converter.ConvertDir(TempValue);
+						RawMesh.WedgeTangentY[WedgeIndex] = TangentY.GetSafeNormal();
+					}
+
+					FbxVector4 TempValue = LayerElementNormal->GetDirectArray().GetAt(NormalValueIndex);
+					TempValue = TotalMatrixForNormal.MultT(TempValue);
+					FVector TangentZ = Converter.ConvertDir(TempValue);
+					RawMesh.WedgeTangentZ[WedgeIndex] = TangentZ.GetSafeNormal();
+				}
 
+				//
+				// vertex colors
+				//
+				if (VertexColorImportOption == EVertexColorImportOption::Replace)
+				{
+					if (LayerElementVertexColor)
+					{
+						int32 VertexColorMappingIndex = (VertexColorMappingMode == FbxLayerElement::eByControlPoint) ?
+							Mesh->GetPolygonVertex(TriangleIndex, CornerIndex) : (TriangleIndex * 3 + CornerIndex);
 
-					int lControlPointIndex = Mesh->GetPolygonVertex(TriangleIndex, CornerIndex);
-					int UVMapIndex = (FBXUVs.UVMappingMode[UVLayerIndex] == FbxLayerElement::eByControlPoint) ? lControlPointIndex : TriangleIndex*3+CornerIndex;
-					int32 UVIndex = (FBXUVs.UVReferenceMode[UVLayerIndex] == FbxLayerElement::eDirect) ? 
-						UVMapIndex : FBXUVs.LayerElementUV[UVLayerIndex]->GetIndexArray().GetAt(UVMapIndex);
+						int32 VectorColorIndex = (VertexColorReferenceMode == FbxLayerElement::eDirect) ?
+						VertexColorMappingIndex : LayerElementVertexColor->GetIndexArray().GetAt(VertexColorMappingIndex);
 
-					FbxVector2	UVVector = FBXUVs.LayerElementUV[UVLayerIndex]->GetDirectArray().GetAt(UVIndex);
+						FbxColor VertexColor = LayerElementVertexColor->GetDirectArray().GetAt(VectorColorIndex);
 
-					RawMesh.WedgeTexCoords[UVLayerIndex][WedgeIndex].X = static_cast<float>(UVVector[0]);
-					RawMesh.WedgeTexCoords[UVLayerIndex][WedgeIndex].Y = 1.f-static_cast<float>(UVVector[1]);   //flip the Y of UVs for DirectX
+						RawMesh.WedgeColors[WedgeIndex] = FColor(
+							uint8(255.f*VertexColor.mRed),
+							uint8(255.f*VertexColor.mGreen),
+							uint8(255.f*VertexColor.mBlue),
+							uint8(255.f*VertexColor.mAlpha)
+							);
+					}
+				}
+				else if (VertexColorImportOption == EVertexColorImportOption::Ignore)
+				{
+					// try to match this triangles current vertex with one that existed in the previous mesh.
+					// This is a find in a tmap which uses a fast hash table lookup.
+					FVector Position = RawMesh.VertexPositions[RawMesh.WedgeIndices[WedgeIndex]];
+					const FColor* PaintedColor = ExistingVertexColorData.Find(Position);
+					if (PaintedColor)
+					{
+						// A matching color for this vertex was found
+						RawMesh.WedgeColors[WedgeIndex] = *PaintedColor;
+					}
+				}
+				else
+				{
+					// set the triangle's vertex color to a constant override
+					check(VertexColorImportOption == EVertexColorImportOption::Override);
+					RawMesh.WedgeColors[WedgeIndex] = VertexOverrideColor;
 				}
 			}
-		}
 
-		//
-		// material index
-		//
-		int32 MaterialIndex = 0;
-		if (MaterialCount>0)
-		{
-			if (LayerElementMaterial)
+			// Check if the triangle just discovered is non-degenerate if we haven't found one yet
+			if (!bHasNonDegenerateTriangles)
 			{
-				switch(MaterialMappingMode)
+				float ComparisonThreshold = ImportOptions->bRemoveDegenerates ? THRESH_POINTS_ARE_SAME : 0.0f;
+
+				if (!(CornerPositions[0].Equals(CornerPositions[1], ComparisonThreshold)
+					  || CornerPositions[0].Equals(CornerPositions[2], ComparisonThreshold)
+					  || CornerPositions[1].Equals(CornerPositions[2], ComparisonThreshold)))
 				{
-					// material index is stored in the IndexArray, not the DirectArray (which is irrelevant with 2009.1)
-				case FbxLayerElement::eAllSame:
-					{	
-						MaterialIndex = LayerElementMaterial->GetIndexArray().GetAt(0);
+					bHasNonDegenerateTriangles = true;
+				}
+			}
+
+			//
+			// smoothing mask
+			//
+			if (bSmoothingAvailable && SmoothingInfo)
+			{
+				if (SmoothingMappingMode == FbxLayerElement::eByPolygon)
+				{
+					int lSmoothingIndex = (SmoothingReferenceMode == FbxLayerElement::eDirect) ? TriangleIndex : SmoothingInfo->GetIndexArray().GetAt(TriangleIndex);
+					RawMesh.FaceSmoothingMasks[DestTriangleIndex] = SmoothingInfo->GetDirectArray().GetAt(lSmoothingIndex);
+				}
+				else
+				{
+					AddTokenizedErrorMessage(FTokenizedMessage::Create(EMessageSeverity::Warning, FText::Format(LOCTEXT("Error_UnsupportedSmoothingGroup", "Unsupported Smoothing group mapping mode on mesh  '{0}'"), FText::FromString(Mesh->GetName()))), FFbxErrors::Generic_Mesh_UnsupportingSmoothingGroup);
+				}
+			}
+
+			//
+			// uvs
+			//
+			// In FBX file, the same UV may be saved multiple times, i.e., there may be same UV in LayerElementUV
+			// So we don't import the duplicate UVs
+			int32 UVLayerIndex;
+			for (UVLayerIndex = 0; UVLayerIndex<FBXUVs.UniqueUVCount; UVLayerIndex++)
+			{
+				if (FBXUVs.LayerElementUV[UVLayerIndex] != NULL) 
+				{
+					for (int32 CornerIndex=0;CornerIndex<3;CornerIndex++)
+					{
+						// If there are odd number negative scale, invert the vertex order for triangles
+						int32 WedgeIndex = WedgeOffset + TriangleIndex * 3 + (OddNegativeScale ? 2 - CornerIndex : CornerIndex);
+
+
+						int lControlPointIndex = Mesh->GetPolygonVertex(TriangleIndex, CornerIndex);
+						int UVMapIndex = (FBXUVs.UVMappingMode[UVLayerIndex] == FbxLayerElement::eByControlPoint) ? lControlPointIndex : TriangleIndex*3+CornerIndex;
+						int32 UVIndex = (FBXUVs.UVReferenceMode[UVLayerIndex] == FbxLayerElement::eDirect) ? 
+							UVMapIndex : FBXUVs.LayerElementUV[UVLayerIndex]->GetIndexArray().GetAt(UVMapIndex);
+
+						FbxVector2	UVVector = FBXUVs.LayerElementUV[UVLayerIndex]->GetDirectArray().GetAt(UVIndex);
+
+						RawMesh.WedgeTexCoords[UVLayerIndex][WedgeIndex].X = static_cast<float>(UVVector[0]);
+						RawMesh.WedgeTexCoords[UVLayerIndex][WedgeIndex].Y = 1.f-static_cast<float>(UVVector[1]);   //flip the Y of UVs for DirectX
 					}
-					break;
-				case FbxLayerElement::eByPolygon:
-					{	
-						MaterialIndex = LayerElementMaterial->GetIndexArray().GetAt(TriangleIndex);
+				}
+			}
+
+			//
+			// material index
+			//
+			int32 MaterialIndex = 0;
+			if (MaterialCount>0)
+			{
+				if (LayerElementMaterial)
+				{
+					switch(MaterialMappingMode)
+					{
+						// material index is stored in the IndexArray, not the DirectArray (which is irrelevant with 2009.1)
+					case FbxLayerElement::eAllSame:
+						{	
+							MaterialIndex = LayerElementMaterial->GetIndexArray().GetAt(0);
+						}
+						break;
+					case FbxLayerElement::eByPolygon:
+						{	
+							MaterialIndex = LayerElementMaterial->GetIndexArray().GetAt(TriangleIndex);
+						}
+						break;
 					}
-					break;
 				}
 			}
-		}
-		MaterialIndex += MaterialIndexOffset;
+			MaterialIndex += MaterialIndexOffset;
 
-		if (MaterialIndex >= MaterialCount + MaterialIndexOffset || MaterialIndex < 0)
-		{
-			AddTokenizedErrorMessage(FTokenizedMessage::Create(EMessageSeverity::Warning, LOCTEXT("Error_MaterialIndexInconsistency", "Face material index inconsistency - forcing to 0")), FFbxErrors::Generic_Mesh_MaterialIndexInconsistency);
-			MaterialIndex = 0;
+			if (MaterialIndex >= MaterialCount + MaterialIndexOffset || MaterialIndex < 0)
+			{
+				AddTokenizedErrorMessage(FTokenizedMessage::Create(EMessageSeverity::Warning, LOCTEXT("Error_MaterialIndexInconsistency", "Face material index inconsistency - forcing to 0")), FFbxErrors::Generic_Mesh_MaterialIndexInconsistency);
+				MaterialIndex = 0;
+			}
+	
+			RawMesh.FaceMaterialIndices[DestTriangleIndex] = MaterialIndex;
 		}
 	
-		RawMesh.FaceMaterialIndices[DestTriangleIndex] = MaterialIndex;
+		// needed?
+		FBXUVs.Cleanup();
 	}
 	
-	// needed?
-	FBXUVs.Cleanup();
-
 	if (!bHasNonDegenerateTriangles)
 	{
 		FFormatNamedArguments Arguments;
@@ -1103,6 +1453,7 @@ UStaticMesh* UnFbx::FFbxImporter::ImportStaticMeshAsSingle(UObject* InParent, TA
 
 	// create empty mesh
 	UStaticMesh*	StaticMesh = NULL;
+	UMeshDescription* MeshDescription = nullptr;
 
 	UStaticMesh* ExistingMesh = NULL;
 	UObject* ExistingObject = NULL;
@@ -1192,6 +1543,17 @@ UStaticMesh* UnFbx::FFbxImporter::ImportStaticMeshAsSingle(UObject* InParent, TA
 		StaticMesh = NewObject<UStaticMesh>(Package, FName(*MeshName), Flags | RF_Public);
 	}
 
+	if (ImportOptions->bImportEditableMesh)
+	{
+		MeshDescription = StaticMesh->MeshDescription;
+		if (MeshDescription == nullptr)
+		{
+			//Create private asset in the same package as the StaticMesh, and make sure reference are set to avoid GC
+			FString MeshDescriptionName = MeshName + TEXT("_MeshData");
+			StaticMesh->MeshDescription = NewObject<UMeshDescription>(StaticMesh, FName(*MeshDescriptionName), RF_NoFlags);
+		}
+	}
+
 	if (StaticMesh->SourceModels.Num() < LODIndex+1)
 	{
 		// Add one LOD 
@@ -1238,230 +1600,236 @@ UStaticMesh* UnFbx::FFbxImporter::ImportStaticMeshAsSingle(UObject* InParent, TA
 		}
 	}
 
-	// Store the new raw mesh.
-	SrcModel.RawMeshBulkData->SaveRawMesh(NewRawMesh);
+	if (MeshDescription == nullptr)
+	{
+		// Store the new raw mesh.
+		SrcModel.RawMeshBulkData->SaveRawMesh(NewRawMesh);
+	}
 
 
 	if (bBuildStatus)
 	{
-		UE_LOG(LogFbx, Verbose, TEXT("== Initial material list:"));
-		for (int32 MaterialIndex = 0; MaterialIndex < MeshMaterials.Num(); ++MaterialIndex)
+		int32 MaxMaterialIndex = 0;
+		int32 FirstOpenUVChannel = 1;
+		if (MeshDescription == nullptr)
 		{
-			UE_LOG(LogFbx, Verbose, TEXT("%d: %s"), MaterialIndex, *MeshMaterials[MaterialIndex].GetName());
-		}
+			UE_LOG(LogFbx, Verbose, TEXT("== Initial material list:"));
+			for (int32 MaterialIndex = 0; MaterialIndex < MeshMaterials.Num(); ++MaterialIndex)
+			{
+				UE_LOG(LogFbx, Verbose, TEXT("%d: %s"), MaterialIndex, *MeshMaterials[MaterialIndex].GetName());
+			}
 
-		// Compress the materials array by removing any duplicates.
-		bool bDoRemap = false;
-		TArray<int32> MaterialMap;
-		TArray<FFbxMaterial> UniqueMaterials;
-		for (int32 MaterialIndex = 0; MaterialIndex < MeshMaterials.Num(); ++MaterialIndex)
-		{
-			bool bUnique = true;
-			for (int32 OtherMaterialIndex = MaterialIndex - 1; OtherMaterialIndex >= 0; --OtherMaterialIndex)
+			// Compress the materials array by removing any duplicates.
+			bool bDoRemap = false;
+			TArray<int32> MaterialMap;
+			TArray<FFbxMaterial> UniqueMaterials;
+			for (int32 MaterialIndex = 0; MaterialIndex < MeshMaterials.Num(); ++MaterialIndex)
 			{
-				if (MeshMaterials[MaterialIndex].FbxMaterial == MeshMaterials[OtherMaterialIndex].FbxMaterial &&
-					MeshMaterials[MaterialIndex].Material == MeshMaterials[OtherMaterialIndex].Material)
+				bool bUnique = true;
+				for (int32 OtherMaterialIndex = MaterialIndex - 1; OtherMaterialIndex >= 0; --OtherMaterialIndex)
 				{
-					int32 UniqueIndex = MaterialMap[OtherMaterialIndex];
+					if (MeshMaterials[MaterialIndex].FbxMaterial == MeshMaterials[OtherMaterialIndex].FbxMaterial &&
+						MeshMaterials[MaterialIndex].Material == MeshMaterials[OtherMaterialIndex].Material)
+					{
+						int32 UniqueIndex = MaterialMap[OtherMaterialIndex];
+
+						MaterialMap.Add(UniqueIndex);
+						bDoRemap = true;
+						bUnique = false;
+						break;
+					}
+				}
+				if (bUnique)
+				{
+					int32 UniqueIndex = UniqueMaterials.Add(MeshMaterials[MaterialIndex]);
 
 					MaterialMap.Add(UniqueIndex);
-					bDoRemap = true;
-					bUnique = false;
-					break;
+				}
+				else
+				{
+					UE_LOG(LogFbx, Verbose, TEXT("  remap %d -> %d"), MaterialIndex, MaterialMap[MaterialIndex]);
 				}
 			}
-			if (bUnique)
-			{
-				int32 UniqueIndex = UniqueMaterials.Add(MeshMaterials[MaterialIndex]);
 
-				MaterialMap.Add(UniqueIndex);
-			}
-			else
+			if (UniqueMaterials.Num() > LARGE_MESH_MATERIAL_INDEX_THRESHOLD)
 			{
-				UE_LOG(LogFbx, Verbose, TEXT("  remap %d -> %d"), MaterialIndex, MaterialMap[MaterialIndex]);
-			}
-		}
-
-		if (UniqueMaterials.Num() > LARGE_MESH_MATERIAL_INDEX_THRESHOLD)
-		{
-			AddTokenizedErrorMessage(
-				FTokenizedMessage::Create(
-					EMessageSeverity::Warning,
-					FText::Format(LOCTEXT("Error_TooManyMaterials", "StaticMesh has a large number({0}) of materials and may render inefficently.  Consider breaking up the mesh into multiple Static Mesh Assets"),
-						FText::AsNumber(UniqueMaterials.Num())
+				AddTokenizedErrorMessage(
+					FTokenizedMessage::Create(
+						EMessageSeverity::Warning,
+						FText::Format(LOCTEXT("Error_TooManyMaterials", "StaticMesh has a large number({0}) of materials and may render inefficently.  Consider breaking up the mesh into multiple Static Mesh Assets"),
+							FText::AsNumber(UniqueMaterials.Num())
 						)),
-				FFbxErrors::StaticMesh_TooManyMaterials);
-		}
+					FFbxErrors::StaticMesh_TooManyMaterials);
+			}
 
-		//The fix is required for blender file. The static mesh build have change and now required that
-		//the sections (face declaration) must be declare in the same order as the material index
-		TArray<uint32> SortedMaterialIndex;
-		TArray<int32> UsedMaterials;
-		for (int32 FaceMaterialIndex = 0; FaceMaterialIndex < NewRawMesh.FaceMaterialIndices.Num(); ++FaceMaterialIndex)
-		{
-			int32 MaterialIndex = NewRawMesh.FaceMaterialIndices[FaceMaterialIndex];
-			if (!UsedMaterials.Contains(MaterialIndex))
+			//The fix is required for blender file. The static mesh build have change and now required that
+			//the sections (face declaration) must be declare in the same order as the material index
+			TArray<uint32> SortedMaterialIndex;
+			TArray<int32> UsedMaterials;
+			for (int32 FaceMaterialIndex = 0; FaceMaterialIndex < NewRawMesh.FaceMaterialIndices.Num(); ++FaceMaterialIndex)
 			{
-				int32 NewIndex = UsedMaterials.Add(MaterialIndex);
-				if (NewIndex != MaterialIndex)
+				int32 MaterialIndex = NewRawMesh.FaceMaterialIndices[FaceMaterialIndex];
+				if (!UsedMaterials.Contains(MaterialIndex))
 				{
-					bDoRemap = true;
+					int32 NewIndex = UsedMaterials.Add(MaterialIndex);
+					if (NewIndex != MaterialIndex)
+					{
+						bDoRemap = true;
+					}
 				}
 			}
-		}
 
-		for (int32 MaterialIndex = 0; MaterialIndex < MeshMaterials.Num(); ++MaterialIndex)
-		{
-			int32 SkinIndex = 0xffff;
-			if (bDoRemap)
+			for (int32 MaterialIndex = 0; MaterialIndex < MeshMaterials.Num(); ++MaterialIndex)
 			{
-				int32 UsedIndex = 0;
-				if (UsedMaterials.Find(MaterialIndex, UsedIndex))
+				int32 SkinIndex = 0xffff;
+				if (bDoRemap)
 				{
-					SkinIndex = UsedIndex;
+					int32 UsedIndex = 0;
+					if (UsedMaterials.Find(MaterialIndex, UsedIndex))
+					{
+						SkinIndex = UsedIndex;
+					}
+				}
+				int32 RemappedIndex = MaterialMap[MaterialIndex];
+				uint32 SortedMaterialKey = ((uint32)SkinIndex << 16) | ((uint32)RemappedIndex & 0xffff);
+				if (!SortedMaterialIndex.IsValidIndex(SortedMaterialKey))
+				{
+					SortedMaterialIndex.Add(SortedMaterialKey);
 				}
 			}
-			int32 RemappedIndex = MaterialMap[MaterialIndex];
-			uint32 SortedMaterialKey = ((uint32)SkinIndex << 16) | ((uint32)RemappedIndex & 0xffff);
-			if (!SortedMaterialIndex.IsValidIndex(SortedMaterialKey))
-			{
-				SortedMaterialIndex.Add(SortedMaterialKey);
-			}
-		}
 
-		SortedMaterialIndex.Sort();
+			SortedMaterialIndex.Sort();
 
 
-		UE_LOG(LogFbx, Verbose, TEXT("== After sorting:"));
-		TArray<FFbxMaterial> SortedMaterials;
-		for (int32 SortedIndex = 0; SortedIndex < SortedMaterialIndex.Num(); ++SortedIndex)
-		{
-			int32 RemappedIndex = SortedMaterialIndex[SortedIndex] & 0xffff;
-			SortedMaterials.Add(UniqueMaterials[RemappedIndex]);
-			UE_LOG(LogFbx, Verbose, TEXT("%d: %s"), SortedIndex, *UniqueMaterials[RemappedIndex].GetName());
-		}
-		UE_LOG(LogFbx, Verbose, TEXT("== Mapping table:"));
-		for (int32 MaterialIndex = 0; MaterialIndex < MaterialMap.Num(); ++MaterialIndex)
-		{
+			UE_LOG(LogFbx, Verbose, TEXT("== After sorting:"));
+			TArray<FFbxMaterial> SortedMaterials;
 			for (int32 SortedIndex = 0; SortedIndex < SortedMaterialIndex.Num(); ++SortedIndex)
 			{
 				int32 RemappedIndex = SortedMaterialIndex[SortedIndex] & 0xffff;
-				if (MaterialMap[MaterialIndex] == RemappedIndex)
-				{
-					UE_LOG(LogFbx, Verbose, TEXT("  sort %d -> %d"), MaterialIndex, SortedIndex);
-					MaterialMap[MaterialIndex] = SortedIndex;
-					break;
-				}
+				SortedMaterials.Add(UniqueMaterials[RemappedIndex]);
+				UE_LOG(LogFbx, Verbose, TEXT("%d: %s"), SortedIndex, *UniqueMaterials[RemappedIndex].GetName());
 			}
-		}
-
-		// Remap material indices.
-		int32 MaxMaterialIndex = 0;
-		int32 FirstOpenUVChannel = 1;
-		{
-			FRawMesh LocalRawMesh;
-			SrcModel.RawMeshBulkData->LoadRawMesh(LocalRawMesh);
-
-			if (bDoRemap)
+			UE_LOG(LogFbx, Verbose, TEXT("== Mapping table:"));
+			for (int32 MaterialIndex = 0; MaterialIndex < MaterialMap.Num(); ++MaterialIndex)
 			{
-				for (int32 TriIndex = 0; TriIndex < LocalRawMesh.FaceMaterialIndices.Num(); ++TriIndex)
+				for (int32 SortedIndex = 0; SortedIndex < SortedMaterialIndex.Num(); ++SortedIndex)
 				{
-					LocalRawMesh.FaceMaterialIndices[TriIndex] = MaterialMap[LocalRawMesh.FaceMaterialIndices[TriIndex]];
+					int32 RemappedIndex = SortedMaterialIndex[SortedIndex] & 0xffff;
+					if (MaterialMap[MaterialIndex] == RemappedIndex)
+					{
+						UE_LOG(LogFbx, Verbose, TEXT("  sort %d -> %d"), MaterialIndex, SortedIndex);
+						MaterialMap[MaterialIndex] = SortedIndex;
+						break;
+					}
 				}
 			}
 
-			// Compact material indices so that we won't have any sections with zero triangles.
-			LocalRawMesh.CompactMaterialIndices();
-
-			// Also compact the sorted materials array.
-			if (LocalRawMesh.MaterialIndexToImportIndex.Num() > 0)
+			// Remap material indices.
 			{
-				TArray<FFbxMaterial> OldSortedMaterials;
+				FRawMesh LocalRawMesh;
+				SrcModel.RawMeshBulkData->LoadRawMesh(LocalRawMesh);
 
-				Exchange(OldSortedMaterials, SortedMaterials);
-				SortedMaterials.Empty(LocalRawMesh.MaterialIndexToImportIndex.Num());
-				for (int32 MaterialIndex = 0; MaterialIndex < LocalRawMesh.MaterialIndexToImportIndex.Num(); ++MaterialIndex)
+				if (bDoRemap)
 				{
-					FFbxMaterial Material;
-					int32 ImportIndex = LocalRawMesh.MaterialIndexToImportIndex[MaterialIndex];
-					if (OldSortedMaterials.IsValidIndex(ImportIndex))
+					for (int32 TriIndex = 0; TriIndex < LocalRawMesh.FaceMaterialIndices.Num(); ++TriIndex)
 					{
-						Material = OldSortedMaterials[ImportIndex];
+						LocalRawMesh.FaceMaterialIndices[TriIndex] = MaterialMap[LocalRawMesh.FaceMaterialIndices[TriIndex]];
 					}
-					SortedMaterials.Add(Material);
 				}
-			}
 
-			for (int32 TriIndex = 0; TriIndex < LocalRawMesh.FaceMaterialIndices.Num(); ++TriIndex)
-			{
-				MaxMaterialIndex = FMath::Max<int32>(MaxMaterialIndex, LocalRawMesh.FaceMaterialIndices[TriIndex]);
-			}
+				// Compact material indices so that we won't have any sections with zero triangles.
+				LocalRawMesh.CompactMaterialIndices();
 
-			for (int32 i = 0; i < MAX_MESH_TEXTURE_COORDS; i++)
-			{
-				if (LocalRawMesh.WedgeTexCoords[i].Num() == 0)
+				// Also compact the sorted materials array.
+				if (LocalRawMesh.MaterialIndexToImportIndex.Num() > 0)
 				{
-					FirstOpenUVChannel = i;
-					break;
-				}
-			}
+					TArray<FFbxMaterial> OldSortedMaterials;
 
-			SrcModel.RawMeshBulkData->SaveRawMesh(LocalRawMesh);
-		}
+					Exchange(OldSortedMaterials, SortedMaterials);
+					SortedMaterials.Empty(LocalRawMesh.MaterialIndexToImportIndex.Num());
+					for (int32 MaterialIndex = 0; MaterialIndex < LocalRawMesh.MaterialIndexToImportIndex.Num(); ++MaterialIndex)
+					{
+						FFbxMaterial Material;
+						int32 ImportIndex = LocalRawMesh.MaterialIndexToImportIndex[MaterialIndex];
+						if (OldSortedMaterials.IsValidIndex(ImportIndex))
+						{
+							Material = OldSortedMaterials[ImportIndex];
+						}
+						SortedMaterials.Add(Material);
+					}
+				}
 
-		// Setup per-section info and the materials array.
-		if (LODIndex == 0)
-		{
-			StaticMesh->StaticMaterials.Empty();
-		}
+				for (int32 TriIndex = 0; TriIndex < LocalRawMesh.FaceMaterialIndices.Num(); ++TriIndex)
+				{
+					MaxMaterialIndex = FMath::Max<int32>(MaxMaterialIndex, LocalRawMesh.FaceMaterialIndices[TriIndex]);
+				}
 
-		// Replace map of sections with the unique material set
-		int32 NumMaterials = FMath::Min(SortedMaterials.Num(), MaxMaterialIndex + 1);
-		for (int32 MaterialIndex = 0; MaterialIndex < NumMaterials; ++MaterialIndex)
-		{
-			FMeshSectionInfo Info = StaticMesh->SectionInfoMap.Get(LODIndex, MaterialIndex);
+				for (int32 i = 0; i < MAX_MESH_TEXTURE_COORDS; i++)
+				{
+					if (LocalRawMesh.WedgeTexCoords[i].Num() == 0)
+					{
+						FirstOpenUVChannel = i;
+						break;
+					}
+				}
 
-			int32 Index = 0;
+				SrcModel.RawMeshBulkData->SaveRawMesh(LocalRawMesh);
+			}
 
-			FName MaterialFName = FName(*(SortedMaterials[MaterialIndex].GetName()));
-			FString CleanMaterialSlotName = MaterialFName.ToString();
-			int32 SkinOffset = CleanMaterialSlotName.Find(TEXT("_skin"));
-			if (SkinOffset != INDEX_NONE)
+			// Setup per-section info and the materials array.
+			if (LODIndex == 0)
 			{
-				CleanMaterialSlotName = CleanMaterialSlotName.LeftChop(CleanMaterialSlotName.Len() - SkinOffset);
+				StaticMesh->StaticMaterials.Empty();
 			}
 
-			if (InStaticMesh)
+			// Replace map of sections with the unique material set
+			int32 NumMaterials = FMath::Min(SortedMaterials.Num(), MaxMaterialIndex + 1);
+			for (int32 MaterialIndex = 0; MaterialIndex < NumMaterials; ++MaterialIndex)
 			{
-				Index = INDEX_NONE;
-				FStaticMaterial StaticMaterialImported(SortedMaterials[MaterialIndex].Material, FName(*CleanMaterialSlotName), MaterialFName);
-				for (int32 OriginalMaterialIndex = 0; OriginalMaterialIndex < InStaticMesh->StaticMaterials.Num(); ++OriginalMaterialIndex)
+				FMeshSectionInfo Info = StaticMesh->SectionInfoMap.Get(LODIndex, MaterialIndex);
+
+				int32 Index = 0;
+
+				FName MaterialFName = FName(*(SortedMaterials[MaterialIndex].GetName()));
+				FString CleanMaterialSlotName = MaterialFName.ToString();
+				int32 SkinOffset = CleanMaterialSlotName.Find(TEXT("_skin"));
+				if (SkinOffset != INDEX_NONE)
+				{
+					CleanMaterialSlotName = CleanMaterialSlotName.LeftChop(CleanMaterialSlotName.Len() - SkinOffset);
+				}
+
+				if (InStaticMesh)
 				{
-					if (InStaticMesh->StaticMaterials[OriginalMaterialIndex] == StaticMaterialImported)
+					Index = INDEX_NONE;
+					FStaticMaterial StaticMaterialImported(SortedMaterials[MaterialIndex].Material, FName(*CleanMaterialSlotName), MaterialFName);
+					for (int32 OriginalMaterialIndex = 0; OriginalMaterialIndex < InStaticMesh->StaticMaterials.Num(); ++OriginalMaterialIndex)
 					{
-						Index = OriginalMaterialIndex;
-						break;
+						if (InStaticMesh->StaticMaterials[OriginalMaterialIndex] == StaticMaterialImported)
+						{
+							Index = OriginalMaterialIndex;
+							break;
+						}
+					}
+					if (Index == INDEX_NONE || (Index >= NumMaterials && Index >= InStaticMesh->StaticMaterials.Num()))
+					{
+						Index = StaticMesh->StaticMaterials.Add(FStaticMaterial(SortedMaterials[MaterialIndex].Material, FName(*CleanMaterialSlotName), MaterialFName));
 					}
 				}
-				if (Index == INDEX_NONE || (Index >= NumMaterials && Index >= InStaticMesh->StaticMaterials.Num()))
+				else
 				{
+
 					Index = StaticMesh->StaticMaterials.Add(FStaticMaterial(SortedMaterials[MaterialIndex].Material, FName(*CleanMaterialSlotName), MaterialFName));
 				}
-			}
-			else
-			{
 
-				Index = StaticMesh->StaticMaterials.Add(FStaticMaterial(SortedMaterials[MaterialIndex].Material, FName(*CleanMaterialSlotName), MaterialFName));
+				Info.MaterialIndex = Index;
+				StaticMesh->SectionInfoMap.Remove(LODIndex, MaterialIndex);
+				StaticMesh->SectionInfoMap.Set(LODIndex, MaterialIndex, Info);
 			}
 
-			Info.MaterialIndex = Index;
-			StaticMesh->SectionInfoMap.Remove(LODIndex, MaterialIndex);
-			StaticMesh->SectionInfoMap.Set(LODIndex, MaterialIndex, Info);
+			FRawMesh LocalRawMesh;
+			SrcModel.RawMeshBulkData->LoadRawMesh(LocalRawMesh);
 		}
 
-		FRawMesh LocalRawMesh;
-		SrcModel.RawMeshBulkData->LoadRawMesh(LocalRawMesh);
-
 		// Setup default LOD settings based on the selected LOD group.
 		if (LODIndex == 0)
 		{
@@ -1609,11 +1977,18 @@ void UnFbx::FFbxImporter::PostImportStaticMesh(UStaticMesh* StaticMesh, TArray<F
 				}
 				bool bOriginalGenerateMeshDistanceField = StaticMesh->bGenerateMeshDistanceField;
 				StaticMesh->bGenerateMeshDistanceField = false;
-
-				StaticMesh->Build(false, &BuildErrors);
-				for (FText& Error : BuildErrors)
+				
+				if (ImportOptions->bImportEditableMesh)
 				{
-					AddTokenizedErrorMessage(FTokenizedMessage::Create(EMessageSeverity::Warning, Error), FFbxErrors::StaticMesh_BuildError);
+					//TODO: Build the Mesh Description
+				}
+				else
+				{
+					StaticMesh->Build(false, &BuildErrors);
+					for (FText& Error : BuildErrors)
+					{
+						AddTokenizedErrorMessage(FTokenizedMessage::Create(EMessageSeverity::Warning, Error), FFbxErrors::StaticMesh_BuildError);
+					}
 				}
 
 				StaticMesh->bGenerateMeshDistanceField = bOriginalGenerateMeshDistanceField;
@@ -1647,11 +2022,17 @@ void UnFbx::FFbxImporter::PostImportStaticMesh(UStaticMesh* StaticMesh, TArray<F
 			}
 		}
 	}
-
-	StaticMesh->Build(false, &BuildErrors);
-	for (FText& Error : BuildErrors)
+	if (ImportOptions->bImportEditableMesh)
 	{
-		AddTokenizedErrorMessage(FTokenizedMessage::Create(EMessageSeverity::Warning, Error), FFbxErrors::StaticMesh_BuildError);
+		//TODO: Build the mesh description
+	}
+	else
+	{
+		StaticMesh->Build(false, &BuildErrors);
+		for (FText& Error : BuildErrors)
+		{
+			AddTokenizedErrorMessage(FTokenizedMessage::Create(EMessageSeverity::Warning, Error), FFbxErrors::StaticMesh_BuildError);
+		}
 	}
 
 	//Set the specified LOD distances for every LODs we have to do this after the build in case there is a specified Lod Group
@@ -1670,17 +2051,20 @@ void UnFbx::FFbxImporter::PostImportStaticMesh(UStaticMesh* StaticMesh, TArray<F
 	// all index is pointing correctly and they're all used. Otherwise we remove them
 	FMeshSectionInfoMap TempOldSectionInfoMap = StaticMesh->SectionInfoMap;
 	StaticMesh->SectionInfoMap.Clear();
-	// fix up section data
-	for (int32 LODResoureceIndex = 0; LODResoureceIndex < StaticMesh->RenderData->LODResources.Num(); ++LODResoureceIndex)
+	if (StaticMesh->RenderData)
 	{
-		FStaticMeshLODResources& LOD = StaticMesh->RenderData->LODResources[LODResoureceIndex];
-		int32 NumSections = LOD.Sections.Num();
-		for (int32 SectionIndex = 0; SectionIndex < NumSections; ++SectionIndex)
+		// fix up section data
+		for (int32 LODResoureceIndex = 0; LODResoureceIndex < StaticMesh->RenderData->LODResources.Num(); ++LODResoureceIndex)
 		{
-			FMeshSectionInfo Info = TempOldSectionInfoMap.Get(LODResoureceIndex, SectionIndex);
-			if (StaticMesh->StaticMaterials.IsValidIndex(Info.MaterialIndex))
+			FStaticMeshLODResources& LOD = StaticMesh->RenderData->LODResources[LODResoureceIndex];
+			int32 NumSections = LOD.Sections.Num();
+			for (int32 SectionIndex = 0; SectionIndex < NumSections; ++SectionIndex)
 			{
-				StaticMesh->SectionInfoMap.Set(LODResoureceIndex, SectionIndex, Info);
+				FMeshSectionInfo Info = TempOldSectionInfoMap.Get(LODResoureceIndex, SectionIndex);
+				if (StaticMesh->StaticMaterials.IsValidIndex(Info.MaterialIndex))
+				{
+					StaticMesh->SectionInfoMap.Set(LODResoureceIndex, SectionIndex, Info);
+				}
 			}
 		}
 	}
@@ -1857,7 +2241,7 @@ void UnFbx::FFbxImporter::PostImportStaticMesh(UStaticMesh* StaticMesh, TArray<F
 
 void UnFbx::FFbxImporter::UpdateStaticMeshImportData(UStaticMesh *StaticMesh, UFbxStaticMeshImportData* StaticMeshImportData)
 {
-	if (StaticMesh == nullptr)
+	if (StaticMesh == nullptr || StaticMesh->RenderData == nullptr)
 	{
 		return;
 	}
diff --git a/Engine/Source/Editor/UnrealEd/Private/FbxImporter.h b/Engine/Source/Editor/UnrealEd/Private/FbxImporter.h
--- a/Engine/Source/Editor/UnrealEd/Private/FbxImporter.h
+++ b/Engine/Source/Editor/UnrealEd/Private/FbxImporter.h
@@ -107,6 +107,7 @@ namespace UnFbx
 struct FBXImportOptions
 {
 	// General options
+	bool bImportEditableMesh;
 	bool bImportScene;
 	bool bImportMaterials;
 	bool bResetMaterialSlots;
diff --git a/Engine/Source/Editor/UnrealEd/UnrealEd.Build.cs b/Engine/Source/Editor/UnrealEd/UnrealEd.Build.cs
--- a/Engine/Source/Editor/UnrealEd/UnrealEd.Build.cs
+++ b/Engine/Source/Editor/UnrealEd/UnrealEd.Build.cs
@@ -95,7 +95,8 @@ public class UnrealEd : ModuleRules
 				"AudioEditor",
 				"NetworkFileSystem",
                 "UMG",
-			}
+                "MeshDescription",
+            }
 		);
 
 		PrivateDependencyModuleNames.AddRange(
@@ -159,7 +160,7 @@ public class UnrealEd : ModuleRules
 				"ClothingSystemRuntime",
 				"ClothingSystemRuntimeInterface",
 				"PIEPreviewDeviceProfileSelector",
-			}
+            }
 		);
 
 		DynamicallyLoadedModuleNames.AddRange(
diff --git a/Engine/Source/Runtime/Engine/Classes/Engine/StaticMesh.h b/Engine/Source/Runtime/Engine/Classes/Engine/StaticMesh.h
--- a/Engine/Source/Runtime/Engine/Classes/Engine/StaticMesh.h
+++ b/Engine/Source/Runtime/Engine/Classes/Engine/StaticMesh.h
@@ -622,6 +622,10 @@ public:
 	UPROPERTY(Instanced, VisibleAnywhere, Category = EditableMesh)
 	class UObject* EditableMesh;
 
+	/** The MeshDescription use to build the render data. */
+	UPROPERTY(VisibleAnywhere, Category = MeshDescription)
+	class UMeshDescription* MeshDescription;
+
 	/** Pre-build navigation collision */
 	UPROPERTY(VisibleAnywhere, transient, duplicatetransient, Instanced, Category = Navigation)
 	class UNavCollision* NavCollision;
@@ -881,11 +885,13 @@ private:
 	 */
 	void FixupZeroTriangleSections();
 
+public:
 	/**
 	 * Caches derived renderable data.
 	 */
-	void CacheDerivedData();
+	ENGINE_API void CacheDerivedData();
 
+private:
 
 	FOnPreMeshBuild PreMeshBuild;
 	FOnPostMeshBuild PostMeshBuild;
diff --git a/Engine/Source/Runtime/Engine/Engine.Build.cs b/Engine/Source/Runtime/Engine/Engine.Build.cs
--- a/Engine/Source/Runtime/Engine/Engine.Build.cs
+++ b/Engine/Source/Runtime/Engine/Engine.Build.cs
@@ -78,6 +78,7 @@ public class Engine : ModuleRules
 				"DatabaseSupport",
                 "PacketHandler",
 				"HardwareSurvey",
+                "MeshDescription",
             }
         );
 
diff --git a/Engine/Source/Runtime/Engine/Public/StaticMeshResources.h b/Engine/Source/Runtime/Engine/Public/StaticMeshResources.h
--- a/Engine/Source/Runtime/Engine/Public/StaticMeshResources.h
+++ b/Engine/Source/Runtime/Engine/Public/StaticMeshResources.h
@@ -271,7 +271,7 @@ struct FStaticMeshLODResources
 	struct FSplineMeshVertexFactory* SplineVertexFactoryOverrideColorVertexBuffer;
 
 	/** Default constructor. */
-	FStaticMeshLODResources();
+	ENGINE_API FStaticMeshLODResources();
 
 	ENGINE_API ~FStaticMeshLODResources();
 
diff --git a/Engine/Source/Runtime/MeshDescription/Private/MeshDescription.cpp b/Engine/Source/Runtime/MeshDescription/Private/MeshDescription.cpp
--- a/Engine/Source/Runtime/MeshDescription/Private/MeshDescription.cpp
+++ b/Engine/Source/Runtime/MeshDescription/Private/MeshDescription.cpp
@@ -6,25 +6,25 @@
 UMeshDescription::UMeshDescription()
 {
 	// Add basic vertex attributes
-	VertexAttributes().RegisterAttribute<FVector>( "Position" );
-	VertexAttributes().RegisterAttribute<float>( "CornerSharpness" );
+	VertexAttributes().RegisterAttribute<FVector>(UEditableMeshAttribute::VertexPosition());
+	VertexAttributes().RegisterAttribute<float>(UEditableMeshAttribute::VertexCornerSharpness());
 
 	// Add basic vertex instance attributes
-	VertexInstanceAttributes().RegisterAttribute<FVector2D>( "TextureCoordinate", 2 );
-	VertexInstanceAttributes().RegisterAttribute<FVector>( "Normal" );
-	VertexInstanceAttributes().RegisterAttribute<FVector>( "Tangent" );
-	VertexInstanceAttributes().RegisterAttribute<float>( "BinormalSign" );
-	VertexInstanceAttributes().RegisterAttribute<FVector4>( "Color" );
+	VertexInstanceAttributes().RegisterAttribute<FVector2D>(UEditableMeshAttribute::VertexTextureCoordinate(), 2 );
+	VertexInstanceAttributes().RegisterAttribute<FVector>(UEditableMeshAttribute::VertexNormal());
+	VertexInstanceAttributes().RegisterAttribute<FVector>(UEditableMeshAttribute::VertexTangent());
+	VertexInstanceAttributes().RegisterAttribute<float>(UEditableMeshAttribute::VertexBinormalSign());
+	VertexInstanceAttributes().RegisterAttribute<FVector4>(UEditableMeshAttribute::VertexColor());
 
 	// Add basic edge attributes
-	EdgeAttributes().RegisterAttribute<bool>( "EdgeIsHard" );
-	EdgeAttributes().RegisterAttribute<float>( "CreaseSharpness" );
+	EdgeAttributes().RegisterAttribute<bool>(UEditableMeshAttribute::EdgeIsHard());
+	EdgeAttributes().RegisterAttribute<float>(UEditableMeshAttribute::EdgeCreaseSharpness());
 
 	// Add basic polygon attributes
-	PolygonAttributes().RegisterAttribute<FVector>( "PolygonNormal" );
+	PolygonAttributes().RegisterAttribute<FVector>(UEditableMeshAttribute::PolygonNormal());
 	PolygonAttributes().RegisterAttribute<FVector>( "PolygonTangent" );
 	PolygonAttributes().RegisterAttribute<FVector>( "PolygonBinormal" );
-	PolygonAttributes().RegisterAttribute<FVector>( "PolygonCenter" );
+	PolygonAttributes().RegisterAttribute<FVector>(UEditableMeshAttribute::PolygonCenter());
 
 	// Add basic polygon group attributes
 	PolygonGroupAttributes().RegisterAttribute<FStringAssetReference>( "MaterialAsset" );
diff --git a/Engine/Source/Runtime/MeshDescription/Public/MeshDescription.h b/Engine/Source/Runtime/MeshDescription/Public/MeshDescription.h
--- a/Engine/Source/Runtime/MeshDescription/Public/MeshDescription.h
+++ b/Engine/Source/Runtime/MeshDescription/Public/MeshDescription.h
@@ -1,6 +1,7 @@
 // Copyright 1998-2017 Epic Games, Inc. All Rights Reserved.
 
 #pragma once
+#pragma warning(disable:4503)
 
 #include "CoreMinimal.h"
 #include "MeshTypes.h"
@@ -493,9 +494,17 @@ struct TAttributesSet : public TMakeAttributesSet<ElementIDType, AttributeTypes>
 	}
 
 	template <typename AttributeType>
-	AttributeType SetAttribute( const ElementIDType ElementID, const FName AttributeName, const int32 AttributeIndex, const AttributeType& AttributeValue )
+	void SetAttribute( const ElementIDType ElementID, const FName AttributeName, const int32 AttributeIndex, const AttributeType& AttributeValue )
 	{
-		this->Get<TTupleIndex<AttributeType, AttributeTypes>::Value>().FindChecked( AttributeName )[ AttributeIndex ][ ElementID ] = AttributeValue;
+		TMeshElementArray<AttributeType, ElementIDType>& MeshElementArray = this->Get<TTupleIndex<AttributeType, AttributeTypes>::Value>().FindChecked(AttributeName)[AttributeIndex];
+		if(!MeshElementArray.IsValid(ElementID))
+		{
+			MeshElementArray.Insert(ElementID, AttributeValue);
+		}
+		else
+		{
+			MeshElementArray[ElementID] = AttributeValue;
+		}
 	}
 
 	friend FArchive& operator<<( FArchive& Ar, TAttributesSet& AttributesSet )
@@ -871,4 +880,4 @@ private:
 	TAttributesSet<FEdgeID> EdgeAttributesSet;
 	TAttributesSet<FPolygonID> PolygonAttributesSet;
 	TAttributesSet<FPolygonGroupID> PolygonGroupAttributesSet;
-};
+};
\ No newline at end of file
