Remove use or rendering array
Use the Create API from meshdescription
Fix the combine mesh
Fix the material remap

diff --git a/Engine/Source/Editor/UnrealEd/Private/Fbx/FbxStaticMeshImport.cpp b/Engine/Source/Editor/UnrealEd/Private/Fbx/FbxStaticMeshImport.cpp
--- a/Engine/Source/Editor/UnrealEd/Private/Fbx/FbxStaticMeshImport.cpp
+++ b/Engine/Source/Editor/UnrealEd/Private/Fbx/FbxStaticMeshImport.cpp
@@ -549,17 +549,16 @@ bool UnFbx::FFbxImporter::BuildStaticMeshFromGeometry(FbxNode* Node, UStaticMesh
 			FbxPosition = TotalMatrix.MultT(FbxPosition);
 			const FVector VertexPosition = Converter.ConvertPos(FbxPosition);
 			
-			FMeshVertex NewVertex;
+			FVertexID AddedVertexId = MeshDescription->CreateVertex();
+			FMeshVertex& NewVertex = MeshDescription->GetVertex(AddedVertexId);
 			NewVertex.VertexPosition = VertexPosition;
 			NewVertex.CornerSharpness = 0.0f;
-			FVertexID AddedVertexId = MeshDescription->Vertices().Add(NewVertex);
-			if (!MeshDescription->Vertices().IsValid(AddedVertexId) || AddedVertexId.GetValue() != RealVertexIndex)
+			if (AddedVertexId.GetValue() != RealVertexIndex)
 			{
 				AddTokenizedErrorMessage(FTokenizedMessage::Create(EMessageSeverity::Error, FText::Format(LOCTEXT("Error_CannotCreateVertex", "Cannot create valid vertex for mesh '{0}'"), FText::FromString(Mesh->GetName()))), FFbxErrors::StaticMesh_BuildError);
 				return false;
 			}
 		}
-
 		//Polygons
 		for (int32 TriangleIndex = 0; TriangleIndex < TriangleCount; TriangleIndex++)
 		{
@@ -576,21 +575,17 @@ bool UnFbx::FFbxImporter::BuildStaticMeshFromGeometry(FbxNode* Node, UStaticMesh
 				FVector VertexPosition = MeshDescription->GetVertex(VertexID).VertexPosition;
 				CornerVerticesIDs[CornerIndex] = VertexID;
 
-				FMeshVertexInstance TmpVertexInstance;
-				TmpVertexInstance.VertexID = VertexID;
-				FVertexInstanceID AddedVertexInstanceId = MeshDescription->VertexInstances().Add(TmpVertexInstance);
+				FVertexInstanceID AddedVertexInstanceId = MeshDescription->CreateVertexInstance(VertexID);
 				
 				//Make sure the Added vertex instance ID is matching the expected vertex instance ID
 				check(AddedVertexInstanceId == VertexInstanceID);
 				
 				FMeshVertexInstance& NewVertexInstance = MeshDescription->GetVertexInstance(AddedVertexInstanceId);
-				if (!MeshDescription->VertexInstances().IsValid(AddedVertexInstanceId) || AddedVertexInstanceId.GetValue() != VertexInstanceIndex)
+				if (AddedVertexInstanceId.GetValue() != VertexInstanceIndex)
 				{
 					AddTokenizedErrorMessage(FTokenizedMessage::Create(EMessageSeverity::Error, FText::Format(LOCTEXT("Error_CannotCreateVertexInstance", "Cannot create valid vertex instance for mesh '{0}'"), FText::FromString(Mesh->GetName()))), FFbxErrors::StaticMesh_BuildError);
 					return false;
 				}
-				FMeshVertex& MeshVertex = MeshDescription->GetVertex(VertexID);
-				MeshVertex.VertexInstanceIDs.Add(AddedVertexInstanceId);
 
 				//UVs attributes
 				for (int32 UVLayerIndex = 0; UVLayerIndex < FBXUVs.UniqueUVCount; UVLayerIndex++)
@@ -759,13 +754,13 @@ bool UnFbx::FFbxImporter::BuildStaticMeshFromGeometry(FbxNode* Node, UStaticMesh
 				}
 				if (ExistingPolygonGroup == FPolygonGroupID::Invalid)
 				{
-					FMeshPolygonGroup NewPolygonGroup;
+					ExistingPolygonGroup = MeshDescription->CreatePolygonGroup();
+					FMeshPolygonGroup& NewPolygonGroup = MeshDescription->GetPolygonGroup(ExistingPolygonGroup);
 					NewPolygonGroup.bCastShadow = true;
 					NewPolygonGroup.bEnableCollision = bEnableCollision;
 					NewPolygonGroup.ImportedMaterialSlotName = Material->GetFName();
 					NewPolygonGroup.MaterialSlotName = Material->GetFName();
 					NewPolygonGroup.MaterialAsset = FStringAssetReference(Material);
-					ExistingPolygonGroup = MeshDescription->PolygonGroups().Add(NewPolygonGroup);
 				}
 				PolygonGroupMapping.Add(RealMaterialIndex, ExistingPolygonGroup);
 			}
@@ -777,7 +772,7 @@ bool UnFbx::FFbxImporter::BuildStaticMeshFromGeometry(FbxNode* Node, UStaticMesh
 			TmpPolygon.PolygonGroupID = PolygonGroupID;
 			const FPolygonID NewPolygonID = MeshDescription->Polygons().Add(TmpPolygon);
 			FMeshPolygonGroup& MeshPolygonGroup = MeshDescription->GetPolygonGroup(PolygonGroupID);
-			MeshPolygonGroup.Polygons.Add(NewPolygonID);
+			MeshPolygonGroup.Polygons.AddUnique(NewPolygonID);
 			FMeshPolygon& NewPolygon = MeshDescription->GetPolygon(NewPolygonID);
 			FMeshTriangle NewTriangle;
 			for (int32 TriangleVertexIndex = 0; TriangleVertexIndex < 3; ++TriangleVertexIndex)
@@ -787,7 +782,7 @@ bool UnFbx::FFbxImporter::BuildStaticMeshFromGeometry(FbxNode* Node, UStaticMesh
 				FMeshVertexInstance& MeshVertexInstance = MeshDescription->GetVertexInstance(VertexInstanceID);
 				MeshVertexInstance.ConnectedPolygons.AddUnique(NewPolygonID);
 				const FVertexID &VertexID = CornerVerticesIDs[TriangleVertexIndex];
-				NewPolygon.PerimeterContour.VertexInstanceIDs.Add(VertexInstanceID);
+				NewPolygon.PerimeterContour.VertexInstanceIDs.AddUnique(VertexInstanceID);
 			}
 			NewPolygon.Triangles.Add(NewTriangle);
 
@@ -808,7 +803,8 @@ bool UnFbx::FFbxImporter::BuildStaticMeshFromGeometry(FbxNode* Node, UStaticMesh
 					EdgeVertexIDs[1] = CornerVerticesIDs[CornerIndices[1]];
 
 					FEdgeID MatchEdgeId = FEdgeID::Invalid;
-					for (const FEdgeID EdgeID : MeshDescription->Edges().GetElementIDs())
+					const FMeshVertex& EdgeMeshVertex = MeshDescription->GetVertex(EdgeVertexIDs[0]);
+					for (const FEdgeID& EdgeID : EdgeMeshVertex.ConnectedEdgeIDs)
 					{
 						FVertexID OutEdgeVertexID0 = MeshDescription->Edges()[EdgeID].VertexIDs[0];
 						FVertexID OutEdgeVertexID1 = MeshDescription->Edges()[EdgeID].VertexIDs[1];
@@ -820,34 +816,24 @@ bool UnFbx::FFbxImporter::BuildStaticMeshFromGeometry(FbxNode* Node, UStaticMesh
 						}
 					}
 					int32 SmoothMaskValue = 0;
-
 					if (MatchEdgeId == FEdgeID::Invalid)
 					{
-						FMeshEdge NewEdge;
-						NewEdge.VertexIDs[0] = EdgeVertexIDs[0];
-						NewEdge.VertexIDs[1] = EdgeVertexIDs[1];
-						NewEdge.bIsHardEdge = false;
-						NewEdge.CreaseSharpness = 0.0f;
-						MatchEdgeId = MeshDescription->Edges().Add(NewEdge);
+						MatchEdgeId = MeshDescription->CreateEdge(EdgeVertexIDs[0], EdgeVertexIDs[1]);
 					}
-
 					FMeshEdge& ExistingEdge = MeshDescription->GetEdge(MatchEdgeId);
-					ExistingEdge.ConnectedPolygons.Add(NewPolygonID);
-
-					FMeshVertex& MeshVertex0 = MeshDescription->GetVertex(EdgeVertexIDs[0]);
-					MeshVertex0.ConnectedEdgeIDs.Add(MatchEdgeId);
-					FMeshVertex& MeshVertex1 = MeshDescription->GetVertex(EdgeVertexIDs[1]);
-					MeshVertex1.ConnectedEdgeIDs.Add(MatchEdgeId);
 
+					ExistingEdge.ConnectedPolygons.AddUnique(NewPolygonID);
+					int32 EdgeIndex = Mesh->GetMeshEdgeIndexForPolygon(TriangleIndex, TriangleEdgeNumber);
+					ExistingEdge.CreaseSharpness = (float)Mesh->GetEdgeCreaseInfo(EdgeIndex);
 					if (!ExistingEdge.bIsHardEdge && bSmoothingAvailable && SmoothingInfo)
 					{
 						if (SmoothingMappingMode == FbxLayerElement::eByEdge)
 						{
-							int32 EdgeIndex = Mesh->GetMeshEdgeIndexForPolygon(TriangleIndex, TriangleEdgeNumber);
-							int lSmoothingIndex = (SmoothingReferenceMode == FbxLayerElement::eDirect) ? EdgeIndex : SmoothingInfo->GetIndexArray().GetAt(EdgeIndex);
+							int32 lSmoothingIndex = (SmoothingReferenceMode == FbxLayerElement::eDirect) ? EdgeIndex : SmoothingInfo->GetIndexArray().GetAt(EdgeIndex);
 							SmoothMaskValue = SmoothingInfo->GetDirectArray().GetAt(lSmoothingIndex);
 							if (SmoothMaskValue > 0)
 							{
+								//TODO: not sure I'm doing the right thing here? Maybe I should use eByPolygon smoothing to detect if I have a hard edge by comparing the connected polygons smooth group
 								//Set the hard edges
 								EdgeIDs.Add(MatchEdgeId);
 								EdgesIsHard.Add(true);
@@ -1902,36 +1888,11 @@ UStaticMesh* UnFbx::FFbxImporter::ImportStaticMeshAsSingle(UObject* InParent, TA
 			if (bDoRemap)
 			{
 				int32 GroupNumber = MeshDescription->PolygonGroups().Num();
-				TMap<int32, int32> PolygonGroupIDMap;
-				for (const FPolygonGroupID& PolygonGroupID : MeshDescription->PolygonGroups().GetElementIDs())
-				{
-					int32 MaterialIndex = PolygonGroupID.GetValue();
-					PolygonGroupIDMap.Add(MaterialIndex, MaterialMap[MaterialIndex]);
-				}
 
 				for (const FPolygonID& PolygonID : MeshDescription->Polygons().GetElementIDs())
 				{
 					FMeshPolygon MeshPolygon = MeshDescription->GetPolygon(PolygonID);
-					int32 MaterialIndex = MeshPolygon.PolygonGroupID.GetValue();
-					MeshPolygon.PolygonGroupID = FPolygonGroupID(PolygonGroupIDMap[MeshPolygon.PolygonGroupID.GetValue()]);
-				}
-				FPolygonGroupArray ReorderPolygonGroups;
-				for (int32 GroupIndex = 0; GroupIndex < GroupNumber; ++GroupIndex)
-				{
-					for (auto Kvp : PolygonGroupIDMap)
-					{
-						if (Kvp.Value == GroupIndex)
-						{
-							const FMeshPolygonGroup& MeshPolygonGroup = MeshDescription->GetPolygonGroup(FPolygonGroupID(Kvp.Key));
-							ReorderPolygonGroups.Insert(FPolygonGroupID(Kvp.Value), MeshPolygonGroup);
-							break;
-						}
-					}
-				}
-				MeshDescription->PolygonGroups().Reset(GroupNumber);
-				for (const FPolygonGroupID PolygonGroupID : ReorderPolygonGroups.GetElementIDs())
-				{
-					MeshDescription->PolygonGroups().Insert(PolygonGroupID, ReorderPolygonGroups[PolygonGroupID]);
+					MeshPolygon.PolygonGroupID = FPolygonGroupID(MaterialMap[MeshPolygon.PolygonGroupID.GetValue()]);
 				}
 			}
 
@@ -1939,7 +1900,8 @@ UStaticMesh* UnFbx::FFbxImporter::ImportStaticMeshAsSingle(UObject* InParent, TA
 			{
 				int32 MaterialIndex = PolygonGroupID.GetValue();
 				const FMeshPolygonGroup& MeshPolygonGroup = MeshDescription->GetPolygonGroup(PolygonGroupID);
-				int32 AddIndex = StaticMesh->StaticMaterials.Add(FStaticMaterial(Cast<UMaterialInterface>(MeshPolygonGroup.MaterialAsset.TryLoad()), MeshPolygonGroup.MaterialSlotName, MeshPolygonGroup.ImportedMaterialSlotName));
+				int32 AddedIndex = StaticMesh->StaticMaterials.Add(FStaticMaterial(Cast<UMaterialInterface>(MeshPolygonGroup.MaterialAsset.TryLoad()), MeshPolygonGroup.MaterialSlotName, MeshPolygonGroup.ImportedMaterialSlotName));
+				check(MaterialIndex == AddedIndex);
 			}
 		}
 
diff --git a/Engine/Source/Runtime/MeshBuilder/Private/StaticMeshBuilder.cpp b/Engine/Source/Runtime/MeshBuilder/Private/StaticMeshBuilder.cpp
--- a/Engine/Source/Runtime/MeshBuilder/Private/StaticMeshBuilder.cpp
+++ b/Engine/Source/Runtime/MeshBuilder/Private/StaticMeshBuilder.cpp
@@ -6,13 +6,29 @@
 #include "PhysicsEngine/BodySetup.h"
 #include "MeshDescription.h"
 
-const FTriangleID FTriangleID::Invalid(TNumericLimits<uint32>::Max());
-
 FStaticMeshBuilder::FStaticMeshBuilder()
 {
 
 }
 
+int32 FStaticMeshBuilder::GetPolygonGroupTriangles(UStaticMesh* StaticMesh, TArray<FMeshTriangle>& OutTriangles, const FPolygonGroupID& PolygonGroupID)
+{
+	int32 TriangleCount = 0;
+	for (const FPolygonID& PolygonID : StaticMesh->MeshDescription->Polygons().GetElementIDs())
+	{
+		FMeshPolygon MeshPolygon = StaticMesh->MeshDescription->GetPolygon(PolygonID);
+		if (MeshPolygon.PolygonGroupID == PolygonGroupID)
+		{
+			for (const FMeshTriangle& MeshTriangle : MeshPolygon.Triangles)
+			{
+				OutTriangles.Add(MeshTriangle);
+				TriangleCount++;
+			}
+		}
+	}
+	return TriangleCount;
+}
+
 bool FStaticMeshBuilder::Build(UStaticMesh* StaticMesh)
 {
 	OnBuildRenderMeshStart(StaticMesh, false);
@@ -62,52 +78,20 @@ bool FStaticMeshBuilder::Build(UStaticMesh* StaticMesh)
 		}
 	}
 
-	
-	/** All of the polygon groups in this mesh */
-	TMeshElementArray<FRenderingPolygonGroup, FPolygonGroupID> RenderingPolygonGroups;
-
-	for (const FPolygonID& PolygonID : StaticMesh->MeshDescription->Polygons().GetElementIDs())
-	{
-		FMeshPolygon MeshPolygon = StaticMesh->MeshDescription->GetPolygon(PolygonID);
-		FPolygonGroupID PolygonGroupID = MeshPolygon.PolygonGroupID;
-
-		if (!RenderingPolygonGroups.IsValid(PolygonGroupID))
-		{
-			// Create a rendering polygon group for holding the triangulated data and references to the static mesh rendering section.
-			// This is indexed by the same FPolygonGroupID as the PolygonGroups.
-			RenderingPolygonGroups.Insert(PolygonGroupID);
-		}
-		FRenderingPolygonGroup& RenderingPolygonGroup = RenderingPolygonGroups[PolygonGroupID];
-		for (const FMeshTriangle& MeshTriangle : MeshPolygon.Triangles)
-		{
-			RenderingPolygonGroup.Triangles.Add(MeshTriangle);
-		}
-		RenderingPolygonGroup.MaxTriangles = RenderingPolygonGroup.Triangles.Num();
-	}
-
-
 	// Set up index buffer
-	for (const FPolygonGroupID PolygonGroupID : PolygonGroups.GetElementIDs())
+	for (const FPolygonGroupID& PolygonGroupID : PolygonGroups.GetElementIDs())
 	{
+		TArray<FMeshTriangle> Triangles;
+		int32 SectionTriangleCount = GetPolygonGroupTriangles(StaticMesh, Triangles, PolygonGroupID);
+
 		const FMeshPolygonGroup& PolygonGroup = PolygonGroups[PolygonGroupID];
-		
-		if (!RenderingPolygonGroups.IsValid(PolygonGroupID))
-		{
-			//There is no polygon using this polygon group
-			continue;
-		}
-		FRenderingPolygonGroup& RenderingPolygonGroup = RenderingPolygonGroups[PolygonGroupID];
-		RenderingPolygonGroup.RenderingSectionIndex = StaticMeshLOD.Sections.Num();
 
 		// Create new rendering section
 		StaticMeshLOD.Sections.Add(FStaticMeshSection());
 		FStaticMeshSection& StaticMeshSection = StaticMeshLOD.Sections.Last();
 
 		StaticMeshSection.FirstIndex = IndexBuffer.Num();
-
-		//TODO: Get the number of triangles
-		StaticMeshSection.NumTriangles = RenderingPolygonGroup.Triangles.GetArraySize();
-		check(RenderingPolygonGroup.Triangles.GetArraySize() <= RenderingPolygonGroup.MaxTriangles);
+		StaticMeshSection.NumTriangles = SectionTriangleCount;
 
 		const int32 MaterialIndex = StaticMesh->GetMaterialIndex(PolygonGroup.MaterialSlotName);
 		check(MaterialIndex != INDEX_NONE);
@@ -116,50 +100,24 @@ bool FStaticMeshBuilder::Build(UStaticMesh* StaticMesh)
 		StaticMeshSection.bEnableCollision = PolygonGroup.bEnableCollision;
 		StaticMeshSection.bCastShadow = PolygonGroup.bCastShadow;
 
-
-		if (RenderingPolygonGroup.Triangles.Num() > 0)
+		if (Triangles.Num() > 0)
 		{
-			IndexBuffer.Reserve(IndexBuffer.Num() + RenderingPolygonGroup.Triangles.GetArraySize() * 3);
+			IndexBuffer.Reserve(IndexBuffer.Num() + Triangles.Num() * 3);
 			uint32 MinIndex = TNumericLimits< uint32 >::Max();
 			uint32 MaxIndex = TNumericLimits< uint32 >::Min();
-
-			// Find the first valid vertex instance index, so that we have a value we can use for our degenerates
-			check(RenderingPolygonGroup.Triangles.Num() > 0);
-			const FVertexInstanceID FirstValidRenderingID = RenderingPolygonGroup.Triangles[RenderingPolygonGroup.Triangles.GetFirstValidID()].GetVertexInstanceID(0);
-
-			for (int32 TriangleIndex = 0; TriangleIndex < RenderingPolygonGroup.Triangles.GetArraySize(); ++TriangleIndex)
+			for (int32 TriangleIndex = 0; TriangleIndex < Triangles.Num(); ++TriangleIndex)
 			{
-				const FTriangleID TriangleID(TriangleIndex);
-				if (RenderingPolygonGroup.Triangles.IsValid(TriangleID))
-				{
-					const FMeshTriangle& Triangle = RenderingPolygonGroup.Triangles[TriangleID];
-					for (int32 TriVert = 0; TriVert < 3; ++TriVert)
-					{
-						const uint32 RenderingVertexIndex = Triangle.GetVertexInstanceID(TriVert).GetValue();
-						IndexBuffer.Add(RenderingVertexIndex);
-						MinIndex = FMath::Min(MinIndex, RenderingVertexIndex);
-						MaxIndex = FMath::Max(MaxIndex, RenderingVertexIndex);
-					}
-				}
-				else
+				const FMeshTriangle& Triangle = Triangles[TriangleIndex];
+				for (int32 TriVert = 0; TriVert < 3; ++TriVert)
 				{
-					IndexBuffer.Add(FirstValidRenderingID.GetValue());
-					IndexBuffer.Add(FirstValidRenderingID.GetValue());
-					IndexBuffer.Add(FirstValidRenderingID.GetValue());
+					const uint32 RenderingVertexIndex = Triangle.GetVertexInstanceID(TriVert).GetValue();
+					IndexBuffer.Add(RenderingVertexIndex);
+					MinIndex = FMath::Min(MinIndex, RenderingVertexIndex);
+					MaxIndex = FMath::Max(MaxIndex, RenderingVertexIndex);
 				}
 			}
-
 			StaticMeshSection.MinVertexIndex = MinIndex;
 			StaticMeshSection.MaxVertexIndex = MaxIndex;
-
-			// Add any index buffer padding.
-			// This can be necessary if we have just loaded an editable mesh which had a MaxTriangles count in the editable mesh section
-			// greater than the sparse array max size (i.e. an extra gap had been reserved for tris).
-			const int32 IndexBufferPadding = RenderingPolygonGroup.MaxTriangles - RenderingPolygonGroup.Triangles.GetArraySize();
-			if (IndexBufferPadding > 0)
-			{
-				IndexBuffer.AddZeroed(IndexBufferPadding * 3);
-			}
 		}
 		else
 		{
diff --git a/Engine/Source/Runtime/MeshBuilder/Public/StaticMeshBuilder.h b/Engine/Source/Runtime/MeshBuilder/Public/StaticMeshBuilder.h
--- a/Engine/Source/Runtime/MeshBuilder/Public/StaticMeshBuilder.h
+++ b/Engine/Source/Runtime/MeshBuilder/Public/StaticMeshBuilder.h
@@ -25,6 +25,8 @@ private:
 	void OnBuildRenderMeshStart(class UStaticMesh* StaticMesh, const bool bInvalidateLighting);
 	void OnBuildRenderMeshFinish(class UStaticMesh* StaticMesh, const bool bRebuildBoundsAndCollision);
 	
+	int32 GetPolygonGroupTriangles(class UStaticMesh* StaticMesh, TArray<FMeshTriangle>& OutTriangles, const FPolygonGroupID& PolygonGroupID);
+
 	bool IsOrphanedVertex(class UStaticMesh* StaticMesh, const FVertexID VertexID) const;
 	void UpdateBounds(class UStaticMesh* StaticMesh);
 	void UpdateCollision(UStaticMesh *StaticMesh);
@@ -33,67 +35,3 @@ private:
 	TSharedPtr<class FStaticMeshComponentRecreateRenderStateContext> RecreateRenderStateContext;
 };
 
-struct FTriangleID : public FElementID
-{
-	FTriangleID()
-	{
-	}
-
-	explicit FTriangleID(const FElementID InitElementID)
-		: FElementID(InitElementID.GetValue())
-	{
-	}
-
-	explicit FTriangleID(const uint32 InitIDValue)
-		: FElementID(InitIDValue)
-	{
-	}
-
-	FORCEINLINE friend uint32 GetTypeHash(const FTriangleID& Other)
-	{
-		return GetTypeHash(Other.IDValue);
-	}
-
-	/** Invalid triangle ID */
-	static const FTriangleID Invalid;
-};
-
-
-struct FRenderingPolygon
-{
-
-	/** Which rendering polygon group the polygon is in */
-	FPolygonGroupID PolygonGroupID;
-
-	/** This is a list of indices of triangles in the FRenderingPolygonGroup::Triangles array.
-	We use this to maintain a record of which triangles in the section belong to this polygon. */
-	TArray<FTriangleID> TriangulatedPolygonTriangleIndices;
-};
-
-
-struct FRenderingPolygonGroup
-{
-
-	/** The rendering section index for this mesh section */
-	uint32 RenderingSectionIndex;
-
-	/** Maximum number of triangles which have been reserved in the index buffer */
-	int32 MaxTriangles;
-
-	/** Sparse array of triangles, that matches the triangles in the mesh index buffers.  Elements that
-	aren't allocated will be stored as degenerates in the mesh index buffer. */
-	TMeshElementArray<FMeshTriangle, FTriangleID> Triangles;
-
-
-	/** Converts from an index in our Triangles array to an index of a rendering triangle's first vertex in the rendering mesh's index buffer */
-	inline static FTriangleID RenderingTriangleFirstIndexToTriangleIndex(const FStaticMeshSection& RenderingSection, const uint32 RenderingTriangleFirstIndex)
-	{
-		return FTriangleID((RenderingTriangleFirstIndex - RenderingSection.FirstIndex) / 3);
-	}
-
-	/** Converts from an index of a rendering triangle's first vertex in the rendering mesh's index buffer to an index in our Triangles array */
-	inline static uint32 TriangleIndexToRenderingTriangleFirstIndex(const FStaticMeshSection& RenderingSection, const FTriangleID TriangleIndex)
-	{
-		return TriangleIndex.GetValue() * 3 + RenderingSection.FirstIndex;
-	}
-};
