Use the create polygon API

diff --git a/Engine/Source/Editor/UnrealEd/Private/Fbx/FbxStaticMeshImport.cpp b/Engine/Source/Editor/UnrealEd/Private/Fbx/FbxStaticMeshImport.cpp
--- a/Engine/Source/Editor/UnrealEd/Private/Fbx/FbxStaticMeshImport.cpp
+++ b/Engine/Source/Editor/UnrealEd/Private/Fbx/FbxStaticMeshImport.cpp
@@ -768,38 +768,15 @@ bool UnFbx::FFbxImporter::BuildStaticMeshFromGeometry(FbxNode* Node, UStaticMesh
 				PolygonGroupMapping.Add(RealMaterialIndex, ExistingPolygonGroup);
 			}
 
-			FPolygonGroupID PolygonGroupID = PolygonGroupMapping[RealMaterialIndex];
-
-			// Insert a polygon into the mesh
-			FMeshPolygon TmpPolygon;
-			TmpPolygon.PolygonGroupID = PolygonGroupID;
-			const FPolygonID NewPolygonID = MeshDescription->Polygons().Add(TmpPolygon);
-			FMeshPolygonGroup& MeshPolygonGroup = MeshDescription->GetPolygonGroup(PolygonGroupID);
-			
-			//We cannot do a add unique here since this is too slow, since we just create the polygons it should not be part
-			//of the polygonGroup already
-			MeshPolygonGroup.Polygons.Add(NewPolygonID);
-			
-			FMeshPolygon& NewPolygon = MeshDescription->GetPolygon(NewPolygonID);
-			FMeshTriangle NewTriangle;
-			for (int32 TriangleVertexIndex = 0; TriangleVertexIndex < 3; ++TriangleVertexIndex)
-			{
-				const FVertexInstanceID &VertexInstanceID = CornerInstanceIDs[TriangleVertexIndex];
-				NewTriangle.SetVertexInstanceID(TriangleVertexIndex, VertexInstanceID);
-				FMeshVertexInstance& MeshVertexInstance = MeshDescription->GetVertexInstance(VertexInstanceID);
-				MeshVertexInstance.ConnectedPolygons.AddUnique(NewPolygonID);
-				const FVertexID &VertexID = CornerVerticesIDs[TriangleVertexIndex];
-				NewPolygon.PerimeterContour.VertexInstanceIDs.AddUnique(VertexInstanceID);
-			}
-			NewPolygon.Triangles.Add(NewTriangle);
-
-			// Set edges hardness
+			// Create polygon edges
+			TArray<UMeshDescription::FContourPoint> Contours;
 			{
 				TArray<FEdgeID> EdgeIDs;
 				TArray<bool> EdgesIsHard;
 				// Add the edges of this triangle
 				for (uint32 TriangleEdgeNumber = 0; TriangleEdgeNumber < 3; ++TriangleEdgeNumber)
 				{
+					UMeshDescription::FContourPoint ContourPoint;
 					//Find the matching edge ID
 					uint32 CornerIndices[2];
 					CornerIndices[0] = (TriangleEdgeNumber + 0) % 3;
@@ -808,6 +785,7 @@ bool UnFbx::FFbxImporter::BuildStaticMeshFromGeometry(FbxNode* Node, UStaticMesh
 					FVertexID EdgeVertexIDs[2];
 					EdgeVertexIDs[0] = CornerVerticesIDs[CornerIndices[0]];
 					EdgeVertexIDs[1] = CornerVerticesIDs[CornerIndices[1]];
+					
 
 					FEdgeID MatchEdgeId = FEdgeID::Invalid;
 					const FMeshVertex& EdgeMeshVertex = MeshDescription->GetVertex(EdgeVertexIDs[0]);
@@ -827,9 +805,11 @@ bool UnFbx::FFbxImporter::BuildStaticMeshFromGeometry(FbxNode* Node, UStaticMesh
 					{
 						MatchEdgeId = MeshDescription->CreateEdge(EdgeVertexIDs[0], EdgeVertexIDs[1]);
 					}
+					ContourPoint.EdgeID = MatchEdgeId;
+					ContourPoint.VertexInstanceID = CornerInstanceIDs[CornerIndices[0]];
+					Contours.Add(ContourPoint);
 					FMeshEdge& ExistingEdge = MeshDescription->GetEdge(MatchEdgeId);
 
-					ExistingEdge.ConnectedPolygons.AddUnique(NewPolygonID);
 					int32 EdgeIndex = Mesh->GetMeshEdgeIndexForPolygon(TriangleIndex, TriangleEdgeNumber);
 					ExistingEdge.CreaseSharpness = (float)Mesh->GetEdgeCreaseInfo(EdgeIndex);
 					if (!ExistingEdge.bIsHardEdge && bSmoothingAvailable && SmoothingInfo)
@@ -858,6 +838,23 @@ bool UnFbx::FFbxImporter::BuildStaticMeshFromGeometry(FbxNode* Node, UStaticMesh
 					MeshDescription->GetEdge(EdgeID).bIsHardEdge = EdgesIsHard[EdgeIndex];
 				}
 			}
+
+			FPolygonGroupID PolygonGroupID = PolygonGroupMapping[RealMaterialIndex];
+			// Insert a polygon into the mesh
+			const FPolygonID NewPolygonID = MeshDescription->CreatePolygon(Contours);
+			FMeshPolygonGroup& MeshPolygonGroup = MeshDescription->GetPolygonGroup(PolygonGroupID);
+			FMeshPolygon& NewPolygon = MeshDescription->GetPolygon(NewPolygonID);
+			NewPolygon.PolygonGroupID = PolygonGroupID;
+			MeshPolygonGroup.Polygons.Add(NewPolygonID);
+			FMeshTriangle NewTriangle;
+			for (int32 TriangleVertexIndex = 0; TriangleVertexIndex < 3; ++TriangleVertexIndex)
+			{
+				const FVertexInstanceID &VertexInstanceID = CornerInstanceIDs[TriangleVertexIndex];
+				NewTriangle.SetVertexInstanceID(TriangleVertexIndex, VertexInstanceID);
+			}
+			NewPolygon.Triangles.Add(NewTriangle);
+
+			
 		}
 		//Call this after all GetMeshEdgeIndexForPolygon call this is for optimization purpose.
 		Mesh->EndGetMeshEdgeIndexForPolygon();
