Change where we store the original mesh description.
Make sure fbx automation tests use Meshdescription instead of rawmesh

diff --git a/Engine/Source/Editor/UnrealEd/Private/Fbx/FbxStaticMeshImport.cpp b/Engine/Source/Editor/UnrealEd/Private/Fbx/FbxStaticMeshImport.cpp
--- a/Engine/Source/Editor/UnrealEd/Private/Fbx/FbxStaticMeshImport.cpp
+++ b/Engine/Source/Editor/UnrealEd/Private/Fbx/FbxStaticMeshImport.cpp
@@ -580,19 +580,27 @@ bool UnFbx::FFbxImporter::BuildStaticMeshFromGeometry(FbxNode* Node, UStaticMesh
 				}
 
 				//UVs attributes
-				for (int32 UVLayerIndex = 0; UVLayerIndex < FBXUVs.UniqueUVCount; UVLayerIndex++)
+				if (FBXUVs.UniqueUVCount > 0)
 				{
-					FVector2D FinalUVVector(0.0f, 0.0f);
-					if (FBXUVs.LayerElementUV[UVLayerIndex] != NULL)
+					for (int32 UVLayerIndex = 0; UVLayerIndex < FBXUVs.UniqueUVCount; UVLayerIndex++)
 					{
-						int UVMapIndex = (FBXUVs.UVMappingMode[UVLayerIndex] == FbxLayerElement::eByControlPoint) ? ControlPointIndex : TriangleIndex * 3 + CornerIndex;
-						int32 UVIndex = (FBXUVs.UVReferenceMode[UVLayerIndex] == FbxLayerElement::eDirect) ?
-							UVMapIndex : FBXUVs.LayerElementUV[UVLayerIndex]->GetIndexArray().GetAt(UVMapIndex);
+						FVector2D FinalUVVector(0.0f, 0.0f);
+						if (FBXUVs.LayerElementUV[UVLayerIndex] != NULL)
+						{
+							int UVMapIndex = (FBXUVs.UVMappingMode[UVLayerIndex] == FbxLayerElement::eByControlPoint) ? ControlPointIndex : TriangleIndex * 3 + CornerIndex;
+							int32 UVIndex = (FBXUVs.UVReferenceMode[UVLayerIndex] == FbxLayerElement::eDirect) ?
+								UVMapIndex : FBXUVs.LayerElementUV[UVLayerIndex]->GetIndexArray().GetAt(UVMapIndex);
 
-						FbxVector2	UVVector = FBXUVs.LayerElementUV[UVLayerIndex]->GetDirectArray().GetAt(UVIndex);
-						FinalUVVector.X = static_cast<float>(UVVector[0]);
-						FinalUVVector.Y = 1.f - static_cast<float>(UVVector[1]);   //flip the Y of UVs for DirectX
+							FbxVector2	UVVector = FBXUVs.LayerElementUV[UVLayerIndex]->GetDirectArray().GetAt(UVIndex);
+							FinalUVVector.X = static_cast<float>(UVVector[0]);
+							FinalUVVector.Y = 1.f - static_cast<float>(UVVector[1]);   //flip the Y of UVs for DirectX
+						}
+						NewVertexInstance.VertexUVs.Add(FinalUVVector);
 					}
+				}
+				else
+				{
+					FVector2D FinalUVVector(0.0f, 0.0f);
 					NewVertexInstance.VertexUVs.Add(FinalUVVector);
 				}
 
@@ -788,17 +796,25 @@ bool UnFbx::FFbxImporter::BuildStaticMeshFromGeometry(FbxNode* Node, UStaticMesh
 
 					int32 EdgeIndex = Mesh->GetMeshEdgeIndexForPolygon(TriangleIndex, TriangleEdgeNumber);
 					ExistingEdge.CreaseSharpness = (float)Mesh->GetEdgeCreaseInfo(EdgeIndex);
-					if (!ExistingEdge.bIsHardEdge && bSmoothingAvailable && SmoothingInfo)
+					if (!ExistingEdge.bIsHardEdge)
 					{
-						if (SmoothingMappingMode == FbxLayerElement::eByEdge)
+						if (bSmoothingAvailable && SmoothingInfo)
 						{
-							int32 lSmoothingIndex = (SmoothingReferenceMode == FbxLayerElement::eDirect) ? EdgeIndex : SmoothingInfo->GetIndexArray().GetAt(EdgeIndex);
-							//Set the hard edges
-							ExistingEdge.bIsHardEdge = (SmoothingInfo->GetDirectArray().GetAt(lSmoothingIndex) == 0);
+							if (SmoothingMappingMode == FbxLayerElement::eByEdge)
+							{
+								int32 lSmoothingIndex = (SmoothingReferenceMode == FbxLayerElement::eDirect) ? EdgeIndex : SmoothingInfo->GetIndexArray().GetAt(EdgeIndex);
+								//Set the hard edges
+								ExistingEdge.bIsHardEdge = (SmoothingInfo->GetDirectArray().GetAt(lSmoothingIndex) == 0);
+							}
+							else
+							{
+								AddTokenizedErrorMessage(FTokenizedMessage::Create(EMessageSeverity::Warning, FText::Format(LOCTEXT("Error_UnsupportedSmoothingGroup", "Unsupported Smoothing group mapping mode on mesh  '{0}'"), FText::FromString(Mesh->GetName()))), FFbxErrors::Generic_Mesh_UnsupportingSmoothingGroup);
+							}
 						}
 						else
 						{
-							AddTokenizedErrorMessage(FTokenizedMessage::Create(EMessageSeverity::Warning, FText::Format(LOCTEXT("Error_UnsupportedSmoothingGroup", "Unsupported Smoothing group mapping mode on mesh  '{0}'"), FText::FromString(Mesh->GetName()))), FFbxErrors::Generic_Mesh_UnsupportingSmoothingGroup);
+							//When there is no smoothing group we set all edge to hard (faceted mesh)
+							ExistingEdge.bIsHardEdge = true;
 						}
 					}
 				}
@@ -1526,6 +1542,17 @@ UStaticMesh* UnFbx::FFbxImporter::ImportStaticMeshAsSingle(UObject* InParent, TA
 		StaticMesh = NewObject<UStaticMesh>(Package, FName(*MeshName), Flags | RF_Public);
 	}
 
+	if (StaticMesh->SourceModels.Num() < LODIndex+1)
+	{
+		// Add one LOD 
+		new(StaticMesh->SourceModels) FStaticMeshSourceModel();
+		
+		if (StaticMesh->SourceModels.Num() < LODIndex+1)
+		{
+			LODIndex = StaticMesh->SourceModels.Num() - 1;
+		}
+	}
+
 	if (ImportOptions->bImportMeshDescription)
 	{
 		MeshDescription = StaticMesh->GetOriginalMeshDescription(LODIndex);
@@ -1538,17 +1565,6 @@ UStaticMesh* UnFbx::FFbxImporter::ImportStaticMeshAsSingle(UObject* InParent, TA
 			StaticMesh->SetOriginalMeshDescription(LODIndex, MeshDescription);
 		}
 	}
-
-	if (StaticMesh->SourceModels.Num() < LODIndex+1)
-	{
-		// Add one LOD 
-		new(StaticMesh->SourceModels) FStaticMeshSourceModel();
-		
-		if (StaticMesh->SourceModels.Num() < LODIndex+1)
-		{
-			LODIndex = StaticMesh->SourceModels.Num() - 1;
-		}
-	}
 	
 	FStaticMeshSourceModel& SrcModel = StaticMesh->SourceModels[LODIndex];
 	if( InStaticMesh != NULL && LODIndex > 0 && !SrcModel.RawMeshBulkData->IsEmpty() )
@@ -1866,7 +1882,17 @@ UStaticMesh* UnFbx::FFbxImporter::ImportStaticMeshAsSingle(UObject* InParent, TA
 				for (const FPolygonID& PolygonID : MeshDescription->Polygons().GetElementIDs())
 				{
 					FMeshPolygon MeshPolygon = MeshDescription->GetPolygon(PolygonID);
-					MeshPolygon.PolygonGroupID = FPolygonGroupID(MaterialMap[MeshPolygon.PolygonGroupID.GetValue()]);
+					if (MaterialMap[MeshPolygon.PolygonGroupID.GetValue()] != MeshPolygon.PolygonGroupID.GetValue())
+					{
+						//Remove from the old polygon group
+						FMeshPolygonGroup& SourcePolygonGroup = MeshDescription->GetPolygonGroup(MeshPolygon.PolygonGroupID);
+						SourcePolygonGroup.Polygons.Remove(PolygonID);
+						//Set the new ID
+						MeshPolygon.PolygonGroupID = FPolygonGroupID(MaterialMap[MeshPolygon.PolygonGroupID.GetValue()]);
+						//Add to the new polygon group
+						FMeshPolygonGroup& DestinationPolygonGroup = MeshDescription->GetPolygonGroup(MeshPolygon.PolygonGroupID);
+						DestinationPolygonGroup.Polygons.Add(PolygonID);
+					}
 				}
 			}
 			TArray<FStaticMaterial> MaterialToAdd;
@@ -1895,11 +1921,6 @@ UStaticMesh* UnFbx::FFbxImporter::ImportStaticMeshAsSingle(UObject* InParent, TA
 							const FStaticMaterial& StaticMeshMaterial = StaticMesh->StaticMaterials[StaticMeshMaterialIndex];
 							if (StaticMeshMaterial.MaterialInterface == CandidateMaterial.MaterialInterface)
 							{
-								FPolygonGroupID PolygonGroupID(MaterialToAddIndex);
-								FMeshPolygonGroup& MeshPolygonGroup = MeshDescription->GetPolygonGroup(PolygonGroupID);
-								//The LOD 0 existing material is driving the polygon groups slot detail, so all LOD mesh description use the same data.
-								MeshPolygonGroup.ImportedMaterialSlotName = StaticMeshMaterial.ImportedMaterialSlotName;
-								MeshPolygonGroup.MaterialSlotName = StaticMeshMaterial.MaterialSlotName;
 								FoundExistingMaterial = true;
 								break;
 							}
diff --git a/Engine/Source/Editor/UnrealEd/Private/GeomFitUtils.cpp b/Engine/Source/Editor/UnrealEd/Private/GeomFitUtils.cpp
--- a/Engine/Source/Editor/UnrealEd/Private/GeomFitUtils.cpp
+++ b/Engine/Source/Editor/UnrealEd/Private/GeomFitUtils.cpp
@@ -20,6 +20,7 @@
 #include "PhysicsEngine/SphylElem.h"
 #include "PhysicsEngine/BodySetup.h"
 #include "Engine/StaticMesh.h"
+#include "MeshDescription.h"
 
 #define LOCAL_EPS (0.01f)
 static void AddVertexIfNotPresent(TArray<FVector> &vertices, FVector &newVertex)
@@ -199,15 +200,35 @@ static void CalcBoundingBox(const FRawMesh& RawMesh, FVector& Center, FVector& E
 	Box.GetCenterAndExtents(Center, Extents);
 }
 
+static void CalcBoundingBox(const UMeshDescription* MeshDescription, FVector& Center, FVector& Extents, FVector& LimitVec)
+{
+	FBox Box(ForceInit);
+	for (const FVertexID& VertexID : MeshDescription->Vertices().GetElementIDs())
+	{
+		const FMeshVertex& Vertex = MeshDescription->GetVertex(VertexID);
+		Box += Vertex.VertexPosition * LimitVec;
+	}
+	Box.GetCenterAndExtents(Center, Extents);
+}
+
 void ComputeBoundingBox(UStaticMesh* StaticMesh, FVector& Center, FVector& Extents)
 {
 	// Calculate bounding Box.
-	FRawMesh RawMesh;
+	
 	FStaticMeshSourceModel& SrcModel = StaticMesh->SourceModels[0];
-	SrcModel.RawMeshBulkData->LoadRawMesh(RawMesh);
-
-	FVector unitVec = FVector(1.f);
-	CalcBoundingBox(RawMesh, Center, Extents, unitVec);
+	if (StaticMesh->GetOriginalMeshDescription(0) != nullptr)
+	{
+		UMeshDescription* MeshDescription = StaticMesh->GetOriginalMeshDescription(0);
+		FVector unitVec = FVector(1.f);
+		CalcBoundingBox(MeshDescription, Center, Extents, unitVec);
+	}
+	else
+	{
+		FRawMesh RawMesh;
+		SrcModel.RawMeshBulkData->LoadRawMesh(RawMesh);
+		FVector unitVec = FVector(1.f);
+		CalcBoundingBox(RawMesh, Center, Extents, unitVec);
+	}
 }
 
 int32 GenerateBoxAsSimpleCollision(UStaticMesh* StaticMesh)
@@ -221,12 +242,21 @@ int32 GenerateBoxAsSimpleCollision(UStaticMesh* StaticMesh)
 
 	// Calculate bounding Box.
 	FRawMesh RawMesh;
+	UMeshDescription* MeshDescription = nullptr;
 	FStaticMeshSourceModel& SrcModel = StaticMesh->SourceModels[0];
-	SrcModel.RawMeshBulkData->LoadRawMesh(RawMesh);
-
+	
 	FVector unitVec = bs->BuildScale3D;
 	FVector Center, Extents;
-	CalcBoundingBox(RawMesh, Center, Extents, unitVec);
+	if (StaticMesh->GetOriginalMeshDescription(0) != nullptr)
+	{
+		MeshDescription = StaticMesh->GetMeshDescription(0);
+		CalcBoundingBox(MeshDescription, Center, Extents, unitVec);
+	}
+	else
+	{
+		SrcModel.RawMeshBulkData->LoadRawMesh(RawMesh);
+		CalcBoundingBox(RawMesh, Center, Extents, unitVec);
+	}
 
 	bs->Modify();
 
diff --git a/Engine/Source/Editor/UnrealEd/Private/StaticMeshEdit.cpp b/Engine/Source/Editor/UnrealEd/Private/StaticMeshEdit.cpp
--- a/Engine/Source/Editor/UnrealEd/Private/StaticMeshEdit.cpp
+++ b/Engine/Source/Editor/UnrealEd/Private/StaticMeshEdit.cpp
@@ -892,6 +892,7 @@ struct ExistingLODMeshData
 	FMeshBuildSettings			ExistingBuildSettings;
 	FMeshReductionSettings		ExistingReductionSettings;
 	FRawMesh					ExistingRawMesh;
+	UMeshDescription*			ExistingMeshDescription;
 	TArray<FStaticMaterial>		ExistingMaterials;
 	float						ExistingScreenSize;
 };
@@ -1049,6 +1050,7 @@ ExistingStaticMeshData* SaveExistingStaticMeshData(UStaticMesh* ExistingMesh, Un
 			ExistingMeshDataPtr->ExistingLODData[i].ExistingBuildSettings = ExistingMesh->SourceModels[i].BuildSettings;
 			ExistingMeshDataPtr->ExistingLODData[i].ExistingReductionSettings = ExistingMesh->SourceModels[i].ReductionSettings;
 			ExistingMeshDataPtr->ExistingLODData[i].ExistingScreenSize = ExistingMesh->SourceModels[i].ScreenSize;
+			ExistingMeshDataPtr->ExistingLODData[i].ExistingMeshDescription = ExistingMesh->SourceModels[i].OriginalMeshDescription;
 			ExistingMesh->SourceModels[i].RawMeshBulkData->LoadRawMesh(ExistingMeshDataPtr->ExistingLODData[i].ExistingRawMesh);
 		}
 
@@ -1455,7 +1457,10 @@ void RestoreExistingMeshData(ExistingStaticMeshData* ExistingMeshDataPtr, UStati
 	for(int32 i=NumCommonLODs; i < ExistingMeshDataPtr->ExistingLODData.Num(); ++i)
 	{
 		FStaticMeshSourceModel* SrcModel = new(NewMesh->SourceModels) FStaticMeshSourceModel();
-
+		if (ExistingMeshDataPtr->ExistingLODData[i].ExistingMeshDescription != nullptr)
+		{
+			SrcModel->OriginalMeshDescription = ExistingMeshDataPtr->ExistingLODData[i].ExistingMeshDescription;
+		}
 		if (ExistingMeshDataPtr->ExistingLODData[i].ExistingRawMesh.IsValidOrFixable())
 		{
 			SrcModel->RawMeshBulkData->SaveRawMesh(ExistingMeshDataPtr->ExistingLODData[i].ExistingRawMesh);
diff --git a/Engine/Source/Editor/UnrealEd/Private/Tests/FbxAutomationCommon.cpp b/Engine/Source/Editor/UnrealEd/Private/Tests/FbxAutomationCommon.cpp
--- a/Engine/Source/Editor/UnrealEd/Private/Tests/FbxAutomationCommon.cpp
+++ b/Engine/Source/Editor/UnrealEd/Private/Tests/FbxAutomationCommon.cpp
@@ -78,6 +78,8 @@ namespace FbxAutomationTestsAPI
 			FbxTestPlan->ImportUI->AnimSequenceImportData = SavedAnimSequenceData;
 			FbxTestPlan->ImportUI->TextureImportData = SavedTextureData;
 
+			//Force usage of mesh descriptions
+			FbxTestPlan->ImportUI->bImportMeshDescription = true;
 			TestPlanArray.Add(FbxTestPlan);
 		}
 	}
diff --git a/Engine/Source/Runtime/Engine/Classes/Engine/StaticMesh.h b/Engine/Source/Runtime/Engine/Classes/Engine/StaticMesh.h
--- a/Engine/Source/Runtime/Engine/Classes/Engine/StaticMesh.h
+++ b/Engine/Source/Runtime/Engine/Classes/Engine/StaticMesh.h
@@ -153,6 +153,10 @@ struct FStaticMeshSourceModel
 #if WITH_EDITOR
 	/** Imported raw mesh data. Optional for all but the first LOD. */
 	class FRawMeshBulkData* RawMeshBulkData;
+
+	/* Original Imported meshdescription data. Optional for all but the first LOD. */
+	UPROPERTY()
+	class UMeshDescription* OriginalMeshDescription;
 #endif // #if WITH_EDITOR
 
 	/** Settings applied when building the mesh. */
@@ -631,15 +635,8 @@ public:
 	ENGINE_API int32 GetMeshDescriptionCount() const;
 
 #if WITH_EDITORONLY_DATA
-	/**
-	* The original mesh description we use to do non destructive reduce of the mesh, we set the original value when we create the mesh (procedurally or at import/re-import...).
-	*/
-	UPROPERTY()
-	TArray<class UMeshDescription*> OriginalMeshDescriptions;
-
 	ENGINE_API class UMeshDescription* GetOriginalMeshDescription(int32 LodIndex = 0) const;
 	ENGINE_API void SetOriginalMeshDescription(int32 LodIndex, class UMeshDescription* MeshDescription);
-	ENGINE_API int32 GetOriginalMeshDescriptionCount() const;
 	ENGINE_API void ClearOriginalMeshDescription(int32 LodIndex);
 #endif
 
diff --git a/Engine/Source/Runtime/Engine/Private/StaticMesh.cpp b/Engine/Source/Runtime/Engine/Private/StaticMesh.cpp
--- a/Engine/Source/Runtime/Engine/Private/StaticMesh.cpp
+++ b/Engine/Source/Runtime/Engine/Private/StaticMesh.cpp
@@ -1142,9 +1142,9 @@ static FString BuildStaticMeshDerivedDataKey(UStaticMesh* Mesh, const FStaticMes
 	for (int32 LODIndex = 0; LODIndex < NumLODs; ++LODIndex)
 	{
 		FStaticMeshSourceModel& SrcModel = Mesh->SourceModels[LODIndex];
-		if (Mesh->GetOriginalMeshDescriptionCount() > 0)
+		if (Mesh->GetOriginalMeshDescription(0) != nullptr)
 		{
-			KeySuffix += (Mesh->GetOriginalMeshDescriptionCount() > LODIndex && Mesh->GetOriginalMeshDescription(LODIndex) != nullptr) ? Mesh->GetOriginalMeshDescription(LODIndex)->GetIdString() : TEXT("NoOriginalMeshDescriptionForLod") + FString::FromInt(LODIndex);
+			KeySuffix += (Mesh->GetOriginalMeshDescription(LODIndex) != nullptr) ? Mesh->GetOriginalMeshDescription(LODIndex)->GetIdString() : TEXT("NoOriginalMeshDescriptionForLod") + FString::FromInt(LODIndex);
 		}
 		else
 		{
@@ -1293,7 +1293,7 @@ void FStaticMeshRenderData::Cache(UStaticMesh* Owner, const FStaticMeshLODSettin
 			Args.Add(TEXT("StaticMeshName"), FText::FromString( Owner->GetName() ) );
 			FStaticMeshStatusMessageContext StatusContext( FText::Format( NSLOCTEXT("Engine", "BuildingStaticMeshStatus", "Building static mesh {StaticMeshName}..."), Args ) );
 
-			if (Owner->GetOriginalMeshDescriptionCount() > 0)
+			if (Owner->GetOriginalMeshDescription(0) != nullptr)
 			{
 				IMeshBuilderModule& MeshBuilderModule = FModuleManager::Get().LoadModuleChecked<IMeshBuilderModule>(TEXT("MeshBuilder"));
 				MeshBuilderModule.BuildMesh(Owner);
@@ -2031,6 +2031,7 @@ FStaticMeshSourceModel::FStaticMeshSourceModel()
 #if WITH_EDITOR
 	RawMeshBulkData = new FRawMeshBulkData();
 	ScreenSize = 0.0f;
+	OriginalMeshDescription = nullptr;
 #endif // #if WITH_EDITOR
 }
 
@@ -2042,8 +2043,9 @@ FStaticMeshSourceModel::~FStaticMeshSourceModel()
 		delete RawMeshBulkData;
 		RawMeshBulkData = NULL;
 	}
+	OriginalMeshDescription = nullptr;
 #endif // #if WITH_EDITOR
-	}
+}
 
 #if WITH_EDITOR
 void FStaticMeshSourceModel::SerializeBulkData(FArchive& Ar, UObject* Owner)
@@ -2233,36 +2235,27 @@ int32 UStaticMesh::GetMeshDescriptionCount() const
 
 UMeshDescription* UStaticMesh::GetOriginalMeshDescription(int32 LodIndex) const
 {
-	if (!OriginalMeshDescriptions.IsValidIndex(LodIndex))
+	if (SourceModels.IsValidIndex(LodIndex))
 	{
-		return nullptr;
+		return SourceModels[LodIndex].OriginalMeshDescription;
 	}
-	return OriginalMeshDescriptions[LodIndex];
+	return nullptr;
 }
 
 void UStaticMesh::SetOriginalMeshDescription(int32 LodIndex, class UMeshDescription* MeshDescription)
 {
-	if (!OriginalMeshDescriptions.IsValidIndex(LodIndex))
+	if (SourceModels.IsValidIndex(LodIndex))
 	{
-		//Add nullptr missing entries
-		OriginalMeshDescriptions.AddZeroed(LodIndex - OriginalMeshDescriptions.Num() + 1);
+		SourceModels[LodIndex].OriginalMeshDescription = MeshDescription;
 	}
-	//Set the original mesh description
-	OriginalMeshDescriptions[LodIndex] = MeshDescription;
-}
-
-int32 UStaticMesh::GetOriginalMeshDescriptionCount() const
-{
-	return OriginalMeshDescriptions.Num();
 }
 
 void UStaticMesh::ClearOriginalMeshDescription(int32 LodIndex)
 {
-	if (!OriginalMeshDescriptions.IsValidIndex(LodIndex))
+	if (SourceModels.IsValidIndex(LodIndex))
 	{
-		return;
+		SourceModels[LodIndex].OriginalMeshDescription = nullptr;
 	}
-	OriginalMeshDescriptions[LodIndex] = nullptr;
 }
 
 #endif
diff --git a/Engine/Source/Runtime/MeshBuilder/Private/LayoutUV.h b/Engine/Source/Runtime/MeshBuilder/Private/LayoutUV.h
--- a/Engine/Source/Runtime/MeshBuilder/Private/LayoutUV.h
+++ b/Engine/Source/Runtime/MeshBuilder/Private/LayoutUV.h
@@ -169,7 +169,7 @@ inline float FLayoutUV::TriangleUVArea( uint32 Tri ) const
 	FVector2D UVs[3];
 	for( int k = 0; k < 3; k++ )
 	{
-		FVertexInstanceID VertexInstanceID(Tri + k);
+		FVertexInstanceID VertexInstanceID((3*Tri) + k);
 		check(MeshDescription->VertexInstances().IsValid(VertexInstanceID));
 		const FMeshVertexInstance& VertexInstance = MeshDescription->GetVertexInstance(VertexInstanceID);
 		check(VertexInstance.VertexUVs.IsValidIndex(SrcChannel));
diff --git a/Engine/Source/Runtime/MeshBuilder/Private/MeshDescriptionHelper.cpp b/Engine/Source/Runtime/MeshBuilder/Private/MeshDescriptionHelper.cpp
--- a/Engine/Source/Runtime/MeshBuilder/Private/MeshDescriptionHelper.cpp
+++ b/Engine/Source/Runtime/MeshBuilder/Private/MeshDescriptionHelper.cpp
@@ -540,43 +540,21 @@ void FMeshDescriptionHelper::CreateMikktTangents(UMeshDescription* MeshDescripti
 	bool bIgnoreDegenerateTriangles = (TangentOptions & FMeshDescriptionHelper::ETangentOptions::IgnoreDegenerateTriangles) != 0;
 	float ComparisonThreshold = bIgnoreDegenerateTriangles ? THRESH_POINTS_ARE_SAME : 0.0f;
 
-	bool bWedgeTSpace = true;
-
-	for (const FVertexInstanceID& VertexInstanceID : MeshDescription->VertexInstances().GetElementIDs())
-	{
-		FMeshVertexInstance& VertexInstance = MeshDescription->GetVertexInstance(VertexInstanceID);
-		FVector BiNormal(0.0f);
-		if (!VertexInstance.Normal.IsNearlyZero() && !VertexInstance.Tangent.IsNearlyZero())
-		{
-			BiNormal = (FVector::CrossProduct(VertexInstance.Normal, VertexInstance.Tangent).GetSafeNormal() * VertexInstance.BinormalSign).GetSafeNormal();
-		}
-
-		if (VertexInstance.Tangent.IsNearlyZero() || BiNormal.IsNearlyZero())
-		{
-			bWedgeTSpace = false;
-			break;
-		}
-	}
-
-	//Set the tangent space
-	if (!bWedgeTSpace)
-	{
-		// we can use mikktspace to calculate the tangents
-		SMikkTSpaceInterface MikkTInterface;
-		MikkTInterface.m_getNormal = MeshDescriptionMikktSpaceInterface::MikkGetNormal;
-		MikkTInterface.m_getNumFaces = MeshDescriptionMikktSpaceInterface::MikkGetNumFaces;
-		MikkTInterface.m_getNumVerticesOfFace = MeshDescriptionMikktSpaceInterface::MikkGetNumVertsOfFace;
-		MikkTInterface.m_getPosition = MeshDescriptionMikktSpaceInterface::MikkGetPosition;
-		MikkTInterface.m_getTexCoord = MeshDescriptionMikktSpaceInterface::MikkGetTexCoord;
-		MikkTInterface.m_setTSpaceBasic = MeshDescriptionMikktSpaceInterface::MikkSetTSpaceBasic;
-		MikkTInterface.m_setTSpace = nullptr;
-
-		SMikkTSpaceContext MikkTContext;
-		MikkTContext.m_pInterface = &MikkTInterface;
-		MikkTContext.m_pUserData = (void*)(MeshDescription);
-		MikkTContext.m_bIgnoreDegenerates = bIgnoreDegenerateTriangles;
-		genTangSpaceDefault(&MikkTContext);
-	}
+	// we can use mikktspace to calculate the tangents
+	SMikkTSpaceInterface MikkTInterface;
+	MikkTInterface.m_getNormal = MeshDescriptionMikktSpaceInterface::MikkGetNormal;
+	MikkTInterface.m_getNumFaces = MeshDescriptionMikktSpaceInterface::MikkGetNumFaces;
+	MikkTInterface.m_getNumVerticesOfFace = MeshDescriptionMikktSpaceInterface::MikkGetNumVertsOfFace;
+	MikkTInterface.m_getPosition = MeshDescriptionMikktSpaceInterface::MikkGetPosition;
+	MikkTInterface.m_getTexCoord = MeshDescriptionMikktSpaceInterface::MikkGetTexCoord;
+	MikkTInterface.m_setTSpaceBasic = MeshDescriptionMikktSpaceInterface::MikkSetTSpaceBasic;
+	MikkTInterface.m_setTSpace = nullptr;
+
+	SMikkTSpaceContext MikkTContext;
+	MikkTContext.m_pInterface = &MikkTInterface;
+	MikkTContext.m_pUserData = (void*)(MeshDescription);
+	MikkTContext.m_bIgnoreDegenerates = bIgnoreDegenerateTriangles;
+	genTangSpaceDefault(&MikkTContext);
 }
 
 namespace MeshDescriptionMikktSpaceInterface
diff --git a/Engine/Source/Runtime/MeshBuilder/Private/StaticMeshBuilder.cpp b/Engine/Source/Runtime/MeshBuilder/Private/StaticMeshBuilder.cpp
--- a/Engine/Source/Runtime/MeshBuilder/Private/StaticMeshBuilder.cpp
+++ b/Engine/Source/Runtime/MeshBuilder/Private/StaticMeshBuilder.cpp
@@ -19,6 +19,7 @@ void UpdateBounds(class UStaticMesh* StaticMesh);
 void UpdateCollision(UStaticMesh *StaticMesh);
 void BuildVertexBuffer(
 	  class UStaticMesh *StaticMesh
+	, int32 LodIndex
 	, const class UMeshDescription* MeshDescription
 	, struct FStaticMeshLODResources& StaticMeshLOD
 	, const struct FMeshBuildSettings& LODBuildSettings
@@ -39,18 +40,17 @@ FStaticMeshBuilder::FStaticMeshBuilder()
 
 bool FStaticMeshBuilder::Build(UStaticMesh* StaticMesh)
 {
-	int32 MeshDescriptionCount = StaticMesh->GetOriginalMeshDescriptionCount();
-	if (MeshDescriptionCount <= 0)
+	if (StaticMesh->GetOriginalMeshDescription(0) == nullptr)
 	{
 		//TODO: Warn the user that there is no mesh description data
 		return false;
 	}
-	StaticMesh->RenderData->AllocateLODResources(MeshDescriptionCount);
+	StaticMesh->RenderData->AllocateLODResources(StaticMesh->SourceModels.Num());
 
 	OnBuildRenderMeshStart(StaticMesh, false);
 
 	FStaticMeshRenderData& StaticMeshRenderData = *StaticMesh->RenderData;
-	for (int32 LodIndex = 0; LodIndex < MeshDescriptionCount; ++LodIndex)
+	for (int32 LodIndex = 0; LodIndex < StaticMesh->SourceModels.Num(); ++LodIndex)
 	{
 		FMeshBuildSettings& LODBuildSettings = StaticMesh->SourceModels[LodIndex].BuildSettings;
 		const UMeshDescription* OriginalMeshDescription = StaticMesh->GetOriginalMeshDescription(LodIndex);
@@ -93,7 +93,7 @@ bool FStaticMeshBuilder::Build(UStaticMesh* StaticMesh)
 		}
 
 		//Build the vertex and index buffer
-		BuildVertexBuffer(StaticMesh, MeshDescription, StaticMeshLOD, LODBuildSettings, IndexBuffer, WedgeMap, PerSectionIndices, StaticMeshBuildVertices, MeshDescriptionHelper.GetOverlappingCorners(), VertexComparisonThreshold, RemapVerts);
+		BuildVertexBuffer(StaticMesh, LodIndex, MeshDescription, StaticMeshLOD, LODBuildSettings, IndexBuffer, WedgeMap, PerSectionIndices, StaticMeshBuildVertices, MeshDescriptionHelper.GetOverlappingCorners(), VertexComparisonThreshold, RemapVerts);
 
 		// Figure out which index buffer stride we need
 		bool bNeeds32BitIndices = false;
@@ -316,6 +316,7 @@ bool AreVerticesEqual(FStaticMeshBuildVertex const& A, FStaticMeshBuildVertex co
 
 void BuildVertexBuffer(
 	  UStaticMesh *StaticMesh
+	, int32 LodIndex
 	, const UMeshDescription* MeshDescription
 	, FStaticMeshLODResources& StaticMeshLOD
 	, const FMeshBuildSettings& LODBuildSettings
@@ -358,7 +359,7 @@ void BuildVertexBuffer(
 		TArray<uint32>& SectionIndices = OutPerSectionIndices[PolygonGroupID.GetValue()];
 
 		// Create new rendering section
-		StaticMeshLOD.Sections.Add(FStaticMeshSection());
+		int32 SectionIndex = StaticMeshLOD.Sections.Add(FStaticMeshSection());
 		FStaticMeshSection& StaticMeshSection = StaticMeshLOD.Sections.Last();
 
 		StaticMeshSection.FirstIndex = IndexBuffer.Num();
@@ -370,6 +371,17 @@ void BuildVertexBuffer(
 		StaticMeshSection.MaterialIndex = MaterialIndex;
 		StaticMeshSection.bEnableCollision = PolygonGroup.bEnableCollision;
 		StaticMeshSection.bCastShadow = PolygonGroup.bCastShadow;
+		
+		if (LodIndex > 0)
+		{
+			//Set the overwrite section info map
+			FMeshSectionInfo SectionInfo = StaticMesh->SectionInfoMap.Get(LodIndex, SectionIndex);
+			SectionInfo.bCastShadow = StaticMeshSection.bCastShadow;
+			SectionInfo.bEnableCollision = StaticMeshSection.bEnableCollision;
+			SectionInfo.MaterialIndex = StaticMeshSection.MaterialIndex;
+			StaticMesh->SectionInfoMap.Set(LodIndex, SectionIndex, SectionInfo);
+		}
+
 		for (FPolygonID& PolygonID : Polygons)
 		{
 			const FMeshPolygon& Polygon = MeshDescription->GetPolygon(PolygonID);
diff --git a/Engine/Source/Runtime/MeshBuilder/Private/ThirdPartyBuildOptimizationHelper.h b/Engine/Source/Runtime/MeshBuilder/Private/ThirdPartyBuildOptimizationHelper.h
--- a/Engine/Source/Runtime/MeshBuilder/Private/ThirdPartyBuildOptimizationHelper.h
+++ b/Engine/Source/Runtime/MeshBuilder/Private/ThirdPartyBuildOptimizationHelper.h
@@ -13,7 +13,7 @@ namespace BuildOptimizationPrivate
 {
 	// CVars
 	static TAutoConsoleVariable<int32> CVarTriangleOrderOptimization(
-		TEXT("r.TriangleOrderOptimization"),
+		TEXT("r.TriangleOrderOptimization_MD"),
 		1,
 		TEXT("Controls the algorithm to use when optimizing the triangle order for the post-transform cache.\n")
 		TEXT("0: Use NVTriStrip (slower)\n")
@@ -22,13 +22,13 @@ namespace BuildOptimizationPrivate
 		ECVF_Default);
 
 	static TAutoConsoleVariable<int32> CVarSupportDepthOnlyIndexBuffers(
-		TEXT("r.SupportDepthOnlyIndexBuffers"),
+		TEXT("r.SupportDepthOnlyIndexBuffers_MD"),
 		1,
 		TEXT("Enables depth-only index buffers. Saves a little time at the expense of doubling the size of index buffers."),
 		ECVF_ReadOnly | ECVF_RenderThreadSafe);
 
 	static TAutoConsoleVariable<int32> CVarSupportReversedIndexBuffers(
-		TEXT("r.SupportReversedIndexBuffers"),
+		TEXT("r.SupportReversedIndexBuffers_MD"),
 		1,
 		TEXT("Enables reversed index buffers. Saves a little time at the expense of doubling the size of index buffers."),
 		ECVF_ReadOnly | ECVF_RenderThreadSafe);
