Port Geom util sphere and sphyl to mesh description

diff --git a/Engine/Source/Editor/UnrealEd/Private/GeomFitUtils.cpp b/Engine/Source/Editor/UnrealEd/Private/GeomFitUtils.cpp
--- a/Engine/Source/Editor/UnrealEd/Private/GeomFitUtils.cpp
+++ b/Engine/Source/Editor/UnrealEd/Private/GeomFitUtils.cpp
@@ -409,6 +409,142 @@ static void CalcBoundingSphere2(const FRawMesh& RawMesh, FSphere& sphere, FVecto
 	sphere.W = FMath::Sqrt(sphere.W);
 }
 
+static void CalcBoundingSphere(const UMeshDescription* MeshDescription, FSphere& sphere, FVector& LimitVec)
+{
+	if (MeshDescription->Vertices().Num() == 0)
+		return;
+
+	FBox Box;
+	FVector MinIx[3];
+	FVector MaxIx[3];
+
+	bool bFirstVertex = true;
+	for (const FVertexID& VertexID : MeshDescription->Vertices().GetElementIDs())
+	{
+		const FMeshVertex& Vertex = MeshDescription->GetVertex(VertexID);
+		FVector p = Vertex.VertexPosition * LimitVec;
+		if (bFirstVertex)
+		{
+			// First, find AABB, remembering furthest points in each dir.
+			Box.Min = p;
+			Box.Max = Box.Min;
+
+			MinIx[0] = Vertex.VertexPosition;
+			MinIx[1] = Vertex.VertexPosition;
+			MinIx[2] = Vertex.VertexPosition;
+
+			MaxIx[0] = Vertex.VertexPosition;
+			MaxIx[1] = Vertex.VertexPosition;
+			MaxIx[2] = Vertex.VertexPosition;
+			bFirstVertex = false;
+			continue;
+		}
+
+		// X //
+		if (p.X < Box.Min.X)
+		{
+			Box.Min.X = p.X;
+			MinIx[0] = Vertex.VertexPosition;
+		}
+		else if (p.X > Box.Max.X)
+		{
+			Box.Max.X = p.X;
+			MaxIx[0] = Vertex.VertexPosition;
+		}
+
+		// Y //
+		if (p.Y < Box.Min.Y)
+		{
+			Box.Min.Y = p.Y;
+			MinIx[1] = Vertex.VertexPosition;
+		}
+		else if (p.Y > Box.Max.Y)
+		{
+			Box.Max.Y = p.Y;
+			MaxIx[1] = Vertex.VertexPosition;
+		}
+
+		// Z //
+		if (p.Z < Box.Min.Z)
+		{
+			Box.Min.Z = p.Z;
+			MinIx[2] = Vertex.VertexPosition;
+		}
+		else if (p.Z > Box.Max.Z)
+		{
+			Box.Max.Z = p.Z;
+			MaxIx[2] = Vertex.VertexPosition;
+		}
+	}
+
+	const FVector Extremes[3] = { (MaxIx[0] - MinIx[0]) * LimitVec,
+		(MaxIx[1] - MinIx[1]) * LimitVec,
+		(MaxIx[2] - MinIx[2]) * LimitVec };
+
+	// Now find extreme points furthest apart, and initial center and radius of sphere.
+	float d2 = 0.f;
+	for (int32 i = 0; i < 3; i++)
+	{
+		const float tmpd2 = Extremes[i].SizeSquared();
+		if (tmpd2 > d2)
+		{
+			d2 = tmpd2;
+			sphere.Center = (MinIx[i] + (0.5f * Extremes[i])) * LimitVec;
+			sphere.W = 0.f;
+		}
+	}
+
+	const FVector Extents = FVector(Extremes[0].X, Extremes[1].Y, Extremes[2].Z);
+
+	// radius and radius squared
+	float r = 0.5f * Extents.GetMax();
+	float r2 = FMath::Square(r);
+
+	// Now check each point lies within this sphere. If not - expand it a bit.
+	for (const FVertexID& VertexID : MeshDescription->Vertices().GetElementIDs())
+	{
+		const FMeshVertex& Vertex = MeshDescription->GetVertex(VertexID);
+
+		const FVector cToP = (Vertex.VertexPosition * LimitVec) - sphere.Center;
+
+		const float pr2 = cToP.SizeSquared();
+
+		// If this point is outside our current bounding sphere's radius
+		if (pr2 > r2)
+		{
+			// ..expand radius just enough to include this point.
+			const float pr = FMath::Sqrt(pr2);
+			r = 0.5f * (r + pr);
+			r2 = FMath::Square(r);
+
+			sphere.Center += ((pr - r) / pr * cToP);
+		}
+	}
+
+	sphere.W = r;
+}
+
+// This is the one thats already used by unreal.
+// Seems to do better with more symmetric input...
+static void CalcBoundingSphere2(const UMeshDescription* MeshDescription, FSphere& sphere, FVector& LimitVec)
+{
+	FVector Center, Extents;
+	CalcBoundingBox(MeshDescription, Center, Extents, LimitVec);
+
+	sphere.Center = Center;
+	sphere.W = 0.0f;
+
+	for (const FVertexID& VertexID : MeshDescription->Vertices().GetElementIDs())
+	{
+		const FMeshVertex& Vertex = MeshDescription->GetVertex(VertexID);
+
+		float Dist = FVector::DistSquared(Vertex.VertexPosition * LimitVec, sphere.Center);
+		if (Dist > sphere.W)
+			sphere.W = Dist;
+	}
+	sphere.W = FMath::Sqrt(sphere.W);
+}
+
 // // //
 
 int32 GenerateSphereAsSimpleCollision(UStaticMesh* StaticMesh)
@@ -419,16 +555,24 @@ int32 GenerateSphereAsSimpleCollision(UStaticMesh* StaticMesh)
 	}
 
 	UBodySetup* bs = StaticMesh->BodySetup;
-
-	// Calculate bounding sphere.
-	FRawMesh RawMesh;
-	FStaticMeshSourceModel& SrcModel = StaticMesh->SourceModels[0];
-	SrcModel.RawMeshBulkData->LoadRawMesh(RawMesh);
-
 	FSphere bSphere, bSphere2, bestSphere;
 	FVector unitVec = bs->BuildScale3D;
-	CalcBoundingSphere(RawMesh, bSphere, unitVec);
-	CalcBoundingSphere2(RawMesh, bSphere2, unitVec);
+
+	// Calculate bounding sphere.
+	if (StaticMesh->GetOriginalMeshDescription(0) != nullptr)
+	{
+		UMeshDescription* MeshDescription = StaticMesh->GetMeshDescription(0);
+		CalcBoundingSphere(MeshDescription, bSphere, unitVec);
+		CalcBoundingSphere2(MeshDescription, bSphere2, unitVec);
+	}
+	else
+	{
+		FRawMesh RawMesh;
+		FStaticMeshSourceModel& SrcModel = StaticMesh->SourceModels[0];
+		SrcModel.RawMeshBulkData->LoadRawMesh(RawMesh);
+		CalcBoundingSphere(RawMesh, bSphere, unitVec);
+		CalcBoundingSphere2(RawMesh, bSphere2, unitVec);
+	}
 
 	if(bSphere.W < bSphere2.W)
 		bestSphere = bSphere;
@@ -464,6 +608,92 @@ int32 GenerateSphereAsSimpleCollision(UStaticMesh* StaticMesh)
 
 /* ******************************** SPHYL ******************************** */
 
+static void CalcBoundingSphyl(const UMeshDescription* MeshDescription, FSphere& sphere, float& length, FRotator& rotation, FVector& LimitVec)
+{
+	if (MeshDescription->Vertices().Num() == 0)
+		return;
+
+	FVector Center, Extents;
+	CalcBoundingBox(MeshDescription, Center, Extents, LimitVec);
+
+	// @todo sphere.Center could perhaps be adjusted to best fit if model is non-symmetric on it's longest axis
+	sphere.Center = Center;
+
+	// Work out best axis aligned orientation (longest side)
+	float Extent = Extents.GetMax();
+	if (Extent == Extents.X)
+	{
+		rotation = FRotator(90.f, 0.f, 0.f);
+		Extents.X = 0.0f;
+	}
+	else if (Extent == Extents.Y)
+	{
+		rotation = FRotator(0.f, 0.f, 90.f);
+		Extents.Y = 0.0f;
+	}
+	else
+	{
+		rotation = FRotator(0.f, 0.f, 0.f);
+		Extents.Z = 0.0f;
+	}
+
+	// Cleared the largest axis above, remaining determines the radius
+	float r = Extents.GetMax();
+	float r2 = FMath::Square(r);
+	
+	// Now check each point lies within this the radius. If not - expand it a bit.
+	for (const FVertexID& VertexID : MeshDescription->Vertices().GetElementIDs())
+	{
+		const FMeshVertex& Vertex = MeshDescription->GetVertex(VertexID);
+		FVector cToP = (Vertex.VertexPosition * LimitVec) - sphere.Center;
+		cToP = rotation.UnrotateVector(cToP);
+
+		const float pr2 = cToP.SizeSquared2D();	// Ignore Z here...
+
+		// If this point is outside our current bounding sphere's radius
+		if (pr2 > r2)
+		{
+			// ..expand radius just enough to include this point.
+			const float pr = FMath::Sqrt(pr2);
+			r = 0.5f * (r + pr);
+			r2 = FMath::Square(r);
+		}
+	}
+	
+	// The length is the longest side minus the radius.
+	float hl = FMath::Max(0.0f, Extent - r);
+
+	// Now check each point lies within the length. If not - expand it a bit.
+	for (const FVertexID& VertexID : MeshDescription->Vertices().GetElementIDs())
+	{
+		const FMeshVertex& Vertex = MeshDescription->GetVertex(VertexID);
+		FVector cToP = (Vertex.VertexPosition * LimitVec) - sphere.Center;
+		cToP = rotation.UnrotateVector(cToP);
+
+		// If this point is outside our current bounding sphyl's length
+		if (FMath::Abs(cToP.Z) > hl)
+		{
+			const bool bFlip = (cToP.Z < 0.f ? true : false);
+			const FVector cOrigin(0.f, 0.f, (bFlip ? -hl : hl));
+
+			const float pr2 = (cOrigin - cToP).SizeSquared();
+
+			// If this point is outside our current bounding sphyl's radius
+			if (pr2 > r2)
+			{
+				FVector cPoint;
+				FMath::SphereDistToLine(cOrigin, r, cToP, (bFlip ? FVector(0.f, 0.f, 1.f) : FVector(0.f, 0.f, -1.f)), cPoint);
+
+				// Don't accept zero as a valid diff when we know it's outside the sphere (saves needless retest on further iterations of like points)
+				hl += FMath::Max(FMath::Abs(cToP.Z - cPoint.Z), 1.e-6f);
+			}
+		}
+	}
+
+	sphere.W = r;
+	length = hl * 2.0f;
+}
+
 static void CalcBoundingSphyl(const FRawMesh& RawMesh, FSphere& sphere, float& length, FRotator& rotation, FVector& LimitVec)
 {
 	if (RawMesh.VertexPositions.Num() == 0)
@@ -559,16 +789,24 @@ int32 GenerateSphylAsSimpleCollision(UStaticMesh* StaticMesh)
 
 	UBodySetup* bs = StaticMesh->BodySetup;
 
-	// Calculate bounding box.
-	FRawMesh RawMesh;
-	FStaticMeshSourceModel& SrcModel = StaticMesh->SourceModels[0];
-	SrcModel.RawMeshBulkData->LoadRawMesh(RawMesh);
-
 	FSphere sphere;
 	float length;
 	FRotator rotation;
 	FVector unitVec = bs->BuildScale3D;
-	CalcBoundingSphyl(RawMesh, sphere, length, rotation, unitVec);
+
+	// Calculate bounding box.
+	if (StaticMesh->GetOriginalMeshDescription(0) != nullptr)
+	{
+		UMeshDescription* MeshDescription = StaticMesh->GetMeshDescription(0);
+		CalcBoundingSphyl(MeshDescription, sphere, length, rotation, unitVec);
+	}
+	else
+	{
+		FRawMesh RawMesh;
+		FStaticMeshSourceModel& SrcModel = StaticMesh->SourceModels[0];
+		SrcModel.RawMeshBulkData->LoadRawMesh(RawMesh);
+		CalcBoundingSphyl(RawMesh, sphere, length, rotation, unitVec);
+	}
 
 	// Dont use if radius is zero.
 	if (sphere.W <= 0.f)
