Fix the import with "import material" disabled

diff --git a/Engine/Source/Editor/UnrealEd/Private/Fbx/FbxStaticMeshImport.cpp b/Engine/Source/Editor/UnrealEd/Private/Fbx/FbxStaticMeshImport.cpp
--- a/Engine/Source/Editor/UnrealEd/Private/Fbx/FbxStaticMeshImport.cpp
+++ b/Engine/Source/Editor/UnrealEd/Private/Fbx/FbxStaticMeshImport.cpp
@@ -753,10 +753,11 @@ bool UnFbx::FFbxImporter::BuildStaticMeshFromGeometry(FbxNode* Node, UStaticMesh
 			if (!PolygonGroupMapping.Contains(RealMaterialIndex))
 			{
 				UMaterialInterface* Material = MeshMaterials.IsValidIndex(RealMaterialIndex) ? MeshMaterials[RealMaterialIndex].Material : UMaterial::GetDefaultMaterial(MD_Surface);
+				FName ImportedMaterialSlotName = MeshMaterials.IsValidIndex(RealMaterialIndex) ? FName(*MeshMaterials[RealMaterialIndex].GetName()) : NAME_None;
 				FPolygonGroupID ExistingPolygonGroup = FPolygonGroupID::Invalid;
 				for (const FPolygonGroupID PolygonGroupID : MeshDescription->PolygonGroups().GetElementIDs())
 				{
-					if (MeshDescription->PolygonGroups()[PolygonGroupID].MaterialAsset.ResolveObject() == Material)
+					if (MeshDescription->PolygonGroups()[PolygonGroupID].ImportedMaterialSlotName == ImportedMaterialSlotName)
 					{
 						ExistingPolygonGroup = PolygonGroupID;
 						break;
@@ -768,8 +769,8 @@ bool UnFbx::FFbxImporter::BuildStaticMeshFromGeometry(FbxNode* Node, UStaticMesh
 					FMeshPolygonGroup& NewPolygonGroup = MeshDescription->GetPolygonGroup(ExistingPolygonGroup);
 					NewPolygonGroup.bCastShadow = true;
 					NewPolygonGroup.bEnableCollision = bEnableCollision;
-					NewPolygonGroup.ImportedMaterialSlotName = Material->GetFName();
-					NewPolygonGroup.MaterialSlotName = Material->GetFName();
+					NewPolygonGroup.ImportedMaterialSlotName = ImportedMaterialSlotName;
+					NewPolygonGroup.MaterialSlotName = ImportedMaterialSlotName;
 					NewPolygonGroup.MaterialAsset = FStringAssetReference(Material);
 				}
 				PolygonGroupMapping.Add(RealMaterialIndex, ExistingPolygonGroup);
