Add an experimental switch to use MeshDescription instead of RawMesh

diff --git a/Engine/Config/BaseEditorPerProjectUserSettings.ini b/Engine/Config/BaseEditorPerProjectUserSettings.ini
--- a/Engine/Config/BaseEditorPerProjectUserSettings.ini
+++ b/Engine/Config/BaseEditorPerProjectUserSettings.ini
@@ -140,6 +140,7 @@ bBreakOnExceptions=False
 bToolbarCustomization=False
 bBehaviorTreeEditor=False
 bEnableFindAndReplaceReferences =False
+bUseMeshDescription=False
 
 [/Script/UnrealEd.EditorLoadingSavingSettings]
 ; True if we should automatically load a default level at start up
@@ -575,7 +576,6 @@ BackupIntervalInMinutes=5
 
 
 [/Script/UnrealEd.FbxImportUI]
-bImportMeshDescription=False
 bImportMaterials=True
 bAutoCreateGroups=True
 bInvertNormalMaps=False
diff --git a/Engine/Source/Editor/UnrealEd/Classes/Factories/FbxImportUI.h b/Engine/Source/Editor/UnrealEd/Classes/Factories/FbxImportUI.h
--- a/Engine/Source/Editor/UnrealEd/Classes/Factories/FbxImportUI.h
+++ b/Engine/Source/Editor/UnrealEd/Classes/Factories/FbxImportUI.h
@@ -50,10 +50,6 @@ public:
 	UPROPERTY(EditAnywhere, AdvancedDisplay, config, Category=Miscellaneous, meta=(OBJRestrict="true"))
 	uint32 bOverrideFullName:1;
 
-	/** Use the string in "Name" field as full name of mesh. The option only works when the scene contains one mesh. */
-	UPROPERTY(EditAnywhere, AdvancedDisplay, config, Category = Miscellaneous, meta = (OBJRestrict = "true"))
-	uint32 bImportMeshDescription:1;
-
 	/** Whether to import the incoming FBX as a skeletal object */
 	UPROPERTY(EditAnywhere, Category = Mesh, meta = (ImportType = "StaticMesh|SkeletalMesh", DisplayName="Skeletal Mesh"))
 	bool bImportAsSkeletal;
diff --git a/Engine/Source/Editor/UnrealEd/Classes/Settings/EditorExperimentalSettings.h b/Engine/Source/Editor/UnrealEd/Classes/Settings/EditorExperimentalSettings.h
--- a/Engine/Source/Editor/UnrealEd/Classes/Settings/EditorExperimentalSettings.h
+++ b/Engine/Source/Editor/UnrealEd/Classes/Settings/EditorExperimentalSettings.h
@@ -124,6 +124,10 @@ public:
 	/** Enables the ability to bake materials inside of the Static Mesh and Skeletal Mesh Editor, and for static mesh instances. */
 	UPROPERTY(EditAnywhere, config, Category = Tools, meta = (ConfigRestartRequired = true))
 	bool bAssetMaterialBaking;
+
+	/** Allows usage of the mesh description instead of the raw mesh */
+	UPROPERTY(EditAnywhere, config, Category = Tools, meta = (DisplayName = "Use Mesh Description"))
+	bool bUseMeshDescription;
 	
 	/**
 	 * Returns an event delegate that is executed when a setting has changed.
diff --git a/Engine/Source/Editor/UnrealEd/Private/Fbx/FbxMainImport.cpp b/Engine/Source/Editor/UnrealEd/Private/Fbx/FbxMainImport.cpp
--- a/Engine/Source/Editor/UnrealEd/Private/Fbx/FbxMainImport.cpp
+++ b/Engine/Source/Editor/UnrealEd/Private/Fbx/FbxMainImport.cpp
@@ -247,7 +247,6 @@ void ApplyImportUIToImportOptions(UFbxImportUI* ImportUI, FBXImportOptions& InOu
 		InOutImportOptions.BaseEmissiveColorName = ImportUI->TextureImportData->BaseEmissiveColorName;
 	}
 	InOutImportOptions.bImportTextures = ImportUI->bImportTextures;
-	InOutImportOptions.bImportMeshDescription = ImportUI->bImportMeshDescription;
 	InOutImportOptions.bUsedAsFullName = ImportUI->bOverrideFullName;
 	InOutImportOptions.bImportAnimations = ImportUI->bImportAnimations;
 	InOutImportOptions.SkeletonForAnimation = ImportUI->Skeleton;
diff --git a/Engine/Source/Editor/UnrealEd/Private/Fbx/FbxStaticMeshImport.cpp b/Engine/Source/Editor/UnrealEd/Private/Fbx/FbxStaticMeshImport.cpp
--- a/Engine/Source/Editor/UnrealEd/Private/Fbx/FbxStaticMeshImport.cpp
+++ b/Engine/Source/Editor/UnrealEd/Private/Fbx/FbxStaticMeshImport.cpp
@@ -35,6 +35,7 @@
 #include "PhysicsEngine/BodySetup.h"
 #include "MeshDescription.h"
 #include "IMeshBuilderModule.h"
+#include "Settings/EditorExperimentalSettings.h"
 
 #define LOCTEXT_NAMESPACE "FbxStaticMeshImport"
 
@@ -266,7 +267,7 @@ bool UnFbx::FFbxImporter::BuildStaticMeshFromGeometry(FbxNode* Node, UStaticMesh
 	FStaticMeshSourceModel& SrcModel = StaticMesh->SourceModels[LODIndex];
 	
 	UMeshDescription *MeshDescription = StaticMesh->GetOriginalMeshDescription(LODIndex);
-	if (ImportOptions->bImportMeshDescription)
+	if (GetDefault<UEditorExperimentalSettings>()->bUseMeshDescription)
 	{
 		//The mesh description should have been created before calling BuildStaticMeshFromGeometry
 		check(MeshDescription);
@@ -1572,7 +1573,7 @@ UStaticMesh* UnFbx::FFbxImporter::ImportStaticMeshAsSingle(UObject* InParent, TA
 		}
 	}
 
-	if (ImportOptions->bImportMeshDescription)
+	if (GetDefault<UEditorExperimentalSettings>()->bUseMeshDescription)
 	{
 		MeshDescription = StaticMesh->GetOriginalMeshDescription(LODIndex);
 		if (MeshDescription == nullptr)
diff --git a/Engine/Source/Editor/UnrealEd/Private/Tests/FbxAutomationCommon.cpp b/Engine/Source/Editor/UnrealEd/Private/Tests/FbxAutomationCommon.cpp
--- a/Engine/Source/Editor/UnrealEd/Private/Tests/FbxAutomationCommon.cpp
+++ b/Engine/Source/Editor/UnrealEd/Private/Tests/FbxAutomationCommon.cpp
@@ -11,7 +11,7 @@
 #include "Factories/FbxStaticMeshImportData.h"
 #include "Factories/FbxTextureImportData.h"
 #include "Factories/FbxImportUI.h"
-
+#include "Settings/EditorExperimentalSettings.h"
 #include "JsonObjectConverter.h"
 
 
@@ -78,8 +78,6 @@ namespace FbxAutomationTestsAPI
 			FbxTestPlan->ImportUI->AnimSequenceImportData = SavedAnimSequenceData;
 			FbxTestPlan->ImportUI->TextureImportData = SavedTextureData;
 
-			//Force usage of mesh descriptions
-			FbxTestPlan->ImportUI->bImportMeshDescription = true;
 			TestPlanArray.Add(FbxTestPlan);
 		}
 	}
diff --git a/Engine/Source/Editor/UnrealEd/Public/FbxImporter.h b/Engine/Source/Editor/UnrealEd/Public/FbxImporter.h
--- a/Engine/Source/Editor/UnrealEd/Public/FbxImporter.h
+++ b/Engine/Source/Editor/UnrealEd/Public/FbxImporter.h
@@ -107,7 +107,6 @@ namespace UnFbx
 struct FBXImportOptions
 {
 	// General options
-	bool bImportMeshDescription;
 	bool bImportScene;
 	bool bImportMaterials;
 	bool bResetMaterialSlots;
