PolygongroupAttribute use FSoftObjectPath instead of UObject pointer
UMeshDescription::GetIdString is now using FObjectWriter archive to support UObject
FMeshDescriptionHelper::GetRenderMeshDescription use StaticDuplicateObject to create the render copy
Fbx automation tests now use the mesh description if the editor preference experimental feature "Use Mesh Description" is on

diff --git a/Engine/Plugins/Editor/MeshEditor/Source/MeshEditor/MeshEditorMode.cpp b/Engine/Plugins/Editor/MeshEditor/Source/MeshEditor/MeshEditorMode.cpp
--- a/Engine/Plugins/Editor/MeshEditor/Source/MeshEditor/MeshEditorMode.cpp
+++ b/Engine/Plugins/Editor/MeshEditor/Source/MeshEditor/MeshEditorMode.cpp
@@ -2301,11 +2301,11 @@ bool FMeshEditorMode::AssignMaterialToSelectedPolygons( UMaterialInterface* Sele
 				// further details if there is more than one polygon group which matches the material.
 				FPolygonGroupID PolygonGroupToAssign = FPolygonGroupID::Invalid;
 
-				const TPolygonGroupAttributeArray<UObject*>& PolygonGroupMaterialAssets = MeshDescription->PolygonGroupAttributes().GetAttributes<UObject*>( MeshAttribute::PolygonGroup::MaterialAsset );
+				const TPolygonGroupAttributeArray<FSoftObjectPath>& PolygonGroupMaterialAssets = MeshDescription->PolygonGroupAttributes().GetAttributes<FSoftObjectPath>( MeshAttribute::PolygonGroup::MaterialAsset );
 
 				for( const FPolygonGroupID PolygonGroupID : MeshDescription->PolygonGroups().GetElementIDs() )
 				{
-					UMaterialInterface* Material = Cast<UMaterialInterface>( PolygonGroupMaterialAssets[ PolygonGroupID ] );
+					UMaterialInterface* Material = Cast<UMaterialInterface>( PolygonGroupMaterialAssets[ PolygonGroupID ].TryLoad() );
 					if( Material == SelectedMaterial )
 					{
 						// We only expect to find one polygon group containing this material at the moment.
@@ -2342,7 +2342,7 @@ bool FMeshEditorMode::AssignMaterialToSelectedPolygons( UMaterialInterface* Sele
 					PolygonGroupsToCreate.Emplace();
 					
 					FPolygonGroupToCreate& PolygonGroupToCreate = PolygonGroupsToCreate.Last();
-					PolygonGroupToCreate.PolygonGroupAttributes.Attributes.Emplace( MeshAttribute::PolygonGroup::MaterialAsset, 0, FMeshElementAttributeValue( SelectedMaterial ) );
+					PolygonGroupToCreate.PolygonGroupAttributes.Attributes.Emplace( MeshAttribute::PolygonGroup::MaterialAsset, 0, FMeshElementAttributeValue( FSoftObjectPath(SelectedMaterial) ) );
 					PolygonGroupToCreate.PolygonGroupAttributes.Attributes.Emplace( MeshAttribute::PolygonGroup::MaterialSlotName, 0, FMeshElementAttributeValue( UniqueSlotName ) );
 					PolygonGroupToCreate.PolygonGroupAttributes.Attributes.Emplace( MeshAttribute::PolygonGroup::CastShadow, 0, FMeshElementAttributeValue( true ) );
 					PolygonGroupToCreate.PolygonGroupAttributes.Attributes.Emplace( MeshAttribute::PolygonGroup::EnableCollision, 0, FMeshElementAttributeValue( true ) );
diff --git a/Engine/Plugins/Runtime/EditableMesh/Source/EditableMesh/EditableStaticMeshAdapter.cpp b/Engine/Plugins/Runtime/EditableMesh/Source/EditableMesh/EditableStaticMeshAdapter.cpp
--- a/Engine/Plugins/Runtime/EditableMesh/Source/EditableMesh/EditableStaticMeshAdapter.cpp
+++ b/Engine/Plugins/Runtime/EditableMesh/Source/EditableMesh/EditableStaticMeshAdapter.cpp
@@ -104,7 +104,7 @@ void UEditableStaticMeshAdapter::InitEditableStaticMesh( UEditableMesh* Editable
 
 				TEdgeAttributeArray<bool>& EdgeHardnesses = MeshDescription->EdgeAttributes().GetAttributes<bool>( MeshAttribute::Edge::IsHard );
 
-				TPolygonGroupAttributeArray<UObject*>& PolygonGroupMaterialAssets = MeshDescription->PolygonGroupAttributes().GetAttributes<UObject*>( MeshAttribute::PolygonGroup::MaterialAsset );
+				TPolygonGroupAttributeArray<FSoftObjectPath>& PolygonGroupMaterialAssets = MeshDescription->PolygonGroupAttributes().GetAttributes<FSoftObjectPath>( MeshAttribute::PolygonGroup::MaterialAsset );
 				TPolygonGroupAttributeArray<FName>& PolygonGroupMaterialSlotNames = MeshDescription->PolygonGroupAttributes().GetAttributes<FName>( MeshAttribute::PolygonGroup::MaterialSlotName );
 				TPolygonGroupAttributeArray<FName>& PolygonGroupImportedMaterialSlotNames = MeshDescription->PolygonGroupAttributes().GetAttributes<FName>( MeshAttribute::PolygonGroup::ImportedMaterialSlotName );
 				TPolygonGroupAttributeArray<bool>& PolygonGroupCollision = MeshDescription->PolygonGroupAttributes().GetAttributes<bool>( MeshAttribute::PolygonGroup::EnableCollision );
@@ -279,7 +279,7 @@ void UEditableStaticMeshAdapter::InitEditableStaticMesh( UEditableMesh* Editable
 
 					// Create a new polygon group
 					const FPolygonGroupID NewPolygonGroupID = MeshDescription->CreatePolygonGroup();
-					PolygonGroupMaterialAssets[ NewPolygonGroupID ] = MaterialInterface;
+					PolygonGroupMaterialAssets[ NewPolygonGroupID ] = FSoftObjectPath(MaterialInterface);
 					PolygonGroupMaterialSlotNames[ NewPolygonGroupID ] = StaticMaterial.MaterialSlotName;
 					PolygonGroupImportedMaterialSlotNames[ NewPolygonGroupID ] = StaticMaterial.ImportedMaterialSlotName;
 					PolygonGroupCollision[ NewPolygonGroupID ] = RenderingSection.bEnableCollision;
@@ -1671,7 +1671,7 @@ void UEditableStaticMeshAdapter::OnCreatePolygonGroups( const UEditableMesh* Edi
 	const UMeshDescription* MeshDescription = EditableMesh->GetMeshDescription();
 	check( MeshDescription );
 
-	const TPolygonGroupAttributeArray<UObject*>& PolygonGroupMaterialAssets = MeshDescription->PolygonGroupAttributes().GetAttributes<UObject*>( MeshAttribute::PolygonGroup::MaterialAsset );
+	const TPolygonGroupAttributeArray<FSoftObjectPath>& PolygonGroupMaterialAssets = MeshDescription->PolygonGroupAttributes().GetAttributes<FSoftObjectPath>( MeshAttribute::PolygonGroup::MaterialAsset );
 	const TPolygonGroupAttributeArray<FName>& PolygonGroupMaterialSlotNames = MeshDescription->PolygonGroupAttributes().GetAttributes<FName>( MeshAttribute::PolygonGroup::MaterialSlotName );
 	const TPolygonGroupAttributeArray<FName>& PolygonGroupImportedMaterialSlotNames = MeshDescription->PolygonGroupAttributes().GetAttributes<FName>( MeshAttribute::PolygonGroup::ImportedMaterialSlotName );
 	const TPolygonGroupAttributeArray<bool>& PolygonGroupCollision = MeshDescription->PolygonGroupAttributes().GetAttributes<bool>( MeshAttribute::PolygonGroup::EnableCollision );
@@ -1716,7 +1716,7 @@ void UEditableStaticMeshAdapter::OnCreatePolygonGroups( const UEditableMesh* Edi
 			StaticMeshSection.bEnableCollision = PolygonGroupCollision[ PolygonGroupID ];
 			StaticMeshSection.bCastShadow = PolygonGroupCastShadow[ PolygonGroupID ];
 
-			UMaterialInterface* Material = Cast<UMaterialInterface>( PolygonGroupMaterialAssets[ PolygonGroupID ] );
+			UMaterialInterface* Material = Cast<UMaterialInterface>( PolygonGroupMaterialAssets[ PolygonGroupID ].TryLoad() );
 			StaticMeshSection.MaterialIndex = StaticMesh->StaticMaterials.Emplace( Material, PolygonGroupMaterialSlotNames[ PolygonGroupID ], PolygonGroupImportedMaterialSlotNames[ PolygonGroupID ] );
 		}
 
diff --git a/Engine/Plugins/Runtime/EditableMesh/Source/EditableMesh/Public/EditableMeshTypes.h b/Engine/Plugins/Runtime/EditableMesh/Source/EditableMesh/Public/EditableMeshTypes.h
--- a/Engine/Plugins/Runtime/EditableMesh/Source/EditableMesh/Public/EditableMeshTypes.h
+++ b/Engine/Plugins/Runtime/EditableMesh/Source/EditableMesh/Public/EditableMeshTypes.h
@@ -335,7 +335,7 @@ enum class EMeshElementAttributeType : uint8
 	Int,
 	Bool,
 	FName,
-	UObject
+	FSoftObjectPath
 };
 
 
@@ -356,7 +356,7 @@ struct FMeshElementAttributeValue
 	explicit FMeshElementAttributeValue( const int Value ) : Type( EMeshElementAttributeType::Int ), Value_Int( Value ) {}
 	explicit FMeshElementAttributeValue( const bool Value ) : Type( EMeshElementAttributeType::Bool ), Value_Bool( Value ) {}
 	explicit FMeshElementAttributeValue( const FName Value ) : Type( EMeshElementAttributeType::FName ), Value_FName( Value ) {}
-	explicit FMeshElementAttributeValue( UObject* const Value ) : Type( EMeshElementAttributeType::UObject ), Value_UObject( Value ) {}
+	explicit FMeshElementAttributeValue(FSoftObjectPath const Value ) : Type( EMeshElementAttributeType::FSoftObjectPath), Value_FSoftObjectPath( Value ) {}
 
 	/** Returns the type of this attribute value */
 	EMeshElementAttributeType GetType() const
@@ -373,7 +373,7 @@ struct FMeshElementAttributeValue
 	template <> int GetValue<int>() const { check( Type == EMeshElementAttributeType::Int ); return Value_Int; }
 	template <> bool GetValue<bool>() const { check( Type == EMeshElementAttributeType::Bool ); return Value_Bool; }
 	template <> FName GetValue<FName>() const { check( Type == EMeshElementAttributeType::FName ); return Value_FName; }
-	template <> UObject* GetValue<UObject*>() const { check( Type == EMeshElementAttributeType::UObject ); return Value_UObject; }
+	template <> FSoftObjectPath GetValue<FSoftObjectPath>() const { check( Type == EMeshElementAttributeType::FSoftObjectPath); return Value_FSoftObjectPath; }
 
 	/**
 	 * Calls the specified polymorphic lambda, with this attribute value as its passed-in parameter.
@@ -393,7 +393,7 @@ struct FMeshElementAttributeValue
 			case EMeshElementAttributeType::Int: Func( Value_Int ); break;
 			case EMeshElementAttributeType::Bool: Func( Value_Bool ); break;
 			case EMeshElementAttributeType::FName: Func( Value_FName ); break;
-			case EMeshElementAttributeType::UObject: Func( Value_UObject ); break;
+			case EMeshElementAttributeType::FSoftObjectPath: Func(Value_FSoftObjectPath); break;
 		}
 	}
 
@@ -410,7 +410,7 @@ struct FMeshElementAttributeValue
 			case EMeshElementAttributeType::Int: return Lex::ToString( Value_Int );
 			case EMeshElementAttributeType::Bool: return Lex::ToString( Value_Bool );
 			case EMeshElementAttributeType::FName: return Value_FName.ToString();
-			case EMeshElementAttributeType::UObject: return GetNameSafe( Value_UObject );
+			case EMeshElementAttributeType::FSoftObjectPath: return GetNameSafe(Value_FSoftObjectPath.TryLoad());
 		}
 
 		return FString();
@@ -430,7 +430,7 @@ struct FMeshElementAttributeValue
 			case EMeshElementAttributeType::Int: Ar << AttributeValue.Value_Int; break;
 			case EMeshElementAttributeType::Bool: Ar << AttributeValue.Value_Bool; break;
 			case EMeshElementAttributeType::FName: Ar << AttributeValue.Value_FName; break;
-			case EMeshElementAttributeType::UObject: Ar << AttributeValue.Value_UObject; break;
+			case EMeshElementAttributeType::FSoftObjectPath: Ar << AttributeValue.Value_FSoftObjectPath; break;
 		}
 
 		return Ar;
@@ -454,7 +454,7 @@ struct FMeshElementAttributeValue
 			case EMeshElementAttributeType::Int: return Value1.Value_Int == Value2.Value_Int;
 			case EMeshElementAttributeType::Bool: return Value1.Value_Bool == Value2.Value_Bool;
 			case EMeshElementAttributeType::FName: return Value1.Value_FName == Value2.Value_FName;
-			case EMeshElementAttributeType::UObject: return Value1.Value_UObject == Value2.Value_UObject;
+			case EMeshElementAttributeType::FSoftObjectPath: return Value1.Value_FSoftObjectPath == Value2.Value_FSoftObjectPath;
 		}
 
 		return false;
@@ -475,7 +475,7 @@ private:
 	int Value_Int;
 	bool Value_Bool;
 	FName Value_FName;
-	UObject* Value_UObject;
+	FSoftObjectPath Value_FSoftObjectPath;
 };
 
 
diff --git a/Engine/Source/Editor/UnrealEd/Private/Fbx/FbxStaticMeshImport.cpp b/Engine/Source/Editor/UnrealEd/Private/Fbx/FbxStaticMeshImport.cpp
--- a/Engine/Source/Editor/UnrealEd/Private/Fbx/FbxStaticMeshImport.cpp
+++ b/Engine/Source/Editor/UnrealEd/Private/Fbx/FbxStaticMeshImport.cpp
@@ -533,7 +533,7 @@ bool UnFbx::FFbxImporter::BuildStaticMeshFromGeometry(FbxNode* Node, UStaticMesh
 		TEdgeAttributeArray<bool>& EdgeHardnesses = MeshDescription->EdgeAttributes().GetAttributes<bool>(MeshAttribute::Edge::IsHard);
 		TEdgeAttributeArray<float>& EdgeCreaseSharpnesses = MeshDescription->EdgeAttributes().GetAttributes<float>(MeshAttribute::Edge::CreaseSharpness);
 
-		TPolygonGroupAttributeArray<UObject*>& PolygonGroupMaterialAssets = MeshDescription->PolygonGroupAttributes().GetAttributes<UObject*>(MeshAttribute::PolygonGroup::MaterialAsset);
+		TPolygonGroupAttributeArray<FSoftObjectPath>& PolygonGroupMaterialAssets = MeshDescription->PolygonGroupAttributes().GetAttributes<FSoftObjectPath>(MeshAttribute::PolygonGroup::MaterialAsset);
 		TPolygonGroupAttributeArray<FName>& PolygonGroupMaterialSlotNames = MeshDescription->PolygonGroupAttributes().GetAttributes<FName>(MeshAttribute::PolygonGroup::MaterialSlotName);
 		TPolygonGroupAttributeArray<FName>& PolygonGroupImportedMaterialSlotNames = MeshDescription->PolygonGroupAttributes().GetAttributes<FName>(MeshAttribute::PolygonGroup::ImportedMaterialSlotName);
 		TPolygonGroupAttributeArray<bool>& PolygonGroupCollision = MeshDescription->PolygonGroupAttributes().GetAttributes<bool>(MeshAttribute::PolygonGroup::EnableCollision);
@@ -767,7 +767,7 @@ bool UnFbx::FFbxImporter::BuildStaticMeshFromGeometry(FbxNode* Node, UStaticMesh
 					PolygonGroupCollision[ExistingPolygonGroup] = bEnableCollision;
 					PolygonGroupImportedMaterialSlotNames[ExistingPolygonGroup] = ImportedMaterialSlotName;
 					PolygonGroupMaterialSlotNames[ExistingPolygonGroup] = ImportedMaterialSlotName;
-					PolygonGroupMaterialAssets[ExistingPolygonGroup] = Material;
+					PolygonGroupMaterialAssets[ExistingPolygonGroup] = FSoftObjectPath(Material);
 				}
 				PolygonGroupMapping.Add(RealMaterialIndex, ExistingPolygonGroup);
 			}
@@ -1833,7 +1833,7 @@ UStaticMesh* UnFbx::FFbxImporter::ImportStaticMeshAsSingle(UObject* InParent, TA
 		}
 		else
 		{
-			TPolygonGroupAttributeArray<UObject*>& PolygonGroupMaterialAssets = MeshDescription->PolygonGroupAttributes().GetAttributes<UObject*>( MeshAttribute::PolygonGroup::MaterialAsset );
+			TPolygonGroupAttributeArray<FSoftObjectPath>& PolygonGroupMaterialAssets = MeshDescription->PolygonGroupAttributes().GetAttributes<FSoftObjectPath>( MeshAttribute::PolygonGroup::MaterialAsset );
 			TPolygonGroupAttributeArray<FName>& PolygonGroupMaterialSlotNames = MeshDescription->PolygonGroupAttributes().GetAttributes<FName>( MeshAttribute::PolygonGroup::MaterialSlotName );
 			TPolygonGroupAttributeArray<FName>& PolygonGroupImportedMaterialSlotNames = MeshDescription->PolygonGroupAttributes().GetAttributes<FName>( MeshAttribute::PolygonGroup::ImportedMaterialSlotName );
 			TPolygonGroupAttributeArray<bool>& PolygonGroupCollision = MeshDescription->PolygonGroupAttributes().GetAttributes<bool>( MeshAttribute::PolygonGroup::EnableCollision );
@@ -1843,7 +1843,7 @@ UStaticMesh* UnFbx::FFbxImporter::ImportStaticMeshAsSingle(UObject* InParent, TA
 			for (const FPolygonGroupID PolygonGroupID : MeshDescription->PolygonGroups().GetElementIDs())
 			{
 				int32 MaterialIndex = PolygonGroupID.GetValue();
-				UMaterialInterface* Material = Cast<UMaterialInterface>(PolygonGroupMaterialAssets[PolygonGroupID]);
+				UMaterialInterface* Material = Cast<UMaterialInterface>(PolygonGroupMaterialAssets[PolygonGroupID].TryLoad());
 				const FName& MaterialSlotName = PolygonGroupMaterialSlotNames[PolygonGroupID];
 				const FName& ImportedMaterialSlotName = PolygonGroupImportedMaterialSlotNames[PolygonGroupID];
 				FStaticMaterial StaticMaterial(Material, MaterialSlotName, ImportedMaterialSlotName);
@@ -2000,6 +2000,13 @@ void UnFbx::FFbxImporter::PostImportStaticMesh(UStaticMesh* StaticMesh, TArray<F
 	// Build the staticmesh, we move the build here because we want to avoid building the staticmesh for every LOD
 	// when we import the mesh.
 	TArray<FText> BuildErrors;
+#if UE_BUILD_DEBUG
+	if (GIsAutomationTesting)
+	{
+		//Generate a random GUID to be sure it rebuild the asset
+		StaticMesh->BuildCacheAutomationTestGuid = FGuid::NewGuid();
+	}
+#endif
 
 	//Prebuild the static mesh when we use LodGroup and we want to modify the LodNumber
 	if (!ImportOptions->bImportScene)
diff --git a/Engine/Source/Runtime/Engine/Classes/Engine/StaticMesh.h b/Engine/Source/Runtime/Engine/Classes/Engine/StaticMesh.h
--- a/Engine/Source/Runtime/Engine/Classes/Engine/StaticMesh.h
+++ b/Engine/Source/Runtime/Engine/Classes/Engine/StaticMesh.h
@@ -606,6 +606,9 @@ class UStaticMesh : public UObject, public IInterface_CollisionDataProvider, pub
 
 #if WITH_EDITOR
 	FOnExtendedBoundsChanged OnExtendedBoundsChanged;
+
+	/** This transient guid is use by the automation framework to modify the DDC key to force a build. */
+	FGuid BuildCacheAutomationTestGuid;
 #endif
 
 protected:
diff --git a/Engine/Source/Runtime/Engine/Private/StaticMesh.cpp b/Engine/Source/Runtime/Engine/Private/StaticMesh.cpp
--- a/Engine/Source/Runtime/Engine/Private/StaticMesh.cpp
+++ b/Engine/Source/Runtime/Engine/Private/StaticMesh.cpp
@@ -1137,6 +1137,15 @@ static FString BuildStaticMeshDerivedDataKey(UStaticMesh* Mesh, const FStaticMes
 	{
 		KeySuffix += Lex::ToString(Mesh->LightmapUVVersion);
 	}
+#if WITH_EDITOR
+#if UE_BUILD_DEBUG
+	if (GIsAutomationTesting && Mesh->BuildCacheAutomationTestGuid.IsValid())
+	{
+		//If we are in automation testing and the BuildCacheAutomationTestGuid was set
+		KeySuffix += Mesh->BuildCacheAutomationTestGuid.ToString(EGuidFormats::Digits);
+	}
+#endif
+#endif
 
 	int32 NumLODs = Mesh->SourceModels.Num();
 	for (int32 LODIndex = 0; LODIndex < NumLODs; ++LODIndex)
@@ -1421,6 +1430,9 @@ UStaticMesh::UStaticMesh(const FObjectInitializer& ObjectInitializer)
 	MinLOD = 0;
 
 	bSupportUniformlyDistributedSampling = false;
+#if WITH_EDITOR
+	BuildCacheAutomationTestGuid.Invalidate();
+#endif
 }
 
 void UStaticMesh::PostInitProperties()
diff --git a/Engine/Source/Runtime/MeshBuilder/Private/MeshDescriptionHelper.cpp b/Engine/Source/Runtime/MeshBuilder/Private/MeshDescriptionHelper.cpp
--- a/Engine/Source/Runtime/MeshBuilder/Private/MeshDescriptionHelper.cpp
+++ b/Engine/Source/Runtime/MeshBuilder/Private/MeshDescriptionHelper.cpp
@@ -5,6 +5,8 @@
 
 #include "CoreMinimal.h"
 #include "CoreTypes.h"
+#include "UObject/UObjectGlobals.h"
+#include "UObject/Package.h"
 #include "Engine/EngineTypes.h"
 #include "Engine/StaticMesh.h"
 #include "MeshDescription.h"
@@ -44,7 +46,7 @@ UMeshDescription* FMeshDescriptionHelper::GetRenderMeshDescription(UObject* Owne
 	UStaticMesh* StaticMesh = Cast<UStaticMesh>(Owner);
 	check(StaticMesh);
 	//Use the build settings to create the RenderMeshDescription
-	UMeshDescription *RenderMeshDescription = NewObject<UMeshDescription>(Owner, NAME_None, RF_NoFlags);
+	UMeshDescription *RenderMeshDescription = nullptr;
 
 	if (OriginalMeshDescription == nullptr)
 	{
@@ -52,7 +54,7 @@ UMeshDescription* FMeshDescriptionHelper::GetRenderMeshDescription(UObject* Owne
 		return RenderMeshDescription;
 	}
 	//Copy The Original Mesh Description in the render mesh description
-	CopyMeshDescription(const_cast<UMeshDescription*>(OriginalMeshDescription), RenderMeshDescription);
+	RenderMeshDescription = Cast<UMeshDescription>(StaticDuplicateObject(OriginalMeshDescription, Owner, NAME_None, RF_NoFlags));
 	float ComparisonThreshold = BuildSettings->bRemoveDegenerates ? THRESH_POINTS_ARE_SAME : 0.0f;
 
 	FVertexInstanceArray& VertexInstanceArray = RenderMeshDescription->VertexInstances();
@@ -157,19 +159,6 @@ void FMeshDescriptionHelper::ReduceLOD(const UMeshDescription* BaseMesh, UMeshDe
 	MeshReduction->ReduceMeshDescription(DestMesh, MaxDeviation, BaseMesh, InOverlappingCorners, ReductionSettings);
 }
 
-void FMeshDescriptionHelper::CopyMeshDescription(UMeshDescription* SourceMeshDescription, UMeshDescription* DestinationMeshDescription) const
-{
-	BuildStatisticManager::FBuildStatisticScope StatScope(TEXT("CopyMeshDescription took"));
-	//Copy the Source into the destination
-	//Save the Source
-	TArray<uint8> TempBytes;
-	FMemoryWriter SaveAr(TempBytes, /*bIsPersistent=*/ true);
-	SourceMeshDescription->Serialize(SaveAr);
-	//Load the save array in the destination
-	FMemoryReader LoadAr(TempBytes, /*bIsPersistent=*/ true);
-	DestinationMeshDescription->Serialize(LoadAr);
-}
-
 bool FMeshDescriptionHelper::IsValidOriginalMeshDescription()
 {
 	return OriginalMeshDescription != nullptr;
diff --git a/Engine/Source/Runtime/MeshBuilder/Private/MeshDescriptionHelper.h b/Engine/Source/Runtime/MeshBuilder/Private/MeshDescriptionHelper.h
--- a/Engine/Source/Runtime/MeshBuilder/Private/MeshDescriptionHelper.h
+++ b/Engine/Source/Runtime/MeshBuilder/Private/MeshDescriptionHelper.h
@@ -43,8 +43,6 @@ public:
 
 	void ReduceLOD(const UMeshDescription* BaseMesh, UMeshDescription* DestMesh, const struct FMeshReductionSettings& ReductionSettings, const TMultiMap<int32, int32>& InOverlappingCorners);
 
-	void CopyMeshDescription(UMeshDescription* SourceMeshDescription, UMeshDescription* DestinationMeshDescription) const;
-
 	//Return true if there is a valid original mesh description, false otherwise(Auto generate LOD).
 	bool IsValidOriginalMeshDescription();
 
diff --git a/Engine/Source/Runtime/MeshBuilder/Private/StaticMeshBuilder.cpp b/Engine/Source/Runtime/MeshBuilder/Private/StaticMeshBuilder.cpp
--- a/Engine/Source/Runtime/MeshBuilder/Private/StaticMeshBuilder.cpp
+++ b/Engine/Source/Runtime/MeshBuilder/Private/StaticMeshBuilder.cpp
@@ -366,7 +366,7 @@ void BuildVertexBuffer(
 	const uint32 NumTextureCoord = MeshDescription->VertexInstanceAttributes().GetAttributeIndexCount<FVector2D>(MeshAttribute::VertexInstance::TextureCoordinate);
 
 	const TPolygonGroupAttributeArray<FName>& PolygonGroupImportedMaterialSlotNames = MeshDescription->PolygonGroupAttributes().GetAttributes<FName>(MeshAttribute::PolygonGroup::ImportedMaterialSlotName);
-	const TPolygonGroupAttributeArray<UObject*>& PolygonGroupMaterialAssets = MeshDescription->PolygonGroupAttributes().GetAttributes<UObject*>(MeshAttribute::PolygonGroup::MaterialAsset);
+	const TPolygonGroupAttributeArray<FSoftObjectPath>& PolygonGroupMaterialAssets = MeshDescription->PolygonGroupAttributes().GetAttributes<FSoftObjectPath>(MeshAttribute::PolygonGroup::MaterialAsset);
 	const TPolygonGroupAttributeArray<bool>& PolygonGroupCollision = MeshDescription->PolygonGroupAttributes().GetAttributes<bool>(MeshAttribute::PolygonGroup::EnableCollision);
 	const TPolygonGroupAttributeArray<bool>& PolygonGroupCastShadow = MeshDescription->PolygonGroupAttributes().GetAttributes<bool>(MeshAttribute::PolygonGroup::CastShadow);
 
@@ -396,7 +396,7 @@ void BuildVertexBuffer(
 
 		const int32 MaterialIndex = StaticMesh->GetMaterialIndex(PolygonGroupImportedMaterialSlotNames[PolygonGroupID]);
 		check(MaterialIndex != INDEX_NONE);
-		check(StaticMesh->StaticMaterials[MaterialIndex].MaterialInterface == Cast<UMaterialInterface>(PolygonGroupMaterialAssets[PolygonGroupID]));
+		check(StaticMesh->StaticMaterials[MaterialIndex].MaterialInterface == Cast<UMaterialInterface>(PolygonGroupMaterialAssets[PolygonGroupID].TryLoad()));
 		StaticMeshSection.MaterialIndex = MaterialIndex;
 		StaticMeshSection.bEnableCollision = PolygonGroupCollision[PolygonGroupID];
 		StaticMeshSection.bCastShadow = PolygonGroupCastShadow[PolygonGroupID];
diff --git a/Engine/Source/Runtime/MeshDescription/Private/MeshDescription.cpp b/Engine/Source/Runtime/MeshDescription/Private/MeshDescription.cpp
--- a/Engine/Source/Runtime/MeshDescription/Private/MeshDescription.cpp
+++ b/Engine/Source/Runtime/MeshDescription/Private/MeshDescription.cpp
@@ -4,6 +4,7 @@
 #include "Misc/SecureHash.h"
 #include "Serialization/MemoryWriter.h"
 #include "MeshAttributes.h"
+#include "Serialization/ObjectWriter.h"
 
 UMeshDescription::UMeshDescription()
 {
@@ -29,7 +30,7 @@ UMeshDescription::UMeshDescription()
 	PolygonAttributes().RegisterAttribute<FVector>( MeshAttribute::Polygon::Center, 1, FVector::ZeroVector, EMeshAttributeFlags::AutoGenerated );
 
 	// Add basic polygon group attributes
-	PolygonGroupAttributes().RegisterAttribute<UObject*>( MeshAttribute::PolygonGroup::MaterialAsset );
+	PolygonGroupAttributes().RegisterAttribute<FSoftObjectPath>( MeshAttribute::PolygonGroup::MaterialAsset );
 	PolygonGroupAttributes().RegisterAttribute<FName>( MeshAttribute::PolygonGroup::MaterialSlotName );
 	PolygonGroupAttributes().RegisterAttribute<FName>( MeshAttribute::PolygonGroup::ImportedMaterialSlotName );
 	PolygonGroupAttributes().RegisterAttribute<bool>( MeshAttribute::PolygonGroup::EnableCollision );
@@ -59,7 +60,7 @@ void UMeshDescription::Serialize( FArchive& Ar )
 FString UMeshDescription::GetIdString()
 {
 	TArray<uint8> TempBytes;
-	FMemoryWriter Ar(TempBytes, /*bIsPersistent=*/ true);
+	FObjectWriter Ar(this, TempBytes);
 	Serialize(Ar);
 	FSHA1 Sha;
 	TArray<TCHAR> OwnerName = GetPathName().GetCharArray();
diff --git a/Engine/Source/Runtime/MeshDescription/Public/MeshAttributeArray.h b/Engine/Source/Runtime/MeshDescription/Public/MeshAttributeArray.h
--- a/Engine/Source/Runtime/MeshDescription/Public/MeshAttributeArray.h
+++ b/Engine/Source/Runtime/MeshDescription/Public/MeshAttributeArray.h
@@ -20,7 +20,7 @@ using AttributeTypes = TTuple
 	int,
 	bool,
 	FName,
-	UObject*
+	FSoftObjectPath
 >;
 
 
