Port the conversion of the brush to mesh description

diff --git a/Engine/Source/Editor/UnrealEd/Private/StaticMeshEdit.cpp b/Engine/Source/Editor/UnrealEd/Private/StaticMeshEdit.cpp
--- a/Engine/Source/Editor/UnrealEd/Private/StaticMeshEdit.cpp
+++ b/Engine/Source/Editor/UnrealEd/Private/StaticMeshEdit.cpp
@@ -25,6 +25,9 @@
 
 #include "Materials/MaterialInterface.h"
 #include "Materials/Material.h"
+#include "Settings/EditorExperimentalSettings.h"
+#include "MeshDescription.h"
+#include "MeshAttributes.h"
 
 
 bool GBuildStaticMeshCollision = 1;
@@ -572,6 +575,40 @@ void AddConvexGeomFromVertices( const TArray<FVector>& Verts, FKAggregateGeom* A
 	ConvexElem->UpdateElemBox();
 }
 
+/**
+* Creates a static mesh object from raw triangle data.
+*/
+UStaticMesh* CreateStaticMesh(UMeshDescription* MeshDescription, TArray<FStaticMaterial>& Materials, UObject* InOuter, FName InName)
+{
+	// Create the UStaticMesh object.
+	FStaticMeshComponentRecreateRenderStateContext RecreateRenderStateContext(FindObject<UStaticMesh>(InOuter, *InName.ToString()));
+	auto StaticMesh = NewObject<UStaticMesh>(InOuter, InName, RF_Public | RF_Standalone);
+
+	// Add one LOD for the base mesh
+	FStaticMeshSourceModel* SrcModel = new(StaticMesh->SourceModels) FStaticMeshSourceModel();
+	StaticMesh->SetOriginalMeshDescription(StaticMesh->SourceModels.Num() - 1, MeshDescription);
+	StaticMesh->StaticMaterials = Materials;
+
+	int32 NumSections = StaticMesh->StaticMaterials.Num();
+
+	// Set up the SectionInfoMap to enable collision
+	for (int32 SectionIdx = 0; SectionIdx < NumSections; ++SectionIdx)
+	{
+		FMeshSectionInfo Info = StaticMesh->SectionInfoMap.Get(0, SectionIdx);
+		Info.MaterialIndex = SectionIdx;
+		Info.bEnableCollision = true;
+		StaticMesh->SectionInfoMap.Set(0, SectionIdx, Info);
+		StaticMesh->OriginalSectionInfoMap.Set(0, SectionIdx, Info);
+	}
+
+	//Set the Imported version before calling the build
+	StaticMesh->ImportVersion = EImportStaticMeshVersion::LastVersion;
+
+	StaticMesh->Build();
+	StaticMesh->MarkPackageDirty();
+	return StaticMesh;
+}
+
 /**
  * Creates a static mesh object from raw triangle data.
  */
@@ -667,6 +704,136 @@ inline bool FVerticesEqual(FVector& V1,FVector& V2)
 	return 1;
 }
 
+void GetBrushMesh(ABrush* Brush, UModel* Model, UMeshDescription* MeshDescription, TArray<FStaticMaterial>& OutMaterials)
+{
+	check(MeshDescription != nullptr);
+
+	TVertexAttributeArray<FVector>& VertexPositions = MeshDescription->VertexAttributes().GetAttributes<FVector>(MeshAttribute::Vertex::Position);
+	TVertexInstanceAttributeArray<FVector>& VertexInstanceNormals = MeshDescription->VertexInstanceAttributes().GetAttributes<FVector>(MeshAttribute::VertexInstance::Normal);
+	TVertexInstanceAttributeArray<FVector>& VertexInstanceTangents = MeshDescription->VertexInstanceAttributes().GetAttributes<FVector>(MeshAttribute::VertexInstance::Tangent);
+	TVertexInstanceAttributeArray<float>& VertexInstanceBinormalSigns = MeshDescription->VertexInstanceAttributes().GetAttributes<float>(MeshAttribute::VertexInstance::BinormalSign);
+	TVertexInstanceAttributeArray<FVector4>& VertexInstanceColors = MeshDescription->VertexInstanceAttributes().GetAttributes<FVector4>(MeshAttribute::VertexInstance::Color);
+	TVertexInstanceAttributeIndicesArray<FVector2D>& VertexInstanceUVs = MeshDescription->VertexInstanceAttributes().GetAttributesSet<FVector2D>(MeshAttribute::VertexInstance::TextureCoordinate);
+	TEdgeAttributeArray<bool>& EdgeHardnesses = MeshDescription->EdgeAttributes().GetAttributes<bool>(MeshAttribute::Edge::IsHard);
+	TEdgeAttributeArray<float>& EdgeCreaseSharpnesses = MeshDescription->EdgeAttributes().GetAttributes<float>(MeshAttribute::Edge::CreaseSharpness);
+	TPolygonGroupAttributeArray<FName>& PolygonGroupImportedMaterialSlotNames = MeshDescription->PolygonGroupAttributes().GetAttributes<FName>(MeshAttribute::PolygonGroup::ImportedMaterialSlotName);
+	TPolygonGroupAttributeArray<int>& PolygonGroupMaterialIndex = MeshDescription->PolygonGroupAttributes().GetAttributes<int>(MeshAttribute::PolygonGroup::MaterialIndex);
+	
+	//Make sure we have one UVChannel
+	VertexInstanceUVs.SetNumIndices(1);
+	
+	// Calculate the local to world transform for the source brush.
+	FMatrix	ActorToWorld = Brush ? Brush->ActorToWorld().ToMatrixWithScale() : FMatrix::Identity;
+	bool	ReverseVertices = 0;
+	FVector4	PostSub = Brush ? FVector4(Brush->GetActorLocation()) : FVector4(0, 0, 0, 0);
+
+	TMap<uint32, FEdgeID> RemapEdgeID;
+	int32 NumPolys = Model->Polys->Element.Num();
+	//Create Fill the vertex position
+	for (int32 PolygonIndex = 0; PolygonIndex < NumPolys; ++PolygonIndex)
+	{
+		FPoly& Polygon = Model->Polys->Element[PolygonIndex];
+
+		// Find a material index for this polygon.
+		UMaterialInterface*	Material = Polygon.Material;
+		int32 MaterialIndex = OutMaterials.AddUnique(FStaticMaterial(Material));
+		FPolygonGroupID CurrentPolygonGroupID = FPolygonGroupID::Invalid;
+		for (const FPolygonGroupID& PolygonGroupID : MeshDescription->PolygonGroups().GetElementIDs())
+		{
+			if (MaterialIndex == PolygonGroupMaterialIndex[PolygonGroupID])
+			{
+				CurrentPolygonGroupID = PolygonGroupID;
+				break;
+			}
+		}
+		if (CurrentPolygonGroupID == FPolygonGroupID::Invalid)
+		{
+			CurrentPolygonGroupID = MeshDescription->CreatePolygonGroup();
+			PolygonGroupMaterialIndex[CurrentPolygonGroupID] = MaterialIndex;
+		}
+
+		// Cache the texture coordinate system for this polygon.
+		FVector	TextureBase = Polygon.Base - (Brush ? Brush->GetPivotOffset() : FVector::ZeroVector),
+			TextureX = Polygon.TextureU / UModel::GetGlobalBSPTexelScale(),
+			TextureY = Polygon.TextureV / UModel::GetGlobalBSPTexelScale();
+		// For each vertex after the first two vertices...
+		for (int32 VertexIndex = 2; VertexIndex < Polygon.Vertices.Num(); VertexIndex++)
+		{
+			FVector Positions[3];
+			Positions[ReverseVertices ? 0 : 2] = ActorToWorld.TransformPosition(Polygon.Vertices[0]) - PostSub;
+			Positions[1] = ActorToWorld.TransformPosition(Polygon.Vertices[VertexIndex - 1]) - PostSub;
+			Positions[ReverseVertices ? 2 : 0] = ActorToWorld.TransformPosition(Polygon.Vertices[VertexIndex]) - PostSub;
+			FVertexID VertexID[3] = { FVertexID::Invalid, FVertexID::Invalid, FVertexID::Invalid };
+			for (FVertexID IterVertexID : MeshDescription->Vertices().GetElementIDs())
+			{
+				if (FVerticesEqual(Positions[0], VertexPositions[IterVertexID]))
+				{
+					VertexID[0] = IterVertexID;
+				}
+				if (FVerticesEqual(Positions[1], VertexPositions[IterVertexID]))
+				{
+					VertexID[1] = IterVertexID;
+				}
+				if (FVerticesEqual(Positions[2], VertexPositions[IterVertexID]))
+				{
+					VertexID[2] = IterVertexID;
+				}
+			}
+
+			//Create the vertex instances
+			FVertexInstanceID VertexInstanceIDs[3];
+			for (int32 CornerIndex = 0; CornerIndex < 3; ++CornerIndex)
+			{
+				if (VertexID[CornerIndex] == FVertexID::Invalid)
+				{
+					VertexID[CornerIndex] = MeshDescription->CreateVertex();
+					VertexPositions[VertexID[CornerIndex]] = Positions[CornerIndex];
+				}
+				VertexInstanceIDs[CornerIndex] = MeshDescription->CreateVertexInstance(VertexID[CornerIndex]);
+				VertexInstanceUVs.GetArrayForIndex(0)[VertexInstanceIDs[CornerIndex]] = FVector2D(
+					(Positions[CornerIndex] - TextureBase) | TextureX,
+					(Positions[CornerIndex] - TextureBase) | TextureY);
+			}
+
+			// Create a polygon with the 3 vertex instances
+			
+			TArray<UMeshDescription::FContourPoint> Contours;
+			for (uint32 TriangleEdgeNumber = 0; TriangleEdgeNumber < 3; ++TriangleEdgeNumber)
+			{
+				int32 ContourPointIndex = Contours.AddDefaulted();
+				UMeshDescription::FContourPoint& ContourPoint = Contours[ContourPointIndex];
+				//Find the matching edge ID
+				uint32 CornerIndices[2];
+				CornerIndices[0] = (TriangleEdgeNumber + 0) % 3;
+				CornerIndices[1] = (TriangleEdgeNumber + 1) % 3;
+
+				FVertexID EdgeVertexIDs[2];
+				EdgeVertexIDs[0] = VertexID[CornerIndices[0]];
+				EdgeVertexIDs[1] = VertexID[CornerIndices[1]];
+
+				FEdgeID MatchEdgeId = MeshDescription->GetVertexPairEdge(EdgeVertexIDs[0], EdgeVertexIDs[1]);
+				if (MatchEdgeId == FEdgeID::Invalid)
+				{
+					MatchEdgeId = MeshDescription->CreateEdge(EdgeVertexIDs[0], EdgeVertexIDs[1]);
+				}
+				ContourPoint.EdgeID = MatchEdgeId;
+				ContourPoint.VertexInstanceID = VertexInstanceIDs[CornerIndices[0]];
+				//All edge are hard for BSP
+				EdgeHardnesses[MatchEdgeId] = true;
+			}
+			// Insert a polygon into the mesh
+			const FPolygonID NewPolygonID = MeshDescription->CreatePolygon(CurrentPolygonGroupID, Contours);
+			int32 NewTriangleIndex = MeshDescription->GetPolygonTriangles(NewPolygonID).AddDefaulted();
+			FMeshTriangle& NewTriangle = MeshDescription->GetPolygonTriangles(NewPolygonID)[NewTriangleIndex];
+			for (int32 TriangleVertexIndex = 0; TriangleVertexIndex < 3; ++TriangleVertexIndex)
+			{
+				const FVertexInstanceID &VertexInstanceID = VertexInstanceIDs[TriangleVertexIndex];
+				NewTriangle.SetVertexInstanceID(TriangleVertexIndex, VertexInstanceID);
+			}
+		}
+	}
+}
+
 void GetBrushMesh(ABrush* Brush,UModel* Model,struct FRawMesh& OutMesh,TArray<FStaticMaterial>& OutMaterials)
 {
 	// Calculate the local to world transform for the source brush.
@@ -757,12 +924,21 @@ void GetBrushMesh(ABrush* Brush,UModel* Model,struct FRawMesh& OutMesh,TArray<FS
 UStaticMesh* CreateStaticMeshFromBrush(UObject* Outer,FName Name,ABrush* Brush,UModel* Model)
 {
 	GWarn->BeginSlowTask( NSLOCTEXT("UnrealEd", "CreatingStaticMeshE", "Creating static mesh..."), true );
-
-	FRawMesh RawMesh;
-	TArray<FStaticMaterial> Materials;
-	GetBrushMesh(Brush,Model,RawMesh,Materials);
-
-	UStaticMesh* StaticMesh = CreateStaticMesh(RawMesh, Materials,Outer,Name);
+	UStaticMesh* StaticMesh = nullptr;
+	if (GetDefault<UEditorExperimentalSettings>()->bUseMeshDescription)
+	{
+		UMeshDescription* MeshDescription = NewObject<UMeshDescription>(Outer);
+		TArray<FStaticMaterial> Materials;
+		GetBrushMesh(Brush, Model, MeshDescription, Materials);
+		StaticMesh = CreateStaticMesh(MeshDescription, Materials, Outer, Name);
+	}
+	else
+	{
+		FRawMesh RawMesh;
+		TArray<FStaticMaterial> Materials;
+		GetBrushMesh(Brush, Model, RawMesh, Materials);
+		StaticMesh = CreateStaticMesh(RawMesh, Materials, Outer, Name);
+	}
 	GWarn->EndSlowTask();
 
 	return StaticMesh;
