- Remove the editor pref settings bUseMeshDescription
- Fix fbx importer material remap when using mesh description

diff --git a/Engine/Source/Developer/MeshDescriptionOperations/Private/MeshDescriptionOperations.cpp b/Engine/Source/Developer/MeshDescriptionOperations/Private/MeshDescriptionOperations.cpp
--- a/Engine/Source/Developer/MeshDescriptionOperations/Private/MeshDescriptionOperations.cpp
+++ b/Engine/Source/Developer/MeshDescriptionOperations/Private/MeshDescriptionOperations.cpp
@@ -481,65 +481,74 @@ void FMeshDescriptionOperations::CreatePolygonNTB(UMeshDescription* MeshDescript
 
 	for (const FPolygonID PolygonID : MeshDescription->Polygons().GetElementIDs())
 	{
-		FVector TangentX(0.0f);
-		FVector TangentY(0.0f);
-		FVector TangentZ(0.0f);
-
 		if (!PolygonNormals[PolygonID].IsNearlyZero())
 		{
 			//By pass normal calculation if its already done
 			continue;
 		}
 		const TArray<FMeshTriangle>& MeshTriangles = MeshDescription->GetPolygonTriangles(PolygonID);
+		FVector TangentX(0.0f);
+		FVector TangentY(0.0f);
+		FVector TangentZ(0.0f);
+		for (const FMeshTriangle& MeshTriangle : MeshTriangles)
+		{
+			int32 UVIndex = 0;
 
-		//We need only the first triangle since all triangle of a polygon must have the same normals (planar polygon)
-		const FMeshTriangle& MeshTriangle = MeshTriangles[0];
-		int32 UVIndex = 0;
-
-		FVector P[3];
-		FVector2D UVs[3];
+			FVector P[3];
+			FVector2D UVs[3];
 
-		for (int32 i = 0; i < 3; ++i)
-		{
-			const FVertexInstanceID VertexInstanceID = MeshTriangle.GetVertexInstanceID(i);
-			UVs[i] = VertexUVs[VertexInstanceID];
-			P[i] = VertexPositions[MeshDescription->GetVertexInstanceVertex(VertexInstanceID)];
-		}
+			for (int32 i = 0; i < 3; ++i)
+			{
+				const FVertexInstanceID VertexInstanceID = MeshTriangle.GetVertexInstanceID(i);
+				UVs[i] = VertexUVs[VertexInstanceID];
+				P[i] = VertexPositions[MeshDescription->GetVertexInstanceVertex(VertexInstanceID)];
+			}
 
-		const FVector Normal = ((P[1] - P[2]) ^ (P[0] - P[2])).GetSafeNormal(ComparisonThreshold);
-		//Check for degenerated polygons, avoid NAN
-		if (!Normal.IsNearlyZero(ComparisonThreshold))
-		{
-			FMatrix	ParameterToLocal(
-				FPlane(P[1].X - P[0].X, P[1].Y - P[0].Y, P[1].Z - P[0].Z, 0),
-				FPlane(P[2].X - P[0].X, P[2].Y - P[0].Y, P[2].Z - P[0].Z, 0),
-				FPlane(P[0].X, P[0].Y, P[0].Z, 0),
-				FPlane(0, 0, 0, 1)
-			);
-
-			FMatrix ParameterToTexture(
-				FPlane(UVs[1].X - UVs[0].X, UVs[1].Y - UVs[0].Y, 0, 0),
-				FPlane(UVs[2].X - UVs[0].X, UVs[2].Y - UVs[0].Y, 0, 0),
-				FPlane(UVs[0].X, UVs[0].Y, 1, 0),
-				FPlane(0, 0, 0, 1)
-			);
-
-			// Use InverseSlow to catch singular matrices.  Inverse can miss this sometimes.
-			const FMatrix TextureToLocal = ParameterToTexture.Inverse() * ParameterToLocal;
-
-			TangentX = TextureToLocal.TransformVector(FVector(1, 0, 0)).GetSafeNormal();
-			TangentY = TextureToLocal.TransformVector(FVector(0, 1, 0)).GetSafeNormal();
-			TangentZ = Normal;
-			FVector::CreateOrthonormalBasis(TangentX, TangentY, TangentZ);
-		}
-		else
-		{
-			DegeneratePolygons.Add(PolygonID);
+			const FVector Normal = ((P[1] - P[2]) ^ (P[0] - P[2])).GetSafeNormal(ComparisonThreshold);
+			//Check for degenerated polygons, avoid NAN
+			if (!Normal.IsNearlyZero(ComparisonThreshold))
+			{
+				FMatrix	ParameterToLocal(
+					FPlane(P[1].X - P[0].X, P[1].Y - P[0].Y, P[1].Z - P[0].Z, 0),
+					FPlane(P[2].X - P[0].X, P[2].Y - P[0].Y, P[2].Z - P[0].Z, 0),
+					FPlane(P[0].X, P[0].Y, P[0].Z, 0),
+					FPlane(0, 0, 0, 1)
+				);
+
+				FMatrix ParameterToTexture(
+					FPlane(UVs[1].X - UVs[0].X, UVs[1].Y - UVs[0].Y, 0, 0),
+					FPlane(UVs[2].X - UVs[0].X, UVs[2].Y - UVs[0].Y, 0, 0),
+					FPlane(UVs[0].X, UVs[0].Y, 1, 0),
+					FPlane(0, 0, 0, 1)
+				);
+
+				// Use InverseSlow to catch singular matrices.  Inverse can miss this sometimes.
+				const FMatrix TextureToLocal = ParameterToTexture.Inverse() * ParameterToLocal;
+
+				FVector TmpTangentX(0.0f);
+				FVector TmpTangentY(0.0f);
+				FVector TmpTangentZ(0.0f);
+				TmpTangentX = TextureToLocal.TransformVector(FVector(1, 0, 0)).GetSafeNormal();
+				TmpTangentY = TextureToLocal.TransformVector(FVector(0, 1, 0)).GetSafeNormal();
+				TmpTangentZ = Normal;
+				FVector::CreateOrthonormalBasis(TmpTangentX, TmpTangentY, TmpTangentZ);
+				TangentX += TmpTangentX;
+				TangentY += TmpTangentY;
+				TangentZ += TmpTangentZ;
+			}
+			else
+			{
+				DegeneratePolygons.AddUnique(PolygonID);
+			}
 		}
 
+		TangentX.Normalize();
+		TangentY.Normalize();
+		TangentZ.Normalize();
 		PolygonTangents[PolygonID] = TangentX;
 		PolygonBinormals[PolygonID] = TangentY;
 		PolygonNormals[PolygonID] = TangentZ;
+		
 	}
 
 	//Delete the degenerated polygons. The array is fill only if the remove degenerated option is turn on.
@@ -824,13 +833,16 @@ namespace MeshDescriptionMikktSpaceInterface
 	int MikkGetNumVertsOfFace(const SMikkTSpaceContext* Context, const int FaceIdx)
 	{
 		// All of our meshes are triangles.
-		return 3;
+		UMeshDescription *MeshDescription = (UMeshDescription*)(Context->m_pUserData);
+		const FMeshPolygon& Polygon = MeshDescription->GetPolygon(FPolygonID(FaceIdx));
+		return Polygon.PerimeterContour.VertexInstanceIDs.Num();
 	}
 
 	void MikkGetPosition(const SMikkTSpaceContext* Context, float Position[3], const int FaceIdx, const int VertIdx)
 	{
 		UMeshDescription *MeshDescription = (UMeshDescription*)(Context->m_pUserData);
-		const FVertexInstanceID VertexInstanceID(FaceIdx * 3 + VertIdx);
+		const FMeshPolygon& Polygon = MeshDescription->GetPolygon(FPolygonID(FaceIdx));
+		const FVertexInstanceID VertexInstanceID = Polygon.PerimeterContour.VertexInstanceIDs[VertIdx];
 		const FVertexID VertexID = MeshDescription->GetVertexInstanceVertex(VertexInstanceID);
 		const FVector& VertexPosition = MeshDescription->VertexAttributes().GetAttribute<FVector>(VertexID, MeshAttribute::Vertex::Position);
 		Position[0] = VertexPosition.X;
@@ -841,7 +853,8 @@ namespace MeshDescriptionMikktSpaceInterface
 	void MikkGetNormal(const SMikkTSpaceContext* Context, float Normal[3], const int FaceIdx, const int VertIdx)
 	{
 		UMeshDescription *MeshDescription = (UMeshDescription*)(Context->m_pUserData);
-		const FVertexInstanceID VertexInstanceID(FaceIdx * 3 + VertIdx);
+		const FMeshPolygon& Polygon = MeshDescription->GetPolygon(FPolygonID(FaceIdx));
+		const FVertexInstanceID VertexInstanceID = Polygon.PerimeterContour.VertexInstanceIDs[VertIdx];
 		const FVector& VertexNormal = MeshDescription->VertexInstanceAttributes().GetAttribute<FVector>(VertexInstanceID, MeshAttribute::VertexInstance::Normal);
 		Normal[0] = VertexNormal.X;
 		Normal[1] = VertexNormal.Y;
@@ -851,7 +864,8 @@ namespace MeshDescriptionMikktSpaceInterface
 	void MikkSetTSpaceBasic(const SMikkTSpaceContext* Context, const float Tangent[3], const float BitangentSign, const int FaceIdx, const int VertIdx)
 	{
 		UMeshDescription *MeshDescription = (UMeshDescription*)(Context->m_pUserData);
-		const FVertexInstanceID VertexInstanceID(FaceIdx * 3 + VertIdx);
+		const FMeshPolygon& Polygon = MeshDescription->GetPolygon(FPolygonID(FaceIdx));
+		const FVertexInstanceID VertexInstanceID = Polygon.PerimeterContour.VertexInstanceIDs[VertIdx];
 		const FVector VertexTangent(Tangent[0], Tangent[1], Tangent[2]);
 		MeshDescription->VertexInstanceAttributes().SetAttribute<FVector>(VertexInstanceID, MeshAttribute::VertexInstance::Tangent, 0, VertexTangent);
 		MeshDescription->VertexInstanceAttributes().SetAttribute<float>(VertexInstanceID, MeshAttribute::VertexInstance::BinormalSign, 0, -BitangentSign);
@@ -860,7 +874,8 @@ namespace MeshDescriptionMikktSpaceInterface
 	void MikkGetTexCoord(const SMikkTSpaceContext* Context, float UV[2], const int FaceIdx, const int VertIdx)
 	{
 		UMeshDescription *MeshDescription = (UMeshDescription*)(Context->m_pUserData);
-		const FVertexInstanceID VertexInstanceID(FaceIdx * 3 + VertIdx);
+		const FMeshPolygon& Polygon = MeshDescription->GetPolygon(FPolygonID(FaceIdx));
+		const FVertexInstanceID VertexInstanceID = Polygon.PerimeterContour.VertexInstanceIDs[VertIdx];
 		const FVector2D& TexCoord = MeshDescription->VertexInstanceAttributes().GetAttribute<FVector2D>(VertexInstanceID, MeshAttribute::VertexInstance::TextureCoordinate, 0);
 		UV[0] = TexCoord.X;
 		UV[1] = TexCoord.Y;
diff --git a/Engine/Source/Editor/UnrealEd/Classes/Settings/EditorExperimentalSettings.h b/Engine/Source/Editor/UnrealEd/Classes/Settings/EditorExperimentalSettings.h
--- a/Engine/Source/Editor/UnrealEd/Classes/Settings/EditorExperimentalSettings.h
+++ b/Engine/Source/Editor/UnrealEd/Classes/Settings/EditorExperimentalSettings.h
@@ -136,10 +136,6 @@ public:
 	UPROPERTY(EditAnywhere, config, Category = Materials)
 	bool bExampleLayersAndBlends;
 
-	/** Allows usage of the mesh description instead of the raw mesh when we build static mesh */
-	UPROPERTY(EditAnywhere, config, Category = Tools, meta = (DisplayName = "Use Mesh Description"))
-	bool bUseMeshDescription;
-	
 	/**
 	 * Returns an event delegate that is executed when a setting has changed.
 	 *
diff --git a/Engine/Source/Editor/UnrealEd/Private/Fbx/FbxStaticMeshImport.cpp b/Engine/Source/Editor/UnrealEd/Private/Fbx/FbxStaticMeshImport.cpp
--- a/Engine/Source/Editor/UnrealEd/Private/Fbx/FbxStaticMeshImport.cpp
+++ b/Engine/Source/Editor/UnrealEd/Private/Fbx/FbxStaticMeshImport.cpp
@@ -268,11 +268,8 @@ bool UnFbx::FFbxImporter::BuildStaticMeshFromGeometry(FbxNode* Node, UStaticMesh
 	FStaticMeshSourceModel& SrcModel = StaticMesh->SourceModels[LODIndex];
 	
 	UMeshDescription *MeshDescription = StaticMesh->GetOriginalMeshDescription(LODIndex);
-	if (GetDefault<UEditorExperimentalSettings>()->bUseMeshDescription)
-	{
-		//The mesh description should have been created before calling BuildStaticMeshFromGeometry
-		check(MeshDescription);
-	}
+	//The mesh description should have been created before calling BuildStaticMeshFromGeometry
+	check(MeshDescription);
 	//remove the bad polygons before getting any data from mesh
 	Mesh->RemoveBadPolygons();
 
@@ -415,22 +412,15 @@ bool UnFbx::FFbxImporter::BuildStaticMeshFromGeometry(FbxNode* Node, UStaticMesh
 	FbxLayerElement::EMappingMode SmoothingMappingMode(FbxLayerElement::eByEdge);
 	if (SmoothingInfo)
 	{
-		if(GetDefault<UEditorExperimentalSettings>()->bUseMeshDescription)
+		if (SmoothingInfo->GetMappingMode() == FbxLayerElement::eByPolygon)
 		{
-			if (SmoothingInfo->GetMappingMode() == FbxLayerElement::eByPolygon)
-			{
-				//Convert the base layer to edge smoothing
-				GeometryConverter->ComputeEdgeSmoothingFromPolygonSmoothing(Mesh, 0);
-				BaseLayer = Mesh->GetLayer(0);
-				SmoothingInfo = BaseLayer->GetSmoothing();
-			}
-
-			if (SmoothingInfo->GetMappingMode() == FbxLayerElement::eByEdge)
-			{
-				bSmoothingAvailable = true;
-			}
+			//Convert the base layer to edge smoothing
+			GeometryConverter->ComputeEdgeSmoothingFromPolygonSmoothing(Mesh, 0);
+			BaseLayer = Mesh->GetLayer(0);
+			SmoothingInfo = BaseLayer->GetSmoothing();
 		}
-		else if (SmoothingInfo->GetMappingMode() == FbxLayerElement::eByPolygon)
+
+		if (SmoothingInfo->GetMappingMode() == FbxLayerElement::eByEdge)
 		{
 			bSmoothingAvailable = true;
 		}
@@ -525,677 +515,400 @@ bool UnFbx::FFbxImporter::BuildStaticMeshFromGeometry(FbxNode* Node, UStaticMesh
 	TotalMatrix = ComputeTotalMatrix(Node);
 	TotalMatrixForNormal = TotalMatrix.Inverse();
 	TotalMatrixForNormal = TotalMatrixForNormal.Transpose();
-	int32 TriangleCount = Mesh->GetPolygonCount();
+	int32 PolygonCount = Mesh->GetPolygonCount();
 
-	if(TriangleCount == 0)
+	if(PolygonCount == 0)
 	{
-		AddTokenizedErrorMessage(FTokenizedMessage::Create(EMessageSeverity::Error, FText::Format(LOCTEXT("Error_NoTrianglesFoundInMesh", "No triangles were found on mesh  '{0}'"), FText::FromString(Mesh->GetName()))), FFbxErrors::StaticMesh_NoTriangles);
+		AddTokenizedErrorMessage(FTokenizedMessage::Create(EMessageSeverity::Error, FText::Format(LOCTEXT("Error_NoPolygonFoundInMesh", "No polygon were found on mesh  '{0}'"), FText::FromString(Mesh->GetName()))), FFbxErrors::StaticMesh_NoTriangles);
 		return false;
 	}
 
 	int32 VertexCount = Mesh->GetControlPointsCount();
 	bool OddNegativeScale = IsOddNegativeScale(TotalMatrix);
-	bool bHasNonDegenerateTriangles = false;
-	if(GetDefault<UEditorExperimentalSettings>()->bUseMeshDescription)
-	{
-		TVertexAttributeArray<FVector>& VertexPositions = MeshDescription->VertexAttributes().GetAttributes<FVector>(MeshAttribute::Vertex::Position);
+	bool bHasNonDegeneratePolygons = false;
+	TVertexAttributeArray<FVector>& VertexPositions = MeshDescription->VertexAttributes().GetAttributes<FVector>(MeshAttribute::Vertex::Position);
 
-		TVertexInstanceAttributeArray<FVector>& VertexInstanceNormals = MeshDescription->VertexInstanceAttributes().GetAttributes<FVector>(MeshAttribute::VertexInstance::Normal);
-		TVertexInstanceAttributeArray<FVector>& VertexInstanceTangents = MeshDescription->VertexInstanceAttributes().GetAttributes<FVector>(MeshAttribute::VertexInstance::Tangent);
-		TVertexInstanceAttributeArray<float>& VertexInstanceBinormalSigns = MeshDescription->VertexInstanceAttributes().GetAttributes<float>(MeshAttribute::VertexInstance::BinormalSign);
-		TVertexInstanceAttributeArray<FVector4>& VertexInstanceColors = MeshDescription->VertexInstanceAttributes().GetAttributes<FVector4>(MeshAttribute::VertexInstance::Color);
-		TVertexInstanceAttributeIndicesArray<FVector2D>& VertexInstanceUVs = MeshDescription->VertexInstanceAttributes().GetAttributesSet<FVector2D>(MeshAttribute::VertexInstance::TextureCoordinate);
+	TVertexInstanceAttributeArray<FVector>& VertexInstanceNormals = MeshDescription->VertexInstanceAttributes().GetAttributes<FVector>(MeshAttribute::VertexInstance::Normal);
+	TVertexInstanceAttributeArray<FVector>& VertexInstanceTangents = MeshDescription->VertexInstanceAttributes().GetAttributes<FVector>(MeshAttribute::VertexInstance::Tangent);
+	TVertexInstanceAttributeArray<float>& VertexInstanceBinormalSigns = MeshDescription->VertexInstanceAttributes().GetAttributes<float>(MeshAttribute::VertexInstance::BinormalSign);
+	TVertexInstanceAttributeArray<FVector4>& VertexInstanceColors = MeshDescription->VertexInstanceAttributes().GetAttributes<FVector4>(MeshAttribute::VertexInstance::Color);
+	TVertexInstanceAttributeIndicesArray<FVector2D>& VertexInstanceUVs = MeshDescription->VertexInstanceAttributes().GetAttributesSet<FVector2D>(MeshAttribute::VertexInstance::TextureCoordinate);
 
-		TEdgeAttributeArray<bool>& EdgeHardnesses = MeshDescription->EdgeAttributes().GetAttributes<bool>(MeshAttribute::Edge::IsHard);
-		TEdgeAttributeArray<float>& EdgeCreaseSharpnesses = MeshDescription->EdgeAttributes().GetAttributes<float>(MeshAttribute::Edge::CreaseSharpness);
+	TEdgeAttributeArray<bool>& EdgeHardnesses = MeshDescription->EdgeAttributes().GetAttributes<bool>(MeshAttribute::Edge::IsHard);
+	TEdgeAttributeArray<float>& EdgeCreaseSharpnesses = MeshDescription->EdgeAttributes().GetAttributes<float>(MeshAttribute::Edge::CreaseSharpness);
 
-		TPolygonGroupAttributeArray<FName>& PolygonGroupImportedMaterialSlotNames = MeshDescription->PolygonGroupAttributes().GetAttributes<FName>(MeshAttribute::PolygonGroup::ImportedMaterialSlotName);
+	TPolygonGroupAttributeArray<FName>& PolygonGroupImportedMaterialSlotNames = MeshDescription->PolygonGroupAttributes().GetAttributes<FName>(MeshAttribute::PolygonGroup::ImportedMaterialSlotName);
 
-		int32 VertexOffset = MeshDescription->Vertices().Num();
-		int32 VertexInstanceOffset = MeshDescription->VertexInstances().Num();
-		int32 TriangleOffset = MeshDescription->Polygons().Num();
+	int32 VertexOffset = MeshDescription->Vertices().Num();
+	int32 VertexInstanceOffset = MeshDescription->VertexInstances().Num();
+	int32 PolygonOffset = MeshDescription->Polygons().Num();
 
-		int32 MaxMaterialIndex = 0;
-		TMap<int32, FPolygonGroupID> PolygonGroupMapping;
+	TMap<int32, FPolygonGroupID> PolygonGroupMapping;
 
-		// When importing multiple mesh pieces to the same static mesh.  Ensure each mesh piece has the same number of Uv's
-		int32 ExistingUVCount = 0;
-		for (int32 UVChannelIndex = 0; UVChannelIndex < VertexInstanceUVs.GetNumIndices(); ++UVChannelIndex)
+	// When importing multiple mesh pieces to the same static mesh.  Ensure each mesh piece has the same number of Uv's
+	int32 ExistingUVCount = 0;
+	for (int32 UVChannelIndex = 0; UVChannelIndex < VertexInstanceUVs.GetNumIndices(); ++UVChannelIndex)
+	{
+		if (VertexInstanceUVs.GetArrayForIndex(UVChannelIndex).Num() > 0)
 		{
-			if (VertexInstanceUVs.GetArrayForIndex(UVChannelIndex).Num() > 0)
-			{
-				ExistingUVCount++;
-			}
-			else
-			{
-				break;
-			}
+			ExistingUVCount++;
+		}
+		else
+		{
+			break;
 		}
+	}
 
-		int32 NumUVs = FMath::Max(FBXUVs.UniqueUVCount, ExistingUVCount);
-		// At least one UV set must exist.  
-		NumUVs = FMath::Max(1, NumUVs);
+	int32 NumUVs = FMath::Max(FBXUVs.UniqueUVCount, ExistingUVCount);
+	// At least one UV set must exist.  
+	NumUVs = FMath::Max(1, NumUVs);
 
-		//Make sure all Vertex instance have the correct number of UVs
-		VertexInstanceUVs.SetNumIndices( NumUVs );
+	//Make sure all Vertex instance have the correct number of UVs
+	VertexInstanceUVs.SetNumIndices( NumUVs );
 
-		//Fill the vertex array
-		for (int32 VertexIndex = 0; VertexIndex < VertexCount; ++VertexIndex)
-		{
-			int32 RealVertexIndex = VertexOffset + VertexIndex;
-			FbxVector4 FbxPosition = Mesh->GetControlPoints()[VertexIndex];
-			FbxPosition = TotalMatrix.MultT(FbxPosition);
-			const FVector VertexPosition = Converter.ConvertPos(FbxPosition);
+	//Fill the vertex array
+	for (int32 VertexIndex = 0; VertexIndex < VertexCount; ++VertexIndex)
+	{
+		int32 RealVertexIndex = VertexOffset + VertexIndex;
+		FbxVector4 FbxPosition = Mesh->GetControlPoints()[VertexIndex];
+		FbxPosition = TotalMatrix.MultT(FbxPosition);
+		const FVector VertexPosition = Converter.ConvertPos(FbxPosition);
 			
-			FVertexID AddedVertexId = MeshDescription->CreateVertex();
-			VertexPositions[AddedVertexId] = VertexPosition;
-			if (AddedVertexId.GetValue() != RealVertexIndex)
-			{
-				AddTokenizedErrorMessage(FTokenizedMessage::Create(EMessageSeverity::Error, FText::Format(LOCTEXT("Error_CannotCreateVertex", "Cannot create valid vertex for mesh '{0}'"), FText::FromString(Mesh->GetName()))), FFbxErrors::StaticMesh_BuildError);
-				return false;
-			}
+		FVertexID AddedVertexId = MeshDescription->CreateVertex();
+		VertexPositions[AddedVertexId] = VertexPosition;
+		if (AddedVertexId.GetValue() != RealVertexIndex)
+		{
+			AddTokenizedErrorMessage(FTokenizedMessage::Create(EMessageSeverity::Error, FText::Format(LOCTEXT("Error_CannotCreateVertex", "Cannot create valid vertex for mesh '{0}'"), FText::FromString(Mesh->GetName()))), FFbxErrors::StaticMesh_BuildError);
+			return false;
 		}
+	}
 		
-		Mesh->BeginGetMeshEdgeVertices();
-		TMap<uint32, int32> RemapEdgeID;
-		//Fill the edge array
-		int32 FbxEdgeCount = Mesh->GetMeshEdgeCount();
-		RemapEdgeID.Reserve(FbxEdgeCount*2);
-		for (int32 FbxEdgeIndex = 0; FbxEdgeIndex < FbxEdgeCount; ++FbxEdgeIndex)
-		{
-			int32 EdgeStartVertexIndex = -1;
-			int32 EdgeEndVertexIndex = -1;
-			Mesh->GetMeshEdgeVertices(FbxEdgeIndex, EdgeStartVertexIndex, EdgeEndVertexIndex);
-			FVertexID EdgeVertexStart(EdgeStartVertexIndex + VertexOffset);
-			check(MeshDescription->Vertices().IsValid(EdgeVertexStart));
-			FVertexID EdgeVertexEnd(EdgeEndVertexIndex + VertexOffset);
-			check(MeshDescription->Vertices().IsValid(EdgeVertexEnd));
-			uint32 CompactedKey = ((uint32)EdgeVertexStart.GetValue() << 16) | ((uint32)EdgeVertexEnd.GetValue() & 0xffff);
-			RemapEdgeID.Add(CompactedKey, FbxEdgeIndex);
-			//Add the other edge side
-			CompactedKey = ((uint32)EdgeVertexEnd.GetValue() << 16) | ((uint32)EdgeVertexStart.GetValue() & 0xffff);
-			RemapEdgeID.Add(CompactedKey, FbxEdgeIndex);
-		}
-		//Call this after all GetMeshEdgeIndexForPolygon call this is for optimization purpose.
-		Mesh->EndGetMeshEdgeVertices();
-
-		//Call this before all GetMeshEdgeIndexForPolygon call this is for optimization purpose.
-		Mesh->BeginGetMeshEdgeIndexForPolygon();
-		int32 VertexInstanceCut = 0;
-		//Polygons
-		for (int32 TriangleIndex = 0; TriangleIndex < TriangleCount; TriangleIndex++)
-		{
-			//Verify if the polygon is degenerate, in this case do not add them
-			{
-				float ComparisonThreshold = ImportOptions->bRemoveDegenerates ? SMALL_NUMBER : 0.0f;
-				FVector P[3];
-				for (int32 CornerIndex = 0; CornerIndex < 3; CornerIndex++)
-				{
-					const int32 ControlPointIndex = Mesh->GetPolygonVertex(TriangleIndex, CornerIndex);
-					const FVertexID VertexID(VertexOffset + ControlPointIndex);
-					P[CornerIndex] = VertexPositions[VertexID];
-				}
-				const FVector Normal = ((P[1] - P[2]) ^ (P[0] - P[2])).GetSafeNormal(ComparisonThreshold);
-				//Check for degenerated polygons, avoid NAN
-				if (Normal.IsNearlyZero(ComparisonThreshold))
-				{
-					VertexInstanceCut += 3;
-					continue;
-				}
+	Mesh->BeginGetMeshEdgeVertices();
+	TMap<uint32, int32> RemapEdgeID;
+	//Fill the edge array
+	int32 FbxEdgeCount = Mesh->GetMeshEdgeCount();
+	RemapEdgeID.Reserve(FbxEdgeCount*2);
+	for (int32 FbxEdgeIndex = 0; FbxEdgeIndex < FbxEdgeCount; ++FbxEdgeIndex)
+	{
+		int32 EdgeStartVertexIndex = -1;
+		int32 EdgeEndVertexIndex = -1;
+		Mesh->GetMeshEdgeVertices(FbxEdgeIndex, EdgeStartVertexIndex, EdgeEndVertexIndex);
+		FVertexID EdgeVertexStart(EdgeStartVertexIndex + VertexOffset);
+		check(MeshDescription->Vertices().IsValid(EdgeVertexStart));
+		FVertexID EdgeVertexEnd(EdgeEndVertexIndex + VertexOffset);
+		check(MeshDescription->Vertices().IsValid(EdgeVertexEnd));
+		uint32 CompactedKey = ((uint32)EdgeVertexStart.GetValue() << 16) | ((uint32)EdgeVertexEnd.GetValue() & 0xffff);
+		RemapEdgeID.Add(CompactedKey, FbxEdgeIndex);
+		//Add the other edge side
+		CompactedKey = ((uint32)EdgeVertexEnd.GetValue() << 16) | ((uint32)EdgeVertexStart.GetValue() & 0xffff);
+		RemapEdgeID.Add(CompactedKey, FbxEdgeIndex);
+	}
+	//Call this after all GetMeshEdgeIndexForPolygon call this is for optimization purpose.
+	Mesh->EndGetMeshEdgeVertices();
+
+	//Call this before all GetMeshEdgeIndexForPolygon call this is for optimization purpose.
+	Mesh->BeginGetMeshEdgeIndexForPolygon();
+	int32 CurrentVertexInstanceIndex = 0;
+	//Polygons
+	for (int32 PolygonIndex = 0; PolygonIndex < PolygonCount; PolygonIndex++)
+	{
+		int32 PolygonVertexCount = Mesh->GetPolygonSize(PolygonIndex);
+		//Verify if the polygon is degenerate, in this case do not add them
+		{
+			float ComparisonThreshold = ImportOptions->bRemoveDegenerates ? SMALL_NUMBER : 0.0f;
+			TArray<FVector> P;
+			P.AddUninitialized(PolygonVertexCount);
+			for (int32 CornerIndex = 0; CornerIndex < PolygonVertexCount; CornerIndex++)
+			{
+				const int32 ControlPointIndex = Mesh->GetPolygonVertex(PolygonIndex, CornerIndex);
+				const FVertexID VertexID(VertexOffset + ControlPointIndex);
+				P[CornerIndex] = VertexPositions[VertexID];
 			}
-
-			int32 RealTriangleIndex = TriangleOffset + TriangleIndex;
-			FVertexInstanceID CornerInstanceIDs[3];
-			FVertexID CornerVerticesIDs[3];
-			for (int32 CornerIndex = 0; CornerIndex < 3; CornerIndex++)
+			check(P.Num() > 2); //triangle is the smallest polygon we can have
+			const FVector Normal = ((P[1] - P[2]) ^ (P[0] - P[2])).GetSafeNormal(ComparisonThreshold);
+			//Check for degenerated polygons, avoid NAN
+			if (Normal.IsNearlyZero(ComparisonThreshold))
 			{
-				int32 VertexInstanceIndex = RealTriangleIndex * 3 + CornerIndex - VertexInstanceCut;
-				const FVertexInstanceID VertexInstanceID(VertexInstanceIndex);
-				CornerInstanceIDs[CornerIndex] = VertexInstanceID;
-				const int32 ControlPointIndex = Mesh->GetPolygonVertex(TriangleIndex, CornerIndex);
-				const FVertexID VertexID(VertexOffset + ControlPointIndex);
-				const FVector VertexPosition = VertexPositions[VertexID];
-				CornerVerticesIDs[CornerIndex] = VertexID;
+				continue;
+			}
+		}
 
-				FVertexInstanceID AddedVertexInstanceId = MeshDescription->CreateVertexInstance(VertexID);
+		int32 RealPolygonIndex = PolygonOffset + PolygonIndex;
+		TArray<FVertexInstanceID> CornerInstanceIDs;
+		CornerInstanceIDs.AddUninitialized(PolygonVertexCount);
+		TArray<FVertexID> CornerVerticesIDs;
+		CornerVerticesIDs.AddUninitialized(PolygonVertexCount);
+		for (int32 CornerIndex = 0; CornerIndex < PolygonVertexCount; CornerIndex++, CurrentVertexInstanceIndex++)
+		{
+			int32 VertexInstanceIndex = VertexInstanceOffset + CurrentVertexInstanceIndex;
+			const FVertexInstanceID VertexInstanceID(VertexInstanceIndex);
+			CornerInstanceIDs[CornerIndex] = VertexInstanceID;
+			const int32 ControlPointIndex = Mesh->GetPolygonVertex(PolygonIndex, CornerIndex);
+			const FVertexID VertexID(VertexOffset + ControlPointIndex);
+			const FVector VertexPosition = VertexPositions[VertexID];
+			CornerVerticesIDs[CornerIndex] = VertexID;
+
+			FVertexInstanceID AddedVertexInstanceId = MeshDescription->CreateVertexInstance(VertexID);
 				
-				//Make sure the Added vertex instance ID is matching the expected vertex instance ID
-				check(AddedVertexInstanceId == VertexInstanceID);
+			//Make sure the Added vertex instance ID is matching the expected vertex instance ID
+			check(AddedVertexInstanceId == VertexInstanceID);
 				
-				if (AddedVertexInstanceId.GetValue() != VertexInstanceIndex)
-				{
-					AddTokenizedErrorMessage(FTokenizedMessage::Create(EMessageSeverity::Error, FText::Format(LOCTEXT("Error_CannotCreateVertexInstance", "Cannot create valid vertex instance for mesh '{0}'"), FText::FromString(Mesh->GetName()))), FFbxErrors::StaticMesh_BuildError);
-					return false;
-				}
+			if (AddedVertexInstanceId.GetValue() != VertexInstanceIndex)
+			{
+				AddTokenizedErrorMessage(FTokenizedMessage::Create(EMessageSeverity::Error, FText::Format(LOCTEXT("Error_CannotCreateVertexInstance", "Cannot create valid vertex instance for mesh '{0}'"), FText::FromString(Mesh->GetName()))), FFbxErrors::StaticMesh_BuildError);
+				return false;
+			}
 
-				//UVs attributes
-				for (int32 UVLayerIndex = 0; UVLayerIndex < FBXUVs.UniqueUVCount; UVLayerIndex++)
+			//UVs attributes
+			for (int32 UVLayerIndex = 0; UVLayerIndex < FBXUVs.UniqueUVCount; UVLayerIndex++)
+			{
+				FVector2D FinalUVVector(0.0f, 0.0f);
+				if (FBXUVs.LayerElementUV[UVLayerIndex] != NULL)
 				{
-					FVector2D FinalUVVector(0.0f, 0.0f);
-					if (FBXUVs.LayerElementUV[UVLayerIndex] != NULL)
-					{
-						int UVMapIndex = (FBXUVs.UVMappingMode[UVLayerIndex] == FbxLayerElement::eByControlPoint) ? ControlPointIndex : TriangleIndex * 3 + CornerIndex;
-						int32 UVIndex = (FBXUVs.UVReferenceMode[UVLayerIndex] == FbxLayerElement::eDirect) ?
-							UVMapIndex : FBXUVs.LayerElementUV[UVLayerIndex]->GetIndexArray().GetAt(UVMapIndex);
+					int UVMapIndex = (FBXUVs.UVMappingMode[UVLayerIndex] == FbxLayerElement::eByControlPoint) ? ControlPointIndex : CurrentVertexInstanceIndex;
+					int32 UVIndex = (FBXUVs.UVReferenceMode[UVLayerIndex] == FbxLayerElement::eDirect) ?
+						UVMapIndex : FBXUVs.LayerElementUV[UVLayerIndex]->GetIndexArray().GetAt(UVMapIndex);
 
-						FbxVector2	UVVector = FBXUVs.LayerElementUV[UVLayerIndex]->GetDirectArray().GetAt(UVIndex);
-						FinalUVVector.X = static_cast<float>(UVVector[0]);
-						FinalUVVector.Y = 1.f - static_cast<float>(UVVector[1]);   //flip the Y of UVs for DirectX
-					}
-					VertexInstanceUVs.GetArrayForIndex(UVLayerIndex)[AddedVertexInstanceId] = FinalUVVector;
+					FbxVector2	UVVector = FBXUVs.LayerElementUV[UVLayerIndex]->GetDirectArray().GetAt(UVIndex);
+					FinalUVVector.X = static_cast<float>(UVVector[0]);
+					FinalUVVector.Y = 1.f - static_cast<float>(UVVector[1]);   //flip the Y of UVs for DirectX
 				}
+				VertexInstanceUVs.GetArrayForIndex(UVLayerIndex)[AddedVertexInstanceId] = FinalUVVector;
+			}
 
-				//Color attribute
-				if (VertexColorImportOption == EVertexColorImportOption::Replace)
+			//Color attribute
+			if (VertexColorImportOption == EVertexColorImportOption::Replace)
+			{
+				if (LayerElementVertexColor)
 				{
-					if (LayerElementVertexColor)
-					{
-						int32 VertexColorMappingIndex = (VertexColorMappingMode == FbxLayerElement::eByControlPoint) ?
-							Mesh->GetPolygonVertex(TriangleIndex, CornerIndex) : (TriangleIndex * 3 + CornerIndex);
-
-						int32 VectorColorIndex = (VertexColorReferenceMode == FbxLayerElement::eDirect) ?
-							VertexColorMappingIndex : LayerElementVertexColor->GetIndexArray().GetAt(VertexColorMappingIndex);
+					int32 VertexColorMappingIndex = (VertexColorMappingMode == FbxLayerElement::eByControlPoint) ?
+						Mesh->GetPolygonVertex(PolygonIndex, CornerIndex) : (CurrentVertexInstanceIndex);
 
-						FbxColor VertexColor = LayerElementVertexColor->GetDirectArray().GetAt(VectorColorIndex);
+					int32 VectorColorIndex = (VertexColorReferenceMode == FbxLayerElement::eDirect) ?
+						VertexColorMappingIndex : LayerElementVertexColor->GetIndexArray().GetAt(VertexColorMappingIndex);
 
-						FColor VertexInstanceColor(
-							uint8(255.f*VertexColor.mRed),
-							uint8(255.f*VertexColor.mGreen),
-							uint8(255.f*VertexColor.mBlue),
-							uint8(255.f*VertexColor.mAlpha)
-						);
-						VertexInstanceColors[AddedVertexInstanceId] = FVector4(FLinearColor(VertexInstanceColor));
-					}
-				}
-				else if (VertexColorImportOption == EVertexColorImportOption::Ignore)
-				{
-					// try to match this triangles current vertex with one that existed in the previous mesh.
-					// This is a find in a tmap which uses a fast hash table lookup.
-					const FColor* PaintedColor = ExistingVertexColorData.Find(VertexPosition);
-					if (PaintedColor)
-					{
-						// A matching color for this vertex was found
-						VertexInstanceColors[AddedVertexInstanceId] = FVector4(FLinearColor(*PaintedColor));
-					}
-				}
-				else
-				{
-					// set the triangle's vertex color to a constant override
-					check(VertexColorImportOption == EVertexColorImportOption::Override);
-					VertexInstanceColors[AddedVertexInstanceId] = FVector4(FLinearColor(VertexOverrideColor));
-				}
+					FbxColor VertexColor = LayerElementVertexColor->GetDirectArray().GetAt(VectorColorIndex);
 
-				if (LayerElementNormal)
-				{
-					int TriangleCornerIndex = TriangleIndex * 3 + CornerIndex;
-					//normals may have different reference and mapping mode than tangents and binormals
-					int NormalMapIndex = (NormalMappingMode == FbxLayerElement::eByControlPoint) ?
-						ControlPointIndex : TriangleCornerIndex;
-					int NormalValueIndex = (NormalReferenceMode == FbxLayerElement::eDirect) ?
-						NormalMapIndex : LayerElementNormal->GetIndexArray().GetAt(NormalMapIndex);
-					
-					FbxVector4 TempValue = LayerElementNormal->GetDirectArray().GetAt(NormalValueIndex);
-					TempValue = TotalMatrixForNormal.MultT(TempValue);
-					FVector TangentZ = Converter.ConvertDir(TempValue);
-					VertexInstanceNormals[AddedVertexInstanceId] = TangentZ.GetSafeNormal();
-					//tangents and binormals share the same reference, mapping mode and index array
-					if (bHasNTBInformation)
-					{
-						int TangentMapIndex = (TangentMappingMode == FbxLayerElement::eByControlPoint) ?
-							ControlPointIndex : TriangleCornerIndex;
-						int TangentValueIndex = (TangentReferenceMode == FbxLayerElement::eDirect) ?
-							TangentMapIndex : LayerElementTangent->GetIndexArray().GetAt(TangentMapIndex);
-
-						TempValue = LayerElementTangent->GetDirectArray().GetAt(TangentValueIndex);
-						TempValue = TotalMatrixForNormal.MultT(TempValue);
-						FVector TangentX = Converter.ConvertDir(TempValue);
-						VertexInstanceTangents[AddedVertexInstanceId] = TangentX.GetSafeNormal();
-
-						int BinormalMapIndex = (BinormalMappingMode == FbxLayerElement::eByControlPoint) ?
-							ControlPointIndex : TriangleCornerIndex;
-						int BinormalValueIndex = (BinormalReferenceMode == FbxLayerElement::eDirect) ?
-							BinormalMapIndex : LayerElementBinormal->GetIndexArray().GetAt(BinormalMapIndex);
-
-						TempValue = LayerElementBinormal->GetDirectArray().GetAt(BinormalValueIndex);
-						TempValue = TotalMatrixForNormal.MultT(TempValue);
-						FVector TangentY = -Converter.ConvertDir(TempValue);
-						VertexInstanceBinormalSigns[AddedVertexInstanceId] = GetBasisDeterminantSign(TangentX.GetSafeNormal(), TangentY.GetSafeNormal(), TangentZ.GetSafeNormal());
-					}
+					FColor VertexInstanceColor(
+						uint8(255.f*VertexColor.mRed),
+						uint8(255.f*VertexColor.mGreen),
+						uint8(255.f*VertexColor.mBlue),
+						uint8(255.f*VertexColor.mAlpha)
+					);
+					VertexInstanceColors[AddedVertexInstanceId] = FVector4(FLinearColor(VertexInstanceColor));
 				}
 			}
-			
-			// Check if the triangle just discovered is non-degenerate if we haven't found one yet
-			if (!bHasNonDegenerateTriangles)
-			{
-				float TriangleComparisonThreshold = ImportOptions->bRemoveDegenerates ? THRESH_POINTS_ARE_SAME : 0.0f;
-				FVector VertexPosition[3];
-				VertexPosition[0] = VertexPositions[CornerVerticesIDs[0]];
-				VertexPosition[1] = VertexPositions[CornerVerticesIDs[1]];
-				VertexPosition[2] = VertexPositions[CornerVerticesIDs[2]];
-				if (!( VertexPosition[0].Equals(VertexPosition[1], TriangleComparisonThreshold)
-					|| VertexPosition[0].Equals(VertexPosition[2], TriangleComparisonThreshold)
-					|| VertexPosition[1].Equals(VertexPosition[2], TriangleComparisonThreshold)))
-				{
-					bHasNonDegenerateTriangles = true;
-				}
-			}
-
-			//
-			// material index
-			//
-			int32 MaterialIndex = 0;
-			if (MaterialCount > 0)
+			else if (VertexColorImportOption == EVertexColorImportOption::Ignore)
 			{
-				if (LayerElementMaterial)
+				// try to match this triangles current vertex with one that existed in the previous mesh.
+				// This is a find in a tmap which uses a fast hash table lookup.
+				const FColor* PaintedColor = ExistingVertexColorData.Find(VertexPosition);
+				if (PaintedColor)
 				{
-					switch (MaterialMappingMode)
-					{
-						// material index is stored in the IndexArray, not the DirectArray (which is irrelevant with 2009.1)
-					case FbxLayerElement::eAllSame:
-					{
-						MaterialIndex = LayerElementMaterial->GetIndexArray().GetAt(0);
-					}
-					break;
-					case FbxLayerElement::eByPolygon:
-					{
-						MaterialIndex = LayerElementMaterial->GetIndexArray().GetAt(TriangleIndex);
-					}
-					break;
-					}
+					// A matching color for this vertex was found
+					VertexInstanceColors[AddedVertexInstanceId] = FVector4(FLinearColor(*PaintedColor));
 				}
 			}
-
-			if (MaterialIndex >= MaterialCount || MaterialIndex < 0)
+			else
 			{
-				AddTokenizedErrorMessage(FTokenizedMessage::Create(EMessageSeverity::Warning, LOCTEXT("Error_MaterialIndexInconsistency", "Face material index inconsistency - forcing to 0")), FFbxErrors::Generic_Mesh_MaterialIndexInconsistency);
-				MaterialIndex = 0;
+				// set the triangle's vertex color to a constant override
+				check(VertexColorImportOption == EVertexColorImportOption::Override);
+				VertexInstanceColors[AddedVertexInstanceId] = FVector4(FLinearColor(VertexOverrideColor));
 			}
 
-			//Create a polygon with the 3 vertex intances Add it to the material group
-			int32 RealMaterialIndex = MaterialIndexOffset + MaterialIndex;
-			if (!PolygonGroupMapping.Contains(RealMaterialIndex))
+			if (LayerElementNormal)
 			{
-				UMaterialInterface* Material = MeshMaterials.IsValidIndex(RealMaterialIndex) ? MeshMaterials[RealMaterialIndex].Material : UMaterial::GetDefaultMaterial(MD_Surface);
-				FName ImportedMaterialSlotName = MeshMaterials.IsValidIndex(RealMaterialIndex) ? FName(*MeshMaterials[RealMaterialIndex].GetName()) : NAME_None;
-				FPolygonGroupID ExistingPolygonGroup = FPolygonGroupID::Invalid;
-				for (const FPolygonGroupID PolygonGroupID : MeshDescription->PolygonGroups().GetElementIDs())
-				{
-					if (PolygonGroupImportedMaterialSlotNames[PolygonGroupID] == ImportedMaterialSlotName)
-					{
-						ExistingPolygonGroup = PolygonGroupID;
-						break;
-					}
-				}
-				if (ExistingPolygonGroup == FPolygonGroupID::Invalid)
-				{
-					ExistingPolygonGroup = MeshDescription->CreatePolygonGroup();
-					PolygonGroupImportedMaterialSlotNames[ExistingPolygonGroup] = ImportedMaterialSlotName;
-				}
-				PolygonGroupMapping.Add(RealMaterialIndex, ExistingPolygonGroup);
-			}
-
-			// Create polygon edges
-			TArray<UMeshDescription::FContourPoint> Contours;
-			{
-				// Add the edges of this triangle
-				for (uint32 TriangleEdgeNumber = 0; TriangleEdgeNumber < 3; ++TriangleEdgeNumber)
-				{
-					int32 ContourPointIndex = Contours.AddDefaulted();
-					UMeshDescription::FContourPoint& ContourPoint = Contours[ContourPointIndex];
-					//Find the matching edge ID
-					uint32 CornerIndices[2];
-					CornerIndices[0] = (TriangleEdgeNumber + 0) % 3;
-					CornerIndices[1] = (TriangleEdgeNumber + 1) % 3;
-
-					FVertexID EdgeVertexIDs[2];
-					EdgeVertexIDs[0] = CornerVerticesIDs[CornerIndices[0]];
-					EdgeVertexIDs[1] = CornerVerticesIDs[CornerIndices[1]];
-
-					FEdgeID MatchEdgeId = MeshDescription->GetVertexPairEdge(EdgeVertexIDs[0], EdgeVertexIDs[1]);
-					if (MatchEdgeId == FEdgeID::Invalid)
-					{
-						MatchEdgeId = MeshDescription->CreateEdge(EdgeVertexIDs[0], EdgeVertexIDs[1]);
-					}
-					ContourPoint.EdgeID = MatchEdgeId;
-					ContourPoint.VertexInstanceID = CornerInstanceIDs[CornerIndices[0]];
-					//RawMesh do not have edges, so by ordering the edge with the triangle construction we can ensure back and forth conversion with RawMesh
-					//When raw mesh will be completly remove we can create the edges right after the vertex creation.
-					int32 EdgeIndex = INDEX_NONE;
-					uint32 CompactedKey = ((uint32)EdgeVertexIDs[0].GetValue() << 16) | ((uint32)EdgeVertexIDs[1].GetValue() & 0xffff);
-					if (RemapEdgeID.Contains(CompactedKey))
-					{
-						EdgeIndex = RemapEdgeID[CompactedKey];
-					}
-					else
-					{
-						EdgeIndex = Mesh->GetMeshEdgeIndexForPolygon(TriangleIndex, TriangleEdgeNumber);
-					}
+				int TriangleCornerIndex = CurrentVertexInstanceIndex;
+				//normals may have different reference and mapping mode than tangents and binormals
+				int NormalMapIndex = (NormalMappingMode == FbxLayerElement::eByControlPoint) ?
+					ControlPointIndex : TriangleCornerIndex;
+				int NormalValueIndex = (NormalReferenceMode == FbxLayerElement::eDirect) ?
+					NormalMapIndex : LayerElementNormal->GetIndexArray().GetAt(NormalMapIndex);
 					
-					EdgeCreaseSharpnesses[MatchEdgeId] = (float)Mesh->GetEdgeCreaseInfo(EdgeIndex);
-					if (!EdgeHardnesses[MatchEdgeId])
-					{
-						if (bSmoothingAvailable && SmoothingInfo)
-						{
-							if (SmoothingMappingMode == FbxLayerElement::eByEdge)
-							{
-								int32 lSmoothingIndex = (SmoothingReferenceMode == FbxLayerElement::eDirect) ? EdgeIndex : SmoothingInfo->GetIndexArray().GetAt(EdgeIndex);
-								//Set the hard edges
-								EdgeHardnesses[MatchEdgeId] = (SmoothingInfo->GetDirectArray().GetAt(lSmoothingIndex) == 0);
-							}
-							else
-							{
-								AddTokenizedErrorMessage(FTokenizedMessage::Create(EMessageSeverity::Warning, FText::Format(LOCTEXT("Error_UnsupportedSmoothingGroup", "Unsupported Smoothing group mapping mode on mesh  '{0}'"), FText::FromString(Mesh->GetName()))), FFbxErrors::Generic_Mesh_UnsupportingSmoothingGroup);
-							}
-						}
-						else
-						{
-							//When there is no smoothing group we set all edge to hard (faceted mesh)
-							EdgeHardnesses[MatchEdgeId] = true;
-						}
-					}
-				}
-			}
-
-			FPolygonGroupID PolygonGroupID = PolygonGroupMapping[RealMaterialIndex];
-			// Insert a polygon into the mesh
-			const FPolygonID NewPolygonID = MeshDescription->CreatePolygon(PolygonGroupID, Contours);
-			int32 NewTriangleIndex = MeshDescription->GetPolygonTriangles(NewPolygonID).AddDefaulted();
-			FMeshTriangle& NewTriangle = MeshDescription->GetPolygonTriangles(NewPolygonID)[NewTriangleIndex];
-			for (int32 TriangleVertexIndex = 0; TriangleVertexIndex < 3; ++TriangleVertexIndex)
-			{
-				const FVertexInstanceID &VertexInstanceID = CornerInstanceIDs[TriangleVertexIndex];
-				NewTriangle.SetVertexInstanceID(TriangleVertexIndex, VertexInstanceID);
-			}
-		}
-		//Call this after all GetMeshEdgeIndexForPolygon call this is for optimization purpose.
-		Mesh->EndGetMeshEdgeIndexForPolygon();
-		// needed?
-		FBXUVs.Cleanup();
-	}
-	else
-	{
-		int32 WedgeCount = TriangleCount * 3;
-		int32 VertexOffset = RawMesh.VertexPositions.Num();
-		int32 WedgeOffset = RawMesh.WedgeIndices.Num();
-		int32 TriangleOffset = RawMesh.FaceMaterialIndices.Num();
-
-		int32 MaxMaterialIndex = 0;
+				FbxVector4 TempValue = LayerElementNormal->GetDirectArray().GetAt(NormalValueIndex);
+				TempValue = TotalMatrixForNormal.MultT(TempValue);
+				FVector TangentZ = Converter.ConvertDir(TempValue);
+				VertexInstanceNormals[AddedVertexInstanceId] = TangentZ.GetSafeNormal();
+				//tangents and binormals share the same reference, mapping mode and index array
+				if (bHasNTBInformation)
+				{
+					int TangentMapIndex = (TangentMappingMode == FbxLayerElement::eByControlPoint) ?
+						ControlPointIndex : TriangleCornerIndex;
+					int TangentValueIndex = (TangentReferenceMode == FbxLayerElement::eDirect) ?
+						TangentMapIndex : LayerElementTangent->GetIndexArray().GetAt(TangentMapIndex);
 
-		// Reserve space for attributes.
-		RawMesh.FaceMaterialIndices.AddZeroed(TriangleCount);
-		RawMesh.FaceSmoothingMasks.AddZeroed(TriangleCount);
-		RawMesh.WedgeIndices.AddZeroed(WedgeCount);
+					TempValue = LayerElementTangent->GetDirectArray().GetAt(TangentValueIndex);
+					TempValue = TotalMatrixForNormal.MultT(TempValue);
+					FVector TangentX = Converter.ConvertDir(TempValue);
+					VertexInstanceTangents[AddedVertexInstanceId] = TangentX.GetSafeNormal();
 
-		if (bHasNTBInformation || RawMesh.WedgeTangentX.Num() > 0 || RawMesh.WedgeTangentY.Num() > 0)
-		{
-			RawMesh.WedgeTangentX.AddZeroed(WedgeOffset + WedgeCount - RawMesh.WedgeTangentX.Num());
-			RawMesh.WedgeTangentY.AddZeroed(WedgeOffset + WedgeCount - RawMesh.WedgeTangentY.Num());
-		}
+					int BinormalMapIndex = (BinormalMappingMode == FbxLayerElement::eByControlPoint) ?
+						ControlPointIndex : TriangleCornerIndex;
+					int BinormalValueIndex = (BinormalReferenceMode == FbxLayerElement::eDirect) ?
+						BinormalMapIndex : LayerElementBinormal->GetIndexArray().GetAt(BinormalMapIndex);
 
-		if (LayerElementNormal || RawMesh.WedgeTangentZ.Num() > 0 )
-		{
-			RawMesh.WedgeTangentZ.AddZeroed(WedgeOffset + WedgeCount - RawMesh.WedgeTangentZ.Num());
+					TempValue = LayerElementBinormal->GetDirectArray().GetAt(BinormalValueIndex);
+					TempValue = TotalMatrixForNormal.MultT(TempValue);
+					FVector TangentY = -Converter.ConvertDir(TempValue);
+					VertexInstanceBinormalSigns[AddedVertexInstanceId] = GetBasisDeterminantSign(TangentX.GetSafeNormal(), TangentY.GetSafeNormal(), TangentZ.GetSafeNormal());
+				}
+			}
 		}
-
-		if (LayerElementVertexColor || VertexColorImportOption != EVertexColorImportOption::Replace || RawMesh.WedgeColors.Num() )
+			
+		// Check if the polygon just discovered is non-degenerate if we haven't found one yet
+		//TODO check all polygon vertex, not just the first 3 vertex
+		if (!bHasNonDegeneratePolygons)
 		{
-			int32 NumNewColors = WedgeOffset + WedgeCount - RawMesh.WedgeColors.Num();
-			int32 FirstNewColor = RawMesh.WedgeColors.Num();
-			RawMesh.WedgeColors.AddUninitialized(NumNewColors);
-			for (int32 WedgeIndex = FirstNewColor; WedgeIndex < FirstNewColor + NumNewColors; ++WedgeIndex)
+			float TriangleComparisonThreshold = ImportOptions->bRemoveDegenerates ? THRESH_POINTS_ARE_SAME : 0.0f;
+			FVector VertexPosition[3];
+			VertexPosition[0] = VertexPositions[CornerVerticesIDs[0]];
+			VertexPosition[1] = VertexPositions[CornerVerticesIDs[1]];
+			VertexPosition[2] = VertexPositions[CornerVerticesIDs[2]];
+			if (!( VertexPosition[0].Equals(VertexPosition[1], TriangleComparisonThreshold)
+				|| VertexPosition[0].Equals(VertexPosition[2], TriangleComparisonThreshold)
+				|| VertexPosition[1].Equals(VertexPosition[2], TriangleComparisonThreshold)))
 			{
-				RawMesh.WedgeColors[WedgeIndex] = FColor::White;
+				bHasNonDegeneratePolygons = true;
 			}
 		}
 
-		// When importing multiple mesh pieces to the same static mesh.  Ensure each mesh piece has the same number of Uv's
-		int32 ExistingUVCount = 0;
-		for( int32 ExistingUVIndex = 0; ExistingUVIndex < MAX_MESH_TEXTURE_COORDS; ++ExistingUVIndex )
+		//
+		// material index
+		//
+		int32 MaterialIndex = 0;
+		if (MaterialCount > 0)
 		{
-			if( RawMesh.WedgeTexCoords[ExistingUVIndex].Num() > 0 )
+			if (LayerElementMaterial)
 			{
-				// Mesh already has UVs at this index
-				++ExistingUVCount;
-			}
-			else
-			{
-				// No more UVs
+				switch (MaterialMappingMode)
+				{
+					// material index is stored in the IndexArray, not the DirectArray (which is irrelevant with 2009.1)
+				case FbxLayerElement::eAllSame:
+				{
+					MaterialIndex = LayerElementMaterial->GetIndexArray().GetAt(0);
+				}
+				break;
+				case FbxLayerElement::eByPolygon:
+				{
+					MaterialIndex = LayerElementMaterial->GetIndexArray().GetAt(PolygonIndex);
+				}
 				break;
+				}
 			}
 		}
-	
-		int32 UVCount = FMath::Max( FBXUVs.UniqueUVCount, ExistingUVCount );
-
-		// At least one UV set must exist.  
-		UVCount = FMath::Max( 1, UVCount );
 
-		for (int32 UVLayerIndex = 0; UVLayerIndex<UVCount; UVLayerIndex++)
+		if (MaterialIndex >= MaterialCount || MaterialIndex < 0)
 		{
-			RawMesh.WedgeTexCoords[UVLayerIndex].AddZeroed(WedgeOffset + WedgeCount - RawMesh.WedgeTexCoords[UVLayerIndex].Num());
+			AddTokenizedErrorMessage(FTokenizedMessage::Create(EMessageSeverity::Warning, LOCTEXT("Error_MaterialIndexInconsistency", "Face material index inconsistency - forcing to 0")), FFbxErrors::Generic_Mesh_MaterialIndexInconsistency);
+			MaterialIndex = 0;
 		}
 
-		int32 TriangleIndex;
-		TMap<int32,int32> IndexMap;
-
-		for( TriangleIndex = 0 ; TriangleIndex < TriangleCount ; TriangleIndex++ )
+		//Create a polygon with the 3 vertex intances Add it to the material group
+		int32 RealMaterialIndex = MaterialIndexOffset + MaterialIndex;
+		if (!PolygonGroupMapping.Contains(RealMaterialIndex))
 		{
-			int32 DestTriangleIndex = TriangleOffset + TriangleIndex;
-			FVector CornerPositions[3];
-
-			for ( int32 CornerIndex=0; CornerIndex<3; CornerIndex++)
+			UMaterialInterface* Material = MeshMaterials.IsValidIndex(RealMaterialIndex) ? MeshMaterials[RealMaterialIndex].Material : UMaterial::GetDefaultMaterial(MD_Surface);
+			FName ImportedMaterialSlotName = MeshMaterials.IsValidIndex(RealMaterialIndex) ? FName(*MeshMaterials[RealMaterialIndex].GetName()) : NAME_None;
+			FPolygonGroupID ExistingPolygonGroup = FPolygonGroupID::Invalid;
+			for (const FPolygonGroupID PolygonGroupID : MeshDescription->PolygonGroups().GetElementIDs())
 			{
-				// If there are odd number negative scale, invert the vertex order for triangles
-				int32 WedgeIndex = WedgeOffset + TriangleIndex * 3 + (OddNegativeScale ? 2 - CornerIndex : CornerIndex);
-
-				// Store vertex index and position.
-				int32 ControlPointIndex = Mesh->GetPolygonVertex(TriangleIndex, CornerIndex);
-				int32* ExistingIndex = IndexMap.Find(ControlPointIndex);
-				if (ExistingIndex)
-				{
-					RawMesh.WedgeIndices[WedgeIndex] = *ExistingIndex;
-					CornerPositions[CornerIndex] = RawMesh.VertexPositions[*ExistingIndex];
-				}
-				else
-				{
-					FbxVector4 FbxPosition = Mesh->GetControlPoints()[ControlPointIndex];
-					FbxVector4 FinalPosition = TotalMatrix.MultT(FbxPosition);
-					int32 VertexIndex = RawMesh.VertexPositions.Add(Converter.ConvertPos(FinalPosition));
-					RawMesh.WedgeIndices[WedgeIndex] = VertexIndex;
-					IndexMap.Add(ControlPointIndex, VertexIndex);
-					CornerPositions[CornerIndex] = RawMesh.VertexPositions[VertexIndex];
-				}
-
-				//
-				// normals, tangents and binormals
-				//
-				if (LayerElementNormal)
-				{
-					int TriangleCornerIndex = TriangleIndex*3 + CornerIndex;
-					//normals may have different reference and mapping mode than tangents and binormals
-					int NormalMapIndex = (NormalMappingMode == FbxLayerElement::eByControlPoint) ? 
-						ControlPointIndex : TriangleCornerIndex;
-					int NormalValueIndex = (NormalReferenceMode == FbxLayerElement::eDirect) ? 
-						NormalMapIndex : LayerElementNormal->GetIndexArray().GetAt(NormalMapIndex);
-
-					//tangents and binormals share the same reference, mapping mode and index array
-					if (bHasNTBInformation)
-					{
-						int TangentMapIndex = (TangentMappingMode == FbxLayerElement::eByControlPoint) ? 
-							ControlPointIndex : TriangleCornerIndex;
-						int TangentValueIndex = (TangentReferenceMode == FbxLayerElement::eDirect) ? 
-							TangentMapIndex : LayerElementTangent->GetIndexArray().GetAt(TangentMapIndex);
-
-						FbxVector4 TempValue = LayerElementTangent->GetDirectArray().GetAt(TangentValueIndex);
-						TempValue = TotalMatrixForNormal.MultT(TempValue);
-						FVector TangentX = Converter.ConvertDir(TempValue);
-						RawMesh.WedgeTangentX[WedgeIndex] = TangentX.GetSafeNormal();
-
-						int BinormalMapIndex = (BinormalMappingMode == FbxLayerElement::eByControlPoint) ? 
-							ControlPointIndex : TriangleCornerIndex;
-						int BinormalValueIndex = (BinormalReferenceMode == FbxLayerElement::eDirect) ? 
-							BinormalMapIndex : LayerElementBinormal->GetIndexArray().GetAt(BinormalMapIndex);
-
-						TempValue = LayerElementBinormal->GetDirectArray().GetAt(BinormalValueIndex);
-						TempValue = TotalMatrixForNormal.MultT(TempValue);
-						FVector TangentY = -Converter.ConvertDir(TempValue);
-						RawMesh.WedgeTangentY[WedgeIndex] = TangentY.GetSafeNormal();
-					}
-
-					FbxVector4 TempValue = LayerElementNormal->GetDirectArray().GetAt(NormalValueIndex);
-					TempValue = TotalMatrixForNormal.MultT(TempValue);
-					FVector TangentZ = Converter.ConvertDir(TempValue);
-					RawMesh.WedgeTangentZ[WedgeIndex] = TangentZ.GetSafeNormal();
-				}
-
-				//
-				// vertex colors
-				//
-				if (VertexColorImportOption == EVertexColorImportOption::Replace)
-				{
-					if (LayerElementVertexColor)
-					{
-						int32 VertexColorMappingIndex = (VertexColorMappingMode == FbxLayerElement::eByControlPoint) ?
-							Mesh->GetPolygonVertex(TriangleIndex, CornerIndex) : (TriangleIndex * 3 + CornerIndex);
-
-						int32 VectorColorIndex = (VertexColorReferenceMode == FbxLayerElement::eDirect) ?
-						VertexColorMappingIndex : LayerElementVertexColor->GetIndexArray().GetAt(VertexColorMappingIndex);
-
-						FbxColor VertexColor = LayerElementVertexColor->GetDirectArray().GetAt(VectorColorIndex);
-
-						RawMesh.WedgeColors[WedgeIndex] = FColor(
-							uint8(255.f*VertexColor.mRed),
-							uint8(255.f*VertexColor.mGreen),
-							uint8(255.f*VertexColor.mBlue),
-							uint8(255.f*VertexColor.mAlpha)
-							);
-					}
-				}
-				else if (VertexColorImportOption == EVertexColorImportOption::Ignore)
-				{
-					// try to match this triangles current vertex with one that existed in the previous mesh.
-					// This is a find in a tmap which uses a fast hash table lookup.
-					FVector Position = RawMesh.VertexPositions[RawMesh.WedgeIndices[WedgeIndex]];
-					const FColor* PaintedColor = ExistingVertexColorData.Find(Position);
-					if (PaintedColor)
-					{
-						// A matching color for this vertex was found
-						RawMesh.WedgeColors[WedgeIndex] = *PaintedColor;
-					}
-				}
-				else
+				if (PolygonGroupImportedMaterialSlotNames[PolygonGroupID] == ImportedMaterialSlotName)
 				{
-					// set the triangle's vertex color to a constant override
-					check(VertexColorImportOption == EVertexColorImportOption::Override);
-					RawMesh.WedgeColors[WedgeIndex] = VertexOverrideColor;
+					ExistingPolygonGroup = PolygonGroupID;
+					break;
 				}
 			}
-
-			// Check if the triangle just discovered is non-degenerate if we haven't found one yet
-			if (!bHasNonDegenerateTriangles)
+			if (ExistingPolygonGroup == FPolygonGroupID::Invalid)
 			{
-				float ComparisonThreshold = ImportOptions->bRemoveDegenerates ? THRESH_POINTS_ARE_SAME : 0.0f;
-
-				if (!(CornerPositions[0].Equals(CornerPositions[1], ComparisonThreshold)
-					  || CornerPositions[0].Equals(CornerPositions[2], ComparisonThreshold)
-					  || CornerPositions[1].Equals(CornerPositions[2], ComparisonThreshold)))
-				{
-					bHasNonDegenerateTriangles = true;
-				}
+				ExistingPolygonGroup = MeshDescription->CreatePolygonGroup();
+				PolygonGroupImportedMaterialSlotNames[ExistingPolygonGroup] = ImportedMaterialSlotName;
 			}
+			PolygonGroupMapping.Add(RealMaterialIndex, ExistingPolygonGroup);
+		}
 
-			//
-			// smoothing mask
-			//
-			if (bSmoothingAvailable && SmoothingInfo)
+		// Create polygon edges
+		TArray<UMeshDescription::FContourPoint> Contours;
+		{
+			// Add the edges of this polygon
+			for (uint32 PolygonEdgeNumber = 0; PolygonEdgeNumber < (uint32)PolygonVertexCount; ++PolygonEdgeNumber)
 			{
-				if (SmoothingMappingMode == FbxLayerElement::eByPolygon)
+				int32 ContourPointIndex = Contours.AddDefaulted();
+				UMeshDescription::FContourPoint& ContourPoint = Contours[ContourPointIndex];
+				//Find the matching edge ID
+				uint32 CornerIndices[2];
+				CornerIndices[0] = (PolygonEdgeNumber + 0) % PolygonVertexCount;
+				CornerIndices[1] = (PolygonEdgeNumber + 1) % PolygonVertexCount;
+
+				FVertexID EdgeVertexIDs[2];
+				EdgeVertexIDs[0] = CornerVerticesIDs[CornerIndices[0]];
+				EdgeVertexIDs[1] = CornerVerticesIDs[CornerIndices[1]];
+
+				FEdgeID MatchEdgeId = MeshDescription->GetVertexPairEdge(EdgeVertexIDs[0], EdgeVertexIDs[1]);
+				if (MatchEdgeId == FEdgeID::Invalid)
 				{
-					int lSmoothingIndex = (SmoothingReferenceMode == FbxLayerElement::eDirect) ? TriangleIndex : SmoothingInfo->GetIndexArray().GetAt(TriangleIndex);
-					RawMesh.FaceSmoothingMasks[DestTriangleIndex] = SmoothingInfo->GetDirectArray().GetAt(lSmoothingIndex);
+					MatchEdgeId = MeshDescription->CreateEdge(EdgeVertexIDs[0], EdgeVertexIDs[1]);
 				}
-				else
+				ContourPoint.EdgeID = MatchEdgeId;
+				ContourPoint.VertexInstanceID = CornerInstanceIDs[CornerIndices[0]];
+				//RawMesh do not have edges, so by ordering the edge with the triangle construction we can ensure back and forth conversion with RawMesh
+				//When raw mesh will be completly remove we can create the edges right after the vertex creation.
+				int32 EdgeIndex = INDEX_NONE;
+				uint32 CompactedKey = ((uint32)EdgeVertexIDs[0].GetValue() << 16) | ((uint32)EdgeVertexIDs[1].GetValue() & 0xffff);
+				if (RemapEdgeID.Contains(CompactedKey))
 				{
-					AddTokenizedErrorMessage(FTokenizedMessage::Create(EMessageSeverity::Warning, FText::Format(LOCTEXT("Error_UnsupportedSmoothingGroup", "Unsupported Smoothing group mapping mode on mesh  '{0}'"), FText::FromString(Mesh->GetName()))), FFbxErrors::Generic_Mesh_UnsupportingSmoothingGroup);
+					EdgeIndex = RemapEdgeID[CompactedKey];
 				}
-			}
-
-			//
-			// uvs
-			//
-			// In FBX file, the same UV may be saved multiple times, i.e., there may be same UV in LayerElementUV
-			// So we don't import the duplicate UVs
-			int32 UVLayerIndex;
-			for (UVLayerIndex = 0; UVLayerIndex<FBXUVs.UniqueUVCount; UVLayerIndex++)
-			{
-				if (FBXUVs.LayerElementUV[UVLayerIndex] != NULL) 
+				else
 				{
-					for (int32 CornerIndex=0;CornerIndex<3;CornerIndex++)
-					{
-						// If there are odd number negative scale, invert the vertex order for triangles
-						int32 WedgeIndex = WedgeOffset + TriangleIndex * 3 + (OddNegativeScale ? 2 - CornerIndex : CornerIndex);
-
-
-						int lControlPointIndex = Mesh->GetPolygonVertex(TriangleIndex, CornerIndex);
-						int UVMapIndex = (FBXUVs.UVMappingMode[UVLayerIndex] == FbxLayerElement::eByControlPoint) ? lControlPointIndex : TriangleIndex*3+CornerIndex;
-						int32 UVIndex = (FBXUVs.UVReferenceMode[UVLayerIndex] == FbxLayerElement::eDirect) ? 
-							UVMapIndex : FBXUVs.LayerElementUV[UVLayerIndex]->GetIndexArray().GetAt(UVMapIndex);
-
-						FbxVector2	UVVector = FBXUVs.LayerElementUV[UVLayerIndex]->GetDirectArray().GetAt(UVIndex);
-
-						RawMesh.WedgeTexCoords[UVLayerIndex][WedgeIndex].X = static_cast<float>(UVVector[0]);
-						RawMesh.WedgeTexCoords[UVLayerIndex][WedgeIndex].Y = 1.f-static_cast<float>(UVVector[1]);   //flip the Y of UVs for DirectX
-					}
+					EdgeIndex = Mesh->GetMeshEdgeIndexForPolygon(PolygonIndex, PolygonEdgeNumber);
 				}
-			}
-
-			//
-			// material index
-			//
-			int32 MaterialIndex = 0;
-			if (MaterialCount>0)
-			{
-				if (LayerElementMaterial)
+					
+				EdgeCreaseSharpnesses[MatchEdgeId] = (float)Mesh->GetEdgeCreaseInfo(EdgeIndex);
+				if (!EdgeHardnesses[MatchEdgeId])
 				{
-					switch(MaterialMappingMode)
+					if (bSmoothingAvailable && SmoothingInfo)
 					{
-						// material index is stored in the IndexArray, not the DirectArray (which is irrelevant with 2009.1)
-					case FbxLayerElement::eAllSame:
-						{	
-							MaterialIndex = LayerElementMaterial->GetIndexArray().GetAt(0);
+						if (SmoothingMappingMode == FbxLayerElement::eByEdge)
+						{
+							int32 lSmoothingIndex = (SmoothingReferenceMode == FbxLayerElement::eDirect) ? EdgeIndex : SmoothingInfo->GetIndexArray().GetAt(EdgeIndex);
+							//Set the hard edges
+							EdgeHardnesses[MatchEdgeId] = (SmoothingInfo->GetDirectArray().GetAt(lSmoothingIndex) == 0);
 						}
-						break;
-					case FbxLayerElement::eByPolygon:
-						{	
-							MaterialIndex = LayerElementMaterial->GetIndexArray().GetAt(TriangleIndex);
+						else
+						{
+							AddTokenizedErrorMessage(FTokenizedMessage::Create(EMessageSeverity::Warning, FText::Format(LOCTEXT("Error_UnsupportedSmoothingGroup", "Unsupported Smoothing group mapping mode on mesh  '{0}'"), FText::FromString(Mesh->GetName()))), FFbxErrors::Generic_Mesh_UnsupportingSmoothingGroup);
 						}
-						break;
+					}
+					else
+					{
+						//When there is no smoothing group we set all edge to hard (faceted mesh)
+						EdgeHardnesses[MatchEdgeId] = true;
 					}
 				}
 			}
-			MaterialIndex += MaterialIndexOffset;
-
-			if (MaterialIndex >= MaterialCount + MaterialIndexOffset || MaterialIndex < 0)
-			{
-				AddTokenizedErrorMessage(FTokenizedMessage::Create(EMessageSeverity::Warning, LOCTEXT("Error_MaterialIndexInconsistency", "Face material index inconsistency - forcing to 0")), FFbxErrors::Generic_Mesh_MaterialIndexInconsistency);
-				MaterialIndex = 0;
-			}
-	
-			RawMesh.FaceMaterialIndices[DestTriangleIndex] = MaterialIndex;
 		}
-	
-		// needed?
-		FBXUVs.Cleanup();
+		FPolygonGroupID PolygonGroupID = PolygonGroupMapping[RealMaterialIndex];
+		// Insert a polygon into the mesh
+		const FPolygonID NewPolygonID = MeshDescription->CreatePolygon(PolygonGroupID, Contours);
+		//Triangulate the polygon
+		FMeshPolygon& Polygon = MeshDescription->GetPolygon(NewPolygonID);
+		MeshDescription->ComputePolygonTriangulation(NewPolygonID, Polygon.Triangles);
 	}
+	//Call this after all GetMeshEdgeIndexForPolygon call this is for optimization purpose.
+	Mesh->EndGetMeshEdgeIndexForPolygon();
+	// needed?
+	FBXUVs.Cleanup();
+	
 	
-	if (!bHasNonDegenerateTriangles)
+	if (!bHasNonDegeneratePolygons)
 	{
 		FFormatNamedArguments Arguments;
 		Arguments.Add( TEXT("MeshName"), FText::FromString(StaticMesh->GetName()));
-		FText ErrorMsg = LOCTEXT("MeshHasNoRenderableTriangles", "{MeshName} could not be created because all of its triangles are degenerate.");
+		FText ErrorMsg = LOCTEXT("MeshHasNoRenderableTriangles", "{MeshName} could not be created because all of its polygons are degenerate.");
 		AddTokenizedErrorMessage(FTokenizedMessage::Create(EMessageSeverity::Error, FText::Format(ErrorMsg, Arguments)), FFbxErrors::StaticMesh_AllTrianglesDegenerate);
 	}
 
-	bool bIsValidMesh = bHasNonDegenerateTriangles;
+	bool bIsValidMesh = bHasNonDegeneratePolygons;
 
 	return bIsValidMesh;
 }
@@ -1625,17 +1338,14 @@ UStaticMesh* UnFbx::FFbxImporter::ImportStaticMeshAsSingle(UObject* InParent, TA
 		}
 	}
 
-	if (GetDefault<UEditorExperimentalSettings>()->bUseMeshDescription)
+	MeshDescription = StaticMesh->GetOriginalMeshDescription(LODIndex);
+	if (MeshDescription == nullptr)
 	{
-		MeshDescription = StaticMesh->GetOriginalMeshDescription(LODIndex);
-		if (MeshDescription == nullptr)
-		{
-			StaticMesh->ClearOriginalMeshDescription(LODIndex);
-			//Create private asset in the same package as the StaticMesh, and make sure reference are set to avoid GC
-			MeshDescription = NewObject<UMeshDescription>(StaticMesh, NAME_None, RF_NoFlags);
-			check(MeshDescription != nullptr);
-			StaticMesh->SetOriginalMeshDescription(LODIndex, MeshDescription);
-		}
+		StaticMesh->ClearOriginalMeshDescription(LODIndex);
+		//Create private asset in the same package as the StaticMesh, and make sure reference are set to avoid GC
+		MeshDescription = NewObject<UMeshDescription>(StaticMesh, NAME_None, RF_NoFlags);
+		check(MeshDescription != nullptr);
+		StaticMesh->SetOriginalMeshDescription(LODIndex, MeshDescription);
 	}
 	
 	FStaticMeshSourceModel& SrcModel = StaticMesh->SourceModels[LODIndex];
@@ -1675,281 +1385,66 @@ UStaticMesh* UnFbx::FFbxImporter::ImportStaticMeshAsSingle(UObject* InParent, TA
 		}
 	}
 
-	if (!GetDefault<UEditorExperimentalSettings>()->bUseMeshDescription)
-	{
-		// Store the new raw mesh.
-		SrcModel.SaveRawMesh(NewRawMesh, false);
-	}
-
 	if (bBuildStatus)
 	{
-		int32 MaxMaterialIndex = 0;
-		int32 FirstOpenUVChannel = 1;
-		if (!GetDefault<UEditorExperimentalSettings>()->bUseMeshDescription)
-		{
-			UE_LOG(LogFbx, Verbose, TEXT("== Initial material list:"));
-			for (int32 MaterialIndex = 0; MaterialIndex < MeshMaterials.Num(); ++MaterialIndex)
-			{
-				UE_LOG(LogFbx, Verbose, TEXT("%d: %s"), MaterialIndex, *MeshMaterials[MaterialIndex].GetName());
-			}
-
-			// Compress the materials array by removing any duplicates.
-			bool bDoRemap = false;
-			TArray<int32> MaterialMap;
-			TArray<FFbxMaterial> UniqueMaterials;
-			for (int32 MaterialIndex = 0; MaterialIndex < MeshMaterials.Num(); ++MaterialIndex)
-			{
-				bool bUnique = true;
-				for (int32 OtherMaterialIndex = MaterialIndex - 1; OtherMaterialIndex >= 0; --OtherMaterialIndex)
-				{
-					if (MeshMaterials[MaterialIndex].FbxMaterial == MeshMaterials[OtherMaterialIndex].FbxMaterial &&
-						MeshMaterials[MaterialIndex].Material == MeshMaterials[OtherMaterialIndex].Material)
-					{
-						int32 UniqueIndex = MaterialMap[OtherMaterialIndex];
-
-						MaterialMap.Add(UniqueIndex);
-						bDoRemap = true;
-						bUnique = false;
-						break;
-					}
-				}
-				if (bUnique)
-				{
-					int32 UniqueIndex = UniqueMaterials.Add(MeshMaterials[MaterialIndex]);
-
-					MaterialMap.Add(UniqueIndex);
-				}
-				else
-				{
-					UE_LOG(LogFbx, Verbose, TEXT("  remap %d -> %d"), MaterialIndex, MaterialMap[MaterialIndex]);
-				}
-			}
-
-			if (UniqueMaterials.Num() > LARGE_MESH_MATERIAL_INDEX_THRESHOLD)
-			{
-				AddTokenizedErrorMessage(
-					FTokenizedMessage::Create(
-						EMessageSeverity::Warning,
-						FText::Format(LOCTEXT("Error_TooManyMaterials", "StaticMesh has a large number({0}) of materials and may render inefficently.  Consider breaking up the mesh into multiple Static Mesh Assets"),
-							FText::AsNumber(UniqueMaterials.Num())
-						)),
-					FFbxErrors::StaticMesh_TooManyMaterials);
-			}
-
-			//The fix is required for blender file. The static mesh build have change and now required that
-			//the sections (face declaration) must be declare in the same order as the material index
-			TArray<uint32> SortedMaterialIndex;
-			TArray<int32> UsedMaterials;
-			for (int32 FaceMaterialIndex = 0; FaceMaterialIndex < NewRawMesh.FaceMaterialIndices.Num(); ++FaceMaterialIndex)
-			{
-				int32 MaterialIndex = NewRawMesh.FaceMaterialIndices[FaceMaterialIndex];
-				if (!UsedMaterials.Contains(MaterialIndex))
-				{
-					int32 NewIndex = UsedMaterials.Add(MaterialIndex);
-					if (NewIndex != MaterialIndex)
-					{
-						bDoRemap = true;
-					}
-				}
-			}
-
-			for (int32 MaterialIndex = 0; MaterialIndex < MeshMaterials.Num(); ++MaterialIndex)
-			{
-				int32 SkinIndex = 0xffff;
-				if (bDoRemap)
-				{
-					int32 UsedIndex = 0;
-					if (UsedMaterials.Find(MaterialIndex, UsedIndex))
-					{
-						SkinIndex = UsedIndex;
-					}
-				}
-				int32 RemappedIndex = MaterialMap[MaterialIndex];
-				uint32 SortedMaterialKey = ((uint32)SkinIndex << 16) | ((uint32)RemappedIndex & 0xffff);
-				if (!SortedMaterialIndex.IsValidIndex(SortedMaterialKey))
-				{
-					SortedMaterialIndex.Add(SortedMaterialKey);
-				}
-			}
-
-			SortedMaterialIndex.Sort();
-
-
-			UE_LOG(LogFbx, Verbose, TEXT("== After sorting:"));
-			TArray<FFbxMaterial> SortedMaterials;
-			for (int32 SortedIndex = 0; SortedIndex < SortedMaterialIndex.Num(); ++SortedIndex)
-			{
-				int32 RemappedIndex = SortedMaterialIndex[SortedIndex] & 0xffff;
-				SortedMaterials.Add(UniqueMaterials[RemappedIndex]);
-				UE_LOG(LogFbx, Verbose, TEXT("%d: %s"), SortedIndex, *UniqueMaterials[RemappedIndex].GetName());
-			}
-			UE_LOG(LogFbx, Verbose, TEXT("== Mapping table:"));
-			for (int32 MaterialIndex = 0; MaterialIndex < MaterialMap.Num(); ++MaterialIndex)
-			{
-				for (int32 SortedIndex = 0; SortedIndex < SortedMaterialIndex.Num(); ++SortedIndex)
-				{
-					int32 RemappedIndex = SortedMaterialIndex[SortedIndex] & 0xffff;
-					if (MaterialMap[MaterialIndex] == RemappedIndex)
-					{
-						UE_LOG(LogFbx, Verbose, TEXT("  sort %d -> %d"), MaterialIndex, SortedIndex);
-						MaterialMap[MaterialIndex] = SortedIndex;
-						break;
-					}
-				}
-			}
+		TVertexInstanceAttributeIndicesArray<FVector2D>& VertexInstanceUVs = MeshDescription->VertexInstanceAttributes().GetAttributesSet<FVector2D>(MeshAttribute::VertexInstance::TextureCoordinate);
+		int32 FirstOpenUVChannel = VertexInstanceUVs.GetNumIndices();
+		TPolygonGroupAttributeArray<FName>& PolygonGroupImportedMaterialSlotNames = MeshDescription->PolygonGroupAttributes().GetAttributes<FName>(MeshAttribute::PolygonGroup::ImportedMaterialSlotName);
 
-			// Remap material indices.
+		TArray<FStaticMaterial> MaterialToAdd;
+		for (const FPolygonGroupID PolygonGroupID : MeshDescription->PolygonGroups().GetElementIDs())
+		{
+			const FName& ImportedMaterialSlotName = PolygonGroupImportedMaterialSlotNames[PolygonGroupID];
+			const FString ImportedMaterialSlotNameString = ImportedMaterialSlotName.ToString();
+			const FName MaterialSlotName = ImportedMaterialSlotName;
+			int32 MaterialIndex = INDEX_NONE;
+			for (int32 FbxMaterialIndex = 0; FbxMaterialIndex < MeshMaterials.Num(); ++FbxMaterialIndex)
 			{
-				FRawMesh LocalRawMesh;
-				SrcModel.RawMeshBulkData->LoadRawMesh(LocalRawMesh);
-
-				if (bDoRemap)
-				{
-					for (int32 TriIndex = 0; TriIndex < LocalRawMesh.FaceMaterialIndices.Num(); ++TriIndex)
-					{
-						LocalRawMesh.FaceMaterialIndices[TriIndex] = MaterialMap[LocalRawMesh.FaceMaterialIndices[TriIndex]];
-					}
-				}
-
-				// Compact material indices so that we won't have any sections with zero triangles.
-				LocalRawMesh.CompactMaterialIndices();
-
-				// Also compact the sorted materials array.
-				if (LocalRawMesh.MaterialIndexToImportIndex.Num() > 0)
-				{
-					TArray<FFbxMaterial> OldSortedMaterials;
-
-					Exchange(OldSortedMaterials, SortedMaterials);
-					SortedMaterials.Empty(LocalRawMesh.MaterialIndexToImportIndex.Num());
-					for (int32 MaterialIndex = 0; MaterialIndex < LocalRawMesh.MaterialIndexToImportIndex.Num(); ++MaterialIndex)
-					{
-						FFbxMaterial Material;
-						int32 ImportIndex = LocalRawMesh.MaterialIndexToImportIndex[MaterialIndex];
-						if (OldSortedMaterials.IsValidIndex(ImportIndex))
-						{
-							Material = OldSortedMaterials[ImportIndex];
-						}
-						SortedMaterials.Add(Material);
-					}
-				}
-
-				for (int32 TriIndex = 0; TriIndex < LocalRawMesh.FaceMaterialIndices.Num(); ++TriIndex)
-				{
-					MaxMaterialIndex = FMath::Max<int32>(MaxMaterialIndex, LocalRawMesh.FaceMaterialIndices[TriIndex]);
-				}
-
-				for (int32 i = 0; i < MAX_MESH_TEXTURE_COORDS; i++)
+				FFbxMaterial& FbxMaterial = MeshMaterials[FbxMaterialIndex];
+				if (FbxMaterial.GetName().Equals(ImportedMaterialSlotNameString))
 				{
-					if (LocalRawMesh.WedgeTexCoords[i].Num() == 0)
-					{
-						FirstOpenUVChannel = i;
-						break;
-					}
+					MaterialIndex = FbxMaterialIndex;
+					break;
 				}
-
-				SrcModel.SaveRawMesh(LocalRawMesh, false);
 			}
-
-			// Setup per-section info and the materials array.
-			if (LODIndex == 0)
+			if (MaterialIndex == INDEX_NONE)
 			{
-				StaticMesh->StaticMaterials.Empty();
+				MaterialIndex = PolygonGroupID.GetValue();
 			}
-
-			// Replace map of sections with the unique material set
-			int32 NumMaterials = FMath::Min(SortedMaterials.Num(), MaxMaterialIndex + 1);
-			for (int32 MaterialIndex = 0; MaterialIndex < NumMaterials; ++MaterialIndex)
+			UMaterialInterface* Material = MeshMaterials.IsValidIndex(MaterialIndex) ? MeshMaterials[MaterialIndex].Material : UMaterial::GetDefaultMaterial(MD_Surface);
+			FStaticMaterial StaticMaterial(Material, MaterialSlotName, ImportedMaterialSlotName);
+			(LODIndex > 0) ? MaterialToAdd.Add(StaticMaterial) : StaticMesh->StaticMaterials.Add(StaticMaterial);
+		}
+		if (LODIndex > 0)
+		{
+			//Insert the new materials in the static mesh
+			//The build function will search for imported slot name to find the appropriate slot
+			int32 StaticMeshMaterialCount = StaticMesh->StaticMaterials.Num();
+			if (StaticMeshMaterialCount > 0)
 			{
-				FMeshSectionInfo Info = StaticMesh->SectionInfoMap.Get(LODIndex, MaterialIndex);
-
-				int32 Index = 0;
-
-				FName MaterialFName = FName(*(SortedMaterials[MaterialIndex].GetName()));
-				FString CleanMaterialSlotName = MaterialFName.ToString();
-				int32 SkinOffset = CleanMaterialSlotName.Find(TEXT("_skin"));
-				if (SkinOffset != INDEX_NONE)
-				{
-					CleanMaterialSlotName = CleanMaterialSlotName.LeftChop(CleanMaterialSlotName.Len() - SkinOffset);
-				}
-
-				if (InStaticMesh)
+				for (int32 MaterialToAddIndex = 0; MaterialToAddIndex < MaterialToAdd.Num(); ++MaterialToAddIndex)
 				{
-					Index = INDEX_NONE;
-					FStaticMaterial StaticMaterialImported(SortedMaterials[MaterialIndex].Material, FName(*CleanMaterialSlotName), MaterialFName);
-					for (int32 OriginalMaterialIndex = 0; OriginalMaterialIndex < InStaticMesh->StaticMaterials.Num(); ++OriginalMaterialIndex)
+					const FStaticMaterial& CandidateMaterial = MaterialToAdd[MaterialToAddIndex];
+					bool FoundExistingMaterial = false;
+					//Found matching existing material
+					for (int32 StaticMeshMaterialIndex = 0; StaticMeshMaterialIndex < StaticMeshMaterialCount; ++StaticMeshMaterialIndex)
 					{
-						if (InStaticMesh->StaticMaterials[OriginalMaterialIndex] == StaticMaterialImported)
+						const FStaticMaterial& StaticMeshMaterial = StaticMesh->StaticMaterials[StaticMeshMaterialIndex];
+						if (StaticMeshMaterial.MaterialInterface == CandidateMaterial.MaterialInterface)
 						{
-							Index = OriginalMaterialIndex;
+							FoundExistingMaterial = true;
 							break;
 						}
 					}
-					if (Index == INDEX_NONE || (Index >= NumMaterials && Index >= InStaticMesh->StaticMaterials.Num()))
-					{
-						Index = StaticMesh->StaticMaterials.Add(FStaticMaterial(SortedMaterials[MaterialIndex].Material, FName(*CleanMaterialSlotName), MaterialFName));
-					}
-				}
-				else
-				{
-
-					Index = StaticMesh->StaticMaterials.Add(FStaticMaterial(SortedMaterials[MaterialIndex].Material, FName(*CleanMaterialSlotName), MaterialFName));
-				}
-
-				Info.MaterialIndex = Index;
-				StaticMesh->SectionInfoMap.Remove(LODIndex, MaterialIndex);
-				StaticMesh->SectionInfoMap.Set(LODIndex, MaterialIndex, Info);
-			}
-
-			FRawMesh LocalRawMesh;
-			SrcModel.RawMeshBulkData->LoadRawMesh(LocalRawMesh);
-			//This create the MeshDescription
-			SrcModel.SaveRawMesh(LocalRawMesh);
-		}
-		else
-		{
-			TPolygonGroupAttributeArray<FName>& PolygonGroupImportedMaterialSlotNames = MeshDescription->PolygonGroupAttributes().GetAttributes<FName>(MeshAttribute::PolygonGroup::ImportedMaterialSlotName);
-
-			TArray<FStaticMaterial> MaterialToAdd;
-			for (const FPolygonGroupID PolygonGroupID : MeshDescription->PolygonGroups().GetElementIDs())
-			{
-				int32 MaterialIndex = PolygonGroupID.GetValue();
-				UMaterialInterface* Material = MeshMaterials.IsValidIndex(MaterialIndex) ? MeshMaterials[MaterialIndex].Material : UMaterial::GetDefaultMaterial(MD_Surface);
-				const FName& ImportedMaterialSlotName = PolygonGroupImportedMaterialSlotNames[PolygonGroupID];
-				const FName MaterialSlotName = ImportedMaterialSlotName;
-				FStaticMaterial StaticMaterial(Material, MaterialSlotName, ImportedMaterialSlotName);
-				int32 AddedIndex = (LODIndex > 0) ? MaterialToAdd.Add(StaticMaterial) : StaticMesh->StaticMaterials.Add(StaticMaterial);
-			}
-			if (LODIndex > 0)
-			{
-				//Insert the new materials in the static mesh
-				//The build function will search for imported slot name to find the appropriate slot
-				int32 StaticMeshMaterialCount = StaticMesh->StaticMaterials.Num();
-				if (StaticMeshMaterialCount > 0)
-				{
-					for (int32 MaterialToAddIndex = 0; MaterialToAddIndex < MaterialToAdd.Num(); ++MaterialToAddIndex)
+					if (!FoundExistingMaterial)
 					{
-						const FStaticMaterial& CandidateMaterial = MaterialToAdd[MaterialToAddIndex];
-						bool FoundExistingMaterial = false;
-						//Found matching existing material
-						for (int32 StaticMeshMaterialIndex = 0; StaticMeshMaterialIndex < StaticMeshMaterialCount; ++StaticMeshMaterialIndex)
-						{
-							const FStaticMaterial& StaticMeshMaterial = StaticMesh->StaticMaterials[StaticMeshMaterialIndex];
-							if (StaticMeshMaterial.MaterialInterface == CandidateMaterial.MaterialInterface)
-							{
-								FoundExistingMaterial = true;
-								break;
-							}
-						}
-						if (!FoundExistingMaterial)
-						{
-							StaticMesh->StaticMaterials.Add(CandidateMaterial);
-						}
+						StaticMesh->StaticMaterials.Add(CandidateMaterial);
 					}
 				}
 			}
-			//Set the original mesh description to be able to do non destructive reduce
-			StaticMesh->SetOriginalMeshDescription(LODIndex, MeshDescription);
 		}
+		//Set the original mesh description to be able to do non destructive reduce
+		StaticMesh->SetOriginalMeshDescription(LODIndex, MeshDescription);
 
 		// Setup default LOD settings based on the selected LOD group.
 		if (LODIndex == 0)
diff --git a/Engine/Source/Editor/UnrealEd/Private/GeomFitUtils.cpp b/Engine/Source/Editor/UnrealEd/Private/GeomFitUtils.cpp
--- a/Engine/Source/Editor/UnrealEd/Private/GeomFitUtils.cpp
+++ b/Engine/Source/Editor/UnrealEd/Private/GeomFitUtils.cpp
@@ -219,19 +219,9 @@ void ComputeBoundingBox(UStaticMesh* StaticMesh, FVector& Center, FVector& Exten
 	// Calculate bounding Box.
 	
 	FStaticMeshSourceModel& SrcModel = StaticMesh->SourceModels[0];
-	if (GetDefault<UEditorExperimentalSettings>()->bUseMeshDescription)
-	{
-		UMeshDescription* MeshDescription = StaticMesh->GetOriginalMeshDescription(0);
-		FVector unitVec = FVector(1.f);
-		CalcBoundingBox(MeshDescription, Center, Extents, unitVec);
-	}
-	else
-	{
-		FRawMesh RawMesh;
-		SrcModel.LoadRawMesh(RawMesh);
-		FVector unitVec = FVector(1.f);
-		CalcBoundingBox(RawMesh, Center, Extents, unitVec);
-	}
+	UMeshDescription* MeshDescription = StaticMesh->GetOriginalMeshDescription(0);
+	FVector unitVec = FVector(1.f);
+	CalcBoundingBox(MeshDescription, Center, Extents, unitVec);
 }
 
 int32 GenerateBoxAsSimpleCollision(UStaticMesh* StaticMesh)
@@ -250,17 +240,8 @@ int32 GenerateBoxAsSimpleCollision(UStaticMesh* StaticMesh)
 	
 	FVector unitVec = bs->BuildScale3D;
 	FVector Center, Extents;
-	if (GetDefault<UEditorExperimentalSettings>()->bUseMeshDescription)
-	{
-		MeshDescription = StaticMesh->GetMeshDescription(0);
-		CalcBoundingBox(MeshDescription, Center, Extents, unitVec);
-	}
-	else
-	{
-		SrcModel.LoadRawMesh(RawMesh);
-		CalcBoundingBox(RawMesh, Center, Extents, unitVec);
-	}
-
+	MeshDescription = StaticMesh->GetMeshDescription(0);
+	CalcBoundingBox(MeshDescription, Center, Extents, unitVec);
 	bs->Modify();
 
 	// Create new GUID
@@ -561,20 +542,9 @@ int32 GenerateSphereAsSimpleCollision(UStaticMesh* StaticMesh)
 	FVector unitVec = bs->BuildScale3D;
 
 	// Calculate bounding sphere.
-	if (GetDefault<UEditorExperimentalSettings>()->bUseMeshDescription)
-	{
-		UMeshDescription* MeshDescription = StaticMesh->GetMeshDescription(0);
-		CalcBoundingSphere(MeshDescription, bSphere, unitVec);
-		CalcBoundingSphere2(MeshDescription, bSphere2, unitVec);
-	}
-	else
-	{
-		FRawMesh RawMesh;
-		FStaticMeshSourceModel& SrcModel = StaticMesh->SourceModels[0];
-		SrcModel.LoadRawMesh(RawMesh);
-		CalcBoundingSphere(RawMesh, bSphere, unitVec);
-		CalcBoundingSphere2(RawMesh, bSphere2, unitVec);
-	}
+	UMeshDescription* MeshDescription = StaticMesh->GetMeshDescription(0);
+	CalcBoundingSphere(MeshDescription, bSphere, unitVec);
+	CalcBoundingSphere2(MeshDescription, bSphere2, unitVec);
 
 	if(bSphere.W < bSphere2.W)
 		bestSphere = bSphere;
@@ -797,18 +767,8 @@ int32 GenerateSphylAsSimpleCollision(UStaticMesh* StaticMesh)
 	FVector unitVec = bs->BuildScale3D;
 
 	// Calculate bounding box.
-	if (GetDefault<UEditorExperimentalSettings>()->bUseMeshDescription)
-	{
-		UMeshDescription* MeshDescription = StaticMesh->GetMeshDescription(0);
-		CalcBoundingSphyl(MeshDescription, sphere, length, rotation, unitVec);
-	}
-	else
-	{
-		FRawMesh RawMesh;
-		FStaticMeshSourceModel& SrcModel = StaticMesh->SourceModels[0];
-		SrcModel.LoadRawMesh(RawMesh);
-		CalcBoundingSphyl(RawMesh, sphere, length, rotation, unitVec);
-	}
+	UMeshDescription* MeshDescription = StaticMesh->GetMeshDescription(0);
+	CalcBoundingSphyl(MeshDescription, sphere, length, rotation, unitVec);
 
 	// Dont use if radius is zero.
 	if (sphere.W <= 0.f)
diff --git a/Engine/Source/Editor/UnrealEd/Private/StaticMeshEdit.cpp b/Engine/Source/Editor/UnrealEd/Private/StaticMeshEdit.cpp
--- a/Engine/Source/Editor/UnrealEd/Private/StaticMeshEdit.cpp
+++ b/Engine/Source/Editor/UnrealEd/Private/StaticMeshEdit.cpp
@@ -928,20 +928,10 @@ UStaticMesh* CreateStaticMeshFromBrush(UObject* Outer,FName Name,ABrush* Brush,U
 {
 	GWarn->BeginSlowTask( NSLOCTEXT("UnrealEd", "CreatingStaticMeshE", "Creating static mesh..."), true );
 	UStaticMesh* StaticMesh = nullptr;
-	if (GetDefault<UEditorExperimentalSettings>()->bUseMeshDescription)
-	{
-		UMeshDescription* MeshDescription = NewObject<UMeshDescription>(Outer);
-		TArray<FStaticMaterial> Materials;
-		GetBrushMesh(Brush, Model, MeshDescription, Materials);
-		StaticMesh = CreateStaticMesh(MeshDescription, Materials, Outer, Name);
-	}
-	else
-	{
-		FRawMesh RawMesh;
-		TArray<FStaticMaterial> Materials;
-		GetBrushMesh(Brush, Model, RawMesh, Materials);
-		StaticMesh = CreateStaticMesh(RawMesh, Materials, Outer, Name);
-	}
+	UMeshDescription* MeshDescription = NewObject<UMeshDescription>(Outer);
+	TArray<FStaticMaterial> Materials;
+	GetBrushMesh(Brush, Model, MeshDescription, Materials);
+	StaticMesh = CreateStaticMesh(MeshDescription, Materials, Outer, Name);
 	GWarn->EndSlowTask();
 
 	return StaticMesh;
diff --git a/Engine/Source/Runtime/Engine/Classes/Engine/StaticMesh.h b/Engine/Source/Runtime/Engine/Classes/Engine/StaticMesh.h
--- a/Engine/Source/Runtime/Engine/Classes/Engine/StaticMesh.h
+++ b/Engine/Source/Runtime/Engine/Classes/Engine/StaticMesh.h
@@ -164,8 +164,7 @@ struct FStaticMeshSourceModel
 	class UStaticMesh* StaticMeshOwner;
 	/*
 	 * Accessor to Load and save the raw mesh or the mesh description depending on the editor settings.
-	 * Temporary until we deprecate the RawMesh. It use UEditorExperimentalSettings::bUseMeshDescription to know which structure we use to save/load the data.
-	 * The SaveRawMesh make sure there is only one valid structure data (OriginalMeshDescription or RawMesh) not both. 
+	 * Temporary until we deprecate the RawMesh.
 	 */
 	ENGINE_API bool IsRawMeshEmpty() const;
 	ENGINE_API void LoadRawMesh(struct FRawMesh& OutRawMesh) const;
diff --git a/Engine/Source/Runtime/Engine/Private/StaticMesh.cpp b/Engine/Source/Runtime/Engine/Private/StaticMesh.cpp
--- a/Engine/Source/Runtime/Engine/Private/StaticMesh.cpp
+++ b/Engine/Source/Runtime/Engine/Private/StaticMesh.cpp
@@ -1387,15 +1387,8 @@ static FString BuildStaticMeshDerivedDataKey(UStaticMesh* Mesh, const FStaticMes
 	for (int32 LODIndex = 0; LODIndex < NumLODs; ++LODIndex)
 	{
 		FStaticMeshSourceModel& SrcModel = Mesh->SourceModels[LODIndex];
-		if (GetDefault<UEditorExperimentalSettings>()->bUseMeshDescription)
-		{
-			UMeshDescription* MeshDescription = Mesh->GetOriginalMeshDescription(LODIndex);
-			KeySuffix += (MeshDescription != nullptr) ? MeshDescription->GetIdString() : TEXT("NoOriginalMeshDescriptionForLod") + FString::FromInt(LODIndex);
-		}
-		else
-		{
-			KeySuffix += SrcModel.RawMeshBulkData->GetIdString();
-		}
+		UMeshDescription* MeshDescription = Mesh->GetOriginalMeshDescription(LODIndex);
+		KeySuffix += (MeshDescription != nullptr) ? MeshDescription->GetIdString() : TEXT("NoOriginalMeshDescriptionForLod") + FString::FromInt(LODIndex);
 
 		// Serialize the build and reduction settings into a temporary array. The archive
 		// is flagged as persistent so that machines of different endianness produce
@@ -1572,12 +1565,9 @@ void FStaticMeshRenderData::Cache(UStaticMesh* Owner, const FStaticMeshLODSettin
 			FStaticMeshStatusMessageContext StatusContext( FText::Format( NSLOCTEXT("Engine", "BuildingStaticMeshStatus", "Building static mesh {StaticMeshName}..."), Args ) );
 
 			bool bUseMeshDescription = false;
-			if (GetDefault<UEditorExperimentalSettings>()->bUseMeshDescription)
+			if (Owner->GetOriginalMeshDescription(0) != nullptr)
 			{
-				if (Owner->GetOriginalMeshDescription(0) != nullptr)
-				{
-					bUseMeshDescription = true;
-				}
+				bUseMeshDescription = true;
 			}
 
 			if (bUseMeshDescription)
