When converting BSP geometry to static mesh, convert null materials into the default surface material.

diff --git a/Engine/Source/Editor/UnrealEd/Private/StaticMeshEdit.cpp b/Engine/Source/Editor/UnrealEd/Private/StaticMeshEdit.cpp
--- a/Engine/Source/Editor/UnrealEd/Private/StaticMeshEdit.cpp
+++ b/Engine/Source/Editor/UnrealEd/Private/StaticMeshEdit.cpp
@@ -739,6 +739,11 @@ void GetBrushMesh(ABrush* Brush, UModel* Model, UMeshDescription* MeshDescriptio
 
 		// Find a material index for this polygon.
 		UMaterialInterface*	Material = Polygon.Material;
+		if (Material == nullptr)
+		{
+			Material = UMaterial::GetDefaultMaterial(MD_Surface);
+		}
+
 		int32 MaterialIndex = OutMaterials.AddUnique(FStaticMaterial(Material, Material->GetFName(), Material->GetFName()));
 		FPolygonGroupID CurrentPolygonGroupID = FPolygonGroupID::Invalid;
 		for (const FPolygonGroupID& PolygonGroupID : MeshDescription->PolygonGroups().GetElementIDs())
