Merging //UE4/Dev-Main to Dev-Geometry (//UE4/Dev-Geometry)

diff --git a/Engine/Source/Editor/SourceControlWindows/Private/SSourceControlHistory.cpp b/Engine/Source/Editor/SourceControlWindows/Private/SSourceControlHistory.cpp
--- a/Engine/Source/Editor/SourceControlWindows/Private/SSourceControlHistory.cpp
+++ b/Engine/Source/Editor/SourceControlWindows/Private/SSourceControlHistory.cpp
@@ -1435,7 +1435,35 @@ void FSourceControlWindows::DisplayRevisionHistory( const TArray<FString>& InPac
 	if(SourceControlProvider.Execute(UpdateStatusOperation, PackageFilenames))
 	{
 		TArray< FSourceControlStateRef > SourceControlStates;
-		SourceControlProvider.GetState(PackageFilenames, SourceControlStates, EStateCacheUsage::Use);
+		for(const FString& PackageFilename : PackageFilenames)
+		{
+			TArray<FString> RevisionName;
+			RevisionName.Add(PackageFilename);
+
+			while(RevisionName.Num() != 0)
+			{
+				int32 InitialNum = SourceControlStates.Num();
+				SourceControlProvider.GetState(RevisionName, SourceControlStates, EStateCacheUsage::Use);
+				int32 NewNum = SourceControlStates.Num();
+				ensure(NewNum >= InitialNum);
+				
+				RevisionName.Empty();
+				// check to see if origin of this file is a branch, append the history from the branch point:
+				if(NewNum > InitialNum)
+				{
+					int32 HistorySize = SourceControlStates.Last()->GetHistorySize();
+					if( HistorySize > 0 )
+					{
+						TSharedPtr<class ISourceControlRevision, ESPMode::ThreadSafe> InitialHistory = SourceControlStates.Last()->GetHistoryItem(HistorySize - 1);
+						TSharedPtr<ISourceControlRevision, ESPMode::ThreadSafe> BranchSource = InitialHistory->GetBranchSource();
+						if( BranchSource.IsValid() )
+						{
+							RevisionName.Add(BranchSource->GetFilename());
+						}
+					}
+				}
+			}
+		}
 
 		TSharedRef<SWindow> NewWindow = SNew(SWindow)
 			.Title( NSLOCTEXT("SourceControl.HistoryWindow", "Title", "File History") )
diff --git a/Engine/Source/Editor/StaticMeshEditor/Private/StaticMeshEditorViewportClient.cpp b/Engine/Source/Editor/StaticMeshEditor/Private/StaticMeshEditorViewportClient.cpp
--- a/Engine/Source/Editor/StaticMeshEditor/Private/StaticMeshEditorViewportClient.cpp
+++ b/Engine/Source/Editor/StaticMeshEditor/Private/StaticMeshEditorViewportClient.cpp
@@ -23,7 +23,7 @@
 
 #include "Runtime/Analytics/Analytics/Public/Interfaces/IAnalyticsProvider.h"
 #include "EngineAnalytics.h"
-#include "AI/Navigation/NavCollision.h"
+#include "AI/Navigation/NavCollisionBase.h"
 #include "PhysicsEngine/BodySetup.h"
 
 #include "Engine/AssetUserData.h"
@@ -616,8 +616,8 @@ void FStaticMeshEditorViewportClient::Draw(const FSceneView* View,FPrimitiveDraw
 			}
 		}
 
-		// The simple nav geometry is only used by dynamic obstacles for now
-		if (StaticMesh->NavCollision && StaticMesh->NavCollision->bIsDynamicObstacle)
+		if (StaticMesh->NavCollision 
+			&& StaticMesh->bHasNavigationData)
 		{
 			// Draw the static mesh's body setup (simple collision)
 			FTransform GeomTransform(StaticMeshComponent->GetComponentTransform());
diff --git a/Engine/Source/Editor/UnrealEd/Classes/Settings/EditorExperimentalSettings.h b/Engine/Source/Editor/UnrealEd/Classes/Settings/EditorExperimentalSettings.h
--- a/Engine/Source/Editor/UnrealEd/Classes/Settings/EditorExperimentalSettings.h
+++ b/Engine/Source/Editor/UnrealEd/Classes/Settings/EditorExperimentalSettings.h
@@ -63,6 +63,10 @@ public:
 	UPROPERTY(EditAnywhere, config, Category = Blueprints, meta = (DisplayName = "Find and Replace All References Tool"))
 	bool bEnableFindAndReplaceReferences;
 
+	/** Bookmarks facilitate quick-jump navigation to saved locations in any Blueprint Graph context in the Blueprint Editor. */
+	UPROPERTY(EditAnywhere, config, Category = Blueprints)
+	bool bEnableBlueprintGraphBookmarks;
+
 	/** Should arrows indicating data/execution flow be drawn halfway along wires? */
 	UPROPERTY(/*EditAnywhere - deprecated (moved into UBlueprintEditorSettings), */config/*, Category=Blueprints, meta=(DisplayName="Draw midpoint arrows in Blueprints")*/)
 	bool bDrawMidpointArrowsInBlueprints;
diff --git a/Engine/Source/Editor/UnrealEd/Private/Commandlets/GenerateBlueprintAPICommandlet.cpp b/Engine/Source/Editor/UnrealEd/Private/Commandlets/GenerateBlueprintAPICommandlet.cpp
--- a/Engine/Source/Editor/UnrealEd/Private/Commandlets/GenerateBlueprintAPICommandlet.cpp
+++ b/Engine/Source/Editor/UnrealEd/Private/Commandlets/GenerateBlueprintAPICommandlet.cpp
@@ -58,6 +58,14 @@ GenerateBlueprintAPI commandlet params: \n\
     -class=<Class>      Used to specify the blueprint's parent class, if left  \n\
                         unset then it will go through EVERY blueprint parent   \n\
                         class available.                                       \n\
+\n\
+    -category=<CategoryPath1, CategoryPath2, ...>                              \n\
+                        Used to specify the blueprint's category path, if      \n\
+                        left unset then it will go through EVERY category      \n\
+                        available. A category path can contain a single        \n\
+                        category, or multiple category names concanated by     \n\
+                        \'|\'. Multiple category paths can be supplied by using\n\
+                        \',\' as separators.                                   \n\
 \n\
     -multifile          Used to keep file size down, will split each blueprint \n\
                         into its own file (meaning only one file will be       \n\
@@ -75,7 +83,7 @@ GenerateBlueprintAPI commandlet params: \n\
                         (as it has been time sync in the past). This is not    \n\
                         ideal for diffs though (since times can easily vary).  \n\
 \n\
-    -experimental       Uses an new way of constructing Blueprint action menus \n\
+    -experimental       Uses a new way of constructing Blueprint action menus \n\
                         (that will replace the current system).                \n\
 \n\
     -name=<Filename>    Overrides the default filename. Leave off the extention\n\
@@ -121,6 +129,7 @@ GenerateBlueprintAPI commandlet params: \n\
 		UClass*    PaletteFilter;
 		FString    SaveDir;
 		FString    Filename;
+		TArray<FString> CategoryFilter;
 	};
 
 	/** Static instance of the command switches (so we don't have to pass one along the call stack) */
@@ -222,6 +231,11 @@ GenerateBlueprintAPI commandlet params: \n\
 	 */
 	static void DumpActionList(uint32 Indent, FGraphActionListBuilderBase& ActionList, FArchive* FileOutWriter);
 
+	/**
+	* Get the category information from a given action.
+	*/
+	static FString GetActionCategory(FGraphActionListBuilderBase::ActionGroup const& Action);
+
 	/**
 	 * Generic function that dumps information on a single action (like it's 
 	 * name, category, an associated node if it has one, etc.).
@@ -295,6 +309,34 @@ GenerateBlueprintAPIUtils::CommandletOptions::CommandletOptions(TArray<FString>
 				}
 			}
 		}
+		else if (Switch.StartsWith("category=", ESearchCase::IgnoreCase))
+		{
+			FString ClassSwitch, Filters;
+			Switch.Split(TEXT("="), &ClassSwitch, &Filters, ESearchCase::IgnoreCase);
+
+			// Remove surounding " or ' if exist.
+			Filters = Filters.Replace(TEXT("\""), TEXT(""), ESearchCase::IgnoreCase);
+			Filters = Filters.Replace(TEXT("\'"), TEXT(""), ESearchCase::IgnoreCase);
+			Filters.TrimStartAndEndInline();
+
+			// Parse multiple category paths.
+			while (Filters.Len() > 0)
+			{
+				FString FilterName, OtherFilters;
+				if (Filters.Split(TEXT(","), &FilterName, &OtherFilters, ESearchCase::IgnoreCase))
+				{
+					FilterName.TrimStartAndEndInline();
+					CategoryFilter.Add(FilterName);
+					Filters = OtherFilters;
+				}
+				else
+				{
+					Filters.TrimStartAndEndInline();
+					CategoryFilter.Add(Filters);
+					break;
+				}
+			}
+		}
 		else if (!Switch.Compare("h", ESearchCase::IgnoreCase) || 
 		         !Switch.Compare("?", ESearchCase::IgnoreCase) || 
 				 !Switch.Compare("help", ESearchCase::IgnoreCase))
@@ -726,6 +768,9 @@ static void GenerateBlueprintAPIUtils::DumpPalette(uint32 Indent, UBlueprint* Bl
 static void GenerateBlueprintAPIUtils::DumpActionList(uint32 Indent, FGraphActionListBuilderBase& ActionList, FArchive* FileOutWriter)
 {
 	TArray< FGraphActionListBuilderBase::ActionGroup const* > SortedActions;
+	const TArray<FString>& CategoryFilter = GenerateBlueprintAPIUtils::CommandOptions.CategoryFilter;
+	bool HasCategoryFilter = CategoryFilter.Num() > 0;
+
 	for (int32 ActionIndex = 0; ActionIndex < ActionList.GetNumActions(); ++ActionIndex)
 	{
 		FGraphActionListBuilderBase::ActionGroup& Action = ActionList.GetAction(ActionIndex);
@@ -734,7 +779,22 @@ static void GenerateBlueprintAPIUtils::DumpActionList(uint32 Indent, FGraphActio
 			continue;
 		}
 
-		SortedActions.Add(&Action);
+		if (!HasCategoryFilter)
+		{
+			SortedActions.Add(&Action);
+		}
+		else
+		{
+			const FString& ActionCategory = GenerateBlueprintAPIUtils::GetActionCategory(Action);
+			for (const FString& filter : CategoryFilter)
+			{
+				if (ActionCategory.Contains(filter, ESearchCase::CaseSensitive, ESearchDir::FromStart))
+				{
+					SortedActions.Add(&Action);
+					break;
+				}
+			}
+		}
 	}
 
 	FString const ActionListIndent = BuildIndentString(Indent);
@@ -778,16 +838,15 @@ static void GenerateBlueprintAPIUtils::DumpActionList(uint32 Indent, FGraphActio
 	EndActionListEntry += "\n" + ActionListIndent + "}";
 	FileOutWriter->Serialize(TCHAR_TO_ANSI(*EndActionListEntry), EndActionListEntry.Len());
 }
-
 //------------------------------------------------------------------------------
-static void GenerateBlueprintAPIUtils::DumpActionMenuItem(uint32 Indent, FGraphActionListBuilderBase::ActionGroup const& Action, FGraphActionListBuilderBase& ActionList, FArchive* FileOutWriter)
+static FString GenerateBlueprintAPIUtils::GetActionCategory(FGraphActionListBuilderBase::ActionGroup const& Action)
 {
 	check(Action.Actions.Num() > 0);
 
 	// Get action category info
 	const TArray<FString>& MenuHierarchy = Action.GetCategoryChain();
 
-	FString ActionCategory = TEXT("");
+	FString ActionCategory;
 
 	bool bHasCategory = (MenuHierarchy.Num() > 0);
 	if (bHasCategory)
@@ -797,7 +856,12 @@ static void GenerateBlueprintAPIUtils::DumpActionMenuItem(uint32 Indent, FGraphA
 			ActionCategory += SubCategory + TEXT("|");
 		}
 	}
-
+	return ActionCategory;
+}
+//------------------------------------------------------------------------------
+static void GenerateBlueprintAPIUtils::DumpActionMenuItem(uint32 Indent, FGraphActionListBuilderBase::ActionGroup const& Action, FGraphActionListBuilderBase& ActionList, FArchive* FileOutWriter)
+{
+	const FString& ActionCategory = GenerateBlueprintAPIUtils::GetActionCategory(Action);
 	TArray<FString> Categories;
 	ActionCategory.ParseIntoArray(Categories, TEXT("|"), true);
 
@@ -840,30 +904,30 @@ static void GenerateBlueprintAPIUtils::DumpActionMenuItem(uint32 Indent, FGraphA
 
 		if (Node->ShouldDrawCompact())
 		{
-			ActionEntry += IndentedNewline + "\"CompactTitle\" : \"" + MakeJsonString(Node->GetCompactNodeTitle().ToString()) + "\"";
+			ActionEntry += "," + IndentedNewline + "\"CompactTitle\" : \"" + MakeJsonString(Node->GetCompactNodeTitle().ToString()) + "\"";
 		}
 
 		if (Node->IsNodePure())
 		{
-			ActionEntry += IndentedNewline + "\"NodeType\"     : \"pure\"";
+			ActionEntry += "," + IndentedNewline + "\"NodeType\"     : \"pure\"";
 		}
 		else if (Node->IsA<UK2Node_Event>())
 		{
-			ActionEntry += IndentedNewline + "\"NodeType\"     : \"event\"";
+			ActionEntry += "," + IndentedNewline + "\"NodeType\"     : \"event\"";
 		}
 		else if (Node->IsA<UK2Node_Switch>())
 		{
-			ActionEntry += IndentedNewline + "\"NodeType\"     : \"switch\"";
+			ActionEntry += "," + IndentedNewline + "\"NodeType\"     : \"switch\"";
 		}
 		else
 		{
-			ActionEntry += IndentedNewline + "\"NodeType\"     : \"function\"";
+			ActionEntry += "," + IndentedNewline + "\"NodeType\"     : \"function\"";
 		}
 
 		if (Node->IsA<UK2Node_CommutativeAssociativeBinaryOperator>()
 			|| (Node->IsA<UK2Node_Switch>() && !Node->IsA<UK2Node_SwitchEnum>()))
 		{
-			ActionEntry += IndentedNewline + "\"ShowAddPin\"   : \"true\"";
+			ActionEntry += "," + IndentedNewline + "\"ShowAddPin\"   : \"true\"";
 		}
 
 		if (Node->Pins.Num() > 0)
@@ -898,7 +962,7 @@ static void GenerateBlueprintAPIUtils::DumpActionMenuItem(uint32 Indent, FGraphA
 					{
 						ActionEntry += PinDetailsIndentedNewline + "\"Name\"                 : \"" + MakeJsonString(DisplayName) + "\",";
 					}
-					ActionEntry += PinDetailsIndentedNewline + "\"Direction\"            : \"" + (Pin->Direction == EGPD_Input ? "input" : "output") + "\"";
+					ActionEntry += PinDetailsIndentedNewline + "\"Direction\"            : \"" + (Pin->Direction == EGPD_Input ? "input" : "output") + "\",";
 
 					ActionEntry += PinDetailsIndentedNewline + "\"TypeText\"             : \"" + UEdGraphSchema_K2::TypeToText(Pin->PinType).ToString() + "\"";
 
diff --git a/Engine/Source/Editor/UnrealEd/Private/CookOnTheFlyServer.cpp b/Engine/Source/Editor/UnrealEd/Private/CookOnTheFlyServer.cpp
--- a/Engine/Source/Editor/UnrealEd/Private/CookOnTheFlyServer.cpp
+++ b/Engine/Source/Editor/UnrealEd/Private/CookOnTheFlyServer.cpp
@@ -962,7 +962,7 @@ bool UCookOnTheFlyServer::StartNetworkFileServer( const bool BindAnyPort )
 		FAssetRegistryGenerator* Generator = RegistryGenerators.FindRef(FName(*Platforms[Index]->PlatformName()));
 		if (Generator)
 		{
-			Generator->SaveAssetRegistry(GetSandboxAssetRegistryFilename(), false);
+			Generator->SaveAssetRegistry(GetSandboxAssetRegistryFilename(), true);
 		}
 	}
 
diff --git a/Engine/Source/Editor/UnrealEd/Private/EditorActor.cpp b/Engine/Source/Editor/UnrealEd/Private/EditorActor.cpp
--- a/Engine/Source/Editor/UnrealEd/Private/EditorActor.cpp
+++ b/Engine/Source/Editor/UnrealEd/Private/EditorActor.cpp
@@ -14,7 +14,7 @@
 #include "GameFramework/Actor.h"
 #include "GameFramework/Pawn.h"
 #include "Engine/World.h"
-#include "AI/Navigation/NavigationSystem.h"
+#include "AI/NavigationSystemBase.h"
 #include "Components/LightComponent.h"
 #include "Model.h"
 #include "Exporters/Exporter.h"
@@ -1025,13 +1025,11 @@ bool UUnrealEdEngine::edactDeleteSelected( UWorld* InWorld, bool bVerifyDeletion
 
 	if( LevelsToRebuildNavigation.Num() )
 	{
-		UWorld* World = GetEditorWorldContext().World();
-		UNavigationSystem* NavSys = UNavigationSystem::GetCurrent(World);
-		if (NavSys)
+		for (ULevel* Level : LevelsToRebuildNavigation)
 		{
-			for ( ULevel* Level : LevelsToRebuildNavigation )
+			if (Level)
 			{
-				NavSys->UpdateLevelCollision(Level);
+				FNavigationSystem::UpdateLevelCollision(*Level);
 			}
 		}
 	}
diff --git a/Engine/Source/Editor/UnrealEd/Private/EditorBuildUtils.cpp b/Engine/Source/Editor/UnrealEd/Private/EditorBuildUtils.cpp
--- a/Engine/Source/Editor/UnrealEd/Private/EditorBuildUtils.cpp
+++ b/Engine/Source/Editor/UnrealEd/Private/EditorBuildUtils.cpp
@@ -17,7 +17,7 @@
 #include "SourceControlOperations.h"
 #include "ISourceControlModule.h"
 #include "Materials/MaterialInterface.h"
-#include "AI/Navigation/NavigationSystem.h"
+#include "AI/NavigationSystemBase.h"
 #include "Editor/UnrealEdEngine.h"
 #include "Settings/LevelEditorMiscSettings.h"
 #include "Misc/ConfigCacheIni.h"
@@ -814,8 +814,7 @@ void FEditorBuildUtils::SubmitPackagesForAutomatedBuild( const TSet<UPackage*>&
 
 void FEditorBuildUtils::TriggerNavigationBuilder(UWorld* InWorld, FName Id)
 {
-	if( InWorld->GetWorldSettings()->bEnableNavigationSystem &&
-		InWorld->GetNavigationSystem() )
+	if (InWorld)
 	{
 		if (Id == FBuildOptions::BuildAIPaths ||
 			Id == FBuildOptions::BuildSelectedAIPaths ||
@@ -831,7 +830,7 @@ void FEditorBuildUtils::TriggerNavigationBuilder(UWorld* InWorld, FName Id)
 		}
 
 		// Invoke navmesh generator
-		InWorld->GetNavigationSystem()->Build();
+		FNavigationSystem::Build(*InWorld);
 	}
 }
 
diff --git a/Engine/Source/Editor/UnrealEd/Private/EditorEngine.cpp b/Engine/Source/Editor/UnrealEd/Private/EditorEngine.cpp
--- a/Engine/Source/Editor/UnrealEd/Private/EditorEngine.cpp
+++ b/Engine/Source/Editor/UnrealEd/Private/EditorEngine.cpp
@@ -21,7 +21,7 @@
 #include "PhysicsEngine/BodyInstance.h"
 #include "Components/PrimitiveComponent.h"
 #include "Components/StaticMeshComponent.h"
-#include "AI/Navigation/NavigationSystem.h"
+#include "AI/NavigationSystemBase.h"
 #include "Components/LightComponent.h"
 #include "Tickable.h"
 #include "TickableEditorObject.h"
@@ -625,7 +625,7 @@ void UEditorEngine::InitEditor(IEngineLoop* InEngineLoop)
 
 	// Set navigation system property indicating whether navigation is supposed to rebuild automatically 
 	FWorldContext &EditorContext = GEditor->GetEditorWorldContext();
-	UNavigationSystem::SetNavigationAutoUpdateEnabled(GetDefault<ULevelEditorMiscSettings>()->bNavigationAutoUpdate, EditorContext.World()->GetNavigationSystem());
+	FNavigationSystem::SetNavigationAutoUpdateEnabled(GetDefault<ULevelEditorMiscSettings>()->bNavigationAutoUpdate, EditorContext.World()->GetNavigationSystem());
 
 	// Allocate temporary model.
 	TempModel = NewObject<UModel>();
diff --git a/Engine/Source/Editor/UnrealEd/Private/EditorExporters.cpp b/Engine/Source/Editor/UnrealEd/Private/EditorExporters.cpp
--- a/Engine/Source/Editor/UnrealEd/Private/EditorExporters.cpp
+++ b/Engine/Source/Editor/UnrealEd/Private/EditorExporters.cpp
@@ -20,7 +20,7 @@
 #include "SceneTypes.h"
 #include "RawIndexBuffer.h"
 #include "RenderingThread.h"
-#include "AI/Navigation/NavigationSystem.h"
+#include "AI/NavigationSystemBase.h"
 #include "Model.h"
 #include "Exporters/Exporter.h"
 #include "Exporters/AnimSequenceExporterFBX.h"
@@ -2100,12 +2100,7 @@ void UEditorEngine::RebuildStaticNavigableGeometry(ULevel* Level)
 			}
 		}
 
-		UWorld* World = GetEditorWorldContext().World();
-		UNavigationSystem* NavSys = UNavigationSystem::GetCurrent(World);
-		if (NavSys)
-		{
-			NavSys->UpdateLevelCollision(Level);
-		}
+		FNavigationSystem::UpdateLevelCollision(*Level);
 	}
 }
 
diff --git a/Engine/Source/Editor/UnrealEd/Private/EditorServer.cpp b/Engine/Source/Editor/UnrealEd/Private/EditorServer.cpp
--- a/Engine/Source/Editor/UnrealEd/Private/EditorServer.cpp
+++ b/Engine/Source/Editor/UnrealEd/Private/EditorServer.cpp
@@ -42,7 +42,7 @@
 #include "Materials/MaterialInterface.h"
 #include "Components/PrimitiveComponent.h"
 #include "Components/MeshComponent.h"
-#include "AI/Navigation/NavigationSystem.h"
+#include "AI/NavigationSystemBase.h"
 #include "Components/LightComponent.h"
 #include "Model.h"
 #include "Exporters/Exporter.h"
@@ -145,7 +145,7 @@
 #include "Framework/Notifications/NotificationManager.h"
 #include "Widgets/Notifications/SNotificationList.h"
 #include "Animation/AnimNotifies/AnimNotify.h"
-#include "AI/Navigation/NavLinkRenderingComponent.h"
+#include "NavLinkRenderingComponent.h"
 #include "Analytics/AnalyticsPrivacySettings.h"
 #include "Kismet2/KismetReinstanceUtilities.h"
 #include "AnalyticsEventAttribute.h"
@@ -2232,7 +2232,7 @@ UWorld* UEditorEngine::NewMap()
 	InitBuilderBrush( Context.World() );
 
 	// Let navigation system know we're done creating new world
-	UNavigationSystem::InitializeForWorld(Context.World(), FNavigationSystemRunMode::EditorMode);
+	FNavigationSystem::AddNavigationSystemToWorld(*Context.World(), FNavigationSystemRunMode::EditorMode);
 
 	// Deselect all
 	GEditor->SelectNone( false, true );
@@ -2636,8 +2636,8 @@ bool UEditorEngine::Map_Load(const TCHAR* Str, FOutputDevice& Ar)
 
 					InitializingFeedback.EnterProgressFrame();
 
-					UNavigationSystem::InitializeForWorld(Context.World(), FNavigationSystemRunMode::EditorMode);
 					Context.World()->CreateAISystem();
+					FNavigationSystem::AddNavigationSystemToWorld(*Context.World(), FNavigationSystemRunMode::EditorMode);
 
 					// Assign stationary light channels for previewing
 					ULightComponent::ReassignStationaryLightChannels(Context.World(), false, NULL);
diff --git a/Engine/Source/Editor/UnrealEd/Private/Kismet2/BlueprintEditorUtils.cpp b/Engine/Source/Editor/UnrealEd/Private/Kismet2/BlueprintEditorUtils.cpp
--- a/Engine/Source/Editor/UnrealEd/Private/Kismet2/BlueprintEditorUtils.cpp
+++ b/Engine/Source/Editor/UnrealEd/Private/Kismet2/BlueprintEditorUtils.cpp
@@ -2529,7 +2529,7 @@ void FBlueprintEditorUtils::RemoveGraph(UBlueprint* Blueprint, class UEdGraph* G
 			// Can't just call Remove, the object is wrapped in a struct
 			for(int EditedDocIdx = 0; EditedDocIdx < Blueprint->LastEditedDocuments.Num(); ++EditedDocIdx)
 			{
-				if(Blueprint->LastEditedDocuments[EditedDocIdx].EditedObject == GraphToRemove)
+				if(Blueprint->LastEditedDocuments[EditedDocIdx].EditedObjectPath.ResolveObject() == GraphToRemove)
 				{
 					Blueprint->LastEditedDocuments.RemoveAt(EditedDocIdx);
 					break;
@@ -3873,11 +3873,16 @@ void FBlueprintEditorUtils::SetBlueprintVariableCategory(UBlueprint* Blueprint,
 		UClass* OuterClass = CastChecked<UClass>(TargetProperty->GetOuter());
 		const bool bIsNativeVar = (OuterClass->ClassGeneratedBy == nullptr);
 
-		// If the category does not change, we will not recompile the Blueprint
-		bool bIsCategoryChanged = false;
 		if (!bIsNativeVar)
 		{
-			TargetProperty->SetMetaData(TEXT("Category"), *SetCategory.ToString());
+			if (UTimelineTemplate* Timeline = Blueprint->FindTimelineTemplateByVariableName(TargetProperty->GetFName()))
+			{
+				Timeline->SetMetaData(TEXT("Category"), *SetCategory.ToString());
+			}
+			else
+			{
+				TargetProperty->SetMetaData(TEXT("Category"), *SetCategory.ToString());
+			}
 			const int32 VarIndex = FBlueprintEditorUtils::FindNewVariableIndex(Blueprint, VarName);
 			if (VarIndex != INDEX_NONE)
 			{
@@ -5917,7 +5922,7 @@ void FBlueprintEditorUtils::PurgeNullGraphs(UBlueprint* Blueprint)
 	CleanNullGraphReferencesInArray(Blueprint, Blueprint->DelegateSignatureGraphs);
 	CleanNullGraphReferencesInArray(Blueprint, Blueprint->MacroGraphs);
 
-	Blueprint->LastEditedDocuments.RemoveAll([](const FEditedDocumentInfo& TestDoc) { return TestDoc.EditedObject == nullptr; });
+	Blueprint->LastEditedDocuments.RemoveAll([](const FEditedDocumentInfo& TestDoc) { return TestDoc.EditedObjectPath.ResolveObject() == nullptr; });
 }
 
 struct FConformCallsToParentFunctionUtils
@@ -9013,7 +9018,7 @@ UClass* FBlueprintEditorUtils::GetNativeParent(const UBlueprint* BP)
 	return Ret;
 }
 
-bool FBlueprintEditorUtils::ImplentsGetWorld(const UBlueprint* BP)
+bool FBlueprintEditorUtils::ImplementsGetWorld(const UBlueprint* BP)
 {
 	if(UClass* NativeParent = GetNativeParent(BP))
 	{
diff --git a/Engine/Source/Editor/UnrealEd/Private/Kismet2/KismetDebugUtilities.cpp b/Engine/Source/Editor/UnrealEd/Private/Kismet2/KismetDebugUtilities.cpp
--- a/Engine/Source/Editor/UnrealEd/Private/Kismet2/KismetDebugUtilities.cpp
+++ b/Engine/Source/Editor/UnrealEd/Private/Kismet2/KismetDebugUtilities.cpp
@@ -19,6 +19,7 @@
 #include "Editor/UnrealEdEngine.h"
 #include "Settings/EditorExperimentalSettings.h"
 #include "CallStackViewer.h"
+#include "WatchPointViewer.h"
 #include "Animation/AnimBlueprintGeneratedClass.h"
 #include "UnrealEdGlobals.h"
 #include "Engine/Breakpoint.h"
@@ -705,6 +706,7 @@ void FKismetDebugUtilities::AttemptToBreakExecution(UBlueprint* BlueprintObj, co
 		Data.LastExceptionMessage = Info.GetDescription();
 		FKismetEditorUtilities::BringKismetToFocusAttentionOnObject(NodeStoppedAt);
 		CallStackViewer::UpdateDisplayedCallstack(FBlueprintExceptionTracker::Get().ScriptStack);
+		WatchViewer::UpdateDisplayedWatches(BlueprintObj);
 		FSlateApplication::Get().EnterDebuggingMode();
 	}
 }
@@ -1040,6 +1042,8 @@ void FKismetDebugUtilities::ClearBreakpoints(UBlueprint* Blueprint)
 	Blueprint->MarkPackageDirty();
 }
 
+FKismetDebugUtilities::FOnWatchedPinsListChanged FKismetDebugUtilities::WatchedPinsListChangedEvent;
+
 bool FKismetDebugUtilities::CanWatchPin(const UBlueprint* Blueprint, const UEdGraphPin* Pin)
 {
 	//@TODO: This function belongs in the schema
@@ -1067,6 +1071,7 @@ void FKismetDebugUtilities::RemovePinWatch(UBlueprint* Blueprint, const UEdGraph
 	Blueprint->WatchedPins.Remove(NonConstPin);
 	Blueprint->MarkPackageDirty();
 	Blueprint->PostEditChange();
+	WatchedPinsListChangedEvent.Broadcast(Blueprint);
 }
 
 void FKismetDebugUtilities::TogglePinWatch(UBlueprint* Blueprint, const UEdGraphPin* Pin)
@@ -1083,6 +1088,8 @@ void FKismetDebugUtilities::TogglePinWatch(UBlueprint* Blueprint, const UEdGraph
 		Blueprint->MarkPackageDirty();
 		Blueprint->PostEditChange();
 	}
+
+	WatchedPinsListChangedEvent.Broadcast(Blueprint);
 }
 
 void FKismetDebugUtilities::ClearPinWatches(UBlueprint* Blueprint)
@@ -1090,10 +1097,44 @@ void FKismetDebugUtilities::ClearPinWatches(UBlueprint* Blueprint)
 	Blueprint->WatchedPins.Empty();
 	Blueprint->MarkPackageDirty();
 	Blueprint->PostEditChange();
+
+	WatchedPinsListChangedEvent.Broadcast(Blueprint);
 }
 
 // Gets the watched tooltip for a specified site
 FKismetDebugUtilities::EWatchTextResult FKismetDebugUtilities::GetWatchText(FString& OutWatchText, UBlueprint* Blueprint, UObject* ActiveObject, const UEdGraphPin* WatchPin)
+{
+	UProperty* PropertyToDebug = nullptr;
+	void* DataPtr = nullptr;
+	void* DeltaPtr = nullptr;
+	UObject* ParentObj = nullptr;
+	FKismetDebugUtilities::EWatchTextResult Result = FindDebuggingData(Blueprint, ActiveObject, WatchPin, PropertyToDebug, DataPtr, DeltaPtr, ParentObj);
+
+	if (Result == FKismetDebugUtilities::EWatchTextResult::EWTR_Valid)
+	{
+		PropertyToDebug->ExportText_InContainer(/*ArrayElement=*/ 0, /*inout*/ OutWatchText, DataPtr, DeltaPtr, /*Parent=*/ ParentObj, PPF_PropertyWindow | PPF_BlueprintDebugView);
+	}
+
+	return Result;
+}
+
+FKismetDebugUtilities::EWatchTextResult FKismetDebugUtilities::GetDebugInfo(FDebugInfo& OutDebugInfo, UBlueprint* Blueprint, UObject* ActiveObject, const UEdGraphPin* WatchPin)
+{
+	UProperty* PropertyToDebug = nullptr;
+	void* DataPtr = nullptr;
+	void* DeltaPtr = nullptr;
+	UObject* ParentObj = nullptr;
+	FKismetDebugUtilities::EWatchTextResult Result = FindDebuggingData(Blueprint, ActiveObject, WatchPin, PropertyToDebug, DataPtr, DeltaPtr, ParentObj);
+
+	if (Result == FKismetDebugUtilities::EWatchTextResult::EWTR_Valid)
+	{
+		GetDebugInfo_InContainer(0, OutDebugInfo, PropertyToDebug, DataPtr);
+	}
+
+	return Result;
+}
+
+FKismetDebugUtilities::EWatchTextResult FKismetDebugUtilities::FindDebuggingData(UBlueprint* Blueprint, UObject* ActiveObject, const UEdGraphPin* WatchPin, UProperty*& OutProperty, void*& OutData, void*& OutDelta, UObject*& OutParent)
 {
 	FKismetDebugUtilitiesData& Data = FKismetDebugUtilitiesData::Get();
 
@@ -1161,7 +1202,7 @@ FKismetDebugUtilities::EWatchTextResult FKismetDebugUtilities::GetWatchText(FStr
 #if USE_UBER_GRAPH_PERSISTENT_FRAME
 			// Try find the propertybase in the persistent ubergraph frame
 			UFunction* OuterFunction = Cast<UFunction>(Property->GetOuter());
-			if(!PropertyBase && OuterFunction)
+			if (!PropertyBase && OuterFunction)
 			{
 				UBlueprintGeneratedClass* BPGC = Cast<UBlueprintGeneratedClass>(Blueprint->GeneratedClass);
 				if (BPGC && ActiveObject->IsA(BPGC))
@@ -1170,15 +1211,15 @@ FKismetDebugUtilities::EWatchTextResult FKismetDebugUtilities::GetWatchText(FStr
 				}
 			}
 #endif // USE_UBER_GRAPH_PERSISTENT_FRAME
-			
+
 			// see if our WatchPin is on a animation node & if so try to get its property info
 			UAnimBlueprintGeneratedClass* AnimBlueprintGeneratedClass = Cast<UAnimBlueprintGeneratedClass>(Blueprint->GeneratedClass);
-			if(!PropertyBase && AnimBlueprintGeneratedClass)
+			if (!PropertyBase && AnimBlueprintGeneratedClass)
 			{
 				// are we linked to an anim graph node?
 				UProperty* LinkedProperty = Property;
 				const UAnimGraphNode_Base* Node = Cast<UAnimGraphNode_Base>(WatchPin->GetOuter());
-				if(Node == nullptr && WatchPin->LinkedTo.Num() > 0)
+				if (Node == nullptr && WatchPin->LinkedTo.Num() > 0)
 				{
 					const UEdGraphPin* LinkedPin = WatchPin->LinkedTo[0];
 					// When we change Node we *must* change Property, so it's still a sub-element of that.
@@ -1186,7 +1227,7 @@ FKismetDebugUtilities::EWatchTextResult FKismetDebugUtilities::GetWatchText(FStr
 					Node = Cast<UAnimGraphNode_Base>(LinkedPin->GetOuter());
 				}
 
-				if(Node && LinkedProperty)
+				if (Node && LinkedProperty)
 				{
 					UStructProperty* NodeStructProperty = Cast<UStructProperty>(FKismetDebugUtilities::FindClassPropertyForNode(Blueprint, Node));
 					if (NodeStructProperty)
@@ -1196,7 +1237,10 @@ FKismetDebugUtilities::EWatchTextResult FKismetDebugUtilities::GetWatchText(FStr
 							if (NodeProperty == NodeStructProperty)
 							{
 								void* NodePtr = NodeProperty->ContainerPtrToValuePtr<void>(ActiveObject);
-								LinkedProperty->ExportText_InContainer(/*ArrayElement=*/ 0, /*inout*/ OutWatchText, NodePtr, NodePtr, /*Parent=*/ ActiveObject, PPF_PropertyWindow|PPF_BlueprintDebugView);
+								OutProperty = LinkedProperty;
+								OutData = NodePtr;
+								OutDelta = NodePtr;
+								OutParent = ActiveObject;
 								return EWTR_Valid;
 							}
 						}
@@ -1207,7 +1251,10 @@ FKismetDebugUtilities::EWatchTextResult FKismetDebugUtilities::GetWatchText(FStr
 			// Now either print out the variable value, or that it was out-of-scope
 			if (PropertyBase != nullptr)
 			{
-				Property->ExportText_InContainer(/*ArrayElement=*/ 0, /*inout*/ OutWatchText, PropertyBase, PropertyBase, /*Parent=*/ ActiveObject, PPF_PropertyWindow|PPF_BlueprintDebugView);
+				OutProperty = Property;
+				OutData = PropertyBase;
+				OutDelta = PropertyBase;
+				OutParent = ActiveObject;
 				return EWTR_Valid;
 			}
 			else
@@ -1226,6 +1273,255 @@ FKismetDebugUtilities::EWatchTextResult FKismetDebugUtilities::GetWatchText(FStr
 	}
 }
 
+void FKismetDebugUtilities::GetDebugInfo_InContainer(int32 Index, FDebugInfo& DebugInfo, UProperty* Property, const void* Data)
+{
+	GetDebugInfoInternal(DebugInfo, Property, Property->ContainerPtrToValuePtr<void>(Data, Index));
+}
+
+void FKismetDebugUtilities::GetDebugInfoInternal(FDebugInfo& DebugInfo, UProperty* Property, const void* PropertyValue)
+{
+	if (Property == nullptr)
+	{
+		return;
+	}
+
+	DebugInfo.Type = UEdGraphSchema_K2::TypeToText(Property);
+	DebugInfo.DisplayName = Property->GetDisplayNameText();
+
+	UByteProperty* ByteProperty = Cast<UByteProperty>(Property);
+	if (ByteProperty)
+	{
+		UEnum* Enum = ByteProperty->GetIntPropertyEnum();
+		if (Enum)
+		{
+			if (Enum->IsValidEnumValue(*(const uint8*)PropertyValue))
+			{
+				DebugInfo.Value = Enum->GetDisplayNameTextByValue(*(const uint8*)PropertyValue);
+			}
+			else
+			{
+				DebugInfo.Value = FText::FromString(TEXT("(INVALID)"));
+			}
+
+			return;
+		}
+
+		// if there is no Enum we need to fall through and treat this as a UNumericProperty
+	}
+
+	UNumericProperty* NumericProperty = Cast<UNumericProperty>(Property);
+	if (NumericProperty)
+	{
+		DebugInfo.Value = FText::FromString(NumericProperty->GetNumericPropertyValueToString(PropertyValue));
+		return;
+	}
+
+	UBoolProperty* BoolProperty = Cast<UBoolProperty>(Property);
+	if (BoolProperty)
+	{
+		DebugInfo.Value = BoolProperty->GetPropertyValue(PropertyValue) ? GTrue : GFalse;
+		return;
+	}
+
+	UNameProperty* NameProperty = Cast<UNameProperty>(Property);
+	if (NameProperty)
+	{
+		DebugInfo.Value = FText::FromName(*(FName*)PropertyValue);
+		return;
+	}
+
+	UTextProperty* TextProperty = Cast<UTextProperty>(Property);
+	if (TextProperty)
+	{
+		DebugInfo.Value = TextProperty->GetPropertyValue(PropertyValue);
+		return;
+	}
+
+	UStrProperty* StringProperty = Cast<UStrProperty>(Property);
+	if (StringProperty)
+	{
+		DebugInfo.Value = FText::FromString(StringProperty->GetPropertyValue(PropertyValue));
+		return;
+	}
+
+	UArrayProperty* ArrayProperty = Cast<UArrayProperty>(Property);
+	if (ArrayProperty)
+	{
+		checkSlow(ArrayProperty->Inner);
+
+		FScriptArrayHelper ArrayHelper(ArrayProperty, PropertyValue);
+
+		DebugInfo.Value = FText::Format(LOCTEXT("ArraySize", "Num={0}"), FText::AsNumber(ArrayHelper.Num()));
+
+		for (int32 i = 0; i < ArrayHelper.Num(); i++)
+		{
+			FDebugInfo ArrayDebugInfo;
+
+			uint8* PropData = ArrayHelper.GetRawPtr(i);
+			GetDebugInfoInternal(ArrayDebugInfo, ArrayProperty->Inner, PropData);
+			// overwrite the display name with the array index for the current element
+			ArrayDebugInfo.DisplayName = FText::Format(LOCTEXT("ArrayIndexName", "[{0}]"), FText::AsNumber(i));
+			DebugInfo.Children.Add(ArrayDebugInfo);
+		}
+
+		return;
+	}
+
+	UStructProperty* StructProperty = Cast<UStructProperty>(Property);
+	if (StructProperty)
+	{
+		FString WatchText;
+		StructProperty->ExportTextItem(WatchText, PropertyValue, PropertyValue, nullptr, PPF_PropertyWindow | PPF_BlueprintDebugView, nullptr);
+		DebugInfo.Value = FText::FromString(WatchText);
+
+		for (TFieldIterator<UProperty> It(StructProperty->Struct); It; ++It)
+		{
+			FDebugInfo StructDebugInfo;
+			GetDebugInfoInternal(StructDebugInfo, *It, It->ContainerPtrToValuePtr<void>(PropertyValue, 0));
+
+			DebugInfo.Children.Add(StructDebugInfo);
+		}
+
+		return;
+	}
+
+	UEnumProperty* EnumProperty = Cast<UEnumProperty>(Property);
+	if (EnumProperty)
+	{
+		UNumericProperty* LocalUnderlyingProp = EnumProperty->GetUnderlyingProperty();
+		UEnum* Enum = EnumProperty->GetEnum();
+
+		int64 Value = LocalUnderlyingProp->GetSignedIntPropertyValue(PropertyValue);
+
+		// if the value is the max value (the autogenerated *_MAX value), export as "INVALID", unless we're exporting text for copy/paste (for copy/paste,
+		// the property text value must actually match an entry in the enum's names array)
+		if (Enum)
+		{
+			if (Enum->IsValidEnumValue(Value))
+			{
+				DebugInfo.Value = Enum->GetDisplayNameTextByValue(Value);
+			}
+			else
+			{
+				DebugInfo.Value = LOCTEXT("Invalid", "(INVALID)");
+			}
+		}
+		else
+		{
+			DebugInfo.Value = FText::AsNumber(Value);
+		}
+
+		return;
+	}
+
+	UMapProperty* MapProperty = Cast<UMapProperty>(Property);
+	if (MapProperty)
+	{
+		FScriptMapHelper MapHelper(MapProperty, PropertyValue);
+		DebugInfo.Value = FText::Format(LOCTEXT("MapSize", "Num={0}"), FText::AsNumber(MapHelper.Num()));
+		uint8* PropData = MapHelper.GetPairPtr(0);
+
+		int32 Index = 0;
+		for (int32 Count = MapHelper.Num(); Count; PropData += MapProperty->MapLayout.SetLayout.Size, ++Index)
+		{
+			if (MapHelper.IsValidIndex(Index))
+			{
+				FDebugInfo ChildInfo;
+
+				GetDebugInfoInternal(ChildInfo, MapProperty->ValueProp, PropData + MapProperty->MapLayout.ValueOffset);
+
+				// use the info from the ValueProp and then overwrite the name with the KeyProp data
+				FString NameStr = TEXT("[");
+				MapProperty->KeyProp->ExportTextItem(NameStr, PropData, nullptr, nullptr, PPF_PropertyWindow | PPF_BlueprintDebugView | PPF_Delimited, nullptr);
+				NameStr += TEXT("] ");
+
+				ChildInfo.DisplayName = FText::FromString(NameStr);
+
+				DebugInfo.Children.Add(ChildInfo);
+
+				--Count;
+			}
+		}
+
+		return;
+	}
+
+	USetProperty* SetProperty = Cast<USetProperty>(Property);
+	if (SetProperty)
+	{
+		FScriptSetHelper SetHelper(SetProperty, PropertyValue);
+		DebugInfo.Value = FText::Format(LOCTEXT("SetSize", "Num={0}"), FText::AsNumber(SetHelper.Num()));
+		uint8* PropData = SetHelper.GetElementPtr(0);
+
+		int32 Index = 0;
+		for (int32 Count = SetHelper.Num(); Count; PropData += SetProperty->SetLayout.Size, ++Index)
+		{
+			if (SetHelper.IsValidIndex(Index))
+			{
+				FDebugInfo ChildInfo;
+				GetDebugInfoInternal(ChildInfo, SetProperty->ElementProp, PropData);
+
+				// members of sets don't have their own names
+				ChildInfo.DisplayName = FText::GetEmpty();
+
+				DebugInfo.Children.Add(ChildInfo);
+
+				--Count;
+			}
+		}
+
+		return;
+	}
+
+	UObjectPropertyBase* ObjectPropertyBase = Cast<UObjectPropertyBase>(Property);
+	if (ObjectPropertyBase)
+	{
+		UObject* Obj = ObjectPropertyBase->GetObjectPropertyValue(PropertyValue);
+		if (Obj != nullptr)
+		{
+			DebugInfo.Value = FText::FromString(Obj->GetFullName());
+		}
+		else
+		{
+			DebugInfo.Value = FText::FromString(TEXT("None"));
+		}
+
+		return;
+	}
+
+	UDelegateProperty* DelegateProperty = Cast<UDelegateProperty>(Property);
+	if (DelegateProperty)
+	{
+		if (DelegateProperty->SignatureFunction)
+		{
+			DebugInfo.Value = DelegateProperty->SignatureFunction->GetDisplayNameText();
+		}
+		else
+		{
+			DebugInfo.Value = LOCTEXT("NoFunc", "(No bound function)");
+		}
+
+		return;
+	}
+
+	UMulticastDelegateProperty* MulticastDelegateProperty = Cast<UMulticastDelegateProperty>(Property);
+	if (MulticastDelegateProperty)
+	{
+		if (MulticastDelegateProperty->SignatureFunction)
+		{
+			DebugInfo.Value = MulticastDelegateProperty->SignatureFunction->GetDisplayNameText();
+		}
+		else
+		{
+			DebugInfo.Value = LOCTEXT("NoFunc", "(No bound function)");
+		}
+
+		return;
+	}
+
+	ensure(false);
+}
+
 FText FKismetDebugUtilities::GetAndClearLastExceptionMessage()
 {
 	FKismetDebugUtilitiesData& Data = FKismetDebugUtilitiesData::Get();
diff --git a/Engine/Source/Editor/UnrealEd/Private/NavigationBuildingNotification.cpp b/Engine/Source/Editor/UnrealEd/Private/NavigationBuildingNotification.cpp
--- a/Engine/Source/Editor/UnrealEd/Private/NavigationBuildingNotification.cpp
+++ b/Engine/Source/Editor/UnrealEd/Private/NavigationBuildingNotification.cpp
@@ -1,7 +1,7 @@
 // Copyright 1998-2018 Epic Games, Inc. All Rights Reserved.
 
 #include "NavigationBuildingNotification.h"
-#include "AI/Navigation/NavigationSystem.h"
+#include "NavigationSystem.h"
 #include "Editor/EditorEngine.h"
 #include "Settings/LevelEditorMiscSettings.h"
 #include "EngineGlobals.h"
@@ -111,9 +111,10 @@ FText FNavigationBuildingNotificationImpl::GetNotificationText() const
 	if (EEngine)
 	{
 		FWorldContext &EditorContext = EEngine->GetEditorWorldContext();
-		if (EditorContext.World() && EditorContext.World()->GetNavigationSystem())
+		UNavigationSystemV1* NavSys = FNavigationSystem::GetCurrent<UNavigationSystemV1>(EditorContext.World());
+		if (NavSys)
 		{
-			RemainingTasks = EditorContext.World()->GetNavigationSystem()->GetNumRemainingBuildTasks();
+			RemainingTasks = NavSys->GetNumRemainingBuildTasks();
 		}
 	}
 		
@@ -136,9 +137,10 @@ void FNavigationBuildingNotificationImpl::Tick(float DeltaTime)
 		const bool bUserRequestedBuild = FEditorBuildUtils::IsBuildingNavigationFromUserRequest();
 		FWorldContext &EditorContext = EEngine->GetEditorWorldContext();
 		
-		const bool bBuildInProgress = EditorContext.World() != NULL && EditorContext.World()->GetNavigationSystem() != NULL 
-			&& EditorContext.World()->GetNavigationSystem()->IsNavigationBuildInProgress( GetDefault<ULevelEditorMiscSettings>()->bNavigationAutoUpdate ? true : false) == true
-			&& EditorContext.World()->GetNavigationSystem()->GetNumRemainingBuildTasks() > 0;
+		UNavigationSystemV1* NavSys = FNavigationSystem::GetCurrent<UNavigationSystemV1>(EditorContext.World());
+		const bool bBuildInProgress = NavSys != nullptr
+			&& NavSys->IsNavigationBuildInProgress( GetDefault<ULevelEditorMiscSettings>()->bNavigationAutoUpdate ? true : false) == true
+			&& NavSys->GetNumRemainingBuildTasks() > 0;
 
 		if (!bPreviouslyDetectedBuild && bBuildInProgress)
 		{
diff --git a/Engine/Source/Editor/UnrealEd/Private/PlayLevel.cpp b/Engine/Source/Editor/UnrealEd/Private/PlayLevel.cpp
--- a/Engine/Source/Editor/UnrealEd/Private/PlayLevel.cpp
+++ b/Engine/Source/Editor/UnrealEd/Private/PlayLevel.cpp
@@ -38,7 +38,7 @@
 #include "Engine/GameInstance.h"
 #include "Engine/World.h"
 #include "Settings/LevelEditorPlaySettings.h"
-#include "AI/Navigation/NavigationSystem.h"
+#include "AI/NavigationSystemBase.h"
 #include "Editor/EditorEngine.h"
 #include "Editor/UnrealEdEngine.h"
 #include "Settings/ProjectPackagingSettings.h"
@@ -370,10 +370,7 @@ void UEditorEngine::EndPlayMap()
 	}
 
 	// Lose the EditorWorld pointer (this is only maintained while PIEing)
-	if (EditorWorld->GetNavigationSystem())
-	{
-		EditorWorld->GetNavigationSystem()->OnPIEEnd();
-	}
+	FNavigationSystem::OnPIEEnd(*EditorWorld);
 
 	FGameDelegates::Get().GetEndPlayMapDelegate().Broadcast();
 
@@ -2255,10 +2252,7 @@ void UEditorEngine::PlayInEditor( UWorld* InWorld, bool bInSimulateInEditor )
 	FEditorDelegates::BeginPIE.Broadcast(bInSimulateInEditor);
 
 	// let navigation know PIE starts so it can avoid any blueprint creation/deletion/instantiation affect editor map's navmesh changes
-	if (InWorld->GetNavigationSystem())
-	{
-		InWorld->GetNavigationSystem()->OnPIEStart();
-	}
+	FNavigationSystem::OnPIEStart(*InWorld);
 
 	ULevelEditorPlaySettings* PlayInSettings = GetMutableDefault<ULevelEditorPlaySettings>();
 	check(PlayInSettings);
@@ -2282,11 +2276,7 @@ void UEditorEngine::PlayInEditor( UWorld* InWorld, bool bInSimulateInEditor )
 		if ( !bContinuePIE )
 		{
 			FEditorDelegates::EndPIE.Broadcast(bInSimulateInEditor);
-			if (InWorld->GetNavigationSystem())
-			{
-				InWorld->GetNavigationSystem()->OnPIEEnd();
-			}
-
+			FNavigationSystem::OnPIEEnd(*InWorld);
 			return;
 		}
 		else
@@ -2905,10 +2895,7 @@ UGameInstance* UEditorEngine::CreatePIEGameInstance(int32 InPIEInstance, bool bI
 
 		FEditorDelegates::EndPIE.Broadcast(bInSimulateInEditor);
 
-		if (EditorWorld->GetNavigationSystem())
-		{
-			EditorWorld->GetNavigationSystem()->OnPIEEnd();
-		}
+		FNavigationSystem::OnPIEEnd(*EditorWorld);
 
 		return nullptr;
 	}
diff --git a/Engine/Source/Editor/UnrealEd/Private/SKeySelector.cpp b/Engine/Source/Editor/UnrealEd/Private/SKeySelector.cpp
--- a/Engine/Source/Editor/UnrealEd/Private/SKeySelector.cpp
+++ b/Engine/Source/Editor/UnrealEd/Private/SKeySelector.cpp
@@ -110,7 +110,7 @@ void SKeySelector::Construct(const FArguments& InArgs)
 	// if we allow NoClear, add a "None" option to be able to clear out a binding
 	if (InArgs._AllowClear)
 	{
-		new (KeyTreeRoot) FKeyTreeItem(MakeShareable(new FKeyTreeInfo(FText(), MakeShareable(new FKey(NAME_None)))));
+		new (KeyTreeRoot) FKeyTreeItem(MakeShareable(new FKeyTreeInfo(FText(), MakeShareable(new FKey(EKeys::Invalid)))));
 	}
 
 	TreeViewWidth = InArgs._TreeViewWidth;
diff --git a/Engine/Source/Editor/UnrealEd/Private/Settings/SettingsClasses.cpp b/Engine/Source/Editor/UnrealEd/Private/Settings/SettingsClasses.cpp
--- a/Engine/Source/Editor/UnrealEd/Private/Settings/SettingsClasses.cpp
+++ b/Engine/Source/Editor/UnrealEd/Private/Settings/SettingsClasses.cpp
@@ -7,7 +7,7 @@
 #include "Misc/PackageName.h"
 #include "InputCoreTypes.h"
 #include "Classes/EditorStyleSettings.h"
-#include "AI/Navigation/NavigationSystem.h"
+#include "AI/NavigationSystemBase.h"
 #include "Model.h"
 #include "ISourceControlModule.h"
 #include "Settings/ContentBrowserSettings.h"
@@ -356,7 +356,7 @@ void ULevelEditorMiscSettings::PostEditChangeProperty( struct FPropertyChangedEv
 	if (Name == FName(TEXT("bNavigationAutoUpdate")))
 	{
 		FWorldContext &EditorContext = GEditor->GetEditorWorldContext();
-		UNavigationSystem::SetNavigationAutoUpdateEnabled(bNavigationAutoUpdate, EditorContext.World()->GetNavigationSystem());
+		FNavigationSystem::SetNavigationAutoUpdateEnabled(bNavigationAutoUpdate, EditorContext.World()->GetNavigationSystem());
 	}
 
 	if (!FUnrealEdMisc::Get().IsDeletePreferences())
diff --git a/Engine/Source/Editor/UnrealEd/Private/UnrealEdPrivatePCH.h b/Engine/Source/Editor/UnrealEd/Private/UnrealEdPrivatePCH.h
--- a/Engine/Source/Editor/UnrealEd/Private/UnrealEdPrivatePCH.h
+++ b/Engine/Source/Editor/UnrealEd/Private/UnrealEdPrivatePCH.h
@@ -484,7 +484,7 @@
 #include "Engine/Blueprint.h"
 #include "Engine/BlueprintCore.h"
 #include "AI/Navigation/NavigationTypes.h"
-#include "AI/Navigation/NavFilters/NavigationQueryFilter.h"
+#include "NavFilters/NavigationQueryFilter.h"
 #include "Engine/LatentActionManager.h"
 #include "Engine/World.h"
 #include "WorldCollision.h"
@@ -608,9 +608,8 @@
 #include "Engine/FontImportOptions.h"
 #include "Animation/AnimBlueprint.h"
 #include "EngineModule.h"
-#include "AI/Navigation/NavigationSystem.h"
+#include "AI/NavigationSystemBase.h"
 #include "GenericOctreePublic.h"
-#include "AI/Navigation/NavigationData.h"
 #include "CanvasTypes.h"
 #include "Sound/SoundBase.h"
 #include "Components/LightComponent.h"
diff --git a/Engine/Source/Editor/UnrealEd/Public/GraphEditor.h b/Engine/Source/Editor/UnrealEd/Public/GraphEditor.h
--- a/Engine/Source/Editor/UnrealEd/Public/GraphEditor.h
+++ b/Engine/Source/Editor/UnrealEd/Public/GraphEditor.h
@@ -174,12 +174,12 @@ public:
 		}
 	}
 
-	/* Set new viewer location */
-	virtual void SetViewLocation(const FVector2D& Location, float ZoomAmount)
+	/* Set new viewer location and optionally set the current bookmark */
+	virtual void SetViewLocation(const FVector2D& Location, float ZoomAmount, const FGuid& BookmarkId = FGuid())
 	{
 		if (Implementation.IsValid())
 		{
-			Implementation->SetViewLocation(Location, ZoomAmount);
+			Implementation->SetViewLocation(Location, ZoomAmount, BookmarkId);
 		}
 	}
 
@@ -197,6 +197,19 @@ public:
 		}
 	}
 
+	/**
+	 * Gets the current graph view bookmark
+	 *
+	 * @param OutBookmarkId		Will have the current bookmark ID
+	 */
+	virtual void GetViewBookmark(FGuid& OutBookmarkId)
+	{
+		if (Implementation.IsValid())
+		{
+			Implementation->GetViewBookmark(OutBookmarkId);
+		}
+	}
+
 	/** Check if node title is visible with optional flag to ensure it is */
 	virtual bool IsNodeTitleVisible(const class UEdGraphNode* Node, bool bRequestRename)
 	{
diff --git a/Engine/Source/Editor/UnrealEd/Public/Kismet2/BlueprintEditorUtils.h b/Engine/Source/Editor/UnrealEd/Public/Kismet2/BlueprintEditorUtils.h
--- a/Engine/Source/Editor/UnrealEd/Public/Kismet2/BlueprintEditorUtils.h
+++ b/Engine/Source/Editor/UnrealEd/Public/Kismet2/BlueprintEditorUtils.h
@@ -1566,7 +1566,10 @@ public:
 	/**
 	 * Returns true if this BP is currently based on a type that returns true for the UObject::ImplementsGetWorld() call:
 	 */
-	static bool ImplentsGetWorld(const UBlueprint* BP);
+	static bool ImplementsGetWorld(const UBlueprint* BP);
+
+	DEPRECATED(4.20, "Use ImplementsGetWorld")
+	static bool ImplentsGetWorld(const UBlueprint* BP) { return ImplementsGetWorld(BP); }
 };
 
 struct UNREALED_API FBlueprintDuplicationScopeFlags
diff --git a/Engine/Source/Editor/UnrealEd/Public/Kismet2/KismetDebugUtilities.h b/Engine/Source/Editor/UnrealEd/Public/Kismet2/KismetDebugUtilities.h
--- a/Engine/Source/Editor/UnrealEd/Public/Kismet2/KismetDebugUtilities.h
+++ b/Engine/Source/Editor/UnrealEd/Public/Kismet2/KismetDebugUtilities.h
@@ -24,6 +24,17 @@ struct FKismetTraceSample
 	double ObservationTime;
 };
 
+//////////////////////////////////////////////////////////////////////////
+// FDebugInfo
+
+struct FDebugInfo
+{
+	FText DisplayName;
+	FText Value;
+	FText Type;
+	TArray<FDebugInfo> Children;
+};
+
 //////////////////////////////////////////////////////////////////////////
 // FObjectsBeingDebuggedIterator
 
@@ -105,6 +116,9 @@ private:
 class UNREALED_API FKismetDebugUtilities
 {
 public:
+	// Delegate for when pins are added or removed from the watchlist
+	DECLARE_MULTICAST_DELEGATE_OneParam(FOnWatchedPinsListChanged, class UBlueprint*);
+
 	static void OnScriptException(const UObject* ActiveObject, const FFrame& StackFrame, const FBlueprintExceptionInfo& Info);
 
 	/** Returns the current instruction; if a PIE/SIE session is started but paused; otherwise NULL */
@@ -195,6 +209,9 @@ public:
 	/** Deletes all breakpoints in this blueprint */
 	static void ClearBreakpoints(UBlueprint* Blueprint);
 
+	// Notifies listeners when a watched pin is added or removed
+	static FOnWatchedPinsListChanged WatchedPinsListChangedEvent;
+
 	static bool CanWatchPin(const UBlueprint* Blueprint, const UEdGraphPin* Pin);
 	static bool IsPinBeingWatched(const UBlueprint* Blueprint, const UEdGraphPin* Pin);
 	static void TogglePinWatch(UBlueprint* Blueprint, const UEdGraphPin* Pin);
@@ -219,6 +236,9 @@ public:
 	// Gets the watched tooltip for a specified site
 	static EWatchTextResult GetWatchText(FString& OutWatchText, UBlueprint* Blueprint, UObject* ActiveObject, const UEdGraphPin* WatchPin);
 
+	// Gets the debug info for a specified site
+	static EWatchTextResult GetDebugInfo(FDebugInfo& OutDebugInfo, UBlueprint* Blueprint, UObject* ActiveObject, const UEdGraphPin* WatchPin);
+
 	//@TODO: Pretty lame way to handle this messaging, ideally the entire Info object gets pushed into the editor when intraframe debugging is triggered!
 	// This doesn't work properly if there is more than one blueprint editor open at once either (one will consume it, the others will be left in the cold)
 	static FText GetAndClearLastExceptionMessage();
@@ -226,6 +246,13 @@ protected:
 	static void CheckBreakConditions(UEdGraphNode* NodeStoppedAt, bool bHitBreakpoint, int32 BreakpointOffset, bool& InOutBreakExecution);
 	static void AttemptToBreakExecution(UBlueprint* BlueprintObj, const UObject* ActiveObject, const FFrame& StackFrame, const FBlueprintExceptionInfo& Info, UEdGraphNode* NodeStoppedAt, int32 DebugOpcodeOffset);
 
+	static void GetDebugInfo_InContainer(int32 Index, FDebugInfo& DebugInfo, UProperty* Property, const void* Data);
+	static void GetDebugInfoInternal(FDebugInfo& DebugInfo, UProperty* Property, const void* PropertyValue);
+
+	// helper function for converting between blueprint and debuggable data
+	// output params are only valid if the return result is EWatchTextResult::EWTR_Valid
+	static EWatchTextResult FindDebuggingData(UBlueprint* Blueprint, UObject* ActiveObject, const UEdGraphPin* WatchPin, UProperty*& OutProperty, void*& OutData, void*& OutDelta, UObject*& OutParent);
+
 private:
 	FKismetDebugUtilities() {}
 };
diff --git a/Engine/Source/Editor/UnrealEd/UnrealEd.Build.cs b/Engine/Source/Editor/UnrealEd/UnrealEd.Build.cs
--- a/Engine/Source/Editor/UnrealEd/UnrealEd.Build.cs
+++ b/Engine/Source/Editor/UnrealEd/UnrealEd.Build.cs
@@ -59,8 +59,9 @@ public class UnrealEd : ModuleRules
 				"ViewportInteraction",
 				"VREditor",
 				"Persona",
-                "PhysicsAssetEditor",
-                "ClothingSystemEditorInterface",
+				"PhysicsAssetEditor",
+				"ClothingSystemEditorInterface",
+				"NavigationSystem",
 			}
 		);
 
@@ -82,7 +83,6 @@ public class UnrealEd : ModuleRules
 				"EditorStyle",
 				"SourceControl",
 				"UnrealEdMessages",
-				"AIModule",
 				"GameplayDebugger",
 				"BlueprintGraph",
 				"Http",
@@ -93,6 +93,7 @@ public class UnrealEd : ModuleRules
 				"AudioEditor",
 				"NetworkFileSystem",
 				"UMG",
+				"NavigationSystem",
                 "MeshDescription",
                 "MeshBuilder",
             }
@@ -186,7 +187,7 @@ public class UnrealEd : ModuleRules
 				"LevelEditor",
 				"PackagesDialog",
 				"Persona",
-                "PhysicsAssetEditor",
+				"PhysicsAssetEditor",
 				"ProjectLauncher",
 				"DeviceManager",
 				"SettingsEditor",
@@ -233,15 +234,15 @@ public class UnrealEd : ModuleRules
 		{
 			DynamicallyLoadedModuleNames.Add("AndroidPlatformEditor");
 		}
-
-		CircularlyReferencedDependentModules.AddRange(
+        
+        CircularlyReferencedDependentModules.AddRange(
 			new string[]
 			{
 				"GraphEditor",
 				"Kismet",
 				"AudioEditor",
 				"ViewportInteraction",
-				"VREditor"
+				"VREditor",
 			}
 		);
 
@@ -257,6 +258,9 @@ public class UnrealEd : ModuleRules
 				"MeshUtilities",
 				"AssetTools",
 				"KismetCompiler",
+				"NavigationSystem",
+				"GameplayTasks",
+				"AIModule",
 			}
 			);
 
diff --git a/Engine/Source/Programs/IncludeTool/IncludeTool/Rules.cs b/Engine/Source/Programs/IncludeTool/IncludeTool/Rules.cs
--- a/Engine/Source/Programs/IncludeTool/IncludeTool/Rules.cs
+++ b/Engine/Source/Programs/IncludeTool/IncludeTool/Rules.cs
@@ -553,6 +553,10 @@ namespace IncludeTool
 			{
 				return true;
 			}
+			if(Markup.Type == PreprocessorMarkupType.Define && Markup.Tokens[0].Text == "DEPRECATED_FORGAME")
+			{
+				return true;
+			}
 			if(Markup.Type == PreprocessorMarkupType.Undef && Markup.Tokens[0].Text == "TEXT")
 			{
 				return true;
diff --git a/Engine/Source/Programs/UnrealBuildTool/Configuration/UEBuildModule.cs b/Engine/Source/Programs/UnrealBuildTool/Configuration/UEBuildModule.cs
--- a/Engine/Source/Programs/UnrealBuildTool/Configuration/UEBuildModule.cs
+++ b/Engine/Source/Programs/UnrealBuildTool/Configuration/UEBuildModule.cs
@@ -513,7 +513,7 @@ namespace UnrealBuildTool
 			bool bWithLegacyPublicIncludePaths
 			)
 		{
-			if (this.Type.IsGameModule())
+			if (!Rules.bTreatAsEngineModule)
 			{
 				Definitions.Add("DEPRECATED_FORGAME=DEPRECATED");
 			}
diff --git a/Engine/Source/Programs/UnrealBuildTool/Configuration/UEBuildModuleCPP.cs b/Engine/Source/Programs/UnrealBuildTool/Configuration/UEBuildModuleCPP.cs
--- a/Engine/Source/Programs/UnrealBuildTool/Configuration/UEBuildModuleCPP.cs
+++ b/Engine/Source/Programs/UnrealBuildTool/Configuration/UEBuildModuleCPP.cs
@@ -221,7 +221,6 @@ namespace UnrealBuildTool
 		/// </summary>
 		static readonly KeyValuePair<string, string>[] WhitelistedCircularDependencies =
 		{
-			new KeyValuePair<string, string>("Engine", "AIModule"),
 			new KeyValuePair<string, string>("Engine", "Landscape"),
 			new KeyValuePair<string, string>("Engine", "UMG"),
 			new KeyValuePair<string, string>("Engine", "GameplayTags"),
@@ -273,6 +272,7 @@ namespace UnrealBuildTool
             new KeyValuePair<string, string>("MeshPaintMode", "ViewportInteraction"),
             new KeyValuePair<string, string>("MeshPaintMode", "VREditor"),
             new KeyValuePair<string, string>("Sequencer", "ViewportInteraction"),
+            new KeyValuePair<string, string>("NavigationSystem", "UnrealEd"),
         };
 
 
diff --git a/Engine/Source/Programs/UnrealHeaderTool/Private/HeaderParser.cpp b/Engine/Source/Programs/UnrealHeaderTool/Private/HeaderParser.cpp
--- a/Engine/Source/Programs/UnrealHeaderTool/Private/HeaderParser.cpp
+++ b/Engine/Source/Programs/UnrealHeaderTool/Private/HeaderParser.cpp
@@ -951,6 +951,18 @@ namespace
 			}
 			break;
 
+			case ECheckedMetadataSpecifier::ExpandEnumAsExecs:
+			{
+				if (UFunction* Function = Cast<UFunction>(Field))
+				{
+					if (FHeaderParser::FindField(Function, *InValue, false) == nullptr)
+					{
+						UE_LOG_ERROR_UHT(TEXT("Function does not have a parameter named '%s'"), *InValue);
+					}
+				}
+			}
+			break;
+
 			case ECheckedMetadataSpecifier::DevelopmentStatus:
 			{
 				const FString EarlyAccessValue(TEXT("EarlyAccess"));
diff --git a/Engine/Source/Programs/UnrealHeaderTool/Private/HeaderParser.h b/Engine/Source/Programs/UnrealHeaderTool/Private/HeaderParser.h
--- a/Engine/Source/Programs/UnrealHeaderTool/Private/HeaderParser.h
+++ b/Engine/Source/Programs/UnrealHeaderTool/Private/HeaderParser.h
@@ -434,6 +434,20 @@ public:
 	static void RequireSpecifierValue(const FPropertySpecifier& Specifier, bool bRequireExactlyOne = false);
 	static FString RequireExactlyOneSpecifierValue(const FPropertySpecifier& Specifier);
 
+	/**
+	* Find a field in the specified context.  Starts with the specified scope, then iterates
+	* through the Outer chain until the field is found.
+	* 
+	* @param	InScope				scope to start searching for the field in 
+	* @param	InIdentifier		name of the field we're searching for
+	* @param	bIncludeParents		whether to allow searching in the scope of a parent struct
+	* @param	FieldClass			class of the field to search for.  used to e.g. search for functions only
+	* @param	Thing				hint text that will be used in the error message if an error is encountered
+	*
+	* @return	a pointer to a UField with a name matching InIdentifier, or NULL if it wasn't found
+	*/
+	static UField* FindField( UStruct* InScope, const TCHAR* InIdentifier, bool bIncludeParents=true, UClass* FieldClass=UField::StaticClass(), const TCHAR* Thing=nullptr );
+
 protected:
 
 	/**
@@ -543,19 +557,6 @@ protected:
 
 	UStruct* GetSuperScope( UStruct* CurrentScope, const FName& SearchName );
 
-	/**
-	 * Find a field in the specified context.  Starts with the specified scope, then iterates
-	 * through the Outer chain until the field is found.
-	 * 
-	 * @param	InScope				scope to start searching for the field in 
-	 * @param	InIdentifier		name of the field we're searching for
-	 * @param	bIncludeParents		whether to allow searching in the scope of a parent struct
-	 * @param	FieldClass			class of the field to search for.  used to e.g. search for functions only
-	 * @param	Thing				hint text that will be used in the error message if an error is encountered
-	 *
-	 * @return	a pointer to a UField with a name matching InIdentifier, or NULL if it wasn't found
-	 */
-	UField* FindField( UStruct* InScope, const TCHAR* InIdentifier, bool bIncludeParents=true, UClass* FieldClass=UField::StaticClass(), const TCHAR* Thing=NULL );
 	void SkipStatements( int32 SubCount, const TCHAR* ErrorTag );
 
 	/**
diff --git a/Engine/Source/Programs/UnrealHeaderTool/Private/Specifiers/CheckedMetadataSpecifiers.def b/Engine/Source/Programs/UnrealHeaderTool/Private/Specifiers/CheckedMetadataSpecifiers.def
--- a/Engine/Source/Programs/UnrealHeaderTool/Private/Specifiers/CheckedMetadataSpecifiers.def
+++ b/Engine/Source/Programs/UnrealHeaderTool/Private/Specifiers/CheckedMetadataSpecifiers.def
@@ -6,6 +6,7 @@ CHECKED_METADATA_SPECIFIER(BlueprintProtected)
 CHECKED_METADATA_SPECIFIER(ClampMax)
 CHECKED_METADATA_SPECIFIER(ClampMin)
 CHECKED_METADATA_SPECIFIER(DevelopmentStatus)
+CHECKED_METADATA_SPECIFIER(ExpandEnumAsExecs)
 CHECKED_METADATA_SPECIFIER(UIMax)
 CHECKED_METADATA_SPECIFIER(UIMin)
 CHECKED_METADATA_SPECIFIER(Units)
diff --git a/Engine/Source/Runtime/AIModule/AIModule.Build.cs b/Engine/Source/Runtime/AIModule/AIModule.Build.cs
--- a/Engine/Source/Runtime/AIModule/AIModule.Build.cs
+++ b/Engine/Source/Runtime/AIModule/AIModule.Build.cs
@@ -2,54 +2,56 @@
 
 namespace UnrealBuildTool.Rules
 {
-	public class AIModule : ModuleRules
-	{
+    public class AIModule : ModuleRules
+    {
         public AIModule(ReadOnlyTargetRules Target) : base(Target)
-		{
-			PublicIncludePaths.AddRange(
-				new string[] {
+        {
+            PublicIncludePaths.AddRange(
+                new string[] {
+                    "Runtime/NavigationSystem/Public",
                     "Runtime/AIModule/Public",
-				}
-				);
+                }
+                );
 
-			PrivateIncludePaths.AddRange(
-				new string[] {
-					"Runtime/AIModule/Private",
+            PrivateIncludePaths.AddRange(
+                new string[] {
+                    "Runtime/AIModule/Private",
                     "Runtime/Engine/Private",
-				}
-				);
+                }
+                );
 
-			PublicDependencyModuleNames.AddRange(
-				new string[] {
-					"Core",
-					"CoreUObject",
-					"Engine",    
-               		"GameplayTags",
-                    "GameplayTasks"
-				}
-				);
+            PublicDependencyModuleNames.AddRange(
+                new string[] {
+                    "Core",
+                    "CoreUObject",
+                    "Engine",
+                    "GameplayTags",
+                    "GameplayTasks",
+                    "NavigationSystem",
+                }
+                );
 
-			PrivateDependencyModuleNames.AddRange(
-				new string[] {
-					"RHI",
+            PrivateDependencyModuleNames.AddRange(
+                new string[] {
+                    "RHI",
                     "RenderCore",
                     "ShaderCore",
-				}
-				);
+                }
+                );
 
-			DynamicallyLoadedModuleNames.AddRange(
-				new string[] {
+            DynamicallyLoadedModuleNames.AddRange(
+                new string[] {
 					// ... add any modules that your module loads dynamically here ...
 				}
-				);
+                );
 
-			if (Target.bBuildEditor == true)
-			{
-				PrivateDependencyModuleNames.Add("UnrealEd");
+            if (Target.bBuildEditor == true)
+            {
+                PrivateDependencyModuleNames.Add("UnrealEd");
 
-                PrivateDependencyModuleNames.Add("AITestSuite");                
+                PrivateDependencyModuleNames.Add("AITestSuite");
                 CircularlyReferencedDependentModules.Add("AITestSuite");
-			}
+            }
 
             if (Target.bCompileRecast)
             {
@@ -64,15 +66,15 @@ namespace UnrealBuildTool.Rules
                 PublicDefinitions.Add("WITH_RECAST=0");
             }
 
-			if (Target.bBuildDeveloperTools || (Target.Configuration != UnrealTargetConfiguration.Shipping && Target.Configuration != UnrealTargetConfiguration.Test))
-			{
-				PrivateDependencyModuleNames.Add("GameplayDebugger");
-				PublicDefinitions.Add("WITH_GAMEPLAY_DEBUGGER=1");
-			}
-			else
-			{
-				PublicDefinitions.Add("WITH_GAMEPLAY_DEBUGGER=0");
-			}
-		}
+            if (Target.bBuildDeveloperTools || (Target.Configuration != UnrealTargetConfiguration.Shipping && Target.Configuration != UnrealTargetConfiguration.Test))
+            {
+                PrivateDependencyModuleNames.Add("GameplayDebugger");
+                PublicDefinitions.Add("WITH_GAMEPLAY_DEBUGGER=1");
+            }
+            else
+            {
+                PublicDefinitions.Add("WITH_GAMEPLAY_DEBUGGER=0");
+            }
+        }
     }
 }
diff --git a/Engine/Source/Runtime/AIModule/Classes/AIController.h b/Engine/Source/Runtime/AIModule/Classes/AIController.h
--- a/Engine/Source/Runtime/AIModule/Classes/AIController.h
+++ b/Engine/Source/Runtime/AIModule/Classes/AIController.h
@@ -7,7 +7,7 @@
 #include "UObject/UObjectGlobals.h"
 #include "Templates/SubclassOf.h"
 #include "EngineDefines.h"
-#include "AI/Navigation/NavFilters/NavigationQueryFilter.h"
+#include "NavFilters/NavigationQueryFilter.h"
 #include "AITypes.h"
 #include "GameplayTaskOwnerInterface.h"
 #include "GameplayTask.h"
@@ -254,10 +254,7 @@ public:
 	/** Updates state of movement block detection. */
 	UFUNCTION(BlueprintCallable, Category = "AI|Navigation")
 	void SetMoveBlockDetection(bool bEnable);
-
-	/** Prepares path finding and path following components. */
-	virtual void InitNavigationControl(UPathFollowingComponent*& PathFollowingComp) override;
-
+	
 	/** Starts executing behavior tree. */
 	UFUNCTION(BlueprintCallable, Category = "AI")
 	virtual bool RunBehaviorTree(UBehaviorTree* BTAsset);
@@ -384,6 +381,7 @@ public:
 	// INavAgentInterface
 	//----------------------------------------------------------------------//
 	virtual bool IsFollowingAPath() const override;
+	virtual IPathFollowingAgentInterface* GetPathFollowingAgent() const override { return PathFollowingComponent; }
 
 	//----------------------------------------------------------------------//
 	// IGenericTeamAgentInterface
diff --git a/Engine/Source/Runtime/AIModule/Classes/AITypes.h b/Engine/Source/Runtime/AIModule/Classes/AITypes.h
--- a/Engine/Source/Runtime/AIModule/Classes/AITypes.h
+++ b/Engine/Source/Runtime/AIModule/Classes/AITypes.h
@@ -6,8 +6,9 @@
 #include "Stats/Stats.h"
 #include "UObject/ObjectMacros.h"
 #include "Templates/SubclassOf.h"
-#include "AI/Navigation/NavFilters/NavigationQueryFilter.h"
+#include "NavFilters/NavigationQueryFilter.h"
 #include "AI/Navigation/NavigationTypes.h"
+#include "NavigationSystemTypes.h"
 #include "GameFramework/Actor.h"
 #include "AITypes.generated.h"
 
diff --git a/Engine/Source/Runtime/AIModule/Classes/Blueprint/AIBlueprintHelperLibrary.h b/Engine/Source/Runtime/AIModule/Classes/Blueprint/AIBlueprintHelperLibrary.h
--- a/Engine/Source/Runtime/AIModule/Classes/Blueprint/AIBlueprintHelperLibrary.h
+++ b/Engine/Source/Runtime/AIModule/Classes/Blueprint/AIBlueprintHelperLibrary.h
@@ -65,4 +65,9 @@ class AIMODULE_API UAIBlueprintHelperLibrary : public UBlueprintFunctionLibrary
 	UFUNCTION(BlueprintPure, Category = "AI", meta = (UnsafeDuringActorConstruction = "true"))
 	static UNavigationPath* GetCurrentPath(AController* Controller);
 
+	UFUNCTION(BlueprintCallable, Category = "AI|Navigation")
+	static void SimpleMoveToActor(AController* Controller, const AActor* Goal);
+
+	UFUNCTION(BlueprintCallable, Category = "AI|Navigation")
+	static void SimpleMoveToLocation(AController* Controller, const FVector& Goal);
 };
diff --git a/Engine/Source/Runtime/AIModule/Classes/DataProviders/AIDataProvider.h b/Engine/Source/Runtime/AIModule/Classes/DataProviders/AIDataProvider.h
--- a/Engine/Source/Runtime/AIModule/Classes/DataProviders/AIDataProvider.h
+++ b/Engine/Source/Runtime/AIModule/Classes/DataProviders/AIDataProvider.h
@@ -52,7 +52,10 @@ public:
 
 	/** return raw data from provider's property */
 	template<typename T>
-	T* GetRawValuePtr() const;
+	T* GetRawValuePtr() const
+	{
+		return CachedProperty ? CachedProperty->ContainerPtrToValuePtr<T>(DataBinding) : nullptr;
+	}
 
 	/** bind data in provider and cache property for faster access */
 	void BindData(const UObject* Owner, int32 RequestId) const;
diff --git a/Engine/Source/Runtime/AIModule/Classes/EnvironmentQuery/Tests/EnvQueryTest_Pathfinding.h b/Engine/Source/Runtime/AIModule/Classes/EnvironmentQuery/Tests/EnvQueryTest_Pathfinding.h
--- a/Engine/Source/Runtime/AIModule/Classes/EnvironmentQuery/Tests/EnvQueryTest_Pathfinding.h
+++ b/Engine/Source/Runtime/AIModule/Classes/EnvironmentQuery/Tests/EnvQueryTest_Pathfinding.h
@@ -5,15 +5,16 @@
 #include "CoreMinimal.h"
 #include "UObject/ObjectMacros.h"
 #include "Templates/SubclassOf.h"
-#include "AI/Navigation/NavFilters/NavigationQueryFilter.h"
+#include "NavFilters/NavigationQueryFilter.h"
 #include "AI/Navigation/NavigationTypes.h"
+#include "NavigationSystemTypes.h"
 #include "EnvironmentQuery/EnvQueryContext.h"
 #include "DataProviders/AIDataProvider.h"
 #include "EnvironmentQuery/EnvQueryTest.h"
 #include "EnvQueryTest_Pathfinding.generated.h"
 
 class ANavigationData;
-class UNavigationSystem;
+class UNavigationSystemV1;
 
 UENUM()
 namespace EEnvTestPathfinding
@@ -64,15 +65,15 @@ class UEnvQueryTest_Pathfinding : public UEnvQueryTest
 
 protected:
 
-	DECLARE_DELEGATE_RetVal_SevenParams(bool, FTestPathSignature, const FVector&, const FVector&, EPathFindingMode::Type, const ANavigationData&, UNavigationSystem&, FSharedConstNavQueryFilter, const UObject*);
-	DECLARE_DELEGATE_RetVal_SevenParams(float, FFindPathSignature, const FVector&, const FVector&, EPathFindingMode::Type, const ANavigationData&, UNavigationSystem&, FSharedConstNavQueryFilter, const UObject*);
+	DECLARE_DELEGATE_RetVal_SevenParams(bool, FTestPathSignature, const FVector&, const FVector&, EPathFindingMode::Type, const ANavigationData&, UNavigationSystemV1&, FSharedConstNavQueryFilter, const UObject*);
+	DECLARE_DELEGATE_RetVal_SevenParams(float, FFindPathSignature, const FVector&, const FVector&, EPathFindingMode::Type, const ANavigationData&, UNavigationSystemV1&, FSharedConstNavQueryFilter, const UObject*);
 
-	bool TestPathFrom(const FVector& ItemPos, const FVector& ContextPos, EPathFindingMode::Type Mode, const ANavigationData& NavData, UNavigationSystem& NavSys, FSharedConstNavQueryFilter NavFilter, const UObject* PathOwner) const;
-	bool TestPathTo(const FVector& ItemPos, const FVector& ContextPos, EPathFindingMode::Type Mode, const ANavigationData& NavData, UNavigationSystem& NavSys, FSharedConstNavQueryFilter NavFilter, const UObject* PathOwner) const;
-	float FindPathCostFrom(const FVector& ItemPos, const FVector& ContextPos, EPathFindingMode::Type Mode, const ANavigationData& NavData, UNavigationSystem& NavSys, FSharedConstNavQueryFilter NavFilter, const UObject* PathOwner) const;
-	float FindPathCostTo(const FVector& ItemPos, const FVector& ContextPos, EPathFindingMode::Type Mode, const ANavigationData& NavData, UNavigationSystem& NavSys, FSharedConstNavQueryFilter NavFilter, const UObject* PathOwner) const;
-	float FindPathLengthFrom(const FVector& ItemPos, const FVector& ContextPos, EPathFindingMode::Type Mode, const ANavigationData& NavData, UNavigationSystem& NavSys, FSharedConstNavQueryFilter NavFilter, const UObject* PathOwner) const;
-	float FindPathLengthTo(const FVector& ItemPos, const FVector& ContextPos, EPathFindingMode::Type Mode, const ANavigationData& NavData, UNavigationSystem& NavSys, FSharedConstNavQueryFilter NavFilter, const UObject* PathOwner) const;
+	bool TestPathFrom(const FVector& ItemPos, const FVector& ContextPos, EPathFindingMode::Type Mode, const ANavigationData& NavData, UNavigationSystemV1& NavSys, FSharedConstNavQueryFilter NavFilter, const UObject* PathOwner) const;
+	bool TestPathTo(const FVector& ItemPos, const FVector& ContextPos, EPathFindingMode::Type Mode, const ANavigationData& NavData, UNavigationSystemV1& NavSys, FSharedConstNavQueryFilter NavFilter, const UObject* PathOwner) const;
+	float FindPathCostFrom(const FVector& ItemPos, const FVector& ContextPos, EPathFindingMode::Type Mode, const ANavigationData& NavData, UNavigationSystemV1& NavSys, FSharedConstNavQueryFilter NavFilter, const UObject* PathOwner) const;
+	float FindPathCostTo(const FVector& ItemPos, const FVector& ContextPos, EPathFindingMode::Type Mode, const ANavigationData& NavData, UNavigationSystemV1& NavSys, FSharedConstNavQueryFilter NavFilter, const UObject* PathOwner) const;
+	float FindPathLengthFrom(const FVector& ItemPos, const FVector& ContextPos, EPathFindingMode::Type Mode, const ANavigationData& NavData, UNavigationSystemV1& NavSys, FSharedConstNavQueryFilter NavFilter, const UObject* PathOwner) const;
+	float FindPathLengthTo(const FVector& ItemPos, const FVector& ContextPos, EPathFindingMode::Type Mode, const ANavigationData& NavData, UNavigationSystemV1& NavSys, FSharedConstNavQueryFilter NavFilter, const UObject* PathOwner) const;
 
-	ANavigationData* FindNavigationData(UNavigationSystem& NavSys, UObject* Owner) const;
+	ANavigationData* FindNavigationData(UNavigationSystemV1& NavSys, UObject* Owner) const;
 };
diff --git a/Engine/Source/Runtime/AIModule/Classes/Navigation/CrowdManager.h b/Engine/Source/Runtime/AIModule/Classes/Navigation/CrowdManager.h
--- a/Engine/Source/Runtime/AIModule/Classes/Navigation/CrowdManager.h
+++ b/Engine/Source/Runtime/AIModule/Classes/Navigation/CrowdManager.h
@@ -2,7 +2,7 @@
 
 #pragma once
 
-#include "AI/Navigation/CrowdManagerBase.h"
+#include "CrowdManagerBase.h"
 #include "CoreMinimal.h"
 #include "Stats/Stats.h"
 #include "UObject/ObjectMacros.h"
@@ -10,6 +10,7 @@
 #include "AI/Navigation/NavigationTypes.h"
 #include "Tickable.h"
 #include "DrawDebugHelpers.h"
+#include "NavFilters/NavigationQueryFilter.h"
 #include "CrowdManager.generated.h"
 
 class ANavigationData;
@@ -39,7 +40,7 @@ class dtObstacleAvoidanceDebugData;
  *  but it's slower than AvoidanceManager solution (RVO, cares only about agents)
  *
  *  All agents will operate on the same navmesh data, which will be picked from
- *  navigation system defaults (UNavigationSystem::SupportedAgents[0])
+ *  navigation system defaults (UNavigationSystemV1::SupportedAgents[0])
  *
  *  To use it, you have to add CrowdFollowingComponent to your agent
  *  (usually: replace class of PathFollowingComponent in AIController by adding 
diff --git a/Engine/Source/Runtime/AIModule/Classes/Navigation/NavLinkProxy.h b/Engine/Source/Runtime/AIModule/Classes/Navigation/NavLinkProxy.h
new file mode 100644
--- /dev/null
+++ b/Engine/Source/Runtime/AIModule/Classes/Navigation/NavLinkProxy.h
@@ -0,0 +1,125 @@
+// Copyright 1998-2018 Epic Games, Inc. All Rights Reserved.
+
+#pragma once
+
+#include "CoreMinimal.h"
+#include "UObject/ObjectMacros.h"
+#include "Templates/SubclassOf.h"
+#include "GameFramework/Actor.h"
+#include "EngineDefines.h"
+#include "AI/Navigation/NavRelevantInterface.h"
+#include "AI/Navigation/NavLinkDefinition.h"
+#include "NavLinkHostInterface.h"
+#include "NavLinkProxy.generated.h"
+
+class UBillboardComponent;
+class UNavLinkCustomComponent;
+class UNavLinkRenderingComponent;
+class UPathFollowingComponent;
+struct FNavigationRelevantData;
+
+DECLARE_DYNAMIC_MULTICAST_DELEGATE_TwoParams( FSmartLinkReachedSignature, AActor*, MovingActor, const FVector&, DestinationPoint );
+
+UCLASS(Blueprintable, autoCollapseCategories=(SmartLink, Actor), hideCategories=(Input))
+class AIMODULE_API ANavLinkProxy : public AActor, public INavLinkHostInterface, public INavRelevantInterface
+{
+	GENERATED_UCLASS_BODY()
+
+	/** Navigation links (point to point) added to navigation data */
+	UPROPERTY(EditAnywhere, Category=SimpleLink)
+	TArray<FNavigationLink> PointLinks;
+	
+	/** Navigation links (segment to segment) added to navigation data
+	*	@todo hidden from use until we fix segment links. Not really working now*/
+	UPROPERTY()
+	TArray<FNavigationSegmentLink> SegmentLinks;
+
+private:
+	/** Smart link: can affect path following */
+	UPROPERTY(VisibleAnywhere, Category=SmartLink)
+	UNavLinkCustomComponent* SmartLinkComp;
+public:
+
+	/** Smart link: toggle relevancy */
+	UPROPERTY(EditAnywhere, Category=SmartLink)
+	bool bSmartLinkIsRelevant;
+
+#if WITH_EDITORONLY_DATA
+private:
+	/** Editor Preview */
+	UPROPERTY()
+	UNavLinkRenderingComponent* EdRenderComp;
+
+	UPROPERTY()
+	UBillboardComponent* SpriteComponent;
+public:
+#endif // WITH_EDITORONLY_DATA
+
+	// BEGIN INavRelevantInterface
+	virtual void GetNavigationData(FNavigationRelevantData& Data) const override;
+	virtual FBox GetNavigationBounds() const override;
+	virtual bool IsNavigationRelevant() const override;
+	// END INavRelevantInterface
+
+	// BEGIN INavLinkHostInterface
+	virtual bool GetNavigationLinksClasses(TArray<TSubclassOf<UNavLinkDefinition> >& OutClasses) const override;
+	virtual bool GetNavigationLinksArray(TArray<FNavigationLink>& OutLink, TArray<FNavigationSegmentLink>& OutSegments) const override;
+	// END INavLinkHostInterface
+
+#if WITH_EDITOR
+	virtual void PostEditChangeProperty(FPropertyChangedEvent& PropertyChangedEvent) override;
+	virtual void PostEditUndo() override;
+	virtual void PostEditImport() override;
+#endif // WITH_EDITOR
+	virtual void PostRegisterAllComponents() override;
+	virtual void PostLoad() override;
+
+#if ENABLE_VISUAL_LOG
+protected:
+	virtual void BeginPlay() override;
+
+public:
+#endif // ENABLE_VISUAL_LOG
+
+	virtual FBox GetComponentsBoundingBox(bool bNonColliding = false) const override;
+
+	//////////////////////////////////////////////////////////////////////////
+	// Blueprint interface for smart links
+
+	/** called when agent reaches smart link during path following, use ResumePathFollowing() to give control back */
+	UFUNCTION(BlueprintImplementableEvent)
+	void ReceiveSmartLinkReached(AActor* Agent, const FVector& Destination);
+
+	/** resume normal path following */
+	UFUNCTION(BlueprintCallable, Category="AI|Navigation")
+	void ResumePathFollowing(AActor* Agent);
+
+	/** check if smart link is enabled */
+	UFUNCTION(BlueprintCallable, Category="AI|Navigation")
+	bool IsSmartLinkEnabled() const;
+
+	/** change state of smart link */
+	UFUNCTION(BlueprintCallable, Category="AI|Navigation")
+	void SetSmartLinkEnabled(bool bEnabled);
+
+	/** check if any agent is moving through smart link right now */
+	UFUNCTION(BlueprintCallable, Category="AI|Navigation")
+	bool HasMovingAgents() const;
+
+protected:
+
+	UPROPERTY(BlueprintAssignable)
+	FSmartLinkReachedSignature OnSmartLinkReached;
+
+	void NotifySmartLinkReached(UNavLinkCustomComponent* LinkComp, UObject* PathingAgent, const FVector& DestPoint);
+
+public:
+	/** Returns SmartLinkComp subobject **/
+	UNavLinkCustomComponent* GetSmartLinkComp() const { return SmartLinkComp; }
+#if WITH_EDITORONLY_DATA
+	/** Returns EdRenderComp subobject **/
+	UNavLinkRenderingComponent* GetEdRenderComp() const { return EdRenderComp; }
+	/** Returns SpriteComponent subobject **/
+	UBillboardComponent* GetSpriteComponent() const { return SpriteComponent; }
+#endif
+};
diff --git a/Engine/Source/Runtime/AIModule/Classes/Navigation/PathFollowingComponent.h b/Engine/Source/Runtime/AIModule/Classes/Navigation/PathFollowingComponent.h
--- a/Engine/Source/Runtime/AIModule/Classes/Navigation/PathFollowingComponent.h
+++ b/Engine/Source/Runtime/AIModule/Classes/Navigation/PathFollowingComponent.h
@@ -9,16 +9,18 @@
 #include "Components/ActorComponent.h"
 #include "EngineDefines.h"
 #include "AI/Navigation/NavigationTypes.h"
+#include "NavigationData.h"
 #include "AITypes.h"
 #include "AIResourceInterface.h"
-#include "AI/Navigation/NavigationData.h"
 #include "GameFramework/NavMovementComponent.h"
+#include "AI/Navigation/PathFollowingAgentInterface.h"
 #include "PathFollowingComponent.generated.h"
 
 class Error;
 class FDebugDisplayInfo;
 class INavLinkCustomInterface;
 class UCanvas;
+class ANavigationData;
 
 AIMODULE_API DECLARE_LOG_CATEGORY_EXTERN(LogPathFollowing, Warning, All);
 
@@ -211,7 +213,7 @@ enum class EPathFollowingReachMode : uint8
 };
 
 UCLASS(config=Engine)
-class AIMODULE_API UPathFollowingComponent : public UActorComponent, public IAIResourceInterface
+class AIMODULE_API UPathFollowingComponent : public UActorComponent, public IAIResourceInterface, public IPathFollowingAgentInterface
 {
 	GENERATED_UCLASS_BODY()
 
@@ -227,6 +229,7 @@ class AIMODULE_API UPathFollowingComponent : public UActorComponent, public IAIR
 	FMoveComplete OnRequestFinished;
 
 	//~ Begin UActorComponent Interface
+	virtual void OnRegister() override;
 	virtual void TickComponent(float DeltaTime, enum ELevelTick TickType, FActorComponentTickFunction *ThisTickFunction) override;
 	//~ End UActorComponent Interface
 
@@ -377,14 +380,12 @@ class AIMODULE_API UPathFollowingComponent : public UActorComponent, public IAIR
 	UFUNCTION()
 	virtual void OnActorBump(AActor* SelfActor, AActor* OtherActor, FVector NormalImpulse, const FHitResult& Hit);
 
-	/** Called when movement is blocked by a collision with another actor.  */
-	virtual void OnMoveBlockedBy(const FHitResult& BlockingImpact) {}
-
-	/** Called when falling movement starts. */
-	virtual void OnStartedFalling();
-
-	/** Called when falling movement ends. */
-	virtual void OnLanded() {}
+	// IPathFollowingAgentInterface begin
+	virtual void OnUnableToMove(const UObject& Instigator) override;		
+	//virtual void OnMoveBlockedBy(const FHitResult& BlockingImpact) {}
+	virtual void OnStartedFalling() override;
+	virtual void OnLanded() override {}
+	// IPathFollowingAgentInterface end
 
 	/** Check if path following can be activated */
 	virtual bool IsPathFollowingAllowed() const;
@@ -596,6 +597,10 @@ protected:
 	/** reset path following data */
 	virtual void Reset();
 
+	/** Called if owning Controller possesses new pawn or ends up pawn-less. 
+	 *	Doesn't get called if owner is not an AContoller */
+	virtual void OnNewPawn(APawn* NewPawn);
+
 	/** should verify if agent if still on path ater movement has been resumed? */
 	virtual bool ShouldCheckPathOnResume() const;
 
diff --git a/Engine/Source/Runtime/AIModule/Classes/Navigation/PathFollowingManager.h b/Engine/Source/Runtime/AIModule/Classes/Navigation/PathFollowingManager.h
new file mode 100644
--- /dev/null
+++ b/Engine/Source/Runtime/AIModule/Classes/Navigation/PathFollowingManager.h
@@ -0,0 +1,21 @@
+// Copyright 1998-2018 Epic Games, Inc. All Rights Reserved.
+
+#pragma once
+
+#include "UObject/Object.h"
+#include "AI/NavigationSystemBase.h"
+#include "PathFollowingManager.generated.h"
+
+class AController;
+
+
+UCLASS()
+class AIMODULE_API UPathFollowingManager : public UObject, public IPathFollowingManagerInterface
+{
+	GENERATED_BODY()
+public:
+	UPathFollowingManager(const FObjectInitializer& ObjectInitializer = FObjectInitializer::Get());
+
+	static void StopMovement(const AController& Controller);
+	static bool IsFollowingAPath(const AController& Controller);
+};
diff --git a/Engine/Source/Runtime/AIModule/Classes/Perception/AISenseConfig_Blueprint.h b/Engine/Source/Runtime/AIModule/Classes/Perception/AISenseConfig_Blueprint.h
--- a/Engine/Source/Runtime/AIModule/Classes/Perception/AISenseConfig_Blueprint.h
+++ b/Engine/Source/Runtime/AIModule/Classes/Perception/AISenseConfig_Blueprint.h
@@ -11,7 +11,7 @@
 
 class UAISense_Blueprint;
 
-UCLASS(Blueprintable, Abstract)
+UCLASS(Blueprintable, Abstract, hidedropdown)
 class AIMODULE_API UAISenseConfig_Blueprint : public UAISenseConfig
 {
 	GENERATED_UCLASS_BODY()
diff --git a/Engine/Source/Runtime/AIModule/Classes/Perception/AISense_Blueprint.h b/Engine/Source/Runtime/AIModule/Classes/Perception/AISense_Blueprint.h
--- a/Engine/Source/Runtime/AIModule/Classes/Perception/AISense_Blueprint.h
+++ b/Engine/Source/Runtime/AIModule/Classes/Perception/AISense_Blueprint.h
@@ -14,7 +14,7 @@ class UAIPerceptionComponent;
 class UAISenseEvent;
 class UUserDefinedStruct;
 
-UCLASS(ClassGroup = AI, Abstract, Blueprintable)
+UCLASS(ClassGroup = AI, Abstract, Blueprintable, hidedropdown)
 class AIMODULE_API UAISense_Blueprint : public UAISense
 {
 	GENERATED_BODY()
diff --git a/Engine/Source/Runtime/AIModule/Private/AIController.cpp b/Engine/Source/Runtime/AIModule/Private/AIController.cpp
--- a/Engine/Source/Runtime/AIModule/Private/AIController.cpp
+++ b/Engine/Source/Runtime/AIModule/Private/AIController.cpp
@@ -2,7 +2,7 @@
 
 #include "AIController.h"
 #include "CollisionQueryParams.h"
-#include "AI/Navigation/NavigationSystem.h"
+#include "NavigationSystem.h"
 #include "Components/CapsuleComponent.h"
 #include "GameFramework/PhysicsVolume.h"
 #include "Actions/PawnActionsComponent.h"
@@ -476,12 +476,8 @@ void AAIController::Possess(APawn* InPawn)
 		return;
 	}
 
-	// no point in doing navigation setup if pawn has no movement component
-	const UPawnMovementComponent* MovementComp = InPawn->GetMovementComponent();
-	if (MovementComp != NULL)
-	{
-		UpdateNavigationComponents();
-	}
+	// not calling UpdateNavigationComponents() anymore. The PathFollowingComponent 
+	// is now observing newly possessed pawns (via OnNewPawn)
 
 	if (PathFollowingComponent)
 	{
@@ -562,11 +558,6 @@ void AAIController::SetPawn(APawn* InPawn)
 	}
 }
 
-void AAIController::InitNavigationControl(UPathFollowingComponent*& PathFollowingComp)
-{
-	PathFollowingComp = PathFollowingComponent;
-}
-
 EPathFollowingRequestResult::Type AAIController::MoveToActor(AActor* Goal, float AcceptanceRadius, bool bStopOnOverlap, bool bUsePathfinding, bool bCanStrafe, TSubclassOf<UNavigationQueryFilter> FilterClass, bool bAllowPartialPaths)
 {
 	// abort active movement to keep only one request running
@@ -647,7 +638,7 @@ FPathFollowingRequestResult AAIController::MoveTo(const FAIMoveRequest& MoveRequ
 		// fail if projection to navigation is required but it failed
 		if (bCanRequestMove && MoveRequest.IsProjectingGoal())
 		{
-			UNavigationSystem* NavSys = UNavigationSystem::GetCurrent(GetWorld());
+			UNavigationSystemV1* NavSys = FNavigationSystem::GetCurrent<UNavigationSystemV1>(GetWorld());
 			const FNavAgentProperties& AgentProps = GetNavAgentPropertiesRef();
 			FNavLocation ProjectedLocation;
 
@@ -753,7 +744,7 @@ bool AAIController::BuildPathfindingQuery(const FAIMoveRequest& MoveRequest, FPa
 {
 	bool bResult = false;
 
-	UNavigationSystem* NavSys = UNavigationSystem::GetCurrent(GetWorld());
+	UNavigationSystemV1* NavSys = FNavigationSystem::GetCurrent<UNavigationSystemV1>(GetWorld());
 	const ANavigationData* NavData = (NavSys == nullptr) ? nullptr :
 		MoveRequest.IsUsingPathfinding() ? NavSys->GetNavDataForProps(GetNavAgentPropertiesRef()) :
 		NavSys->GetAbstractNavData();
@@ -798,7 +789,7 @@ void AAIController::FindPathForMoveRequest(const FAIMoveRequest& MoveRequest, FP
 {
 	SCOPE_CYCLE_COUNTER(STAT_AI_Overall);
 
-	UNavigationSystem* NavSys = UNavigationSystem::GetCurrent(GetWorld());
+	UNavigationSystemV1* NavSys = FNavigationSystem::GetCurrent<UNavigationSystemV1>(GetWorld());
 	if (NavSys)
 	{
 		FPathFindingResult PathResult = NavSys->FindPathSync(Query);
diff --git a/Engine/Source/Runtime/AIModule/Private/AIModule.cpp b/Engine/Source/Runtime/AIModule/Private/AIModule.cpp
--- a/Engine/Source/Runtime/AIModule/Private/AIModule.cpp
+++ b/Engine/Source/Runtime/AIModule/Private/AIModule.cpp
@@ -50,6 +50,7 @@ IMPLEMENT_MODULE(FAIModule, AIModule)
 
 void FAIModule::StartupModule()
 { 
+	FModuleManager::LoadModulePtr< IModuleInterface >("NavigationSystem");
 #if WITH_EDITOR 
 	FModuleManager::LoadModulePtr< IModuleInterface >("AITestSuite");
 
diff --git a/Engine/Source/Runtime/AIModule/Private/BehaviorTree/BehaviorTreeComponent.cpp b/Engine/Source/Runtime/AIModule/Private/BehaviorTree/BehaviorTreeComponent.cpp
--- a/Engine/Source/Runtime/AIModule/Private/BehaviorTree/BehaviorTreeComponent.cpp
+++ b/Engine/Source/Runtime/AIModule/Private/BehaviorTree/BehaviorTreeComponent.cpp
@@ -1009,7 +1009,7 @@ void UBehaviorTreeComponent::ApplySearchUpdates(const TArray<FBehaviorTreeSearch
 		{
 			// special case: service node at root of top most subtree - don't remove/re-add them when tree is in looping mode
 			// don't bother with decorators parent == root means that they are on child branches
-			if (bLoopExecution && UpdateInfo.AuxNode->GetParentNode() == InstanceStack[0].RootNode &&
+			if (bLoopExecution && UpdateInfo.AuxNode->GetMyNode() == InstanceStack[0].RootNode &&
 				UpdateInfo.AuxNode->IsA(UBTService::StaticClass()))
 			{
 				if (UpdateInfo.Mode == EBTNodeUpdateMode::Remove ||
@@ -1907,6 +1907,9 @@ bool UBehaviorTreeComponent::PushInstance(UBehaviorTree& TreeAsset)
 			UBTService* ServiceNode = RootNode->Services[ServiceIndex];
 			uint8* NodeMemory = (uint8*)ServiceNode->GetNodeMemory<uint8>(InstanceStack[ActiveInstanceIdx]);
 
+			// send initial on search start events in case someone is using them for init logic
+			ServiceNode->NotifyParentActivation(SearchData);
+
 			InstanceStack[ActiveInstanceIdx].ActiveAuxNodes.Add(ServiceNode);
 			ServiceNode->WrappedOnBecomeRelevant(*this, NodeMemory);
 		}
diff --git a/Engine/Source/Runtime/AIModule/Private/BehaviorTree/Decorators/BTDecorator_DoesPathExist.cpp b/Engine/Source/Runtime/AIModule/Private/BehaviorTree/Decorators/BTDecorator_DoesPathExist.cpp
--- a/Engine/Source/Runtime/AIModule/Private/BehaviorTree/Decorators/BTDecorator_DoesPathExist.cpp
+++ b/Engine/Source/Runtime/AIModule/Private/BehaviorTree/Decorators/BTDecorator_DoesPathExist.cpp
@@ -3,8 +3,8 @@
 #include "BehaviorTree/Decorators/BTDecorator_DoesPathExist.h"
 #include "UObject/Package.h"
 #include "GameFramework/Actor.h"
-#include "AI/Navigation/NavigationSystem.h"
-#include "AI/Navigation/RecastNavMesh.h"
+#include "NavigationSystem.h"
+#include "NavMesh/RecastNavMesh.h"
 #include "BehaviorTree/BlackboardComponent.h"
 #include "AIController.h"
 
@@ -60,7 +60,7 @@ bool UBTDecorator_DoesPathExist::CalculateRawConditionValue(UBehaviorTreeCompone
 	
 	bool bHasPath = false;
 
-	const UNavigationSystem* NavSys = UNavigationSystem::GetCurrent(OwnerComp.GetWorld());
+	const UNavigationSystemV1* NavSys = FNavigationSystem::GetCurrent<UNavigationSystemV1>(OwnerComp.GetWorld());
 	if (NavSys && bHasPointA && bHasPointB)
 	{
 		const AAIController* AIOwner = OwnerComp.GetAIOwner();
diff --git a/Engine/Source/Runtime/AIModule/Private/Blueprint/AIBlueprintHelperLibrary.cpp b/Engine/Source/Runtime/AIModule/Private/Blueprint/AIBlueprintHelperLibrary.cpp
--- a/Engine/Source/Runtime/AIModule/Private/Blueprint/AIBlueprintHelperLibrary.cpp
+++ b/Engine/Source/Runtime/AIModule/Private/Blueprint/AIBlueprintHelperLibrary.cpp
@@ -13,10 +13,15 @@
 #include "BehaviorTree/BlackboardComponent.h"
 #include "Blueprint/AIAsyncTaskBlueprintProxy.h"
 #include "Animation/AnimInstance.h"
-#include "AI/Navigation/NavigationPath.h"
+#include "NavigationPath.h"
+#include "NavigationData.h"
+#include "NavigationSystem.h"
+#include "Logging/MessageLog.h"
 
 DEFINE_LOG_CATEGORY_STATIC(LogAIBlueprint, Warning, All);
 
+#define LOCTEXT_NAMESPACE "AIBlueprintHelperLibrary"
+
 //----------------------------------------------------------------------//
 // UAIAsyncTaskBlueprintProxy
 //----------------------------------------------------------------------//
@@ -338,3 +343,162 @@ UNavigationPath* UAIBlueprintHelperLibrary::GetCurrentPath(AController* Controll
 	return ResultPath;
 }
 
+
+namespace
+{
+	UPathFollowingComponent* InitNavigationControl(AController& Controller)
+	{
+		AAIController* AsAIController = Cast<AAIController>(&Controller);
+		UPathFollowingComponent* PathFollowingComp = nullptr;
+
+		if (AsAIController)
+		{
+			PathFollowingComp = AsAIController->GetPathFollowingComponent();
+		}
+		else
+		{
+			PathFollowingComp = Controller.FindComponentByClass<UPathFollowingComponent>();
+			if (PathFollowingComp == nullptr)
+			{
+				PathFollowingComp = NewObject<UPathFollowingComponent>(&Controller);
+				PathFollowingComp->RegisterComponentWithWorld(Controller.GetWorld());
+				PathFollowingComp->Initialize();
+			}
+		}
+
+		return PathFollowingComp;
+	}
+}
+
+void UAIBlueprintHelperLibrary::SimpleMoveToActor(AController* Controller, const AActor* Goal)
+{
+	UNavigationSystemV1* NavSys = Controller ? FNavigationSystem::GetCurrent<UNavigationSystemV1>(Controller->GetWorld()) : nullptr;
+	if (NavSys == nullptr || Goal == nullptr || Controller == nullptr || Controller->GetPawn() == nullptr)
+	{
+		UE_LOG(LogNavigation, Warning, TEXT("UNavigationSystemV1::SimpleMoveToActor called for NavSys:%s Controller:%s controlling Pawn:%s with goal actor %s (if any of these is None then there's your problem"),
+			*GetNameSafe(NavSys), *GetNameSafe(Controller), Controller ? *GetNameSafe(Controller->GetPawn()) : TEXT("NULL"), *GetNameSafe(Goal));
+		return;
+	}
+
+	UPathFollowingComponent* PFollowComp = InitNavigationControl(*Controller);
+
+	if (PFollowComp == nullptr)
+	{
+		FMessageLog("PIE").Warning(FText::Format(
+			LOCTEXT("SimpleMoveErrorNoComp", "SimpleMove failed for {0}: missing components"),
+			FText::FromName(Controller->GetFName())
+		));
+		return;
+	}
+
+	if (!PFollowComp->IsPathFollowingAllowed())
+	{
+		FMessageLog("PIE").Warning(FText::Format(
+			LOCTEXT("SimpleMoveErrorMovement", "SimpleMove failed for {0}: movement not allowed"),
+			FText::FromName(Controller->GetFName())
+		));
+		return;
+	}
+
+	const bool bAlreadyAtGoal = PFollowComp->HasReached(*Goal, EPathFollowingReachMode::OverlapAgentAndGoal);
+
+	// script source, keep only one move request at time
+	if (PFollowComp->GetStatus() != EPathFollowingStatus::Idle)
+	{
+		PFollowComp->AbortMove(*NavSys, FPathFollowingResultFlags::ForcedScript | FPathFollowingResultFlags::NewRequest
+			, FAIRequestID::AnyRequest, bAlreadyAtGoal ? EPathFollowingVelocityMode::Reset : EPathFollowingVelocityMode::Keep);
+	}
+
+	if (bAlreadyAtGoal)
+	{
+		PFollowComp->RequestMoveWithImmediateFinish(EPathFollowingResult::Success);
+	}
+	else
+	{
+		const ANavigationData* NavData = NavSys->GetNavDataForProps(Controller->GetNavAgentPropertiesRef());
+		if (NavData)
+		{
+			FPathFindingQuery Query(Controller, *NavData, Controller->GetNavAgentLocation(), Goal->GetActorLocation());
+			FPathFindingResult Result = NavSys->FindPathSync(Query);
+			if (Result.IsSuccessful())
+			{
+				Result.Path->SetGoalActorObservation(*Goal, 100.0f);
+				PFollowComp->RequestMove(FAIMoveRequest(Goal), Result.Path);
+			}
+			else if (PFollowComp->GetStatus() != EPathFollowingStatus::Idle)
+			{
+				PFollowComp->RequestMoveWithImmediateFinish(EPathFollowingResult::Invalid);
+			}
+		}
+	}
+}
+
+void UAIBlueprintHelperLibrary::SimpleMoveToLocation(AController* Controller, const FVector& GoalLocation)
+{
+	UNavigationSystemV1* NavSys = Controller ? FNavigationSystem::GetCurrent<UNavigationSystemV1>(Controller->GetWorld()) : nullptr;
+	if (NavSys == nullptr || Controller == nullptr || Controller->GetPawn() == nullptr)
+	{
+		UE_LOG(LogNavigation, Warning, TEXT("UNavigationSystemV1::SimpleMoveToActor called for NavSys:%s Controller:%s controlling Pawn:%s (if any of these is None then there's your problem"),
+			*GetNameSafe(NavSys), *GetNameSafe(Controller), Controller ? *GetNameSafe(Controller->GetPawn()) : TEXT("NULL"));
+		return;
+	}
+
+	UPathFollowingComponent* PFollowComp = InitNavigationControl(*Controller);
+
+	if (PFollowComp == nullptr)
+	{
+		FMessageLog("PIE").Warning(FText::Format(
+			LOCTEXT("SimpleMoveErrorNoComp", "SimpleMove failed for {0}: missing components"),
+			FText::FromName(Controller->GetFName())
+		));
+		return;
+	}
+
+	if (!PFollowComp->IsPathFollowingAllowed())
+	{
+		FMessageLog("PIE").Warning(FText::Format(
+			LOCTEXT("SimpleMoveErrorMovement", "SimpleMove failed for {0}: movement not allowed"),
+			FText::FromName(Controller->GetFName())
+		));
+		return;
+	}
+
+	const bool bAlreadyAtGoal = PFollowComp->HasReached(GoalLocation, EPathFollowingReachMode::OverlapAgent);
+
+	// script source, keep only one move request at time
+	if (PFollowComp->GetStatus() != EPathFollowingStatus::Idle)
+	{
+		PFollowComp->AbortMove(*NavSys, FPathFollowingResultFlags::ForcedScript | FPathFollowingResultFlags::NewRequest
+			, FAIRequestID::AnyRequest, bAlreadyAtGoal ? EPathFollowingVelocityMode::Reset : EPathFollowingVelocityMode::Keep);
+	}
+
+	// script source, keep only one move request at time
+	if (PFollowComp->GetStatus() != EPathFollowingStatus::Idle)
+	{
+		PFollowComp->AbortMove(*NavSys, FPathFollowingResultFlags::ForcedScript | FPathFollowingResultFlags::NewRequest);
+	}
+
+	if (bAlreadyAtGoal)
+	{
+		PFollowComp->RequestMoveWithImmediateFinish(EPathFollowingResult::Success);
+	}
+	else
+	{
+		const ANavigationData* NavData = NavSys->GetNavDataForProps(Controller->GetNavAgentPropertiesRef());
+		if (NavData)
+		{
+			FPathFindingQuery Query(Controller, *NavData, Controller->GetNavAgentLocation(), GoalLocation);
+			FPathFindingResult Result = NavSys->FindPathSync(Query);
+			if (Result.IsSuccessful())
+			{
+				PFollowComp->RequestMove(FAIMoveRequest(GoalLocation), Result.Path);
+			}
+			else if (PFollowComp->GetStatus() != EPathFollowingStatus::Idle)
+			{
+				PFollowComp->RequestMoveWithImmediateFinish(EPathFollowingResult::Invalid);
+			}
+		}
+	}
+}
+
+#undef LOCTEXT_NAMESPACE
\ No newline at end of file
diff --git a/Engine/Source/Runtime/AIModule/Private/DataProviders/AIDataProvider.cpp b/Engine/Source/Runtime/AIModule/Private/DataProviders/AIDataProvider.cpp
--- a/Engine/Source/Runtime/AIModule/Private/DataProviders/AIDataProvider.cpp
+++ b/Engine/Source/Runtime/AIModule/Private/DataProviders/AIDataProvider.cpp
@@ -29,12 +29,6 @@ void FAIDataProviderValue::GetMatchingProperties(TArray<FName>& MatchingProperti
 	}
 }
 
-template<typename T>
-T* FAIDataProviderValue::GetRawValuePtr() const
-{
-	return CachedProperty ? CachedProperty->ContainerPtrToValuePtr<T>(DataBinding) : nullptr;
-}
-
 void FAIDataProviderValue::BindData(const UObject* Owner, int32 RequestId) const
 {
 	if (DataBinding && ensure(Owner))
diff --git a/Engine/Source/Runtime/AIModule/Private/EnvironmentQuery/EnvQueryInstance.cpp b/Engine/Source/Runtime/AIModule/Private/EnvironmentQuery/EnvQueryInstance.cpp
--- a/Engine/Source/Runtime/AIModule/Private/EnvironmentQuery/EnvQueryInstance.cpp
+++ b/Engine/Source/Runtime/AIModule/Private/EnvironmentQuery/EnvQueryInstance.cpp
@@ -14,6 +14,7 @@
 #include "EnvironmentQuery/Contexts/EnvQueryContext_Item.h"
 #include "EnvironmentQuery/Generators/EnvQueryGenerator_Composite.h"
 #include "AISystem.h"
+#include "AITypes.h"
 
 #if UE_BUILD_SHIPPING
 #define eqs_ensure ensure
@@ -121,11 +122,16 @@ bool FEnvQueryInstance::PrepareContext(UClass* ContextClass, FEnvQueryContextDat
 		if (CachedData == NULL)
 		{
 			UEnvQueryManager* QueryManager = UEnvQueryManager::GetCurrent(World);
-			UEnvQueryContext* ContextOb = QueryManager->PrepareLocalContext(ContextClass);
+			const UEnvQueryContext* ContextOb = QueryManager->PrepareLocalContext(ContextClass);
 
+			check(ContextOb);
 			ContextOb->ProvideContext(*this, ContextData);
 
-			DEC_MEMORY_STAT_BY(STAT_AI_EQS_InstanceMemory, GetContextAllocatedSize());
+#if STATS
+			// this contraption has been added to nail down a location of a rare bug
+			const uint32 AllocatedSize = GetContextAllocatedSize();
+			DEC_MEMORY_STAT_BY(STAT_AI_EQS_InstanceMemory, AllocatedSize);
+#endif // STATS
 
 			ContextCache.Add(ContextClass, ContextData);
 
@@ -198,8 +204,12 @@ bool FEnvQueryInstance::PrepareContext(UClass* Context, TArray<FVector>& Data)
 		Data.SetNumUninitialized(ContextData.NumValues);
 		for (int32 ValueIndex = 0; ValueIndex < ContextData.NumValues; ValueIndex++)
 		{
-			Data[ValueIndex] = DefTypeOb->GetItemLocation(ContextRawData);
-			ContextRawData += DefTypeValueSize;
+			const FVector ItemLocation = DefTypeOb->GetItemLocation(ContextRawData);
+			if (FAISystem::IsValidLocation(ItemLocation))
+			{
+				Data[ValueIndex] = ItemLocation;
+				ContextRawData += DefTypeValueSize;
+			}
 		}
 	}
 
@@ -225,8 +235,12 @@ bool FEnvQueryInstance::PrepareContext(UClass* Context, TArray<FRotator>& Data)
 		Data.SetNumUninitialized(ContextData.NumValues);
 		for (int32 ValueIndex = 0; ValueIndex < ContextData.NumValues; ValueIndex++)
 		{
-			Data[ValueIndex] = DefTypeOb->GetItemRotation(ContextRawData);
-			ContextRawData += DefTypeValueSize;
+			const FRotator ItemRotation = DefTypeOb->GetItemRotation(ContextRawData);
+			if (FAISystem::IsValidRotation(ItemRotation))
+			{
+				Data[ValueIndex] = ItemRotation;
+				ContextRawData += DefTypeValueSize;
+			}
 		}
 	}
 
diff --git a/Engine/Source/Runtime/AIModule/Private/EnvironmentQuery/EnvQueryTypes.cpp b/Engine/Source/Runtime/AIModule/Private/EnvironmentQuery/EnvQueryTypes.cpp
--- a/Engine/Source/Runtime/AIModule/Private/EnvironmentQuery/EnvQueryTypes.cpp
+++ b/Engine/Source/Runtime/AIModule/Private/EnvironmentQuery/EnvQueryTypes.cpp
@@ -3,8 +3,8 @@
 #include "EnvironmentQuery/EnvQueryTypes.h"
 #include "UObject/Package.h"
 #include "AI/Navigation/NavAgentInterface.h"
-#include "AI/Navigation/NavigationSystem.h"
-#include "AI/Navigation/RecastNavMesh.h"
+#include "NavigationSystem.h"
+#include "NavMesh/RecastNavMesh.h"
 #include "EnvironmentQuery/Items/EnvQueryItemType_VectorBase.h"
 #include "EnvironmentQuery/Items/EnvQueryItemType_ActorBase.h"
 #include "BehaviorTree/BTNode.h"
@@ -268,7 +268,7 @@ namespace FEQSHelpers
 {
 	const ANavigationData* FindNavigationDataForQuery(FEnvQueryInstance& QueryInstance)
 	{
-		const UNavigationSystem* NavSys = QueryInstance.World->GetNavigationSystem();
+		const UNavigationSystemV1* NavSys = FNavigationSystem::GetCurrent<UNavigationSystemV1>(QueryInstance.World);
 		
 		if (NavSys == nullptr)
 		{
@@ -280,10 +280,10 @@ namespace FEQSHelpers
 		if (NavAgent)
 		{
 			const FNavAgentProperties& NavAgentProps = NavAgent->GetNavAgentPropertiesRef();
-			return NavSys->GetNavDataForProps(NavAgentProps);
+			return Cast<const ANavigationData>(NavSys->GetNavDataForProps(NavAgentProps));
 		}
 
-		return NavSys->GetMainNavData();
+		return Cast<const ANavigationData>(NavSys->GetDefaultNavDataInstance());
 	}
 }
 
