Fix a crash when generating simplified collision

diff --git a/Engine/Source/Editor/UnrealEd/Private/GeomFitUtils.cpp b/Engine/Source/Editor/UnrealEd/Private/GeomFitUtils.cpp
--- a/Engine/Source/Editor/UnrealEd/Private/GeomFitUtils.cpp
+++ b/Engine/Source/Editor/UnrealEd/Private/GeomFitUtils.cpp
@@ -240,7 +240,7 @@ int32 GenerateBoxAsSimpleCollision(UStaticMesh* StaticMesh)
 	
 	FVector unitVec = bs->BuildScale3D;
 	FVector Center, Extents;
-	MeshDescription = StaticMesh->GetMeshDescription(0);
+	MeshDescription = StaticMesh->GetOriginalMeshDescription(0);
 	CalcBoundingBox(MeshDescription, Center, Extents, unitVec);
 	bs->Modify();
 
@@ -542,7 +542,7 @@ int32 GenerateSphereAsSimpleCollision(UStaticMesh* StaticMesh)
 	FVector unitVec = bs->BuildScale3D;
 
 	// Calculate bounding sphere.
-	UMeshDescription* MeshDescription = StaticMesh->GetMeshDescription(0);
+	UMeshDescription* MeshDescription = StaticMesh->GetOriginalMeshDescription(0);
 	CalcBoundingSphere(MeshDescription, bSphere, unitVec);
 	CalcBoundingSphere2(MeshDescription, bSphere2, unitVec);
 
@@ -767,7 +767,7 @@ int32 GenerateSphylAsSimpleCollision(UStaticMesh* StaticMesh)
 	FVector unitVec = bs->BuildScale3D;
 
 	// Calculate bounding box.
-	UMeshDescription* MeshDescription = StaticMesh->GetMeshDescription(0);
+	UMeshDescription* MeshDescription = StaticMesh->GetOriginalMeshDescription(0);
 	CalcBoundingSphyl(MeshDescription, sphere, length, rotation, unitVec);
 
 	// Dont use if radius is zero.
diff --git a/Engine/Source/Runtime/Engine/Classes/Engine/StaticMesh.h b/Engine/Source/Runtime/Engine/Classes/Engine/StaticMesh.h
--- a/Engine/Source/Runtime/Engine/Classes/Engine/StaticMesh.h
+++ b/Engine/Source/Runtime/Engine/Classes/Engine/StaticMesh.h
@@ -690,10 +690,12 @@ public:
 
 	/**
 	 * Accessors for the mesh description use to build the render data
+	 * This is not the data imported it is a modified version to fit the build settings
+	 * Use GetOriginalMeshDescription to get the imported MeshDescription.
+	 * It can be null 
 	 */
 	ENGINE_API UMeshDescription* GetMeshDescription(int32 LodIndex=0) const;
 	ENGINE_API void SetMeshDescription(int32 LodIndex, UMeshDescription* InMeshDescription);
-	ENGINE_API int32 GetMeshDescriptionCount() const;
 
 	/**
 	 * Accessors for the original mesh description imported data
diff --git a/Engine/Source/Runtime/Engine/Private/StaticMesh.cpp b/Engine/Source/Runtime/Engine/Private/StaticMesh.cpp
--- a/Engine/Source/Runtime/Engine/Private/StaticMesh.cpp
+++ b/Engine/Source/Runtime/Engine/Private/StaticMesh.cpp
@@ -2577,11 +2577,6 @@ void UStaticMesh::SetMeshDescription(int32 LodIndex, class UMeshDescription* InM
 	MeshDescriptions[LodIndex] = InMeshDescription;
 }
 
-int32 UStaticMesh::GetMeshDescriptionCount() const
-{
-	return MeshDescriptions.Num();
-}
-
 UMeshDescription* UStaticMesh::GetOriginalMeshDescription(int32 LodIndex)
 {
 	if (SourceModels.IsValidIndex(LodIndex))
