Fix crash when there is more then 8 UVs

diff --git a/Engine/Source/Editor/UnrealEd/Private/Fbx/FbxStaticMeshImport.cpp b/Engine/Source/Editor/UnrealEd/Private/Fbx/FbxStaticMeshImport.cpp
--- a/Engine/Source/Editor/UnrealEd/Private/Fbx/FbxStaticMeshImport.cpp
+++ b/Engine/Source/Editor/UnrealEd/Private/Fbx/FbxStaticMeshImport.cpp
@@ -560,6 +560,7 @@ bool UnFbx::FFbxImporter::BuildStaticMeshFromGeometry(FbxNode* Node, UStaticMesh
 	}
 
 	int32 NumUVs = FMath::Max(FBXUVs.UniqueUVCount, ExistingUVCount);
+	NumUVs = FMath::Min<int32>(MAX_MESH_TEXTURE_COORDS, NumUVs);
 	// At least one UV set must exist.  
 	NumUVs = FMath::Max(1, NumUVs);
 
@@ -1389,7 +1390,7 @@ UStaticMesh* UnFbx::FFbxImporter::ImportStaticMeshAsSingle(UObject* InParent, TA
 	if (bBuildStatus)
 	{
 		TVertexInstanceAttributeIndicesArray<FVector2D>& VertexInstanceUVs = MeshDescription->VertexInstanceAttributes().GetAttributesSet<FVector2D>(MeshAttribute::VertexInstance::TextureCoordinate);
-		int32 FirstOpenUVChannel = VertexInstanceUVs.GetNumIndices();
+		int32 FirstOpenUVChannel = VertexInstanceUVs.GetNumIndices() >= MAX_MESH_TEXTURE_COORDS ? 1 : VertexInstanceUVs.GetNumIndices();
 		TPolygonGroupAttributeArray<FName>& PolygonGroupImportedMaterialSlotNames = MeshDescription->PolygonGroupAttributes().GetAttributes<FName>(MeshAttribute::PolygonGroup::ImportedMaterialSlotName);
 
 		TArray<FStaticMaterial> MaterialToAdd;
