Degenerate polygons rejection during staticmesh import do not anymore mess the vertexinstance normals
Degenerate polygon rejection during staticmesh build is done has the first step when computing the render mesh

diff --git a/Engine/Source/Editor/UnrealEd/Private/Fbx/FbxStaticMeshImport.cpp b/Engine/Source/Editor/UnrealEd/Private/Fbx/FbxStaticMeshImport.cpp
--- a/Engine/Source/Editor/UnrealEd/Private/Fbx/FbxStaticMeshImport.cpp
+++ b/Engine/Source/Editor/UnrealEd/Private/Fbx/FbxStaticMeshImport.cpp
@@ -585,7 +585,7 @@ bool UnFbx::FFbxImporter::BuildStaticMeshFromGeometry(FbxNode* Node, UStaticMesh
 	}
 		
 	Mesh->BeginGetMeshEdgeVertices();
-	TMap<uint32, int32> RemapEdgeID;
+	TMap<uint64, int32> RemapEdgeID;
 	//Fill the edge array
 	int32 FbxEdgeCount = Mesh->GetMeshEdgeCount();
 	RemapEdgeID.Reserve(FbxEdgeCount*2);
@@ -598,10 +598,10 @@ bool UnFbx::FFbxImporter::BuildStaticMeshFromGeometry(FbxNode* Node, UStaticMesh
 		check(MeshDescription->Vertices().IsValid(EdgeVertexStart));
 		FVertexID EdgeVertexEnd(EdgeEndVertexIndex + VertexOffset);
 		check(MeshDescription->Vertices().IsValid(EdgeVertexEnd));
-		uint32 CompactedKey = ((uint32)EdgeVertexStart.GetValue() << 16) | ((uint32)EdgeVertexEnd.GetValue() & 0xffff);
+		uint64 CompactedKey = (((uint64)EdgeVertexStart.GetValue()) << 32) | ((uint64)EdgeVertexEnd.GetValue());
 		RemapEdgeID.Add(CompactedKey, FbxEdgeIndex);
 		//Add the other edge side
-		CompactedKey = ((uint32)EdgeVertexEnd.GetValue() << 16) | ((uint32)EdgeVertexStart.GetValue() & 0xffff);
+		CompactedKey = (((uint64)EdgeVertexEnd.GetValue()) << 32) | ((uint64)EdgeVertexStart.GetValue());
 		RemapEdgeID.Add(CompactedKey, FbxEdgeIndex);
 	}
 	//Call this after all GetMeshEdgeIndexForPolygon call this is for optimization purpose.
@@ -610,6 +610,7 @@ bool UnFbx::FFbxImporter::BuildStaticMeshFromGeometry(FbxNode* Node, UStaticMesh
 	//Call this before all GetMeshEdgeIndexForPolygon call this is for optimization purpose.
 	Mesh->BeginGetMeshEdgeIndexForPolygon();
 	int32 CurrentVertexInstanceIndex = 0;
+	int32 SkippedVertexInstance = 0;
 	//Polygons
 	for (int32 PolygonIndex = 0; PolygonIndex < PolygonCount; PolygonIndex++)
 	{
@@ -630,6 +631,7 @@ bool UnFbx::FFbxImporter::BuildStaticMeshFromGeometry(FbxNode* Node, UStaticMesh
 			//Check for degenerated polygons, avoid NAN
 			if (Normal.IsNearlyZero(ComparisonThreshold))
 			{
+				SkippedVertexInstance += PolygonVertexCount;
 				continue;
 			}
 		}
@@ -642,6 +644,7 @@ bool UnFbx::FFbxImporter::BuildStaticMeshFromGeometry(FbxNode* Node, UStaticMesh
 		for (int32 CornerIndex = 0; CornerIndex < PolygonVertexCount; CornerIndex++, CurrentVertexInstanceIndex++)
 		{
 			int32 VertexInstanceIndex = VertexInstanceOffset + CurrentVertexInstanceIndex;
+			int32 RealFbxVertexIndex = SkippedVertexInstance + CurrentVertexInstanceIndex;
 			const FVertexInstanceID VertexInstanceID(VertexInstanceIndex);
 			CornerInstanceIDs[CornerIndex] = VertexInstanceID;
 			const int32 ControlPointIndex = Mesh->GetPolygonVertex(PolygonIndex, CornerIndex);
@@ -666,7 +669,7 @@ bool UnFbx::FFbxImporter::BuildStaticMeshFromGeometry(FbxNode* Node, UStaticMesh
 				FVector2D FinalUVVector(0.0f, 0.0f);
 				if (FBXUVs.LayerElementUV[UVLayerIndex] != NULL)
 				{
-					int UVMapIndex = (FBXUVs.UVMappingMode[UVLayerIndex] == FbxLayerElement::eByControlPoint) ? ControlPointIndex : CurrentVertexInstanceIndex;
+					int UVMapIndex = (FBXUVs.UVMappingMode[UVLayerIndex] == FbxLayerElement::eByControlPoint) ? ControlPointIndex : RealFbxVertexIndex;
 					int32 UVIndex = (FBXUVs.UVReferenceMode[UVLayerIndex] == FbxLayerElement::eDirect) ?
 						UVMapIndex : FBXUVs.LayerElementUV[UVLayerIndex]->GetIndexArray().GetAt(UVMapIndex);
 
@@ -683,7 +686,7 @@ bool UnFbx::FFbxImporter::BuildStaticMeshFromGeometry(FbxNode* Node, UStaticMesh
 				if (LayerElementVertexColor)
 				{
 					int32 VertexColorMappingIndex = (VertexColorMappingMode == FbxLayerElement::eByControlPoint) ?
-						Mesh->GetPolygonVertex(PolygonIndex, CornerIndex) : (CurrentVertexInstanceIndex);
+						Mesh->GetPolygonVertex(PolygonIndex, CornerIndex) : (RealFbxVertexIndex);
 
 					int32 VectorColorIndex = (VertexColorReferenceMode == FbxLayerElement::eDirect) ?
 						VertexColorMappingIndex : LayerElementVertexColor->GetIndexArray().GetAt(VertexColorMappingIndex);
@@ -719,10 +722,9 @@ bool UnFbx::FFbxImporter::BuildStaticMeshFromGeometry(FbxNode* Node, UStaticMesh
 
 			if (LayerElementNormal)
 			{
-				int TriangleCornerIndex = CurrentVertexInstanceIndex;
 				//normals may have different reference and mapping mode than tangents and binormals
 				int NormalMapIndex = (NormalMappingMode == FbxLayerElement::eByControlPoint) ?
-					ControlPointIndex : TriangleCornerIndex;
+					ControlPointIndex : RealFbxVertexIndex;
 				int NormalValueIndex = (NormalReferenceMode == FbxLayerElement::eDirect) ?
 					NormalMapIndex : LayerElementNormal->GetIndexArray().GetAt(NormalMapIndex);
 					
@@ -734,7 +736,7 @@ bool UnFbx::FFbxImporter::BuildStaticMeshFromGeometry(FbxNode* Node, UStaticMesh
 				if (bHasNTBInformation)
 				{
 					int TangentMapIndex = (TangentMappingMode == FbxLayerElement::eByControlPoint) ?
-						ControlPointIndex : TriangleCornerIndex;
+						ControlPointIndex : RealFbxVertexIndex;
 					int TangentValueIndex = (TangentReferenceMode == FbxLayerElement::eDirect) ?
 						TangentMapIndex : LayerElementTangent->GetIndexArray().GetAt(TangentMapIndex);
 
@@ -744,7 +746,7 @@ bool UnFbx::FFbxImporter::BuildStaticMeshFromGeometry(FbxNode* Node, UStaticMesh
 					VertexInstanceTangents[AddedVertexInstanceId] = TangentX.GetSafeNormal();
 
 					int BinormalMapIndex = (BinormalMappingMode == FbxLayerElement::eByControlPoint) ?
-						ControlPointIndex : TriangleCornerIndex;
+						ControlPointIndex : RealFbxVertexIndex;
 					int BinormalValueIndex = (BinormalReferenceMode == FbxLayerElement::eDirect) ?
 						BinormalMapIndex : LayerElementBinormal->GetIndexArray().GetAt(BinormalMapIndex);
 
@@ -854,7 +856,7 @@ bool UnFbx::FFbxImporter::BuildStaticMeshFromGeometry(FbxNode* Node, UStaticMesh
 				//RawMesh do not have edges, so by ordering the edge with the triangle construction we can ensure back and forth conversion with RawMesh
 				//When raw mesh will be completly remove we can create the edges right after the vertex creation.
 				int32 EdgeIndex = INDEX_NONE;
-				uint32 CompactedKey = ((uint32)EdgeVertexIDs[0].GetValue() << 16) | ((uint32)EdgeVertexIDs[1].GetValue() & 0xffff);
+				uint64 CompactedKey = (((uint64)EdgeVertexIDs[0].GetValue()) << 32) | ((uint64)EdgeVertexIDs[1].GetValue());
 				if (RemapEdgeID.Contains(CompactedKey))
 				{
 					EdgeIndex = RemapEdgeID[CompactedKey];
