Fix lightmass binary not being transferred to Swarm agents.

diff --git a/Engine/Source/Editor/UnrealEd/Private/Lightmass/Lightmass.cpp b/Engine/Source/Editor/UnrealEd/Private/Lightmass/Lightmass.cpp
--- a/Engine/Source/Editor/UnrealEd/Private/Lightmass/Lightmass.cpp
+++ b/Engine/Source/Editor/UnrealEd/Private/Lightmass/Lightmass.cpp
@@ -2778,7 +2778,8 @@ bool FLightmassProcessor::BeginRun()
 		TEXT("../Win32/UnrealLightmass-Core.dll"),
 		TEXT("../Win32/UnrealLightmass-CoreUObject.dll"),
 		TEXT("../Win32/UnrealLightmass-Projects.dll"),
-		TEXT("../Win32/UnrealLightmass-Json.dll")
+		TEXT("../Win32/UnrealLightmass-Json.dll"),
+		TEXT("../Win32/UnrealLightmass-BuildSettings.dll")
 	};
 	const int32 RequiredDependencyPaths32Count = ARRAY_COUNT(RequiredDependencyPaths32);
 
@@ -2795,6 +2796,7 @@ bool FLightmassProcessor::BeginRun()
 		TEXT("../Win64/UnrealLightmass-CoreUObject.dll"),
 		TEXT("../Win64/UnrealLightmass-Projects.dll"),
 		TEXT("../Win64/UnrealLightmass-Json.dll"),
+		TEXT("../Win64/UnrealLightmass-BuildSettings.dll"),
 		TEXT("../Win64/embree.dll"),
 		TEXT("../Win64/tbb.dll"),
 		TEXT("../Win64/tbbmalloc.dll")
@@ -2809,7 +2811,8 @@ bool FLightmassProcessor::BeginRun()
 		TEXT("../Mac/UnrealLightmass-CoreUObject.dylib"),
 		TEXT("../Mac/UnrealLightmass-Json.dylib"),
 		TEXT("../Mac/UnrealLightmass-Projects.dylib"),
-		TEXT("../Mac/UnrealLightmass-SwarmInterface.dylib")
+		TEXT("../Mac/UnrealLightmass-SwarmInterface.dylib"),
+		TEXT("../Mac/UnrealLightmass-BuildSettings.dylib"),
 		TEXT("../Mac/libembree.2.dylib"),
 		TEXT("../Mac/libtbb.dylib"),
 		TEXT("../Mac/libtbbmalloc.dylib")
@@ -2827,6 +2830,7 @@ bool FLightmassProcessor::BeginRun()
 		TEXT("../Linux/libUnrealLightmass-SwarmInterface.so"),
 		TEXT("../Linux/libUnrealLightmass-Networking.so"),
 		TEXT("../Linux/libUnrealLightmass-Messaging.so"),
+		TEXT("../Linux/libUnrealLightmass-BuildSettings.so"),
 		TEXT("../../Plugins/Messaging/UdpMessaging/Binaries/Linux/libUnrealLightmass-UdpMessaging.so")
 	};
 #else // PLATFORM_LINUX
