Guard against race condition with auto reimport automation test and actual editor trying to perform the same thing

diff --git a/Engine/Source/Editor/UnrealEd/Private/AutoReimport/AutoReimportManager.cpp b/Engine/Source/Editor/UnrealEd/Private/AutoReimport/AutoReimportManager.cpp
--- a/Engine/Source/Editor/UnrealEd/Private/AutoReimport/AutoReimportManager.cpp
+++ b/Engine/Source/Editor/UnrealEd/Private/AutoReimport/AutoReimportManager.cpp
@@ -135,7 +135,11 @@ private:
 
 	/** FTickableEditorObject interface*/
 	virtual void Tick(float DeltaTime) override;
-	virtual ETickableTickType GetTickableTickType() const override { return ETickableTickType::Always; }
+	virtual ETickableTickType GetTickableTickType() const override { return ETickableTickType::Conditional; }
+
+	/** Automation testing will tick its own file cache which can result in race conditions with the reimport manager.  */
+	virtual bool IsTickable() const override { return !GIsAutomationTesting; }
+
 	virtual TStatId GetStatId() const override;
 
 	/** FGCObject interface*/
