Fix rotation widget no longer correctly displaying rotation label while rotating in high dpi

diff --git a/Engine/Source/Editor/UnrealEd/Private/UnrealWidget.cpp b/Engine/Source/Editor/UnrealEd/Private/UnrealWidget.cpp
--- a/Engine/Source/Editor/UnrealEd/Private/UnrealWidget.cpp
+++ b/Engine/Source/Editor/UnrealEd/Private/UnrealWidget.cpp
@@ -219,8 +219,9 @@ void FWidget::DrawHUD (FCanvas* Canvas)
 {
 	if (HUDString.Len())
 	{
-		int32 StringPosX = FMath::FloorToInt(HUDInfoPos.X);
-		int32 StringPosY = FMath::FloorToInt(HUDInfoPos.Y);
+		const float DPIScale = Canvas->GetDPIScale();
+		int32 StringPosX = FMath::FloorToInt(HUDInfoPos.X/DPIScale);
+		int32 StringPosY = FMath::FloorToInt(HUDInfoPos.Y/DPIScale);
 
 		//measure string size
 		int32 StringSizeX, StringSizeY;
