Fixed crash when content hot-reloading the build data for the currently loaded world

diff --git a/Engine/Source/Editor/UnrealEd/Private/PackageTools.cpp b/Engine/Source/Editor/UnrealEd/Private/PackageTools.cpp
--- a/Engine/Source/Editor/UnrealEd/Private/PackageTools.cpp
+++ b/Engine/Source/Editor/UnrealEd/Private/PackageTools.cpp
@@ -38,6 +38,7 @@
 #include "ComponentReregisterContext.h"
 #include "Engine/Selection.h"
 #include "Engine/LevelStreaming.h"
+#include "Engine/MapBuildDataRegistry.h"
 
 #include "ShaderCompiler.h"
 #include "DistanceFieldAtlas.h"
@@ -281,6 +282,40 @@ namespace PackageTools
 			OutErrorMessage = FText::Format( NSLOCTEXT("UnrealEd", "UnloadDirtyPackagesList", "The following assets have been modified and cannot be unloaded:{DirtyPackages}\nSaving these assets will allow them to be unloaded."), Args );
 		}
 
+		if (UWorld* EditorWorld = GEditor->GetEditorWorldContext().World())
+		{
+			// Is the currently loaded world being unloaded? If so, we just reset the current world.
+			// We also need to skip the build data package as that will also be destroyed by the call to CreateNewMapForEditing.
+			if (PackagesToUnload.Contains(EditorWorld->GetOutermost()))
+			{
+				// Remove the world package from the unload list
+				PackagesToUnload.Remove(EditorWorld->GetOutermost());
+
+				// Remove the level build data package from the unload list as creating a new map will unload build data for the current world
+				for (int32 LevelIndex = 0; LevelIndex < EditorWorld->GetNumLevels(); ++LevelIndex)
+				{
+					ULevel* Level = EditorWorld->GetLevel(LevelIndex);
+					if (Level->MapBuildData)
+					{
+						PackagesToUnload.Remove(Level->MapBuildData->GetOutermost());
+					}
+				}
+
+				// Remove any streaming levels from the unload list as creating a new map will unload streaming levels for the current world
+				for (ULevelStreaming* EditorStreamingLevel : EditorWorld->GetStreamingLevels())
+				{
+					if (EditorStreamingLevel->IsLevelLoaded())
+					{
+						UPackage* EditorStreamingLevelPackage = EditorStreamingLevel->GetLoadedLevel()->GetOutermost();
+						PackagesToUnload.Remove(EditorStreamingLevelPackage);
+					}
+				}
+
+				// Unload the current world
+				GEditor->CreateNewMapForEditing();
+			}
+		}
+
 		if ( PackagesToUnload.Num() > 0 )
 		{
 			const FScopedBusyCursor BusyCursor;
@@ -534,49 +569,54 @@ namespace PackageTools
 
 		// Check to see if we need to reload the current world.
 		FName WorldNameToReload;
-		PackagesToReload.RemoveAll([&](UPackage* PackageToReload) -> bool
+		TMap<FName, const UMapBuildDataRegistry*> LevelsToMapBuildData;
 		{
-			// Is this the currently loaded world? If so, we just reset the current world and load it again at the end rather than let it go 
-			// through ReloadPackage (which doesn't work for the editor due to some assumptions it makes about worlds, and their lifetimes).
 			if (UWorld* EditorWorld = GEditor->GetEditorWorldContext().World())
 			{
-				const bool bIsCurrentWorld = EditorWorld->GetOutermost() == PackageToReload;
-				if (bIsCurrentWorld)
+				// Is the currently loaded world being reloaded? If so, we just reset the current world and load it again at the end rather than let it go 
+				// through ReloadPackage (which doesn't work for the editor due to some assumptions it makes about worlds, and their lifetimes).
+				// We also need to skip the build data package as that will also be destroyed by the call to CreateNewMapForEditing.
+				if (PackagesToReload.Contains(EditorWorld->GetOutermost()))
 				{
+					// Cache this so we can reload the world later
 					WorldNameToReload = *EditorWorld->GetPathName();
-					return true; // remove world package
-				}
-			}
-			return false; // leave non-world package
-		});
 
-		// Unload the current world (if needed).
-		TMap<FName, const UMapBuildDataRegistry*> LevelsToMapBuildData;
-		if (!WorldNameToReload.IsNone())
-		{
-			// Creating a new map will unload all streaming levels for the current editor world too, so we need to make sure we're not about to try and reload those later
-			if (UWorld* EditorWorld = GEditor->GetEditorWorldContext().World())
-			{
-				for (ULevelStreaming* EditorStreamingLevel : EditorWorld->GetStreamingLevels())
+					// Remove the world package from the reload list
+					PackagesToReload.Remove(EditorWorld->GetOutermost());
+
+					// Remove the level build data package from the reload list as creating a new map will unload build data for the current world
+					for (int32 LevelIndex = 0; LevelIndex < EditorWorld->GetNumLevels(); ++LevelIndex)
+					{
+						ULevel* Level = EditorWorld->GetLevel(LevelIndex);
+						if (Level->MapBuildData)
+						{
+							PackagesToReload.Remove(Level->MapBuildData->GetOutermost());
+						}
+					}
+
+					// Remove any streaming levels from the reload list as creating a new map will unload streaming levels for the current world
+					for (ULevelStreaming* EditorStreamingLevel : EditorWorld->GetStreamingLevels())
+					{
+						if (EditorStreamingLevel->IsLevelLoaded())
+						{
+							UPackage* EditorStreamingLevelPackage = EditorStreamingLevel->GetLoadedLevel()->GetOutermost();
+							PackagesToReload.Remove(EditorStreamingLevelPackage);
+						}
+					}
+
+					// Unload the current world
+					GEditor->CreateNewMapForEditing();
+				}
+				// Cache the current map build data for the levels of the current world so we can see if they change due to a reload (we can skip this if reloading the current world).
+				else
 				{
-					if (EditorStreamingLevel->IsLevelLoaded())
+					for (int32 LevelIndex = 0; LevelIndex < EditorWorld->GetNumLevels(); ++LevelIndex)
 					{
-						UPackage* EditorStreamingLevelPackage = EditorStreamingLevel->GetLoadedLevel()->GetOutermost();
-						PackagesToReload.Remove(EditorStreamingLevelPackage);
+						ULevel* Level = EditorWorld->GetLevel(LevelIndex);
+						LevelsToMapBuildData.Add(Level->GetFName(), Level->MapBuildData);
 					}
 				}
 			}
-
-			GEditor->CreateNewMapForEditing();
-		}
-		// Cache the current map build data for the levels of the current world so we can see if they change due to a reload (we skip this if reloading the current world).
-		else if (UWorld* EditorWorld = GEditor->GetEditorWorldContext().World())
-		{
-			for (int32 LevelIndex = 0; LevelIndex < EditorWorld->GetNumLevels(); ++LevelIndex)
-			{
-				ULevel* Level = EditorWorld->GetLevel(LevelIndex);
-				LevelsToMapBuildData.Add(Level->GetFName(), Level->MapBuildData);
-			}
 		}
 
 		if (PackagesToReload.Num() > 0)
