Fixed clothing data binding state becomming mismatched after a reimport over multiple LODs, moved the unbind and rebind steps out of the inner import loop to be handled before and after the total reimport.

diff --git a/Engine/Source/Editor/UnrealEd/Private/Fbx/FbxSkeletalMeshImport.cpp b/Engine/Source/Editor/UnrealEd/Private/Fbx/FbxSkeletalMeshImport.cpp
--- a/Engine/Source/Editor/UnrealEd/Private/Fbx/FbxSkeletalMeshImport.cpp
+++ b/Engine/Source/Editor/UnrealEd/Private/Fbx/FbxSkeletalMeshImport.cpp
@@ -1477,19 +1477,10 @@ USkeletalMesh* UnFbx::FFbxImporter::ImportSkeletalMesh(FImportSkeletalMeshArgs &
 	TIndirectArray<FComponentReregisterContext> ComponentContexts;
 	TUniquePtr< FSkinnedMeshComponentRecreateRenderStateContext > SkinnedMeshComponentRecreateRenderStateContext;
 
-	TArray<ClothingAssetUtils::FClothingAssetMeshBinding> ClothingBindings;
-
 	FSkinnedMeshComponentRecreateRenderStateContext* RecreateExistingRenderStateContext = ExistingSkelMesh ? new FSkinnedMeshComponentRecreateRenderStateContext(ExistingSkelMesh, false) : nullptr;
 	//Backup the data before importing the new one
 	if (ExistingSkelMesh)
 	{
-		ClothingAssetUtils::GetMeshClothingAssetBindings(ExistingSkelMesh, ClothingBindings);
-
-		for(ClothingAssetUtils::FClothingAssetMeshBinding& Binding : ClothingBindings)
-		{
-			Binding.Asset->UnbindFromSkeletalMesh(ExistingSkelMesh, Binding.LODIndex);
-		}
-
 		//Release the Renderdata resources before reimporting the skeletal mesh
 		SkinnedMeshComponentRecreateRenderStateContext = MakeUnique<FSkinnedMeshComponentRecreateRenderStateContext>(ExistingSkelMesh);
 
@@ -1782,20 +1773,6 @@ USkeletalMesh* UnFbx::FFbxImporter::ImportSkeletalMesh(FImportSkeletalMeshArgs &
 		}
 	}
 
-	// Reapply any clothing assets we had before the import
-	FSkeletalMeshModel* NewMeshResource = SkeletalMesh->GetImportedModel();
-	if(NewMeshResource)
-	{
-		for(ClothingAssetUtils::FClothingAssetMeshBinding& Binding : ClothingBindings)
-		{
-			if(NewMeshResource->LODModels.IsValidIndex(Binding.LODIndex) &&
-			   NewMeshResource->LODModels[Binding.LODIndex].Sections.IsValidIndex(Binding.SectionIndex))
-			{
-				Binding.Asset->BindToSkeletalMesh(SkeletalMesh, Binding.LODIndex, Binding.SectionIndex, Binding.AssetInternalLodIndex);
-			}
-		}
-	}
-
 	// ComponentContexts will now go out of scope, causing components to be re-registered
 
 	if (RecreateExistingRenderStateContext)
@@ -2016,6 +1993,15 @@ USkeletalMesh* UnFbx::FFbxImporter::ReimportSkeletalMesh(USkeletalMesh* Mesh, UF
 	// support to update rigid animation mesh
 	ImportOptions->bImportRigidMesh = true;
 
+	// Unbind any existing clothing assets before we reimport the geometry
+	TArray<ClothingAssetUtils::FClothingAssetMeshBinding> ClothingBindings;
+	ClothingAssetUtils::GetMeshClothingAssetBindings(Mesh, ClothingBindings);
+
+	for(ClothingAssetUtils::FClothingAssetMeshBinding& Binding : ClothingBindings)
+	{
+		Binding.Asset->UnbindFromSkeletalMesh(Mesh, Binding.LODIndex);
+	}
+
 	// get meshes in Fbx file
 	//the function also fill the collision models, so we can update collision models correctly
 	TArray< TArray<FbxNode*>* > FbxSkelMeshArray;
@@ -2277,6 +2263,23 @@ USkeletalMesh* UnFbx::FFbxImporter::ReimportSkeletalMesh(USkeletalMesh* Mesh, UF
 		}
 		if (NewMesh)
 		{
+			FSkeletalMeshModel* NewMeshResource = NewMesh->GetImportedModel();
+			if(NewMeshResource && ClothingBindings.Num() > 0)
+			{
+				NewMesh->PreEditChange(nullptr);
+
+				for(ClothingAssetUtils::FClothingAssetMeshBinding& Binding : ClothingBindings)
+				{
+					if(NewMeshResource->LODModels.IsValidIndex(Binding.LODIndex) &&
+						NewMeshResource->LODModels[Binding.LODIndex].Sections.IsValidIndex(Binding.SectionIndex))
+					{
+						Binding.Asset->BindToSkeletalMesh(NewMesh, Binding.LODIndex, Binding.SectionIndex, Binding.AssetInternalLodIndex);
+					}
+				}
+
+				NewMesh->PostEditChange();
+			}
+
 			//Update the import data so we can re-import correctly
 			UpdateSkeletalMeshImportData(NewMesh, TemplateImportData, INDEX_NONE, &ImportMaterialOriginalNameData, &ImportMeshLodData);
 			//If we have import some morph target we have to rebuild the render resources since morph target are now using GPU
