Fix for bad Camera Rotation Imported From Maya

Main issue was that we don't support RotationOrders that aren't XYZ.
We had some code to zero out other part's of the FBX transform pipeline but it wasn't setting the Rotation Order correctly and was also based upon old documentation.
Currently recommendation is to always just zero everything out.
Finally had to use ResetPivotSetAndConvertAnimation since there appears to be a bug with ConvertPivotAnimationRecursive into the correct space.

diff --git a/Engine/Source/Editor/UnrealEd/Private/Fbx/FbxCurveDataImport.cpp b/Engine/Source/Editor/UnrealEd/Private/Fbx/FbxCurveDataImport.cpp
--- a/Engine/Source/Editor/UnrealEd/Private/Fbx/FbxCurveDataImport.cpp
+++ b/Engine/Source/Editor/UnrealEd/Private/Fbx/FbxCurveDataImport.cpp
@@ -696,62 +696,39 @@ namespace UnFbx {
 		CurvesAPI.TransformData.Add(AnimNodeHandle.UniqueId, Transform);
 	}
 
+	//This function now clears out all pivots, post and pre rotations and set's the
+	//RotationOrder to XYZ.
+	//Updated per the latest documentation
+	//https://help.autodesk.com/view/FBX/2017/ENU/?guid=__files_GUID_C35D98CB_5148_4B46_82D1_51077D8970EE_htm
 	void FFbxImporter::SetupTransformForNode(FbxNode *Node)
 	{
-		FbxVector4 ZeroVector(0, 0, 0);
+
+		// Activate pivot converting 
 		Node->SetPivotState(FbxNode::eSourcePivot, FbxNode::ePivotActive);
 		Node->SetPivotState(FbxNode::eDestinationPivot, FbxNode::ePivotActive);
 
-		EFbxRotationOrder RotationOrder;
-		Node->GetRotationOrder(FbxNode::eSourcePivot, RotationOrder);
-		Node->SetRotationOrder(FbxNode::eDestinationPivot, RotationOrder);
-
-		// For cameras and lights (without targets) let's compensate the postrotation.
-		if (Node->GetCamera() || Node->GetLight())
-		{
-			// Point light do not need to be adjusted (since they radiate in all the directions).
-			if (Node->GetLight() && Node->GetLight()->LightType.Get() == FbxLight::ePoint)
-			{
-				Node->SetPostRotation(FbxNode::eSourcePivot, FbxVector4(0, 0, 0, 0));
-			}
-
-			// apply Pre rotations only on bones / end of chains
-			if ((Node->GetNodeAttribute() && Node->GetNodeAttribute()->GetAttributeType() == FbxNodeAttribute::eSkeleton)
-				|| (Node->GetMarker() && Node->GetMarker()->GetType() == FbxMarker::eEffectorFK)
-				|| (Node->GetMarker() && Node->GetMarker()->GetType() == FbxMarker::eEffectorIK))
-			{
-				Node->SetPreRotation(FbxNode::eDestinationPivot, Node->GetPreRotation(FbxNode::eSourcePivot));
-
-				// No pivots on bones
-				Node->SetRotationPivot(FbxNode::eDestinationPivot, ZeroVector);
-				Node->SetScalingPivot(FbxNode::eDestinationPivot, ZeroVector);
-				Node->SetRotationOffset(FbxNode::eDestinationPivot, ZeroVector);
-				Node->SetScalingOffset(FbxNode::eDestinationPivot, ZeroVector);
-			}
-			else
-			{
-				// any other type: no pre-rotation support but...
-				Node->SetPreRotation(FbxNode::eDestinationPivot, ZeroVector);
-
-				// support for rotation and scaling pivots.
-				Node->SetRotationPivot(FbxNode::eDestinationPivot, Node->GetRotationPivot(FbxNode::eSourcePivot));
-				Node->SetScalingPivot(FbxNode::eDestinationPivot, Node->GetScalingPivot(FbxNode::eSourcePivot));
-				// Rotation and scaling offset are supported
-				Node->SetRotationOffset(FbxNode::eDestinationPivot, Node->GetRotationOffset(FbxNode::eSourcePivot));
-				Node->SetScalingOffset(FbxNode::eDestinationPivot, Node->GetScalingOffset(FbxNode::eSourcePivot));
-				//
-				// If we supported scaling pivots, we could simply do:
-				// Node->SetRotationPivot(FbxNode::eDESTINATION_SET, ZeroVector);
-				// Node->SetScalingPivot(FbxNode::eDESTINATION_SET, ZeroVector);
-			}
-		}
-
-		// Recursively convert the animation data according to pivot settings.
-		Node->ConvertPivotAnimationRecursive(
-			NULL,																// Use the first animation stack by default
-			FbxNode::eDestinationPivot,											// Convert from Source set to Destination set
-			FbxTime::GetFrameRate(Scene->GetGlobalSettings().GetTimeMode()),	// Resampling frame rate in frames per second
-			false);																// Do not apply key reducing filter
+		FbxVector4 Zero(0, 0, 0);
+
+		// We want to set all these to 0 and bake them into the transforms. 
+		Node->SetPostRotation(FbxNode::eDestinationPivot, Zero);
+		Node->SetPreRotation(FbxNode::eDestinationPivot, Zero);
+		Node->SetRotationOffset(FbxNode::eDestinationPivot, Zero);
+		Node->SetScalingOffset(FbxNode::eDestinationPivot, Zero);
+		Node->SetRotationPivot(FbxNode::eDestinationPivot, Zero);
+		Node->SetScalingPivot(FbxNode::eDestinationPivot, Zero);
+
+		Node->SetRotationOrder(FbxNode::eDestinationPivot, eEulerXYZ);
+		//When  we support other orders do this.
+		//EFbxRotationOrder ro;
+		//Node->GetRotationOrder(FbxNode::eSourcePivot, ro);
+		//Node->SetRotationOrder(FbxNode::eDestinationPivot, ro);
+
+		//Most DCC's don't have this but 3ds Max does
+		Node->SetGeometricTranslation(FbxNode::eDestinationPivot, Zero);
+		Node->SetGeometricRotation(FbxNode::eDestinationPivot, Zero);
+		Node->SetGeometricScaling(FbxNode::eDestinationPivot, Zero);
+		//NOTE THAT ConvertPivotAnimationRecursive did not seem to work when getting the local transform values!!!
+		Node->ResetPivotSetAndConvertAnimation(FbxTime::GetFrameRate(Scene->GetGlobalSettings().GetTimeMode()));
 	}
 
 } // namespace UnFBX
