Copied cl 4077328 from dev-enterprise

Fixed SoftObjectPtr/Paths becoming invalid when saving a new world since it's being moved from /Temp/Untitled to its own package.

diff --git a/Engine/Source/Editor/UnrealEd/Private/FileHelpers.cpp b/Engine/Source/Editor/UnrealEd/Private/FileHelpers.cpp
--- a/Engine/Source/Editor/UnrealEd/Private/FileHelpers.cpp
+++ b/Engine/Source/Editor/UnrealEd/Private/FileHelpers.cpp
@@ -10,6 +10,7 @@
 #include "Misc/ConfigCacheIni.h"
 #include "Misc/FeedbackContext.h"
 #include "Misc/ScopedSlowTask.h"
+#include "Misc/RedirectCollector.h"
 #include "Misc/App.h"
 #include "Misc/FileHelper.h"
 #include "Modules/ModuleManager.h"
@@ -607,6 +608,9 @@ static bool SaveWorld(UWorld* World,
 
 		SlowTask.EnterProgressFrame(25);
 
+		FSoftObjectPath OldPath( World );
+		bool bAddedAssetPathRedirection = false;
+
 		// Rename the package and the object, as necessary
 		UWorld* DuplicatedWorld = nullptr;
 		if ( bRenamePackageToFile )
@@ -650,6 +654,11 @@ static bool SaveWorld(UWorld* World,
 						}
 
 						World->Rename(*NewWorldAssetName, NULL, REN_NonTransactional | REN_DontCreateRedirectors | REN_ForceNoResetLoaders);
+
+						// We're renaming the world, add a path redirector so that soft object paths get fixed on save
+						FSoftObjectPath NewPath( World );
+						GRedirectCollector.AddAssetPathRedirection( *OldPath.GetAssetPathString(), *NewPath.GetAssetPathString() );
+						bAddedAssetPathRedirection = true;
 					}
 				}
 			}
@@ -667,6 +676,11 @@ static bool SaveWorld(UWorld* World,
 			SaveErrors.Flush();
 		}
 
+		if ( bAddedAssetPathRedirection )
+		{
+			GRedirectCollector.RemoveAssetPathRedirection( *OldPath.GetAssetPathString() );
+		}
+
 		// @todo Autosaving should save build data as well
 		if (bSuccess && !bAutosaving)
 		{
