Unshelved from pending changelist '4075677':

Fix fbx scene re-import crash when the user re-import more then one time with some override material.

diff --git a/Engine/Source/Editor/UnrealEd/Private/Fbx/FbxSceneImportFactory.cpp b/Engine/Source/Editor/UnrealEd/Private/Fbx/FbxSceneImportFactory.cpp
--- a/Engine/Source/Editor/UnrealEd/Private/Fbx/FbxSceneImportFactory.cpp
+++ b/Engine/Source/Editor/UnrealEd/Private/Fbx/FbxSceneImportFactory.cpp
@@ -132,6 +132,8 @@ bool GetFbxSceneImportOptions(UnFbx::FFbxImporter* FbxImporter
 	//Prefix materials package name to put all material under Material folder (this avoid name clash with meshes)
 	GlobalImportSettings->MaterialBasePath = NAME_None;
 
+	GlobalImportSettings->OverrideMaterials.Reset();
+
 	TSharedPtr<SWindow> ParentWindow;
 	if (FModuleManager::Get().IsModuleLoaded("MainFrame"))
 	{
diff --git a/Engine/Source/Editor/UnrealEd/Private/Fbx/ReimportFbxSceneFactory.cpp b/Engine/Source/Editor/UnrealEd/Private/Fbx/ReimportFbxSceneFactory.cpp
--- a/Engine/Source/Editor/UnrealEd/Private/Fbx/ReimportFbxSceneFactory.cpp
+++ b/Engine/Source/Editor/UnrealEd/Private/Fbx/ReimportFbxSceneFactory.cpp
@@ -245,6 +245,7 @@ bool GetFbxSceneReImportOptions(UnFbx::FFbxImporter* FbxImporter
 	GlobalImportSettings->bConvertScene = true;
 	GlobalImportSettings->bConvertSceneUnit = true;
 
+	GlobalImportSettings->OverrideMaterials.Reset();
 
 	TSharedPtr<SWindow> ParentWindow;
 	if (FModuleManager::Get().IsModuleLoaded("MainFrame"))
