Fix a crash when re-importing an animation and we have to choose a skeleton for the anim sequence

diff --git a/Engine/Source/Editor/UnrealEd/Private/Factories/EditorFactories.cpp b/Engine/Source/Editor/UnrealEd/Private/Factories/EditorFactories.cpp
--- a/Engine/Source/Editor/UnrealEd/Private/Factories/EditorFactories.cpp
+++ b/Engine/Source/Editor/UnrealEd/Private/Factories/EditorFactories.cpp
@@ -5851,6 +5851,8 @@ EReimportResult::Type UReimportFbxAnimSequenceFactory::Reimport( UObject* Obj )
 			Importer->ReleaseScene();
 			return EReimportResult::Succeeded;
 		}
+		//Set the selected skeleton in the anim sequence
+		AnimSequence->SetSkeleton(Skeleton);
 	}
 
 	if ( UEditorEngine::ReimportFbxAnimation(Skeleton, AnimSequence, ImportData, *Filename) )
