Support non uniform scale scene transform for fbx skeletal mesh exporter

diff --git a/Engine/Source/Editor/UnrealEd/Private/Fbx/FbxAnimationExport.cpp b/Engine/Source/Editor/UnrealEd/Private/Fbx/FbxAnimationExport.cpp
--- a/Engine/Source/Editor/UnrealEd/Private/Fbx/FbxAnimationExport.cpp
+++ b/Engine/Source/Editor/UnrealEd/Private/Fbx/FbxAnimationExport.cpp
@@ -215,10 +215,19 @@ FbxNode* FFbxExporter::ExportAnimSequence( const UAnimSequence* AnimSeq, const U
 
 
 	FbxNode* RootNode = (ActorRootNode)? ActorRootNode : Scene->GetRootNode();
+
+	//Create a temporary node attach to the scene root.
+	//This will allow us to do the binding without the scene transform (non uniform scale is not supported when binding the skeleton)
+	//We then detach from the temp node and attach to the parent and remove the temp node
+	FString FbxNodeName = FGuid::NewGuid().ToString(EGuidFormats::Digits);
+	FbxNode* TmpNodeNoTransform = FbxNode::Create(Scene, TCHAR_TO_UTF8(*FbxNodeName));
+	Scene->GetRootNode()->AddChild(TmpNodeNoTransform);
+
+
 	// Create the Skeleton
 	TArray<FbxNode*> BoneNodes;
 	FbxNode* SkeletonRootNode = CreateSkeleton(SkelMesh, BoneNodes);
-	RootNode->AddChild(SkeletonRootNode);
+	TmpNodeNoTransform->AddChild(SkeletonRootNode);
 
 
 	// Export the anim sequence
@@ -254,7 +263,7 @@ FbxNode* FFbxExporter::ExportAnimSequence( const UAnimSequence* AnimSeq, const U
 		FbxNode* MeshRootNode = CreateMesh(SkelMesh, *MeshNodeName);
 		if(MeshRootNode)
 		{
-			RootNode->AddChild(MeshRootNode);
+			TmpNodeNoTransform->AddChild(MeshRootNode);
 		}
 
 		if(SkeletonRootNode && MeshRootNode)
@@ -265,8 +274,23 @@ FbxNode* FFbxExporter::ExportAnimSequence( const UAnimSequence* AnimSeq, const U
 			// Add the bind pose
 			CreateBindPose(MeshRootNode);
 		}
+
+		if (MeshRootNode)
+		{
+			TmpNodeNoTransform->RemoveChild(MeshRootNode);
+			RootNode->AddChild(MeshRootNode);
+		}
+	}
+	
+	if (SkeletonRootNode)
+	{
+		TmpNodeNoTransform->RemoveChild(SkeletonRootNode);
+		RootNode->AddChild(SkeletonRootNode);
 	}
 
+	Scene->GetRootNode()->RemoveChild(TmpNodeNoTransform);
+	Scene->RemoveNode(TmpNodeNoTransform);
+
 	return SkeletonRootNode;
 }
 
diff --git a/Engine/Source/Editor/UnrealEd/Private/Fbx/FbxSkeletalMeshExport.cpp b/Engine/Source/Editor/UnrealEd/Private/Fbx/FbxSkeletalMeshExport.cpp
--- a/Engine/Source/Editor/UnrealEd/Private/Fbx/FbxSkeletalMeshExport.cpp
+++ b/Engine/Source/Editor/UnrealEd/Private/Fbx/FbxSkeletalMeshExport.cpp
@@ -480,20 +480,27 @@ FbxNode* FFbxExporter::ExportSkeletalMeshToFbx(const USkeletalMesh* SkelMesh, co
 	}
 	else
 	{
+		//Create a temporary node attach to the scene root.
+		//This will allow us to do the binding without the scene transform (non uniform scale is not supported when binding the skeleton)
+		//We then detach from the temp node and attach to the parent and remove the temp node
+		FString FbxNodeName = FGuid::NewGuid().ToString(EGuidFormats::Digits);
+		FbxNode* TmpNodeNoTransform = FbxNode::Create(Scene, TCHAR_TO_UTF8(*FbxNodeName));
+		Scene->GetRootNode()->AddChild(TmpNodeNoTransform);
+
 		TArray<FbxNode*> BoneNodes;
 
 		// Add the skeleton to the scene
 		FbxNode* SkeletonRootNode = CreateSkeleton(SkelMesh, BoneNodes);
 		if(SkeletonRootNode)
 		{
-			ActorRootNode->AddChild(SkeletonRootNode);
+			TmpNodeNoTransform->AddChild(SkeletonRootNode);
 		}
 
 		// Add the mesh
 		FbxNode* MeshRootNode = CreateMesh(SkelMesh, MeshName);
 		if(MeshRootNode)
 		{
-			ActorRootNode->AddChild(MeshRootNode);
+			TmpNodeNoTransform->AddChild(MeshRootNode);
 		}
 
 		if(SkeletonRootNode && MeshRootNode)
@@ -505,6 +512,19 @@ FbxNode* FFbxExporter::ExportSkeletalMeshToFbx(const USkeletalMesh* SkelMesh, co
 			CreateBindPose(MeshRootNode);
 		}
 
+		if (SkeletonRootNode)
+		{
+			TmpNodeNoTransform->RemoveChild(SkeletonRootNode);
+			ActorRootNode->AddChild(SkeletonRootNode);
+		}
+		if (MeshRootNode)
+		{
+			TmpNodeNoTransform->RemoveChild(MeshRootNode);
+			ActorRootNode->AddChild(MeshRootNode);
+		}
+
+		Scene->GetRootNode()->RemoveChild(TmpNodeNoTransform);
+		Scene->RemoveNode(TmpNodeNoTransform);
 		return SkeletonRootNode;
 	}
 
