Fix a render thread crash when reimporting skeletal mesh

diff --git a/Engine/Source/Editor/UnrealEd/Private/Fbx/FbxSkeletalMeshImport.cpp b/Engine/Source/Editor/UnrealEd/Private/Fbx/FbxSkeletalMeshImport.cpp
--- a/Engine/Source/Editor/UnrealEd/Private/Fbx/FbxSkeletalMeshImport.cpp
+++ b/Engine/Source/Editor/UnrealEd/Private/Fbx/FbxSkeletalMeshImport.cpp
@@ -1392,6 +1392,8 @@ USkeletalMesh* UnFbx::FFbxImporter::ImportSkeletalMesh(FImportSkeletalMeshArgs &
 	{
 		SkelType = 1;
 	}
+	//Make sure the render thread is done
+	FlushRenderingCommands();
 
 	// warning for missing smoothing group info
 	CheckSmoothingInfo(FbxMesh);
