Make sure GIsImportingT3D is false when the scene importer open the blueprint editor. This flag force a load of UObject when we do FindObject, which in this case was forcing a load of all blueprints assets.

diff --git a/Engine/Source/Editor/UnrealEd/Private/Fbx/FbxSceneImportFactory.cpp b/Engine/Source/Editor/UnrealEd/Private/Fbx/FbxSceneImportFactory.cpp
--- a/Engine/Source/Editor/UnrealEd/Private/Fbx/FbxSceneImportFactory.cpp
+++ b/Engine/Source/Editor/UnrealEd/Private/Fbx/FbxSceneImportFactory.cpp
@@ -1196,8 +1196,17 @@ FFeedbackContext*	Warn
 				FullnameBP = PackageTools::SanitizePackageName(FullnameBP);
 				FString AssetName = TEXT("FbxScene_") + FPaths::GetBaseFilename(UFactory::CurrentFilename);
 				UPackage *Pkg = CreatePackageForNode(FullnameBP, AssetName);
+
+				//IsImportingT3D will force a load of UObject when finding them, Openning the blueprint editor search for all UBlueprint class object and find them.
+				//This force load everything.
+				GEditor->IsImportingT3D = 0;
+				GIsImportingT3D = false;
 				//Create the blueprint from the actor and replace the actor with a blueprintactor that point on the blueprint
 				UBlueprint* SceneBlueprint = FKismetEditorUtilities::CreateBlueprintFromActor(Pkg->GetName(), HierarchyActor, true, true);
+				//Put back the T3D state
+				GEditor->IsImportingT3D = 1;
+				GIsImportingT3D = GEditor->IsImportingT3D;
+
 				if (SceneBlueprint != nullptr && ReimportData != nullptr)
 				{
 					//let the scene blueprint be the return object for this import
