Fix for materials showing up in the Editing Assets list twice

diff --git a/Engine/Source/Editor/UnrealEd/Private/Toolkits/SGlobalTabSwitchingDialog.cpp b/Engine/Source/Editor/UnrealEd/Private/Toolkits/SGlobalTabSwitchingDialog.cpp
--- a/Engine/Source/Editor/UnrealEd/Private/Toolkits/SGlobalTabSwitchingDialog.cpp
+++ b/Engine/Source/Editor/UnrealEd/Private/Toolkits/SGlobalTabSwitchingDialog.cpp
@@ -392,7 +392,10 @@ void SGlobalTabSwitchingDialog::Construct(const FArguments& InArgs, FVector2D In
 	TArray<UObject*> OpenAssetList = FAssetEditorManager::Get().GetAllEditedAssets();
 	for (UObject* OpenAsset : OpenAssetList)
 	{
-		MainTabsListDataSource.Add(MakeShareable(new FTabSwitchingListItem_Asset(OpenAsset)));
+		if (OpenAsset->GetOuter() != GetTransientPackage())
+		{
+			MainTabsListDataSource.Add(MakeShareable(new FTabSwitchingListItem_Asset(OpenAsset)));
+		}
 	}
 
 	MainTabsListDataSource.Add(FTabSwitchingListItem_World::MakeWorldItem());
