RL - Fix issue where force delete materials crashes because the creation order gets out of order

diff --git a/Engine/Source/Editor/UnrealEd/Private/ObjectTools.cpp b/Engine/Source/Editor/UnrealEd/Private/ObjectTools.cpp
--- a/Engine/Source/Editor/UnrealEd/Private/ObjectTools.cpp
+++ b/Engine/Source/Editor/UnrealEd/Private/ObjectTools.cpp
@@ -743,7 +743,10 @@ namespace ObjectTools
 		ReplacementMap.GenerateKeyArray( OutInfo.ReplaceableObjects );
 
 		// Find all the properties (and their corresponding objects) that refer to any of the objects to be replaced
-		TMap< UObject*, TArray<UProperty*> > ReferencingPropertiesMap;
+		using PropertyArrayType = TArray<UProperty*, TInlineAllocator<1>>;
+		TArray<UObject*> ReferencingPropertiesMapKeys;
+		TArray<PropertyArrayType> ReferencingPropertiesMapValues;
+
 		for ( FObjectIterator ObjIter; ObjIter; ++ObjIter )
 		{
 			UObject* CurObject = *ObjIter;
@@ -761,12 +764,15 @@ namespace ObjectTools
 				TMultiMap<UObject*, UProperty*> CurReferencingPropertiesMMap;
 				if ( FindRefsArchive.GetReferenceCounts( CurNumReferencesMap, CurReferencingPropertiesMMap ) > 0  )
 				{
-					TArray<UProperty*> CurReferencedProperties;
+					PropertyArrayType CurReferencedProperties;
 					CurReferencingPropertiesMMap.GenerateValueArray( CurReferencedProperties );
-					ReferencingPropertiesMap.Add( CurObject, CurReferencedProperties );
+
+					ReferencingPropertiesMapKeys.Add(CurObject);
+					ReferencingPropertiesMapValues.Add(CurReferencedProperties);
+
 					if ( CurReferencedProperties.Num() > 0)
 					{
-						for ( TArray<UProperty*>::TConstIterator RefPropIter( CurReferencedProperties ); RefPropIter; ++RefPropIter )
+						for ( PropertyArrayType::TConstIterator RefPropIter( CurReferencedProperties ); RefPropIter; ++RefPropIter )
 						{
 							CurObject->PreEditChange( *RefPropIter );
 						}
@@ -778,15 +784,41 @@ namespace ObjectTools
 				}
 			}
 		}
-		
+
+		for (int CurrentIndex = 0; CurrentIndex < ReferencingPropertiesMapKeys.Num(); CurrentIndex++)
+		{
+			int CheckForCircles = 0;
+			FFindReferencersArchive FindDependentArchive(ReferencingPropertiesMapKeys[CurrentIndex], ReferencingPropertiesMapKeys);
+			for (int DependentIndex = CurrentIndex + 1; DependentIndex < ReferencingPropertiesMapKeys.Num(); DependentIndex++)
+			{
+				if (FindDependentArchive.GetReferenceCount(ReferencingPropertiesMapKeys[DependentIndex]) > 0)
+				{
+					if (CheckForCircles != DependentIndex)
+					{
+						UObject* Key = ReferencingPropertiesMapKeys[CurrentIndex];
+						PropertyArrayType Value = ReferencingPropertiesMapValues[CurrentIndex];
+						ReferencingPropertiesMapKeys[CurrentIndex] = ReferencingPropertiesMapKeys[DependentIndex];
+						ReferencingPropertiesMapValues[CurrentIndex] = ReferencingPropertiesMapValues[DependentIndex];
+						ReferencingPropertiesMapKeys[DependentIndex] = Key;
+						ReferencingPropertiesMapValues[DependentIndex] = Value;
+
+						FindDependentArchive.ResetPotentialReferencer(ReferencingPropertiesMapKeys[CurrentIndex]);
+						DependentIndex = CurrentIndex;
+
+						CheckForCircles = (CheckForCircles == 0) ? DependentIndex : CheckForCircles;
+					}
+				}
+			}
+		}
+
 		// Iterate over the map of referencing objects/changed properties, forcefully replacing the references and
 		int32 NumObjsReplaced = 0;
-		for ( TMap< UObject*, TArray<UProperty*> >::TConstIterator MapIter( ReferencingPropertiesMap ); MapIter; ++MapIter )
+		for (int32 Index = 0; Index < ReferencingPropertiesMapKeys.Num(); Index++)
 		{
 			++NumObjsReplaced;
-			GWarn->StatusUpdate( NumObjsReplaced, ReferencingPropertiesMap.Num(), NSLOCTEXT("UnrealEd", "ConsolidateAssetsUpdate_ReplacingReferences", "Replacing Asset References...") );
+			GWarn->StatusUpdate( NumObjsReplaced, ReferencingPropertiesMapKeys.Num(), NSLOCTEXT("UnrealEd", "ConsolidateAssetsUpdate_ReplacingReferences", "Replacing Asset References...") );
 
-			UObject* CurReplaceObj = MapIter.Key();
+			UObject* CurReplaceObj = ReferencingPropertiesMapKeys[Index];
 
 			FArchiveReplaceObjectRef<UObject> ReplaceAr( CurReplaceObj, ReplacementMap, false, true, false );
 		}
@@ -794,17 +826,17 @@ namespace ObjectTools
 		// Now alter the referencing objects the change has completed via PostEditChange, 
 		// this is done in a separate loop to prevent reading of data that we want to overwrite
 		int32 NumObjsPostEdited = 0;
-		for ( TMap< UObject*, TArray<UProperty*> >::TConstIterator MapIter( ReferencingPropertiesMap ); MapIter; ++MapIter )
+		for (int32 Index = 0; Index < ReferencingPropertiesMapKeys.Num(); Index++)
 		{
 			++NumObjsPostEdited;
-			GWarn->StatusUpdate( NumObjsPostEdited, ReferencingPropertiesMap.Num(), NSLOCTEXT("UnrealEd", "ConsolidateAssetsUpdate_PostEditing", "Performing Post Update Edits...") );
+			GWarn->StatusUpdate( NumObjsPostEdited, ReferencingPropertiesMapKeys.Num(), NSLOCTEXT("UnrealEd", "ConsolidateAssetsUpdate_PostEditing", "Performing Post Update Edits...") );
 
-			UObject* CurReplaceObj = MapIter.Key();
-			const TArray<UProperty*>& RefPropArray = MapIter.Value();
+			UObject* CurReplaceObj = ReferencingPropertiesMapKeys[Index];
+			const PropertyArrayType& RefPropArray = ReferencingPropertiesMapValues[Index];
 
 			if (RefPropArray.Num() > 0)
 			{
-				for ( TArray<UProperty*>::TConstIterator RefPropIter( RefPropArray ); RefPropIter; ++RefPropIter )
+				for ( PropertyArrayType::TConstIterator RefPropIter( RefPropArray ); RefPropIter; ++RefPropIter )
 				{
 					FPropertyChangedEvent PropertyEvent(*RefPropIter, EPropertyChangeType::Redirected);
 					CurReplaceObj->PostEditChangeProperty( PropertyEvent );
