Avoid Memcmp if size is zero

diff --git a/Engine/Source/Editor/UnrealEd/Private/EditorTransaction.cpp b/Engine/Source/Editor/UnrealEd/Private/EditorTransaction.cpp
--- a/Engine/Source/Editor/UnrealEd/Private/EditorTransaction.cpp
+++ b/Engine/Source/Editor/UnrealEd/Private/EditorTransaction.cpp
@@ -399,7 +399,7 @@ void FTransaction::FObjectRecord::Diff( FTransaction* Owner, const FSerializedOb
 
 			// Binary compare the serialized data to see if something has changed for this property
 			bool bIsPropertyIdentical = OldSerializedProperty->DataSize == NewNamePropertyPair.Value.DataSize;
-			if (bIsPropertyIdentical)
+			if (bIsPropertyIdentical && NewNamePropertyPair.Value.DataSize > 0)
 			{
 				bIsPropertyIdentical = FMemory::Memcmp(&OldSerializedObject.Data[OldSerializedProperty->DataOffset], &NewSerializedObject.Data[NewNamePropertyPair.Value.DataOffset], NewNamePropertyPair.Value.DataSize) == 0;
 			}
@@ -421,7 +421,7 @@ void FTransaction::FObjectRecord::Diff( FTransaction* Owner, const FSerializedOb
 			const int32 CurrentHeaderSize = StartOfNewPropertyBlock;
 
 			bool bIsHeaderIdentical = OldHeaderSize == CurrentHeaderSize;
-			if (bIsHeaderIdentical)
+			if (bIsHeaderIdentical && CurrentHeaderSize > 0)
 			{
 				bIsHeaderIdentical = FMemory::Memcmp(&OldSerializedObject.Data[0], &NewSerializedObject.Data[0], CurrentHeaderSize) == 0;
 			}
@@ -439,7 +439,7 @@ void FTransaction::FObjectRecord::Diff( FTransaction* Owner, const FSerializedOb
 			const int32 CurrentFooterSize = NewSerializedObject.Data.Num() - EndOfNewPropertyBlock;
 
 			bool bIsFooterIdentical = OldFooterSize == CurrentFooterSize;
-			if (bIsFooterIdentical)
+			if (bIsFooterIdentical && CurrentFooterSize > 0)
 			{
 				bIsFooterIdentical = FMemory::Memcmp(&OldSerializedObject.Data[EndOfOldPropertyBlock], &NewSerializedObject.Data[EndOfNewPropertyBlock], CurrentFooterSize) == 0;
 			}
@@ -454,7 +454,7 @@ void FTransaction::FObjectRecord::Diff( FTransaction* Owner, const FSerializedOb
 	{
 		// No properties, so just compare the whole blob
 		bool bIsBlobIdentical = OldSerializedObject.Data.Num() == NewSerializedObject.Data.Num();
-		if (bIsBlobIdentical)
+		if (bIsBlobIdentical && NewSerializedObject.Data.Num() > 0)
 		{
 			bIsBlobIdentical = FMemory::Memcmp(&OldSerializedObject.Data[0], &NewSerializedObject.Data[0], NewSerializedObject.Data.Num()) == 0;
 		}
