Can attempt to enter simulate no matter what the last play session is. Using Alt+S while using any PIE method other than Play in Current Viewport will silently fail; Alt+S using Play In Current Viewport will transition to simulate.

diff --git a/Engine/Source/Editor/UnrealEd/Private/Kismet2/DebuggerCommands.cpp b/Engine/Source/Editor/UnrealEd/Private/Kismet2/DebuggerCommands.cpp
--- a/Engine/Source/Editor/UnrealEd/Private/Kismet2/DebuggerCommands.cpp
+++ b/Engine/Source/Editor/UnrealEd/Private/Kismet2/DebuggerCommands.cpp
@@ -1397,21 +1397,12 @@ void FInternalPlayWorldCommandCallbacks::Simulate_Clicked()
 		return;
 	}
 
-	// Can only use SIE if in PIE
-	const ULevelEditorPlaySettings* PlaySettings = GetDefault<ULevelEditorPlaySettings>();
-	if (PlaySettings->LastExecutedPlayModeType != PlayMode_InViewPort &&
-		PlaySettings->LastExecutedPlayModeType != PlayMode_Simulate)
-	{
-		return;
-	}
-
-	SetLastExecutedPlayMode( PlayMode_Simulate );
-
 	FLevelEditorModule& LevelEditorModule = FModuleManager::GetModuleChecked<FLevelEditorModule>( TEXT("LevelEditor") );
 
 	TSharedPtr<ILevelViewport> ActiveLevelViewport = LevelEditorModule.GetFirstActiveViewport();
-	if( ActiveLevelViewport.IsValid() )
+	if( ActiveLevelViewport.IsValid() && ActiveLevelViewport->HasPlayInEditorViewport())
 	{
+		SetLastExecutedPlayMode(PlayMode_Simulate);
 		// Start a new simulation session!
 		if( !HasPlayWorld() )
 		{
