RL - Fix circlebreaking code in the resorting of force delete

diff --git a/Engine/Source/Editor/UnrealEd/Private/ObjectTools.cpp b/Engine/Source/Editor/UnrealEd/Private/ObjectTools.cpp
--- a/Engine/Source/Editor/UnrealEd/Private/ObjectTools.cpp
+++ b/Engine/Source/Editor/UnrealEd/Private/ObjectTools.cpp
@@ -803,9 +803,13 @@ namespace ObjectTools
 						ReferencingPropertiesMapValues[DependentIndex] = Value;
 
 						FindDependentArchive.ResetPotentialReferencer(ReferencingPropertiesMapKeys[CurrentIndex]);
+						
+						if(CheckForCircles == 0)
+						{
+							CheckForCircles = DependentIndex;
+						}
+						
 						DependentIndex = CurrentIndex;
-
-						CheckForCircles = (CheckForCircles == 0) ? DependentIndex : CheckForCircles;
 					}
 				}
 			}
