Fix for undo spam when changing colors in the gradient editor

diff --git a/Engine/Source/Editor/UnrealEd/Private/SColorGradientEditor.cpp b/Engine/Source/Editor/UnrealEd/Private/SColorGradientEditor.cpp
--- a/Engine/Source/Editor/UnrealEd/Private/SColorGradientEditor.cpp
+++ b/Engine/Source/Editor/UnrealEd/Private/SColorGradientEditor.cpp
@@ -595,7 +595,7 @@ void SColorGradientEditor::OpenGradientStopColorPicker()
 		// Open a color picker
 		FColorPickerArgs ColorPickerArgs;
 
-		ColorPickerArgs.bOnlyRefreshOnMouseUp = false;
+		ColorPickerArgs.bOnlyRefreshOnMouseUp = true;
 		ColorPickerArgs.bIsModal = false;
 		ColorPickerArgs.ParentWidget = SharedThis( this );
 		ColorPickerArgs.bUseAlpha = false;
