Merging 4122739 from Release-Staging-4.20 to address UE-60435

diff --git a/Engine/Source/Editor/UnrealEd/Private/ObjectTools.cpp b/Engine/Source/Editor/UnrealEd/Private/ObjectTools.cpp
--- a/Engine/Source/Editor/UnrealEd/Private/ObjectTools.cpp
+++ b/Engine/Source/Editor/UnrealEd/Private/ObjectTools.cpp
@@ -785,15 +785,15 @@ namespace ObjectTools
 			}
 		}
 
-		for (int CurrentIndex = 0; CurrentIndex < ReferencingPropertiesMapKeys.Num(); CurrentIndex++)
 		{
-			int CheckForCircles = 0;
-			FFindReferencersArchive FindDependentArchive(ReferencingPropertiesMapKeys[CurrentIndex], ReferencingPropertiesMapKeys);
-			for (int DependentIndex = CurrentIndex + 1; DependentIndex < ReferencingPropertiesMapKeys.Num(); DependentIndex++)
+			TBitArray<> TouchedThisItteration(false, ReferencingPropertiesMapKeys.Num());
+			for (int CurrentIndex = 0; CurrentIndex < ReferencingPropertiesMapKeys.Num(); CurrentIndex++)
 			{
-				if (FindDependentArchive.GetReferenceCount(ReferencingPropertiesMapKeys[DependentIndex]) > 0)
+				TouchedThisItteration.Init(false, ReferencingPropertiesMapKeys.Num());
+				FFindReferencersArchive FindDependentArchive(ReferencingPropertiesMapKeys[CurrentIndex], ReferencingPropertiesMapKeys);
+				for (int DependentIndex = CurrentIndex + 1; DependentIndex < ReferencingPropertiesMapKeys.Num(); DependentIndex++)
 				{
-					if (CheckForCircles != DependentIndex)
+					if (!TouchedThisItteration[DependentIndex] && FindDependentArchive.GetReferenceCount(ReferencingPropertiesMapKeys[DependentIndex]) > 0)
 					{
 						UObject* Key = ReferencingPropertiesMapKeys[CurrentIndex];
 						PropertyArrayType Value = ReferencingPropertiesMapValues[CurrentIndex];
@@ -803,12 +803,7 @@ namespace ObjectTools
 						ReferencingPropertiesMapValues[DependentIndex] = Value;
 
 						FindDependentArchive.ResetPotentialReferencer(ReferencingPropertiesMapKeys[CurrentIndex]);
-						
-						if(CheckForCircles == 0)
-						{
-							CheckForCircles = DependentIndex;
-						}
-						
+						TouchedThisItteration[DependentIndex] = true;
 						DependentIndex = CurrentIndex;
 					}
 				}
