Fix for simulate while PIE is active

diff --git a/Engine/Source/Editor/UnrealEd/Private/Kismet2/DebuggerCommands.cpp b/Engine/Source/Editor/UnrealEd/Private/Kismet2/DebuggerCommands.cpp
--- a/Engine/Source/Editor/UnrealEd/Private/Kismet2/DebuggerCommands.cpp
+++ b/Engine/Source/Editor/UnrealEd/Private/Kismet2/DebuggerCommands.cpp
@@ -1400,9 +1400,8 @@ void FInternalPlayWorldCommandCallbacks::Simulate_Clicked()
 	FLevelEditorModule& LevelEditorModule = FModuleManager::GetModuleChecked<FLevelEditorModule>( TEXT("LevelEditor") );
 
 	TSharedPtr<ILevelViewport> ActiveLevelViewport = LevelEditorModule.GetFirstActiveViewport();
-	if( ActiveLevelViewport.IsValid() && ActiveLevelViewport->HasPlayInEditorViewport())
+	if( ActiveLevelViewport.IsValid())
 	{
-		SetLastExecutedPlayMode(PlayMode_Simulate);
 		// Start a new simulation session!
 		if( !HasPlayWorld() )
 		{
@@ -1410,10 +1409,10 @@ void FInternalPlayWorldCommandCallbacks::Simulate_Clicked()
 			{
 				FEngineAnalytics::GetProvider().RecordEvent( TEXT("Editor.Usage.SimulateInEditor") );
 			}
-
+			SetLastExecutedPlayMode(PlayMode_Simulate);
 			GUnrealEd->RequestPlaySession(false, ActiveLevelViewport, true/*bSimulateInEditor*/, NULL, NULL, -1, false );
 		}
-		else
+		else if (ActiveLevelViewport->HasPlayInEditorViewport())
 		{
 			GUnrealEd->RequestToggleBetweenPIEandSIE();
 		}
