Creating a HLOD setup array adds non static mesh elements.
#fix adress issue where invalid actors could be added to HLOD clusters

diff --git a/Engine/Source/Editor/UnrealEd/Private/HierarchicalLOD.cpp b/Engine/Source/Editor/UnrealEd/Private/HierarchicalLOD.cpp
--- a/Engine/Source/Editor/UnrealEd/Private/HierarchicalLOD.cpp
+++ b/Engine/Source/Editor/UnrealEd/Private/HierarchicalLOD.cpp
@@ -323,35 +323,42 @@ void FHierarchicalLODBuilder::InitializeClusters(ULevel* InLevel, const int32 LO
 			for (int32 ActorId = 0; ActorId < InLevel->Actors.Num(); ++ActorId)
 			{
 				AActor* Actor = InLevel->Actors[ActorId];
-				if (ShouldGenerateCluster(Actor, bPreviewBuild, LODIdx))
+				const bool bShouldGenerate = ShouldGenerateCluster(Actor, bPreviewBuild, LODIdx);
+				if (bShouldGenerate)
 				{
 					// Check whether or not this actor falls within a HierarchicalLODVolume, if so add to the Volume's cluster and exclude from normal process
-					bool bAdded = bVolumesOnly;
-					for (TPair<AHierarchicalLODVolume*, FLODCluster>& Cluster : HLODVolumeClusters)
+					auto ProcessVolumeClusters = [this](AActor* InActor) -> bool
 					{
-						if (Cluster.Key->EncompassesPoint(Actor->GetActorLocation(), Cluster.Key->bIncludeOverlappingActors ? Actor->GetComponentsBoundingBox().GetSize().Size() : 0.0f, nullptr))
-						{			
-							FBox BoundingBox = Actor->GetComponentsBoundingBox(true);
-							FBox VolumeBox = Cluster.Key->GetComponentsBoundingBox(true);
-
-							if (VolumeBox.IsInside(BoundingBox)|| (Cluster.Key->bIncludeOverlappingActors && VolumeBox.Intersect(BoundingBox)))
+						for (TPair<AHierarchicalLODVolume*, FLODCluster>& Cluster : HLODVolumeClusters)
+						{
+							if (Cluster.Key->EncompassesPoint(InActor->GetActorLocation(), Cluster.Key->bIncludeOverlappingActors ? InActor->GetComponentsBoundingBox().GetSize().Size() : 0.0f, nullptr))
 							{
-								FLODCluster ActorCluster(Actor);
-								Cluster.Value += ActorCluster;
-								bAdded = true;
-								break;
-							}							
+								FBox BoundingBox = InActor->GetComponentsBoundingBox(true);
+								FBox VolumeBox = Cluster.Key->GetComponentsBoundingBox(true);
+
+								if (VolumeBox.IsInside(BoundingBox) || (Cluster.Key->bIncludeOverlappingActors && VolumeBox.Intersect(BoundingBox)))
+								{
+									FLODCluster ActorCluster(InActor);
+									Cluster.Value += ActorCluster;
+									return true;
+								}
+							}
 						}
-					}
 
-					if (!bAdded)
+						return false;
+					};
+
+					if (bVolumesOnly)
 					{
-						ValidStaticMeshActorsInLevel.Add(Actor);
+						ProcessVolumeClusters(Actor);
 					}
-				}
-				else if ( Actor != nullptr )
-				{
-					ValidStaticMeshActorsInLevel.Add(Actor);
+					else
+					{
+						if (!ProcessVolumeClusters(Actor))
+						{
+							ValidStaticMeshActorsInLevel.Add(Actor);
+						}
+					}					
 				}
 			}
 			
