Fixed broken static rect light

diff --git a/Engine/Source/Editor/UnrealEd/Private/Lightmass/Lightmass.cpp b/Engine/Source/Editor/UnrealEd/Private/Lightmass/Lightmass.cpp
--- a/Engine/Source/Editor/UnrealEd/Private/Lightmass/Lightmass.cpp
+++ b/Engine/Source/Editor/UnrealEd/Private/Lightmass/Lightmass.cpp
@@ -193,9 +193,11 @@ void Copy( const ULightComponent* In, Lightmass::FLightData& Out )
 {	
 	Copy((const ULightComponentBase*)In, Out);
 
+	const ULocalLightComponent* LocalLight = Cast<const ULocalLightComponent>(In);
 	const UPointLightComponent* PointLight = Cast<const UPointLightComponent>(In);
 
-	if( PointLight && PointLight->bUseInverseSquaredFalloff )
+	if( ( LocalLight && LocalLight->GetLightType() == LightType_Rect ) ||
+		( PointLight && PointLight->bUseInverseSquaredFalloff ) )
 	{
 		Out.LightFlags |= Lightmass::GI_LIGHT_INVERSE_SQUARED;
 	}
