Rename compilation manager flag in 4.20 so that users know they will have to use it in 4.21

diff --git a/Engine/Config/BaseEditor.ini b/Engine/Config/BaseEditor.ini
--- a/Engine/Config/BaseEditor.ini
+++ b/Engine/Config/BaseEditor.ini
@@ -511,7 +511,6 @@ bNativizeAnimBPOnlyWhenNonReducibleFuncitons=false
 
 [/Script/UnrealEd.BlueprintEditorProjectSettings]
 bUseCompilationManager=true
-bSkipUnneededDependencyCompilation=true
 bValidateUnloadedSoftActorReferences=true
 
 [/Script/UnrealEd.AssetViewerSettings]
diff --git a/Engine/Source/Editor/Kismet/Private/BlueprintCompilationManager.cpp b/Engine/Source/Editor/Kismet/Private/BlueprintCompilationManager.cpp
--- a/Engine/Source/Editor/Kismet/Private/BlueprintCompilationManager.cpp
+++ b/Engine/Source/Editor/Kismet/Private/BlueprintCompilationManager.cpp
@@ -569,7 +569,7 @@ void FBlueprintCompilationManagerImpl::FlushCompilationQueueImpl(TArray<UObject*
 			// then we can avoid dependency recompilation:
 			TSet<UBlueprint*> BlueprintsWithSignatureChanges;
 			const UBlueprintEditorProjectSettings* EditorProjectSettings = GetDefault<UBlueprintEditorProjectSettings>();
-			bool bSkipUnneededDependencyCompilation = EditorProjectSettings->bSkipUnneededDependencyCompilation;
+			bool bSkipUnneededDependencyCompilation = !EditorProjectSettings->bForceAllDependenciesToRecompile;
 
 			for (FCompilerData& CompilerData : CurrentlyCompilingBPs)
 			{
diff --git a/Engine/Source/Editor/UnrealEd/Private/Settings/EditorProjectSettings.cpp b/Engine/Source/Editor/UnrealEd/Private/Settings/EditorProjectSettings.cpp
--- a/Engine/Source/Editor/UnrealEd/Private/Settings/EditorProjectSettings.cpp
+++ b/Engine/Source/Editor/UnrealEd/Private/Settings/EditorProjectSettings.cpp
@@ -147,12 +147,3 @@ UBlueprintEditorProjectSettings::UBlueprintEditorProjectSettings(const FObjectIn
 	: Super(ObjectInitializer)
 {
 }
-
-void UBlueprintEditorProjectSettings::PostEditChangeProperty(struct FPropertyChangedEvent& PropertyChangedEvent)
-{
-	const FName Name = PropertyChangedEvent.Property ? PropertyChangedEvent.Property->GetFName() : NAME_None;
-	if (Name == GET_MEMBER_NAME_CHECKED(UBlueprintEditorProjectSettings, bUseCompilationManager))
-	{
-		GBlueprintUseCompilationManager = bUseCompilationManager;
-	}
-}
diff --git a/Engine/Source/Editor/UnrealEd/Public/Settings/EditorProjectSettings.h b/Engine/Source/Editor/UnrealEd/Public/Settings/EditorProjectSettings.h
--- a/Engine/Source/Editor/UnrealEd/Public/Settings/EditorProjectSettings.h
+++ b/Engine/Source/Editor/UnrealEd/Public/Settings/EditorProjectSettings.h
@@ -164,20 +164,24 @@ class UNREALED_API UBlueprintEditorProjectSettings : public UDeveloperSettings
 	GENERATED_UCLASS_BODY()
 
 public:
-	/** Flag to disable the new compilation manager for blueprints */
-	UPROPERTY(EditAnywhere, config, Category=Blueprints)
-	uint32 bUseCompilationManager:1;
+	/**
+	 * Flag to disable the compilation manager for blueprints - requires editor restart.
+	 * This flag is deprecated! In 4.21 the compilation manager will be the only way to compile
+	 * blueprints. Report any compilation manager issues immediately.
+	 */
+	UPROPERTY(EditAnywhere, config, Category = Blueprints, DisplayName = "Disable Compilation Manager (DEPRECATED)")
+	uint32 bDisableCompilationManager:1;
 	
-	/** Flag to enable faster compiles for individual blueprints if they have no function signature changes */
-	UPROPERTY(EditAnywhere, config, Category=Blueprints)
-	uint32 bSkipUnneededDependencyCompilation:1;
+	/**
+	 * Flag to disable faster compiles for individual blueprints if they have no function signature
+	 * changes. This flag is deprecated! In 4.21 there will be no way to force all dependencies to 
+	 * compile when no changes are detected. Report any issues immediately.
+	 */
+	UPROPERTY(EditAnywhere, config, Category=Blueprints, DisplayName = "Force All Dependencies To Recompile (DEPRECATED)")
+	uint32 bForceAllDependenciesToRecompile:1;
 
 	/** If enabled, the editor will load packages to look for soft references to actors when deleting/renaming them. This can be slow in large projects so disable this to improve performance but increase the chance of breaking blueprints/sequences that use soft actor references */
 	UPROPERTY(EditAnywhere, config, Category=Actors)
 	uint32 bValidateUnloadedSoftActorReferences : 1;
-
-	// UObject interface
-	virtual void PostEditChangeProperty(struct FPropertyChangedEvent& PropertyChangedEvent) override;
-	// End of UObject interface
 };
 
diff --git a/Engine/Source/Runtime/CoreUObject/Private/Blueprint/BlueprintSupport.cpp b/Engine/Source/Runtime/CoreUObject/Private/Blueprint/BlueprintSupport.cpp
--- a/Engine/Source/Runtime/CoreUObject/Private/Blueprint/BlueprintSupport.cpp
+++ b/Engine/Source/Runtime/CoreUObject/Private/Blueprint/BlueprintSupport.cpp
@@ -28,8 +28,8 @@
 
 DEFINE_LOG_CATEGORY_STATIC(LogBlueprintSupport, Log, All);
 
-// Flag to enable the new BlueprintCompilationManager:
-COREUOBJECT_API bool GBlueprintUseCompilationManager = false;
+// Flag to enable the BlueprintCompilationManager:
+COREUOBJECT_API bool GBlueprintUseCompilationManager = true;
 
 const FName FBlueprintTags::GeneratedClassPath(TEXT("GeneratedClass"));
 const FName FBlueprintTags::ParentClassPath(TEXT("ParentClass"));
@@ -120,7 +120,22 @@ void FBlueprintSupport::InitializeCompilationManager()
 {
 	// 'real' initialization is done lazily because we're in a pretty tough
 	// spot in terms of dependencies:
-	GConfig->GetBool(TEXT("/Script/UnrealEd.BlueprintEditorProjectSettings"), TEXT("bUseCompilationManager"), GBlueprintUseCompilationManager, GEditorIni);
+	bool bCompilationManagerDisabled = false;
+	bool bCompilationManagerEnabled = true;
+	GConfig->GetBool(TEXT("/Script/UnrealEd.BlueprintEditorProjectSettings"), TEXT("bUseCompilationManager"), bCompilationManagerEnabled, GEditorIni);
+	GConfig->GetBool(TEXT("/Script/UnrealEd.BlueprintEditorProjectSettings"), TEXT("bDisableCompilationManager"), bCompilationManagerDisabled, GEditorIni);
+	if (bCompilationManagerDisabled)
+	{
+		GBlueprintUseCompilationManager = false;
+	}
+	else if(!bCompilationManagerEnabled)
+	{
+		UE_LOG(LogScript, Warning,
+			TEXT("Warning: Compilation manager enabled, compilation manager will be mandatory in 4.21. \
+				Use bDisableCompilationManager to disable the compilation manager as a last resort or remove \
+				bUseCompilationManager from project .ini to fix this warning.")
+		);
+	}
 }
 
 static FFlushReinstancingQueueFPtr FlushReinstancingQueueFPtr = nullptr;
