Ensure that streamed levels are properly associated with the world and level caches after being reverted or reloaded.

diff --git a/Engine/Source/Editor/UnrealEd/Private/PackageTools.cpp b/Engine/Source/Editor/UnrealEd/Private/PackageTools.cpp
--- a/Engine/Source/Editor/UnrealEd/Private/PackageTools.cpp
+++ b/Engine/Source/Editor/UnrealEd/Private/PackageTools.cpp
@@ -570,6 +570,7 @@ namespace PackageTools
 		// Check to see if we need to reload the current world.
 		FName WorldNameToReload;
 		TMap<FName, const UMapBuildDataRegistry*> LevelsToMapBuildData;
+		TArray<ULevelStreaming*> RemovedStreamingLevels;
 		{
 			if (UWorld* EditorWorld = GEditor->GetEditorWorldContext().World())
 			{
@@ -610,10 +611,27 @@ namespace PackageTools
 				// Cache the current map build data for the levels of the current world so we can see if they change due to a reload (we can skip this if reloading the current world).
 				else
 				{
-					for (int32 LevelIndex = 0; LevelIndex < EditorWorld->GetNumLevels(); ++LevelIndex)
+					TArray<ULevel*> EditorLevels = EditorWorld->GetLevels();
+
+					for (ULevel* Level : EditorLevels)
 					{
-						ULevel* Level = EditorWorld->GetLevel(LevelIndex);
-						LevelsToMapBuildData.Add(Level->GetFName(), Level->MapBuildData);
+						if (PackagesToReload.Contains(Level->GetOutermost()))
+						{
+							for (ULevelStreaming* StreamingLevel : EditorWorld->GetStreamingLevels())
+							{
+								if (StreamingLevel->GetLoadedLevel() == Level)
+								{
+									EditorWorld->RemoveFromWorld(Level);
+									StreamingLevel->RemoveLevelFromCollectionForReload();
+									RemovedStreamingLevels.Add(StreamingLevel);
+									break;
+								}
+							}
+						}
+						else
+						{
+							LevelsToMapBuildData.Add(Level->GetFName(), Level->MapBuildData);
+						}
 					}
 				}
 			}
@@ -696,20 +714,35 @@ namespace PackageTools
 			FAssetEditorManager::Get().OpenEditorsForAssets(WorldNamesToReload);
 		}
 		// Update the rendering resources for the levels of the current world if their map build data has changed (we skip this if reloading the current world).
-		else if (LevelsToMapBuildData.Num() > 0)
+		else
 		{
-			UWorld* EditorWorld = GEditor->GetEditorWorldContext().World();
-			check(EditorWorld);
+			if (LevelsToMapBuildData.Num() > 0)
+			{
+				UWorld* EditorWorld = GEditor->GetEditorWorldContext().World();
+				check(EditorWorld);
 
-			for (int32 LevelIndex = 0; LevelIndex < EditorWorld->GetNumLevels(); ++LevelIndex)
+				for (int32 LevelIndex = 0; LevelIndex < EditorWorld->GetNumLevels(); ++LevelIndex)
+				{
+					ULevel* Level = EditorWorld->GetLevel(LevelIndex);
+					const UMapBuildDataRegistry* OldMapBuildData = LevelsToMapBuildData.FindRef(Level->GetFName());
+
+					if (OldMapBuildData && OldMapBuildData != Level->MapBuildData)
+					{
+						Level->ReleaseRenderingResources();
+						Level->InitializeRenderingResources();
+					}
+				}
+			}
+			if (RemovedStreamingLevels.Num() > 0)
 			{
-				ULevel* Level = EditorWorld->GetLevel(LevelIndex);
-				const UMapBuildDataRegistry* OldMapBuildData = LevelsToMapBuildData.FindRef(Level->GetFName());
+				UWorld* EditorWorld = GEditor->GetEditorWorldContext().World();
+				check(EditorWorld);
 
-				if (OldMapBuildData && OldMapBuildData != Level->MapBuildData)
+				for (ULevelStreaming* StreamingLevel : RemovedStreamingLevels)
 				{
-					Level->ReleaseRenderingResources();
-					Level->InitializeRenderingResources();
+					ULevel* NewLevel = StreamingLevel->GetLoadedLevel();
+					EditorWorld->AddToWorld(NewLevel, StreamingLevel->LevelTransform, false);
+					StreamingLevel->AddLevelToCollectionAfterReload();
 				}
 			}
 		}
diff --git a/Engine/Source/Runtime/Engine/Classes/Engine/LevelStreaming.h b/Engine/Source/Runtime/Engine/Classes/Engine/LevelStreaming.h
--- a/Engine/Source/Runtime/Engine/Classes/Engine/LevelStreaming.h
+++ b/Engine/Source/Runtime/Engine/Classes/Engine/LevelStreaming.h
@@ -445,6 +445,12 @@ public:
 		NeverBlock,
 	};
 
+#if WITH_EDITOR
+	// After a sub level is reloaded in the editor the cache state needs to be refreshed
+	void RemoveLevelFromCollectionForReload();
+	void AddLevelToCollectionAfterReload();
+#endif
+
 private:
 	/** @return Name of the LOD level package used for loading.																		*/
 	FName GetLODPackageName() const;
diff --git a/Engine/Source/Runtime/Engine/Private/LevelStreaming.cpp b/Engine/Source/Runtime/Engine/Private/LevelStreaming.cpp
--- a/Engine/Source/Runtime/Engine/Private/LevelStreaming.cpp
+++ b/Engine/Source/Runtime/Engine/Private/LevelStreaming.cpp
@@ -706,6 +706,36 @@ FName ULevelStreaming::GetLODPackageNameToLoad() const
 	}
 }
 
+#if WITH_EDITOR
+void ULevelStreaming::RemoveLevelFromCollectionForReload()
+{
+	if (LoadedLevel)
+	{
+		// Remove the loaded level from its current collection, if any.
+		if (LoadedLevel->GetCachedLevelCollection())
+		{
+			LoadedLevel->GetCachedLevelCollection()->RemoveLevel(LoadedLevel);
+		}
+	}
+}
+
+void ULevelStreaming::AddLevelToCollectionAfterReload()
+{
+	if (LoadedLevel)
+	{
+		// Remove the loaded level from its current collection, if any.
+		if (LoadedLevel->GetCachedLevelCollection())
+		{
+			LoadedLevel->GetCachedLevelCollection()->RemoveLevel(LoadedLevel);
+		}
+		// Add this level to the correct collection
+		const ELevelCollectionType CollectionType = bIsStatic ? ELevelCollectionType::StaticLevels : ELevelCollectionType::DynamicSourceLevels;
+		FLevelCollection& LC = GetWorld()->FindOrAddCollectionByType(CollectionType);
+		LC.AddLevel(LoadedLevel);
+	}
+}
+#endif
+
 void ULevelStreaming::SetLoadedLevel(ULevel* Level)
 { 
 	// Pending level should be unloaded at this point
