fix bsp undo not redrawing the viewport and reflecting changes

diff --git a/Engine/Source/Editor/UnrealEd/Private/EditorServer.cpp b/Engine/Source/Editor/UnrealEd/Private/EditorServer.cpp
--- a/Engine/Source/Editor/UnrealEd/Private/EditorServer.cpp
+++ b/Engine/Source/Editor/UnrealEd/Private/EditorServer.cpp
@@ -1365,6 +1365,8 @@ void UEditorEngine::PostUndo(bool bSuccess)
 
 	// Re-instance any actors that need it
 	FBlueprintCompileReinstancer::BatchReplaceInstancesOfClass(OldToNewClassMapToReinstance, false);
+
+	RedrawLevelEditingViewports();
 }
 
 bool UEditorEngine::UndoTransaction(bool bCanRedo)
