Fix incorrectly considering sublevels being closed for asset editors
#rnx

diff --git a/Engine/Source/Editor/UnrealEd/Private/EditorEngine.cpp b/Engine/Source/Editor/UnrealEd/Private/EditorEngine.cpp
--- a/Engine/Source/Editor/UnrealEd/Private/EditorEngine.cpp
+++ b/Engine/Source/Editor/UnrealEd/Private/EditorEngine.cpp
@@ -2398,7 +2398,7 @@ void UEditorEngine::CloseEditedWorldAssets(UWorld* InWorld)
 	{
 		if (LevelStreaming && LevelStreaming->GetLoadedLevel())
 		{
-			ClosingWorlds.Add(LevelStreaming->GetWorld());
+			ClosingWorlds.Add(CastChecked<UWorld>(LevelStreaming->GetLoadedLevel()->GetOuter()));
 		}
 	}
 
