[NO NON-PII COMMIT MESSAGE]

diff --git a/Engine/Source/Editor/UnrealEd/Private/Toolkits/AssetEditorManager.cpp b/Engine/Source/Editor/UnrealEd/Private/Toolkits/AssetEditorManager.cpp
--- a/Engine/Source/Editor/UnrealEd/Private/Toolkits/AssetEditorManager.cpp
+++ b/Engine/Source/Editor/UnrealEd/Private/Toolkits/AssetEditorManager.cpp
@@ -757,13 +757,24 @@ void FAssetEditorManager::HandlePackageReloaded(const EPackageReloadPhase InPack
 
 	if (InPackageReloadPhase == EPackageReloadPhase::PrePackageFixup)
 	{
+		/** Call close for all old assets even if not open, so global callback will go off */
 		TArray<UObject*> OldAssets;
-		for (auto AssetEditorPair : OpenedAssets)
+		const TMap<UObject*, UObject*>& RepointedMap = InPackageReloadedEvent->GetRepointedObjects();
+
+		for (const TPair<UObject*, UObject*> RepointPair : RepointedMap)
+		{
+			if (RepointPair.Key->IsAsset())
+			{
+				OldAssets.Add(RepointPair.Key);
+			}
+		}
+
+		/** Look for replacement for assets that are open now so we can reopen */
+		for (TPair<UObject*, IAssetEditorInstance*>& AssetEditorPair : OpenedAssets)
 		{
 			UObject* NewAsset = nullptr;
 			if (InPackageReloadedEvent->GetRepointedObject(AssetEditorPair.Key, NewAsset))
 			{
-				OldAssets.Add(AssetEditorPair.Key);
 				if (NewAsset)
 				{
 					PendingAssetsToOpen.AddUnique(NewAsset);
