Fix crash when overwriting levels

diff --git a/Engine/Source/Editor/ContentBrowser/Private/ContentBrowserUtils.cpp b/Engine/Source/Editor/ContentBrowser/Private/ContentBrowserUtils.cpp
--- a/Engine/Source/Editor/ContentBrowser/Private/ContentBrowserUtils.cpp
+++ b/Engine/Source/Editor/ContentBrowser/Private/ContentBrowserUtils.cpp
@@ -581,8 +581,10 @@ bool ContentBrowserUtils::DeleteFolders(const TArray<FString>& PathsToDelete)
 			// Make sure we loaded all of them
 			if ( LoadedAssets.Num() == NumAssetsInPaths )
 			{
-				const int32 NumAssetsDeleted = ContentBrowserUtils::DeleteAssets(LoadedAssets);
-				if ( NumAssetsDeleted == NumAssetsInPaths )
+				TArray<UObject*> ToDelete = LoadedAssets;
+				ObjectTools::AddExtraObjectsToDelete(ToDelete);
+				const int32 NumAssetsDeleted = ContentBrowserUtils::DeleteAssets(ToDelete);
+				if ( NumAssetsDeleted == ToDelete.Num() )
 				{
 					// Successfully deleted all assets in the specified path. Allow the folder to be removed.
 					bAllowFolderDelete = true;
diff --git a/Engine/Source/Editor/UnrealEd/Private/ObjectTools.cpp b/Engine/Source/Editor/UnrealEd/Private/ObjectTools.cpp
--- a/Engine/Source/Editor/UnrealEd/Private/ObjectTools.cpp
+++ b/Engine/Source/Editor/UnrealEd/Private/ObjectTools.cpp
@@ -480,6 +480,19 @@ namespace ObjectTools
 				}
 			}
 
+			// Handle map built data
+			const int32 OriginalNumToDelete = ObjectsToDelete.Num();
+			ObjectTools::AddExtraObjectsToDelete(ObjectsToDelete);
+			for (int32 i = OriginalNumToDelete; i < ObjectsToDelete.Num(); ++i)
+			{
+				UPackage* Pkg = ObjectsToDelete[i]->GetOutermost();
+				if ( Pkg && !Pkg->IsRooted() )
+				{
+					DeletedObjectPackages.AddUnique(Pkg);
+					Pkg->AddToRoot();
+				}
+			}
+
 			const int32 NumObjectsDeleted = ObjectTools::DeleteObjects(ObjectsToDelete, bPromptToOverwrite);
 
 			// Remove all packages that we added to the root set above.
@@ -1914,11 +1927,13 @@ namespace ObjectTools
 		return NumPackagesToDelete + NumObjectsToDelete;
 	}
 
-	void AddExtraObjectsToDelete(const TArray< UObject* >& InObjectsToDelete, TArray< UObject* >& ObjectsToDelete)
+	void AddExtraObjectsToDelete(TArray< UObject* >& ObjectsToDelete)
 	{
-		ObjectsToDelete = InObjectsToDelete;
-		for (UObject *ObjectToDelete : InObjectsToDelete)
+		const int32 OriginalNum = ObjectsToDelete.Num();
+		for (int32 i=0; i < OriginalNum; ++i)
 		{
+			UObject* ObjectToDelete = ObjectsToDelete[i];
+
 			// Delete MapBuildData with maps
 			if (UWorld* World = Cast<UWorld>(ObjectToDelete))
 			{
@@ -1954,8 +1969,8 @@ namespace ObjectTools
 
 		const FScopedBusyCursor BusyCursor;
 
-		TArray<UObject*> ObjectsToDelete;
-		AddExtraObjectsToDelete(InObjectsToDelete, ObjectsToDelete);
+		TArray<UObject*> ObjectsToDelete = InObjectsToDelete;
+		AddExtraObjectsToDelete(ObjectsToDelete);
 
 		// Make sure packages being saved are fully loaded.
 		if( !HandleFullyLoadingPackages( ObjectsToDelete, NSLOCTEXT("UnrealEd", "Delete", "Delete") ) )
@@ -2075,8 +2090,8 @@ namespace ObjectTools
 		GWarn->BeginSlowTask( NSLOCTEXT( "UnrealEd", "Deleting", "Deleting" ), true );
 
 		TArray<UObject*> ObjectsDeletedSuccessfully;
-		TArray<UObject*> ObjectsToDelete;
-		AddExtraObjectsToDelete(InObjectsToDelete, ObjectsToDelete);
+		TArray<UObject*> ObjectsToDelete = InObjectsToDelete;
+		AddExtraObjectsToDelete(ObjectsToDelete);
 
 		bool bSawSuccessfulDelete = false;
 		bool bMakeWritable = false;
@@ -2201,8 +2216,8 @@ namespace ObjectTools
 	{
 		int32 NumDeletedObjects = 0;
 
-		TArray<UObject*> ShownObjectsToDelete;
-		AddExtraObjectsToDelete(InObjectsToDelete, ShownObjectsToDelete);
+		TArray<UObject*> ShownObjectsToDelete = InObjectsToDelete;
+		AddExtraObjectsToDelete(ShownObjectsToDelete);
 
 		// Confirm that the delete was intentional
 		if ( ShowConfirmation && !ShowDeleteConfirmationDialog(ShownObjectsToDelete) )
diff --git a/Engine/Source/Editor/UnrealEd/Public/ObjectTools.h b/Engine/Source/Editor/UnrealEd/Public/ObjectTools.h
--- a/Engine/Source/Editor/UnrealEd/Public/ObjectTools.h
+++ b/Engine/Source/Editor/UnrealEd/Public/ObjectTools.h
@@ -312,6 +312,14 @@ namespace ObjectTools
 	 */
 	UNREALED_API int32 ForceDeleteObjects( const TArray< UObject* >& ObjectsToDelete, bool ShowConfirmation = true );
 
+	/**
+	 * Gathers additional objects to delete such as map built data
+	 *
+	 * @param	ObjectsToDelete		List of objects to delete that is appended by additional objects that should be deleted
+	 *
+	 */
+	UNREALED_API void AddExtraObjectsToDelete(TArray< UObject* >& ObjectsToDelete);
+
 	/**
 	 * Utility function to compose a string list of referencing objects
 	 *
