Fix DPI scale issue with large vertices show flag

diff --git a/Engine/Source/Editor/UnrealEd/Private/EdMode.cpp b/Engine/Source/Editor/UnrealEd/Private/EdMode.cpp
--- a/Engine/Source/Editor/UnrealEd/Private/EdMode.cpp
+++ b/Engine/Source/Editor/UnrealEd/Private/EdMode.cpp
@@ -746,6 +746,8 @@ void FEdMode::DrawHUD(FEditorViewportClient* ViewportClient,FViewport* Viewport,
 					Vertices.AddUnique( ActorToWorld.TransformPosition( VertexBuffer.VertexPosition(i) ) );
 				}
 
+				const float InvDpiScale = 1.0f / Canvas->GetDPIScale();
+
 				FCanvasTileItem TileItem( FVector2D( 0.0f, 0.0f ), FVector2D( 0.0f, 0.0f ), FLinearColor::White );
 				TileItem.BlendMode = SE_BLEND_Translucent;
 				for( int32 VertexIndex = 0 ; VertexIndex < Vertices.Num() ; ++VertexIndex )
@@ -754,9 +756,11 @@ void FEdMode::DrawHUD(FEditorViewportClient* ViewportClient,FViewport* Viewport,
 					FVector2D PixelLocation;
 					if(View->ScreenToPixel(View->WorldToScreen(Vertex),PixelLocation))
 					{
+						PixelLocation *= InvDpiScale;
+
 						const bool bOutside =
-							PixelLocation.X < 0.0f || PixelLocation.X > View->UnscaledViewRect.Width() ||
-							PixelLocation.Y < 0.0f || PixelLocation.Y > View->UnscaledViewRect.Height();
+							PixelLocation.X < 0.0f || PixelLocation.X > View->UnscaledViewRect.Width()*InvDpiScale ||
+							PixelLocation.Y < 0.0f || PixelLocation.Y > View->UnscaledViewRect.Height()*InvDpiScale;
 						if( !bOutside )
 						{
 							const float X = PixelLocation.X - (TextureSizeX/2);
