Fix camera safe frames in high dpi

diff --git a/Engine/Source/Editor/UnrealEd/Private/EditorViewportClient.cpp b/Engine/Source/Editor/UnrealEd/Private/EditorViewportClient.cpp
--- a/Engine/Source/Editor/UnrealEd/Private/EditorViewportClient.cpp
+++ b/Engine/Source/Editor/UnrealEd/Private/EditorViewportClient.cpp
@@ -1881,10 +1881,10 @@ void FEditorViewportClient::UpdateLightingShowFlags( FEngineShowFlags& InOutShow
 	}
 }
 
-bool FEditorViewportClient::CalculateEditorConstrainedViewRect(FSlateRect& OutSafeFrameRect, FViewport* InViewport)
+bool FEditorViewportClient::CalculateEditorConstrainedViewRect(FSlateRect& OutSafeFrameRect, FViewport* InViewport, float DPIScale)
 {
-	const int32 SizeX = InViewport->GetSizeXY().X;
-	const int32 SizeY = InViewport->GetSizeXY().Y;
+	const int32 SizeX = InViewport->GetSizeXY().X / DPIScale;
+	const int32 SizeY = InViewport->GetSizeXY().Y / DPIScale;
 
 	OutSafeFrameRect = FSlateRect(0, 0, SizeX, SizeY);
 	float FixedAspectRatio;
@@ -1928,7 +1928,7 @@ void FEditorViewportClient::DrawSafeFrames(FViewport& InViewport, FSceneView& Vi
 	if (EngineShowFlags.CameraAspectRatioBars || EngineShowFlags.CameraSafeFrames)
 	{
 		FSlateRect SafeRect;
-		if (CalculateEditorConstrainedViewRect(SafeRect, &InViewport))
+		if (CalculateEditorConstrainedViewRect(SafeRect, &InViewport, Canvas.GetDPIScale()))
 		{
 			if (EngineShowFlags.CameraSafeFrames)
 			{
@@ -3579,7 +3579,7 @@ void FEditorViewportClient::Draw(FViewport* InViewport, FCanvas* Canvas)
 
 	    FSlateRect SafeFrame;
 	    View->CameraConstrainedViewRect = View->UnscaledViewRect;
-	    if (CalculateEditorConstrainedViewRect(SafeFrame, Viewport))
+	    if (CalculateEditorConstrainedViewRect(SafeFrame, Viewport, Canvas->GetDPIScale()))
 	    {
 		    View->CameraConstrainedViewRect = FIntRect(SafeFrame.Left, SafeFrame.Top, SafeFrame.Right, SafeFrame.Bottom);
 	    }
diff --git a/Engine/Source/Editor/UnrealEd/Private/LevelEditorViewport.cpp b/Engine/Source/Editor/UnrealEd/Private/LevelEditorViewport.cpp
--- a/Engine/Source/Editor/UnrealEd/Private/LevelEditorViewport.cpp
+++ b/Engine/Source/Editor/UnrealEd/Private/LevelEditorViewport.cpp
@@ -4543,7 +4543,7 @@ void FLevelEditorViewportClient::SetCameraSpeedScalar(float SpeedScalar)
 bool FLevelEditorViewportClient::OverrideHighResScreenshotCaptureRegion(FIntRect& OutCaptureRegion)
 {
 	FSlateRect Rect;
-	if (CalculateEditorConstrainedViewRect(Rect, Viewport))
+	if (CalculateEditorConstrainedViewRect(Rect, Viewport, GetDPIScale()))
 	{
 		FSlateRect InnerRect = Rect.InsetBy(FMargin(0.5f * SafePadding * Rect.GetSize().Size()));
 		OutCaptureRegion = FIntRect((int32)InnerRect.Left, (int32)InnerRect.Top, (int32)(InnerRect.Left + InnerRect.GetSize().X), (int32)(InnerRect.Top + InnerRect.GetSize().Y));
diff --git a/Engine/Source/Editor/UnrealEd/Public/EditorViewportClient.h b/Engine/Source/Editor/UnrealEd/Public/EditorViewportClient.h
--- a/Engine/Source/Editor/UnrealEd/Public/EditorViewportClient.h
+++ b/Engine/Source/Editor/UnrealEd/Public/EditorViewportClient.h
@@ -1256,7 +1256,7 @@ protected:
 	virtual bool GetActiveSafeFrame(float& OutAspectRatio) const { return false; }
 
 	/** Helper function to calculate the safe frame rectangle on the current viewport */
-	bool CalculateEditorConstrainedViewRect(FSlateRect& OutSafeFrameRect, FViewport* InViewport);
+	bool CalculateEditorConstrainedViewRect(FSlateRect& OutSafeFrameRect, FViewport* InViewport, const float DPIScale);
 
 	virtual void NudgeSelectedObjects(const struct FInputEventState& InputState) {}
 
