Editor: Check for valid play world before getting sim world counterpart.

diff --git a/Engine/Source/Editor/UnrealEd/Private/Editor.cpp b/Engine/Source/Editor/UnrealEd/Private/Editor.cpp
--- a/Engine/Source/Editor/UnrealEd/Private/Editor.cpp
+++ b/Engine/Source/Editor/UnrealEd/Private/Editor.cpp
@@ -726,17 +726,20 @@ namespace EditorUtilities
 
 			// Find the counterpart level
 			UWorld* PlayWorld = GEditor->PlayWorld;
-			for( auto LevelIt( PlayWorld->GetLevelIterator() ); LevelIt; ++LevelIt )
+			if (PlayWorld != nullptr)
 			{
-				auto* Level = *LevelIt;
-				if( Level->GetFName() == EditorWorldActor->GetLevel()->GetFName() )
+				for (auto LevelIt(PlayWorld->GetLevelIterator()); LevelIt; ++LevelIt)
 				{
-					// Find our counterpart actor
-					const bool bExactClass = false;	// Don't match class exactly, because we support all classes derived from Actor as well!
-					AActor* SimWorldActor = FindObject<AActor>( Level, *EditorWorldActor->GetFName().ToString(), bExactClass );
-					if( SimWorldActor && GEditor->ObjectsThatExistInEditorWorld.Get( SimWorldActor ) )
+					auto* Level = *LevelIt;
+					if (Level->GetFName() == EditorWorldActor->GetLevel()->GetFName())
 					{
-						return SimWorldActor;
+						// Find our counterpart actor
+						const bool bExactClass = false;	// Don't match class exactly, because we support all classes derived from Actor as well!
+						AActor* SimWorldActor = FindObject<AActor>(Level, *EditorWorldActor->GetFName().ToString(), bExactClass);
+						if (SimWorldActor && GEditor->ObjectsThatExistInEditorWorld.Get(SimWorldActor))
+						{
+							return SimWorldActor;
+						}
 					}
 				}
 			}
