Logging to track invalid assets in the content browser??

diff --git a/Engine/Source/Editor/UnrealEd/Private/Toolkits/AssetEditorManager.cpp b/Engine/Source/Editor/UnrealEd/Private/Toolkits/AssetEditorManager.cpp
--- a/Engine/Source/Editor/UnrealEd/Private/Toolkits/AssetEditorManager.cpp
+++ b/Engine/Source/Editor/UnrealEd/Private/Toolkits/AssetEditorManager.cpp
@@ -26,6 +26,8 @@
 
 #define LOCTEXT_NAMESPACE "AssetEditorManager"
 
+DEFINE_LOG_CATEGORY_STATIC(LogAssetEditorManager, Log, All);
+
 FAssetEditorManager* FAssetEditorManager::Instance = NULL;
 
 
@@ -331,6 +333,7 @@ bool FAssetEditorManager::OpenEditorForAsset(UObject* Asset, const EToolkitMode:
 		}
 	}
 
+	UE_LOG(LogAssetEditorManager, Log, TEXT("Opening Asset editor for %s"), *Asset->GetFullName());
 
 	FAssetToolsModule& AssetToolsModule = FModuleManager::LoadModuleChecked<FAssetToolsModule>(TEXT("AssetTools"));
 
