Fix re import assignment of sections (change from dev-editor port to 4.20)

diff --git a/Engine/Source/Editor/UnrealEd/Private/StaticMeshEdit.cpp b/Engine/Source/Editor/UnrealEd/Private/StaticMeshEdit.cpp
--- a/Engine/Source/Editor/UnrealEd/Private/StaticMeshEdit.cpp
+++ b/Engine/Source/Editor/UnrealEd/Private/StaticMeshEdit.cpp
@@ -1667,21 +1667,44 @@ void RestoreExistingMeshData(ExistingStaticMeshData* ExistingMeshDataPtr, UStati
 
 				FMeshSectionInfo NewSectionInfo =  NewMesh->SectionInfoMap.Get(i, SectionIndex);
 				bool bFoundOldMatch = false;
-				if (!ExistingMeshDataPtr->UseMaterialNameSlotWorkflow)
+				bool bKeepOldSectionMaterialIndex = false;
+				int32 OriginalSectionMaterialIndex = INDEX_NONE;
+				if (RemapMaterial.IsValidIndex(NewSectionInfo.MaterialIndex) && NewMesh->StaticMaterials.IsValidIndex(NewSectionInfo.MaterialIndex))
 				{
 					//Find the matching old index
 					for (int32 ExistSectionIndex = 0; ExistSectionIndex < OldSectionNumber; ++ExistSectionIndex)
 					{
+						bKeepOldSectionMaterialIndex = false;
+						OriginalSectionMaterialIndex = INDEX_NONE;
 						FMeshSectionInfo OldSectionInfo = ExistingMeshDataPtr->ExistingSectionInfoMap.Get(i, ExistSectionIndex);
-						if (RemapMaterial.IsValidIndex(NewSectionInfo.MaterialIndex) && RemapMaterial[NewSectionInfo.MaterialIndex] == OldSectionInfo.MaterialIndex)
+						if (ExistingMeshDataPtr->UseMaterialNameSlotWorkflow)
 						{
-							if (NewMesh->StaticMaterials.IsValidIndex(NewSectionInfo.MaterialIndex))
+							if (ExistingMeshDataPtr->ExistingImportData.IsValid() &&
+								ExistingMeshDataPtr->LastImportMeshLodSectionMaterialData.IsValidIndex(i) &&
+								ExistingMeshDataPtr->LastImportMeshLodSectionMaterialData[i].IsValidIndex(ExistSectionIndex))
+							{
+								//Keep the old section material index only if the user has change the section mapping
+								bKeepOldSectionMaterialIndex = ExistingMeshDataPtr->LastImportMeshLodSectionMaterialData[i][ExistSectionIndex] != ExistingMeshDataPtr->ExistingMaterials[OldSectionInfo.MaterialIndex].ImportedMaterialSlotName;
+								for (int32 ExistMaterialIndex = 0; ExistMaterialIndex < ExistingMeshDataPtr->ExistingMaterials.Num(); ++ExistMaterialIndex)
+								{
+									if (ExistingMeshDataPtr->LastImportMeshLodSectionMaterialData[i][ExistSectionIndex] == ExistingMeshDataPtr->ExistingMaterials[ExistMaterialIndex].ImportedMaterialSlotName)
+									{
+										OriginalSectionMaterialIndex = ExistMaterialIndex;
+										break;
+									}
+								}
+							}
+						}
+						int32 OldSectionMatchIndex = OriginalSectionMaterialIndex != INDEX_NONE ? OriginalSectionMaterialIndex : OldSectionInfo.MaterialIndex;
+						if (RemapMaterial[NewSectionInfo.MaterialIndex] == OldSectionMatchIndex)
+						{
+							//Set the remap section
+							if (!bKeepOldSectionMaterialIndex)
 							{
-								//Set the remap section
 								OldSectionInfo.MaterialIndex = NewSectionInfo.MaterialIndex;
-								NewMesh->SectionInfoMap.Set(i, SectionIndex, OldSectionInfo);
-								bFoundOldMatch = true;
 							}
+							NewMesh->SectionInfoMap.Set(i, SectionIndex, OldSectionInfo);
+							bFoundOldMatch = true;
 							break;
 						}
 					}
@@ -1703,6 +1726,7 @@ void RestoreExistingMeshData(ExistingStaticMeshData* ExistingMeshDataPtr, UStati
 								break;
 							}
 						}
+						//If user has change the section info map, we want to keep the change
 						NewSectionInfo.MaterialIndex = RemapMaterial[NewSectionInfo.MaterialIndex];
 						NewMesh->SectionInfoMap.Set(i, SectionIndex, NewSectionInfo);
 					}
