Fix measure tool not showing correct values with high dpi

diff --git a/Engine/Source/Editor/UnrealEd/Private/DragTool_Measure.cpp b/Engine/Source/Editor/UnrealEd/Private/DragTool_Measure.cpp
--- a/Engine/Source/Editor/UnrealEd/Private/DragTool_Measure.cpp
+++ b/Engine/Source/Editor/UnrealEd/Private/DragTool_Measure.cpp
@@ -47,7 +47,7 @@ FVector2D FDragTool_Measure::GetSnappedPixelPos(FVector2D PixelPos)
 	View->WorldToPixel(WorldPos, PixelPos);
 
 	// The canvas we're going to render to factors in dpi scale to final position. 
-	// Since we're basing our position off mouse coordinates it will already be pixel acurate and therefore we must back out the scale the canvas will apply
+	// Since we're basing our position off mouse coordinates it will already be pixel accurate and therefore we must back out the scale the canvas will apply
 	PixelPos /= ViewportClient->GetDPIScale();
 
 	return PixelPos;
@@ -72,7 +72,7 @@ void FDragTool_Measure::AddDelta(const FVector& InDelta)
 void FDragTool_Measure::Render(const FSceneView* View, FCanvas* Canvas)
 {
 	const float OrthoUnitsPerPixel = ViewportClient->GetOrthoUnitsPerPixel(ViewportClient->Viewport);
-	const float Length = FMath::RoundToFloat((PixelEnd - PixelStart).Size() * OrthoUnitsPerPixel);
+	const float Length = FMath::RoundToFloat((PixelEnd - PixelStart).Size() * OrthoUnitsPerPixel * ViewportClient->GetDPIScale());
 
 	if (View != nullptr && Canvas != nullptr && Length >= 1.f)
 	{
