Fix gradient becoming hidden when Hovering the top or bottom of the Curve Gradient with an Alpha Stop set to Opacity = 0.

diff --git a/Engine/Source/Editor/UnrealEd/Private/SColorGradientEditor.cpp b/Engine/Source/Editor/UnrealEd/Private/SColorGradientEditor.cpp
--- a/Engine/Source/Editor/UnrealEd/Private/SColorGradientEditor.cpp
+++ b/Engine/Source/Editor/UnrealEd/Private/SColorGradientEditor.cpp
@@ -183,6 +183,7 @@ int32 SColorGradientEditor::OnPaint( const FPaintArgs& Args, const FGeometry& Al
 				DrawEffects,
 				FLinearColor(.5f, .5f, .5f, .15f)
 				);
+			++LayerId;
 		}
 		else if (bAlphaAreaHovered)
 		{
@@ -196,6 +197,7 @@ int32 SColorGradientEditor::OnPaint( const FPaintArgs& Args, const FGeometry& Al
 				DrawEffects,
 				FLinearColor(.5f, .5f, .5f, .15f)
 			);
+			++LayerId;
 		}
 			
 		// Sample the curve every 2 units.  THe curve could be non-linear so sampling at each stop would display an incorrect gradient
@@ -246,7 +248,7 @@ int32 SColorGradientEditor::OnPaint( const FPaintArgs& Args, const FGeometry& Al
 				GradientAreaGeometry.ToPaintGeometry(),
 				Stops,
 				Orient_Vertical,
-				DrawEffects | ESlateDrawEffect::None
+				DrawEffects
 			);	
 		}
 
