Allow the import of nodes that contains ':' and allow multiple characters through a single FBX.
Fix precision error when importing camera cut via Sequencer.

diff --git a/Engine/Source/Editor/MovieSceneTools/Private/MovieSceneToolHelpers.cpp b/Engine/Source/Editor/MovieSceneTools/Private/MovieSceneToolHelpers.cpp
--- a/Engine/Source/Editor/MovieSceneTools/Private/MovieSceneToolHelpers.cpp
+++ b/Engine/Source/Editor/MovieSceneTools/Private/MovieSceneToolHelpers.cpp
@@ -1411,7 +1411,7 @@ void ImportCameraCut(UnFbx::FFbxImporter* FbxImporter, UMovieScene* InMovieScene
 				FGuid CameraGuid = FindCameraGuid(AllCameras[value], InObjectBindingMap);
 				if (CameraGuid != FGuid())
 				{
-					CameraCutTrack->AddNewCameraCut(FMovieSceneObjectBindingID(CameraGuid, MovieSceneSequenceID::Root), (((float)key.GetTime().GetSecondDouble()) * FrameRate).FloorToFrame());
+					CameraCutTrack->AddNewCameraCut(FMovieSceneObjectBindingID(CameraGuid, MovieSceneSequenceID::Root), (key.GetTime().GetSecondDouble() * FrameRate).RoundToFrame());
 				}
 			}
 		}
diff --git a/Engine/Source/Editor/UnrealEd/Private/Fbx/FbxMainImport.cpp b/Engine/Source/Editor/UnrealEd/Private/Fbx/FbxMainImport.cpp
--- a/Engine/Source/Editor/UnrealEd/Private/Fbx/FbxMainImport.cpp
+++ b/Engine/Source/Editor/UnrealEd/Private/Fbx/FbxMainImport.cpp
@@ -968,6 +968,27 @@ void FFbxImporter::EnsureNodeNameAreValid()
 					FFbxErrors::Generic_LoadingSceneFailed);
 			}
 		}
+		if (NodeName.Contains(TEXT(":")))
+		{
+			NodeName = NodeName.Replace(TEXT(":"), TEXT("_"));
+			Node->SetName(TCHAR_TO_UTF8(*NodeName));
+		}
+		if (AllNodeName.Contains(NodeName))
+		{
+			FString UniqueNodeName;
+			do
+			{
+				UniqueNodeName = NodeName + FString::FromInt(CurrentNameIndex++);
+			} while (AllNodeName.Contains(UniqueNodeName));
+			Node->SetName(TCHAR_TO_UTF8(*UniqueNodeName));
+			if (!GIsAutomationTesting)
+			{
+				AddTokenizedErrorMessage(
+					FTokenizedMessage::Create(EMessageSeverity::Warning,
+						FText::Format(LOCTEXT("FbxImport_NodeNameClash", "FBX File Loading: Found name clash, node '{0}' was rename '{1}'"), FText::FromString(NodeName), FText::FromString(UniqueNodeName))),
+					FFbxErrors::Generic_LoadingSceneFailed);
+			}
+		}
 		AllNodeName.Add(NodeName);
 	}
 }
