Mirror 4163017
Remove expensive FReplaceReferenceHelper::FindAndReplaceReference call that runs for each actor class that is reinstanced. Instead rely on the main FindAndReplaceReference call

diff --git a/Engine/Source/Editor/UnrealEd/Private/Kismet2/KismetReinstanceUtilities.cpp b/Engine/Source/Editor/UnrealEd/Private/Kismet2/KismetReinstanceUtilities.cpp
--- a/Engine/Source/Editor/UnrealEd/Private/Kismet2/KismetReinstanceUtilities.cpp
+++ b/Engine/Source/Editor/UnrealEd/Private/Kismet2/KismetReinstanceUtilities.cpp
@@ -1238,14 +1238,6 @@ void FActorReplacementHelper::Finalize(const TMap<UObject*, UObject*>& OldToNewI
 		GEditor->NotifyToolsOfObjectReplacement(ConstructedComponentReplacementMap);
 	}
 
-	// Make array of component subobjects that have been reinstanced as part of the new Actor.
-	TArray<UObject*> SourceObjects;
-	ConstructedComponentReplacementMap.GenerateKeyArray(SourceObjects);
-
-	// Find and replace any outstanding references to the old Actor's component subobject instances that exist outside of the old Actor instance.
-	// Note: This will typically be references held by the Editor's transaction buffer - we need to find and replace those as well since we also do this for the old->new Actor instance.
-	FReplaceReferenceHelper::FindAndReplaceReferences(SourceObjects, ObjectsThatShouldUseOldStuff, ObjectsToReplace, ConstructedComponentReplacementMap, ReinstancedObjectsWeakReferenceMap);
-	
 	// Destroy actor and clear references.
 	NewActor->Modify();
 	if (GEditor && GEditor->Layers.IsValid())
@@ -2096,6 +2088,8 @@ void FBlueprintCompileReinstancer::ReplaceInstancesOfClass_Inner(TMap<UClass*, U
 		{
 			ReplacementActor.ApplyAttachments(ObjectRemappingHelper.ReplacedObjects, ObjectsThatShouldUseOldStuff, ObjectsReplaced, ReinstancedObjectsWeakReferenceMap);
 		}
+
+		OldToNewInstanceMap.Append(ObjectRemappingHelper.ReplacedObjects);
 	}
 
 	FReplaceReferenceHelper::FindAndReplaceReferences(SourceObjects, ObjectsThatShouldUseOldStuff, ObjectsReplaced, OldToNewInstanceMap, ReinstancedObjectsWeakReferenceMap);
