Fixed bCookAll including localized assets for cultures we don't stage

diff --git a/Engine/Source/Editor/UnrealEd/Private/Commandlets/PackageUtilities.cpp b/Engine/Source/Editor/UnrealEd/Private/Commandlets/PackageUtilities.cpp
--- a/Engine/Source/Editor/UnrealEd/Private/Commandlets/PackageUtilities.cpp
+++ b/Engine/Source/Editor/UnrealEd/Private/Commandlets/PackageUtilities.cpp
@@ -204,7 +204,17 @@ bool NormalizePackageNames( TArray<FString> PackageNames, TArray<FString>& Packa
 
 			if ( (PackageFilter&NORMALIZE_ExcludeNoRedistPackages) != 0 )
 			{
-				if(PackagePathNames.Contains("/NoRedist/") || PackagePathNames.Contains("/NotForLicensees/") || PackagePathNames.Contains("/EpicInternal/"))
+				if (Filename.Contains(TEXT("/NoRedist/")) || Filename.Contains(TEXT("/NotForLicensees/")) || Filename.Contains(TEXT("/EpicInternal/")))
+				{
+					PackagePathNames.RemoveAt(PackageIndex);
+					continue;
+				}
+			}
+
+			if ( (PackageFilter&NORMALIZE_ExcludeLocalizedPackages) != 0 )
+			{
+				FString PackageName;
+				if (FPackageName::TryConvertFilenameToLongPackageName(Filename, PackageName) && FPackageName::IsLocalizedPackage(PackageName))
 				{
 					PackagePathNames.RemoveAt(PackageIndex);
 					continue;
diff --git a/Engine/Source/Editor/UnrealEd/Private/CookOnTheFlyServer.cpp b/Engine/Source/Editor/UnrealEd/Private/CookOnTheFlyServer.cpp
--- a/Engine/Source/Editor/UnrealEd/Private/CookOnTheFlyServer.cpp
+++ b/Engine/Source/Editor/UnrealEd/Private/CookOnTheFlyServer.cpp
@@ -5290,7 +5290,7 @@ void UCookOnTheFlyServer::CollectFilesToCook(TArray<FName>& FilesInPath, const T
 			Tokens.Add(FString("*") + FPackageName::GetAssetPackageExtension());
 			Tokens.Add(FString("*") + FPackageName::GetMapPackageExtension());
 
-			uint8 PackageFilter = NORMALIZE_DefaultFlags | NORMALIZE_ExcludeEnginePackages;
+			uint8 PackageFilter = NORMALIZE_DefaultFlags | NORMALIZE_ExcludeEnginePackages | NORMALIZE_ExcludeLocalizedPackages;
 			if (bMapsOnly)
 			{
 				PackageFilter |= NORMALIZE_ExcludeContentPackages;
diff --git a/Engine/Source/Editor/UnrealEd/Public/PackageHelperFunctions.h b/Engine/Source/Editor/UnrealEd/Public/PackageHelperFunctions.h
--- a/Engine/Source/Editor/UnrealEd/Public/PackageHelperFunctions.h
+++ b/Engine/Source/Editor/UnrealEd/Public/PackageHelperFunctions.h
@@ -39,6 +39,8 @@ enum EPackageNormalizationFlags
 	NORMALIZE_ExcludeEnginePackages		= 0x20,
 	/** do not include packages inside NoRedist or NotForLicensees folders */
 	NORMALIZE_ExcludeNoRedistPackages	= 0x40,
+	/** do not include localized packages */
+	NORMALIZE_ExcludeLocalizedPackages	= 0x80,
 	/** Combo flags */
 	NORMALIZE_DefaultFlags				= NORMALIZE_ResetExistingLoaders,
 };
