Exposed Cook Display Mode to Project Settings

diff --git a/Engine/Source/Editor/UnrealEd/Private/CookOnTheFlyServer.cpp b/Engine/Source/Editor/UnrealEd/Private/CookOnTheFlyServer.cpp
--- a/Engine/Source/Editor/UnrealEd/Private/CookOnTheFlyServer.cpp
+++ b/Engine/Source/Editor/UnrealEd/Private/CookOnTheFlyServer.cpp
@@ -103,10 +103,7 @@ DEFINE_LOG_CATEGORY_STATIC(LogCook, Log, All);
 
 #define PROFILE_NETWORK 0
 
-
-#define COOKER_DISPLAY_REMAINING 1
-#define COOKER_DISPLAY_PACKAGE_NAME 2
-int32 GCookProgressDisplay = COOKER_DISPLAY_REMAINING;
+int32 GCookProgressDisplay = (int32)ECookProgressDisplayMode::RemainingPackages;
 static FAutoConsoleVariableRef CVarCookDisplayMode(
 	TEXT("cook.displaymode"),
 	GCookProgressDisplay,
@@ -1534,7 +1531,7 @@ uint32 UCookOnTheFlyServer::TickCookOnTheSide( const float TimeSlice, uint32 &Co
 	{
 		if (CookRequests.Num() > 0)
 		{
-			UE_CLOG(!(TickFlags & ECookTickFlags::HideProgressDisplay) && (GCookProgressDisplay & COOKER_DISPLAY_REMAINING), 
+			UE_CLOG(!(TickFlags & ECookTickFlags::HideProgressDisplay) && (GCookProgressDisplay & (int32)ECookProgressDisplayMode::RemainingPackages), 
 			LogCook, Display, TEXT("Cooked packages %d Packages Remain %d Total %d"), CookedPackages.Num(), CookRequests.Num(), CookedPackages.Num() + CookRequests.Num());
 		}
 
@@ -3441,7 +3438,7 @@ void UCookOnTheFlyServer::SaveCookedPackage(UPackage* Package, uint32 SaveFlags,
 
 				if (bPackageFullyLoaded)
 				{
-					UE_CLOG(GCookProgressDisplay & COOKER_DISPLAY_PACKAGE_NAME, LogCook, Display, TEXT("Cooking %s -> %s"), *Package->GetName(), *PlatFilename);
+					UE_CLOG(GCookProgressDisplay & (int32)ECookProgressDisplayMode::PackageNames, LogCook, Display, TEXT("Cooking %s -> %s"), *Package->GetName(), *PlatFilename);
 
 					bool bSwap = (!Target->IsLittleEndian()) ^ (!PLATFORM_LITTLE_ENDIAN);
 
diff --git a/Engine/Source/Editor/UnrealEd/Private/CookerSettings.cpp b/Engine/Source/Editor/UnrealEd/Private/CookerSettings.cpp
--- a/Engine/Source/Editor/UnrealEd/Private/CookerSettings.cpp
+++ b/Engine/Source/Editor/UnrealEd/Private/CookerSettings.cpp
@@ -9,6 +9,7 @@ UCookerSettings::UCookerSettings(const FObjectInitializer& ObjectInitializer)
 	, bEnableCookOnTheSide(false)
 	, bEnableBuildDDCInBackground(false)
 	, bIterativeCookingForLaunchOn(false)
+	, CookProgressDisplayMode(ECookProgressDisplayMode::RemainingPackages)
 	, bIgnoreIniSettingsOutOfDateForIteration(false)
 	, bIgnoreScriptPackagesOutOfDateForIteration(false)
 	, bCompileBlueprintsInDevelopmentMode(true)
diff --git a/Engine/Source/Editor/UnrealEd/Public/CookerSettings.h b/Engine/Source/Editor/UnrealEd/Public/CookerSettings.h
--- a/Engine/Source/Editor/UnrealEd/Public/CookerSettings.h
+++ b/Engine/Source/Editor/UnrealEd/Public/CookerSettings.h
@@ -13,6 +13,25 @@
 
 struct FPropertyChangedEvent;
 
+UENUM()
+enum class ECookProgressDisplayMode : int32
+{
+	/** Don't display any progress messages */
+	Nothing = 0,
+
+	/** Display the number of remaining packages */
+	RemainingPackages = 1,
+
+	/** Display names of cooked packages */
+	PackageNames = 2,
+
+	/** Display the number of remaining packages and package names */
+	NamesAndRemainingPackages = 3,
+
+	/** MAX - invalid */
+	Max UMETA(Hidden),
+};
+
 /**
  * Various cooker settings.
  */
@@ -40,6 +59,11 @@ public:
 	UPROPERTY(GlobalConfig, EditAnywhere, Category = Cooker, meta = (DisplayName = "Iterative cooking for the File->Cook Content menu item"))
 	bool bIterativeCookingForFileCookContent;
 
+	UPROPERTY(GlobalConfig, EditAnywhere, Category = Cooker, meta = (
+		ConsoleVariable = "cook.displaymode", DisplayName = "Cooker Progress Display Mode",
+		ToolTip = "Controls log output of the cooker"))
+	ECookProgressDisplayMode CookProgressDisplayMode;
+
 	/** Ignore ini changes when doing iterative cooking, either in editor or out of editor */
 	UPROPERTY(GlobalConfig, EditAnywhere, Category = Cooker, AdvancedDisplay)
 	bool bIgnoreIniSettingsOutOfDateForIteration;
diff --git a/Engine/Source/Runtime/Launch/Private/LaunchEngineLoop.cpp b/Engine/Source/Runtime/Launch/Private/LaunchEngineLoop.cpp
--- a/Engine/Source/Runtime/Launch/Private/LaunchEngineLoop.cpp
+++ b/Engine/Source/Runtime/Launch/Private/LaunchEngineLoop.cpp
@@ -1558,6 +1558,9 @@ int32 FEngineLoop::PreInit(const TCHAR* CmdLine)
 	ApplyCVarSettingsFromIni(TEXT("/Script/Engine.StreamingSettings"), *GEngineIni, ECVF_SetByProjectSetting);
 	ApplyCVarSettingsFromIni(TEXT("/Script/Engine.GarbageCollectionSettings"), *GEngineIni, ECVF_SetByProjectSetting);
 	ApplyCVarSettingsFromIni(TEXT("/Script/Engine.NetworkSettings"), *GEngineIni, ECVF_SetByProjectSetting);
+#if WITH_EDITOR
+	ApplyCVarSettingsFromIni(TEXT("/Script/UnrealEd.CookerSettings"), *GEngineIni, ECVF_SetByProjectSetting);
+#endif
 
 #if !UE_SERVER
 	if (!IsRunningDedicatedServer())
