Fix the re-import LOD by clearing the old data for mesh description.

diff --git a/Engine/Source/Editor/UnrealEd/Private/Fbx/FbxStaticMeshImport.cpp b/Engine/Source/Editor/UnrealEd/Private/Fbx/FbxStaticMeshImport.cpp
--- a/Engine/Source/Editor/UnrealEd/Private/Fbx/FbxStaticMeshImport.cpp
+++ b/Engine/Source/Editor/UnrealEd/Private/Fbx/FbxStaticMeshImport.cpp
@@ -1595,16 +1595,19 @@ UStaticMesh* UnFbx::FFbxImporter::ImportStaticMeshAsSingle(UObject* InParent, TA
 		UStaticMesh::RegisterMeshAttributes(*MeshDescription);
 		StaticMesh->CommitOriginalMeshDescription(LODIndex);
 	}
-	
+	else if (InStaticMesh != NULL && LODIndex > 0)
+	{
+		// clear out the old mesh data
+		MeshDescription->Empty();
+	}
+
 	FStaticMeshSourceModel& SrcModel = StaticMesh->SourceModels[LODIndex];
 	if( InStaticMesh != NULL && LODIndex > 0 && !SrcModel.RawMeshBulkData->IsEmpty() )
 	{
 		// clear out the old mesh data
 		FRawMesh EmptyRawMesh;
-		SrcModel.RawMeshBulkData->SaveRawMesh( EmptyRawMesh );
+		SrcModel.SaveRawMesh(EmptyRawMesh);
 	}
-
-	
 	
 	// make sure it has a new lighting guid
 	StaticMesh->LightingGuid = FGuid::NewGuid();
