Live anim BP recompilation is now non-experimental

diff --git a/Engine/Source/Editor/UnrealEd/Classes/Settings/EditorExperimentalSettings.h b/Engine/Source/Editor/UnrealEd/Classes/Settings/EditorExperimentalSettings.h
--- a/Engine/Source/Editor/UnrealEd/Classes/Settings/EditorExperimentalSettings.h
+++ b/Engine/Source/Editor/UnrealEd/Classes/Settings/EditorExperimentalSettings.h
@@ -113,10 +113,6 @@ public:
 	UPROPERTY(EditAnywhere, config, Category = Tools, meta = (ConfigRestartRequired = true))
 	bool bFacialAnimationImporter;
 
-	/** Allow animation blueprints to be recompiled while a PIE session is running */
-	UPROPERTY(EditAnywhere, config, Category = Tools)
-	bool bEnableLiveRecompilationOfAnimationBlueprints;
-
 	/** Enable experimental PIE preview device launch */
 	UPROPERTY(EditAnywhere, config, Category = PIE, meta = (DisplayName = "Enable mobile PIE with preview device launch options."))
 	bool bMobilePIEPreviewDeviceLaunch;
diff --git a/Engine/Source/Runtime/Engine/Private/AnimBlueprint.cpp b/Engine/Source/Runtime/Engine/Private/AnimBlueprint.cpp
--- a/Engine/Source/Runtime/Engine/Private/AnimBlueprint.cpp
+++ b/Engine/Source/Runtime/Engine/Private/AnimBlueprint.cpp
@@ -178,7 +178,7 @@ void UAnimBlueprint::PostLoad()
 
 bool UAnimBlueprint::CanRecompileWhilePlayingInEditor() const
 {
-	return GetDefault<UEditorExperimentalSettings>()->bEnableLiveRecompilationOfAnimationBlueprints;
+	return true;
 }
 #endif
 
