Remove long deprecated InputGesture.h
Actually mark FInputGesture as deprecated. Unclear why this wasn't correctly done for 4.8

diff --git a/Engine/Plugins/Developer/RenderDocPlugin/Source/RenderDocPlugin/Private/RenderDocPluginCommands.cpp b/Engine/Plugins/Developer/RenderDocPlugin/Source/RenderDocPlugin/Private/RenderDocPluginCommands.cpp
--- a/Engine/Plugins/Developer/RenderDocPlugin/Source/RenderDocPlugin/Private/RenderDocPluginCommands.cpp
+++ b/Engine/Plugins/Developer/RenderDocPlugin/Source/RenderDocPlugin/Private/RenderDocPluginCommands.cpp
@@ -8,7 +8,7 @@
  
 void FRenderDocPluginCommands::RegisterCommands()
 {
-	UI_COMMAND(CaptureFrameCommand, "Capture Frame", "Captures the next frame and launches the renderdoc UI (Alt+F12)", EUserInterfaceActionType::Button, FInputGesture(EKeys::F12, EModifierKey::Alt));
+	UI_COMMAND(CaptureFrameCommand, "Capture Frame", "Captures the next frame and launches the renderdoc UI (Alt+F12)", EUserInterfaceActionType::Button, FInputChord(EKeys::F12, EModifierKey::Alt));
 }
 
 #undef LOCTEXT_NAMESPACE
diff --git a/Engine/Plugins/Experimental/CodeEditor/Source/CodeEditor/Private/CodeProjectEditorCommands.cpp b/Engine/Plugins/Experimental/CodeEditor/Source/CodeEditor/Private/CodeProjectEditorCommands.cpp
--- a/Engine/Plugins/Experimental/CodeEditor/Source/CodeEditor/Private/CodeProjectEditorCommands.cpp
+++ b/Engine/Plugins/Experimental/CodeEditor/Source/CodeEditor/Private/CodeProjectEditorCommands.cpp
@@ -15,8 +15,8 @@ FCodeProjectEditorCommands::FCodeProjectEditorCommands()
 
 void FCodeProjectEditorCommands::RegisterCommands()
 {
-	UI_COMMAND(Save, "Save", "Save the currently active document.", EUserInterfaceActionType::Button, FInputGesture(EModifierKey::Control, EKeys::S));
-	UI_COMMAND(SaveAll, "Save All", "Save all open documents.", EUserInterfaceActionType::Button, FInputGesture(EModifierKey::Control | EModifierKey::Shift, EKeys::S));
+	UI_COMMAND(Save, "Save", "Save the currently active document.", EUserInterfaceActionType::Button, FInputChord(EModifierKey::Control, EKeys::S));
+	UI_COMMAND(SaveAll, "Save All", "Save all open documents.", EUserInterfaceActionType::Button, FInputChord(EModifierKey::Control | EModifierKey::Shift, EKeys::S));
 }
 
 
diff --git a/Engine/Source/Developer/AutomationWindow/Private/SAutomationWindow.cpp b/Engine/Source/Developer/AutomationWindow/Private/SAutomationWindow.cpp
--- a/Engine/Source/Developer/AutomationWindow/Private/SAutomationWindow.cpp
+++ b/Engine/Source/Developer/AutomationWindow/Private/SAutomationWindow.cpp
@@ -69,7 +69,7 @@ public:
 		
 #if WITH_EDITOR
 		// Added button for running the currently open level test.
-		UI_COMMAND(RunLevelTest, "Run Level Test", "Run Level Test", EUserInterfaceActionType::Button, FInputGesture());
+		UI_COMMAND(RunLevelTest, "Run Level Test", "Run Level Test", EUserInterfaceActionType::Button, FInputChord());
 #endif
 	}
 public:
diff --git a/Engine/Source/Editor/UnrealEd/Private/EditorViewportCommands.cpp b/Engine/Source/Editor/UnrealEd/Private/EditorViewportCommands.cpp
--- a/Engine/Source/Editor/UnrealEd/Private/EditorViewportCommands.cpp
+++ b/Engine/Source/Editor/UnrealEd/Private/EditorViewportCommands.cpp
@@ -97,7 +97,7 @@ void FEditorViewportCommands::RegisterCommands()
 	UI_COMMAND( RotateMode, "Rotate Mode", "Select and rotate objects", EUserInterfaceActionType::ToggleButton, FInputChord(EKeys::E) );
 	UI_COMMAND( ScaleMode, "Scale Mode", "Select and scale objects", EUserInterfaceActionType::ToggleButton, FInputChord(EKeys::R) );
 	UI_COMMAND( TranslateRotateMode, "Combined Translate and Rotate Mode", "Select and translate or rotate objects", EUserInterfaceActionType::ToggleButton, FInputChord() );
-	UI_COMMAND( TranslateRotate2DMode, "2D Mode", "Select and translate or rotate objects in 2D", EUserInterfaceActionType::ToggleButton, FInputGesture());
+	UI_COMMAND( TranslateRotate2DMode, "2D Mode", "Select and translate or rotate objects in 2D", EUserInterfaceActionType::ToggleButton, FInputChord());
 
 	UI_COMMAND( ShrinkTransformWidget, "Shrink Transform Widget", "Shrink the level editor transform widget", EUserInterfaceActionType::Button, FInputChord(EModifierKey::Alt, EKeys::LeftBracket) );
 	UI_COMMAND( ExpandTransformWidget, "Expand Transform Widget", "Expand the level editor transform widget", EUserInterfaceActionType::Button, FInputChord(EModifierKey::Alt, EKeys::RightBracket) );
@@ -116,7 +116,7 @@ void FEditorViewportCommands::RegisterCommands()
 
 	UI_COMMAND( LocationGridSnap, "Grid Snap", "Enables or disables snapping to the grid when dragging objects around", EUserInterfaceActionType::ToggleButton, FInputChord() );
 	UI_COMMAND( RotationGridSnap, "Rotation Snap", "Enables or disables snapping objects to a rotation grid", EUserInterfaceActionType::ToggleButton, FInputChord() );
-	UI_COMMAND( Layer2DSnap, "Layer2D Snap", "Enables or disables snapping objects to a 2D layer", EUserInterfaceActionType::ToggleButton, FInputGesture());
+	UI_COMMAND( Layer2DSnap, "Layer2D Snap", "Enables or disables snapping objects to a 2D layer", EUserInterfaceActionType::ToggleButton, FInputChord());
 	UI_COMMAND( ScaleGridSnap, "Scale Snap", "Enables or disables snapping objects to a scale grid", EUserInterfaceActionType::ToggleButton, FInputChord() );
 	UI_COMMAND( SurfaceSnapping, "Surface Snapping", "If enabled, actors will snap to surfaces in the world when dragging", EUserInterfaceActionType::ToggleButton, FInputChord() );
 
diff --git a/Engine/Source/Runtime/Slate/Private/Framework/Commands/InputChord.cpp b/Engine/Source/Runtime/Slate/Private/Framework/Commands/InputChord.cpp
--- a/Engine/Source/Runtime/Slate/Private/Framework/Commands/InputChord.cpp
+++ b/Engine/Source/Runtime/Slate/Private/Framework/Commands/InputChord.cpp
@@ -4,7 +4,7 @@
 
 #define LOCTEXT_NAMESPACE "FInputChord"
 
-/* FInputGesture interface
+/* FInputChord interface
  *****************************************************************************/
 
 /**
diff --git a/Engine/Source/Runtime/Slate/Public/Framework/Commands/InputChord.h b/Engine/Source/Runtime/Slate/Public/Framework/Commands/InputChord.h
--- a/Engine/Source/Runtime/Slate/Public/Framework/Commands/InputChord.h
+++ b/Engine/Source/Runtime/Slate/Public/Framework/Commands/InputChord.h
@@ -8,7 +8,8 @@
 #include "GenericPlatform/GenericApplication.h"
 #include "InputChord.generated.h"
 
-//(4.8, "Use FInputChord instead of FInputGesture")
+// Was supposed to be deprecated in 4.8, but was never properly marked up
+DEPRECATED(4.21, "Use FInputChord instead of FInputGesture") 
 typedef struct FInputChord FInputGesture;
 
 /** An Input Chord is a key and the modifier keys that are to be held with it. */
diff --git a/Engine/Source/Runtime/Slate/Public/Framework/Commands/InputGesture.h b/Engine/Source/Runtime/Slate/Public/Framework/Commands/InputGesture.h
deleted file mode 100644
