Catch-all to avoid editor spamming heartbeat events if the trigger fires more than it should for some reason.

diff --git a/Engine/Source/Editor/UnrealEd/Private/UnrealEdMisc.cpp b/Engine/Source/Editor/UnrealEd/Private/UnrealEdMisc.cpp
--- a/Engine/Source/Editor/UnrealEd/Private/UnrealEdMisc.cpp
+++ b/Engine/Source/Editor/UnrealEd/Private/UnrealEdMisc.cpp
@@ -624,7 +624,20 @@ void FUnrealEdMisc::EditorAnalyticsHeartbeat()
 		return;
 	}
 
-	static double LastHeartbeatTime = FPlatformTime::Seconds();
+	// Analytics has had some very rare instances where a user spams heartbeats at an incredibly high frequency.
+	// So while it's technically impossible to trigger this event more than once per HeartbeatIntervalSeconds,
+	// put an additional guard band around this function in absolute wall-clock time to ensure it doesn't fire too quickly for some strange reason.
+	//////////////////////////////////////////////////////////////////////////////////////////////////////////////////////////////////////////////
+	const double HeartbeatInvervalSec = (double)(UnrealEdMiscDefs::HeartbeatIntervalSeconds);
+	const double Now = FPlatformTime::Seconds();
+	// Initialize to ensure the first time this is called we will DEF execute the heartbeat.
+	static double LastHeartbeatTime = Now - HeartbeatInvervalSec;
+	// allow a little bit of slop in the timing in case the event is only firing slightly too soon, even though that is technically impossible.
+	if (Now <= LastHeartbeatTime + HeartbeatInvervalSec*0.9)
+	{
+		UE_LOG(LogUnrealEdMisc, Warning, TEXT("Heartbeat event firing too frequently (%.3f sec). This should never happen. Something is wrong with the timer delegate!"), (float)(Now - LastHeartbeatTime) );
+		return;
+	}
 
 	bool bIsDebuggerPresent = FPlatformMisc::IsDebuggerPresent();
 	static bool bWasDebuggerPresent = false;
