Maintain hierarchy when harvesting components to make a blueprint

diff --git a/Engine/Source/Editor/UnrealEd/Private/Kismet2/Kismet2.cpp b/Engine/Source/Editor/UnrealEd/Private/Kismet2/Kismet2.cpp
--- a/Engine/Source/Editor/UnrealEd/Private/Kismet2/Kismet2.cpp
+++ b/Engine/Source/Editor/UnrealEd/Private/Kismet2/Kismet2.cpp
@@ -1346,21 +1346,12 @@ void FKismetEditorUtilities::AddComponentsToBlueprint(UBlueprint* Blueprint, con
 		}
 	};
 
-	AActor* Actor = nullptr;
-
 	for (UActorComponent* ActorComponent : Components)
 	{
 		if (ActorComponent)
 		{
-			if (Actor)
-			{
-				check(Actor == ActorComponent->GetOwner());
-			}
-			else
-			{
-				Actor = ActorComponent->GetOwner();
-				check(Actor);
-			}
+			AActor* Actor = ActorComponent->GetOwner();
+			check(Actor);
 
 			if (!ActorComponent->GetClass()->HasMetaData(FBlueprintMetadata::MD_BlueprintSpawnableComponent))
 			{
@@ -1377,7 +1368,7 @@ void FKismetEditorUtilities::AddComponentsToBlueprint(UBlueprint* Blueprint, con
 			}
 			else
 			{
-				if (ActorComponent == Actor->GetRootComponent())
+				if (SceneComponent == Actor->GetRootComponent() && (SceneComponent->GetAttachParent() == nullptr || !Components.Contains(SceneComponent->GetAttachParent())))
 				{
 					if (OptionalNewRootNode != nullptr)
 					{
@@ -1589,68 +1580,99 @@ public:
 			SCS = Blueprint->SimpleConstructionScript;
 
 			// Create a common root if necessary
+			TArray<AActor*> RootActors;
 			USCS_Node* RootNodeOverride = nullptr;
-			if ((SelectedActors.Num() == 1) && (SelectedActors[0]->GetRootComponent() != nullptr))
+
+			// First sort the selected actors such that any actor that has its root attached to one of the other
+			// selected actors comes later in the array.
+			SelectedActors.Sort([](AActor& A, AActor& B)
+			{
+				return (B.IsAttachedTo(&A));
+			});
+
+			// Determine how many of the selected actors are not attached to one another
+			for (int32 OuterIndex = SelectedActors.Num() - 1; OuterIndex >= 0; --OuterIndex)
 			{
-				// We have a single actor that has a scene component, so no need to create a common root
+				AActor* SelectedActor = SelectedActors[OuterIndex];
+				if (SelectedActor->GetRootComponent())
+				{
+					if (OuterIndex == 0)
+					{
+						RootActors.Add(SelectedActor);
+					}
+					else
+					{
+						for (int32 InnerIndex = OuterIndex - 1; InnerIndex >= 0; --InnerIndex)
+						{
+							if (!SelectedActor->IsAttachedTo(SelectedActors[InnerIndex]))
+							{
+								RootActors.Add(SelectedActor);
+								break;
+							}
+						}
+					}
+				}
 			}
-			else
+
+			// If there is not one unique root actor then create a new scene component to serve as the shared root node
+			if (RootActors.Num() != 1)
 			{
-				// Add a new scene component to serve as the shared root node
 				RootNodeOverride = SCS->CreateNode(USceneComponent::StaticClass(), TEXT("SharedRoot"));
 				SCS->AddNode(RootNodeOverride);
 			}
 
 			// Harvest the components from each actor and clone them into the SCS
+			TArray<UActorComponent*> AllSelectedComponents;
 			for (const AActor* Actor : SelectedActors)
 			{
-				TArray<UActorComponent*> Components;
-				Actor->GetComponents(Components);
-
-				// Exclude any components created by other components
-				for (int32 Index = Components.Num() - 1; Index >= 0; --Index)
+				for (UActorComponent* ComponentToConsider : Actor->GetComponents())
 				{
-					UActorComponent* ComponentToConsider = Components[Index];
-					if (ComponentToConsider->IsEditorOnly())
+					if (ComponentToConsider && !ComponentToConsider->IsEditorOnly())
 					{
-						Components.RemoveAt(Index, 1, /*bAllowShrinking=*/ false);
+						AllSelectedComponents.Add(ComponentToConsider);
 					}
 				}
-
-				FKismetEditorUtilities::AddComponentsToBlueprint(Blueprint, Components, /*bHarvesting=*/ true, RootNodeOverride);
 			}
+			FKismetEditorUtilities::AddComponentsToBlueprint(Blueprint, AllSelectedComponents, /*bHarvesting=*/ true, RootNodeOverride);
 
-			// Compute the average origin for all the actors, so it can be backed out when saving them in the blueprint
 			FTransform NewActorTransform = FTransform::Identity;
+			if (RootActors.Num() == 1)
+			{
+				NewActorTransform = RootActors[0]->GetTransform();
+			}
+			else if (RootActors.Num() > 1)
 			{
-				// Find average location of all selected actors
-				FVector AverageLocation = FVector::ZeroVector;
-				for (const AActor* Actor : SelectedActors)
+				// Compute the average origin for all the actors, so it can be backed out when saving them in the blueprint
 				{
-					if (USceneComponent* RootComponent = Actor->GetRootComponent())
+					// Find average location of all selected actors
+					FVector AverageLocation = FVector::ZeroVector;
+					for (const AActor* Actor : RootActors)
 					{
-						AverageLocation += Actor->GetActorLocation();
+						if (USceneComponent* RootComponent = Actor->GetRootComponent())
+						{
+							AverageLocation += Actor->GetActorLocation();
+						}
 					}
-				}
-				AverageLocation /= (float)SelectedActors.Num();
+					AverageLocation /= (float)RootActors.Num();
 
-				// Spawn the new BP at that location
-				NewActorTransform.SetTranslation(AverageLocation);
-			}
+					// Spawn the new BP at that location
+					NewActorTransform.SetTranslation(AverageLocation);
+				}
 
-			// Reposition all of the children of the root node to recenter them around the new pivot
-			for (USCS_Node* TopLevelNode : SCS->GetRootNodes())
-			{
-				if (USceneComponent* TestRoot = Cast<USceneComponent>(TopLevelNode->ComponentTemplate))
+				// Reposition all of the children of the root node to recenter them around the new pivot
+				for (USCS_Node* TopLevelNode : SCS->GetRootNodes())
 				{
-					for (USCS_Node* ChildNode : TopLevelNode->GetChildNodes())
+					if (USceneComponent* TestRoot = Cast<USceneComponent>(TopLevelNode->ComponentTemplate))
 					{
-						if (USceneComponent* ChildComponent = Cast<USceneComponent>(ChildNode->ComponentTemplate))
+						for (USCS_Node* ChildNode : TopLevelNode->GetChildNodes())
 						{
-							const FTransform OldChildToWorld(ChildComponent->RelativeRotation, ChildComponent->RelativeLocation);
-							const FTransform NewRelativeTransform = OldChildToWorld.GetRelativeTransform(NewActorTransform);
-							ChildComponent->RelativeLocation = NewRelativeTransform.GetLocation();
-							ChildComponent->RelativeRotation = NewRelativeTransform.GetRotation().Rotator();
+							if (USceneComponent* ChildComponent = Cast<USceneComponent>(ChildNode->ComponentTemplate))
+							{
+								const FTransform OldChildToWorld(ChildComponent->RelativeRotation, ChildComponent->RelativeLocation);
+								const FTransform NewRelativeTransform = OldChildToWorld.GetRelativeTransform(NewActorTransform);
+								ChildComponent->RelativeLocation = NewRelativeTransform.GetLocation();
+								ChildComponent->RelativeRotation = NewRelativeTransform.GetRotation().Rotator();
+							}
 						}
 					}
 				}
