MarkBlueprintAsStructurallyModified now uses the compilation manager to regnerate the skeleton class

diff --git a/Engine/Source/Editor/UnrealEd/Private/Kismet2/BlueprintEditorUtils.cpp b/Engine/Source/Editor/UnrealEd/Private/Kismet2/BlueprintEditorUtils.cpp
--- a/Engine/Source/Editor/UnrealEd/Private/Kismet2/BlueprintEditorUtils.cpp
+++ b/Engine/Source/Editor/UnrealEd/Private/Kismet2/BlueprintEditorUtils.cpp
@@ -2024,13 +2024,13 @@ void FBlueprintEditorUtils::UpdateDelegatesInBlueprint(UBlueprint* Blueprint)
 void FBlueprintEditorUtils::MarkBlueprintAsStructurallyModified(UBlueprint* Blueprint)
 {
 	FSecondsCounterScope Timer(BlueprintCompileAndLoadTimerData);
-	
+
 	struct FRefreshHelper
 	{
 		static void SkeletalRecompileChildren(TArray<UClass*> SkelClassesToRecompile, bool bIsCompilingOnLoad)
 		{
 			FSecondsCounterScope SkeletalRecompileTimer(BlueprintCompileAndLoadTimerData);
-			
+
 			for (UClass* SkelClass : SkelClassesToRecompile)
 			{
 				if (SkelClass->HasAnyClassFlags(CLASS_NewerVersionExists))
@@ -2080,54 +2080,63 @@ void FBlueprintEditorUtils::MarkBlueprintAsStructurallyModified(UBlueprint* Blue
 	Blueprint->bCachedDependenciesUpToDate = false;
 	if (Blueprint->Status != BS_BeingCreated && !Blueprint->bBeingCompiled)
 	{
-		FCompilerResultsLog Results;
-		Results.bLogInfoOnly = Blueprint->bIsRegeneratingOnLoad;
-
 		BP_SCOPED_COMPILER_EVENT_STAT(EKismetCompilerStats_MarkBlueprintasStructurallyModified);
 
-		TArray<UClass*> ChildrenOfClass;
-		if (UClass* SkelClass = Blueprint->SkeletonGeneratedClass)
+		if (GBlueprintUseCompilationManager)
+		{
+			FBlueprintCompilationManager::CompileSynchronously(
+				FBPCompileRequest(Blueprint, EBlueprintCompileOptions::RegenerateSkeletonOnly, nullptr)
+			);
+		}
+		else
 		{
-			if (!Blueprint->bIsRegeneratingOnLoad)
+			FCompilerResultsLog Results;
+			Results.bLogInfoOnly = Blueprint->bIsRegeneratingOnLoad;
+
+			TArray<UClass*> ChildrenOfClass;
+			if (UClass* SkelClass = Blueprint->SkeletonGeneratedClass)
 			{
-				if (IsInterfaceBlueprint(Blueprint))
+				if (!Blueprint->bIsRegeneratingOnLoad)
 				{
-					// Find all dependent Blueprints that implement the interface. Note: Using
-					// GetDependentBlueprints() here as the result is cached and thus it should
-					// generally be a faster path than iterating through all loaded Blueprints.
-					TArray<UBlueprint*> DependentBlueprints;
-					GetDependentBlueprints(Blueprint, DependentBlueprints, true);
-					for (UBlueprint* DependentBlueprint : DependentBlueprints)
+					if (IsInterfaceBlueprint(Blueprint))
 					{
-						const bool bDependentBPImplementsInterface = DependentBlueprint->ImplementedInterfaces.ContainsByPredicate([&Blueprint](const FBPInterfaceDescription& InterfaceDesc)
+						// Find all dependent Blueprints that implement the interface. Note: Using
+						// GetDependentBlueprints() here as the result is cached and thus it should
+						// generally be a faster path than iterating through all loaded Blueprints.
+						TArray<UBlueprint*> DependentBlueprints;
+						GetDependentBlueprints(Blueprint, DependentBlueprints, true);
+						for (UBlueprint* DependentBlueprint : DependentBlueprints)
 						{
-							return InterfaceDesc.Interface == Blueprint->GeneratedClass;
-						});
+							const bool bDependentBPImplementsInterface = DependentBlueprint->ImplementedInterfaces.ContainsByPredicate([&Blueprint](const FBPInterfaceDescription& InterfaceDesc)
+							{
+								return InterfaceDesc.Interface == Blueprint->GeneratedClass;
+							});
 
-						if (bDependentBPImplementsInterface)
-						{
-							ChildrenOfClass.Add(DependentBlueprint->SkeletonGeneratedClass);
+							if (bDependentBPImplementsInterface)
+							{
+								ChildrenOfClass.Add(DependentBlueprint->SkeletonGeneratedClass);
+							}
 						}
 					}
-				}
-				else
-				{
-					GetDerivedClasses(SkelClass, ChildrenOfClass, false);
+					else
+					{
+						GetDerivedClasses(SkelClass, ChildrenOfClass, false);
+					}
 				}
 			}
-		}
 
-		{
-			// Invoke the compiler to update the skeleton class definition
-			IKismetCompilerInterface& Compiler = FModuleManager::LoadModuleChecked<IKismetCompilerInterface>(KISMET_COMPILER_MODULENAME);
+			{
+				// Invoke the compiler to update the skeleton class definition
+				IKismetCompilerInterface& Compiler = FModuleManager::LoadModuleChecked<IKismetCompilerInterface>(KISMET_COMPILER_MODULENAME);
 
-			FKismetCompilerOptions CompileOptions;
-			CompileOptions.CompileType = EKismetCompileType::SkeletonOnly;
-			Compiler.CompileBlueprint(Blueprint, CompileOptions, Results);
-		}
-		UpdateDelegatesInBlueprint(Blueprint);
+				FKismetCompilerOptions CompileOptions;
+				CompileOptions.CompileType = EKismetCompileType::SkeletonOnly;
+				Compiler.CompileBlueprint(Blueprint, CompileOptions, Results);
+			}
+			UpdateDelegatesInBlueprint(Blueprint);
 
-		FRefreshHelper::SkeletalRecompileChildren(ChildrenOfClass, Blueprint->bIsRegeneratingOnLoad);
+			FRefreshHelper::SkeletalRecompileChildren(ChildrenOfClass, Blueprint->bIsRegeneratingOnLoad);
+		}
 
 		// Call general modification callback as well
 		MarkBlueprintAsModified(Blueprint);
