GenerateBlueprintSkeleton now uses the compilation manager to regnerate the skeleton class

diff --git a/Engine/Source/Editor/UnrealEd/Private/Kismet2/Kismet2.cpp b/Engine/Source/Editor/UnrealEd/Private/Kismet2/Kismet2.cpp
--- a/Engine/Source/Editor/UnrealEd/Private/Kismet2/Kismet2.cpp
+++ b/Engine/Source/Editor/UnrealEd/Private/Kismet2/Kismet2.cpp
@@ -963,23 +963,32 @@ bool FKismetEditorUtilities::GenerateBlueprintSkeleton(UBlueprint* BlueprintObj,
 
 	if( BlueprintObj->SkeletonGeneratedClass == NULL || bForceRegeneration )
 	{
-		UPackage* Package = BlueprintObj->GetOutermost();
-		bool bIsPackageDirty = Package ? Package->IsDirty() : false;
-					
-		IKismetCompilerInterface& Compiler = FModuleManager::LoadModuleChecked<IKismetCompilerInterface>(KISMET_COMPILER_MODULENAME);
+		if (GBlueprintUseCompilationManager)
+		{
+			FBlueprintCompilationManager::CompileSynchronously(
+				FBPCompileRequest(BlueprintObj, EBlueprintCompileOptions::RegenerateSkeletonOnly, nullptr)
+			);
+		}
+		else
+		{
+			UPackage* Package = BlueprintObj->GetOutermost();
+			bool bIsPackageDirty = Package ? Package->IsDirty() : false;
 
-		TGuardValue<bool> GuardTemplateNameFlag(GCompilingBlueprint, true);
-		FCompilerResultsLog Results;
+			IKismetCompilerInterface& Compiler = FModuleManager::LoadModuleChecked<IKismetCompilerInterface>(KISMET_COMPILER_MODULENAME);
 
-		FKismetCompilerOptions CompileOptions;
-		CompileOptions.CompileType = EKismetCompileType::SkeletonOnly;
-		Compiler.CompileBlueprint(BlueprintObj, CompileOptions, Results);
-		bRegeneratedSkeleton = true;
+			TGuardValue<bool> GuardTemplateNameFlag(GCompilingBlueprint, true);
+			FCompilerResultsLog Results;
 
-		// Restore the package dirty flag here
-		if( Package != NULL )
-		{
-			Package->SetDirtyFlag(bIsPackageDirty);
+			FKismetCompilerOptions CompileOptions;
+			CompileOptions.CompileType = EKismetCompileType::SkeletonOnly;
+			Compiler.CompileBlueprint(BlueprintObj, CompileOptions, Results);
+			bRegeneratedSkeleton = true;
+
+			// Restore the package dirty flag here
+			if (Package != NULL)
+			{
+				Package->SetDirtyFlag(bIsPackageDirty);
+			}
 		}
 	}
 	return bRegeneratedSkeleton;
