Blueprint duplication now uses the compilation manager to regnerate the blueprint's classes

diff --git a/Engine/Source/Editor/UnrealEd/Private/Kismet2/BlueprintEditorUtils.cpp b/Engine/Source/Editor/UnrealEd/Private/Kismet2/BlueprintEditorUtils.cpp
--- a/Engine/Source/Editor/UnrealEd/Private/Kismet2/BlueprintEditorUtils.cpp
+++ b/Engine/Source/Editor/UnrealEd/Private/Kismet2/BlueprintEditorUtils.cpp
@@ -1909,7 +1909,16 @@ void FBlueprintEditorUtils::PostDuplicateBlueprint(UBlueprint* Blueprint, bool b
 			Blueprint->Timelines = NewBPGC->Timelines;
 			Blueprint->InheritableComponentHandler = NewBPGC->InheritableComponentHandler;
 
-			Compiler.CompileBlueprint(Blueprint, CompileOptions, Results);
+			if (GBlueprintUseCompilationManager)
+			{
+				FBlueprintCompilationManager::CompileSynchronously(
+					FBPCompileRequest(Blueprint, EBlueprintCompileOptions::RegenerateSkeletonOnly, nullptr)
+				);
+			}
+			else
+			{
+				Compiler.CompileBlueprint(Blueprint, CompileOptions, Results);
+			}
 
 			// Create a new blueprint guid
 			Blueprint->GenerateNewGuid();
@@ -1942,7 +1951,7 @@ void FBlueprintEditorUtils::PostDuplicateBlueprint(UBlueprint* Blueprint, bool b
 							}
 							else
 							{
-								TargetClass = Blueprint->SkeletonGeneratedClass->GetAuthoritativeClass();
+								TargetClass = NewClass;
 							}
 
 							SelfPin->PinType.PinSubCategoryObject = TargetClass;
@@ -1951,9 +1960,18 @@ void FBlueprintEditorUtils::PostDuplicateBlueprint(UBlueprint* Blueprint, bool b
 				}
 			}
 
-			// Needs a full compile to handle the ArchiveReplaceObjectRef
-			CompileOptions.CompileType = EKismetCompileType::Full;
-			Compiler.CompileBlueprint(Blueprint, CompileOptions, Results);
+			if (GBlueprintUseCompilationManager)
+			{
+				FBlueprintCompilationManager::CompileSynchronously(
+					FBPCompileRequest(Blueprint, EBlueprintCompileOptions::None, nullptr)
+				);
+			}
+			else
+			{
+				// Needs a full compile to handle the ArchiveReplaceObjectRef
+				CompileOptions.CompileType = EKismetCompileType::Full;
+				Compiler.CompileBlueprint(Blueprint, CompileOptions, Results);
+			}
 
 			FArchiveReplaceObjectRef<UObject> ReplaceTemplateRefs(NewBPGC, OldToNewMap, /*bNullPrivateRefs=*/ false, /*bIgnoreOuterRef=*/ false, /*bIgnoreArchetypeRef=*/ false);
 
