Fixes to allow DO_BLUEPRINT_GUARD to be disabled in editor builds, lightweight timer so that UObject::ProcessEvent time can still be tracked when not using DO_BLUEPRINT_GUARD

diff --git a/Engine/Source/Editor/Kismet/Private/WatchPointViewer.cpp b/Engine/Source/Editor/Kismet/Private/WatchPointViewer.cpp
--- a/Engine/Source/Editor/Kismet/Private/WatchPointViewer.cpp
+++ b/Engine/Source/Editor/Kismet/Private/WatchPointViewer.cpp
@@ -763,6 +763,7 @@ TSharedRef<SWidget> SWatchTreeWidgetItem::GenerateWidgetForColumn(const FName& C
 
 void WatchViewer::UpdateInstancedWatchDisplay()
 {
+#if DO_BLUEPRINT_GUARD
 	{
 		const TArray<const FFrame*>& ScriptStack = FBlueprintExceptionTracker::Get().ScriptStack;
 
@@ -871,6 +872,7 @@ void WatchViewer::UpdateInstancedWatchDisplay()
 		// Notify subscribers:
 		WatchListSubscribers.Broadcast(&Private_InstanceWatchSource);
 	}
+#endif
 }
 
 void WatchViewer::ContinueExecution()
diff --git a/Engine/Source/Editor/UnrealEd/Private/Kismet2/KismetDebugUtilities.cpp b/Engine/Source/Editor/UnrealEd/Private/Kismet2/KismetDebugUtilities.cpp
--- a/Engine/Source/Editor/UnrealEd/Private/Kismet2/KismetDebugUtilities.cpp
+++ b/Engine/Source/Editor/UnrealEd/Private/Kismet2/KismetDebugUtilities.cpp
@@ -59,6 +59,13 @@ public:
 	{
 	}
 
+	~FKismetDebugUtilitiesData()
+	{
+		volatile int32 a = 0;
+		++a;
+		(void)a;
+	}
+
 	void Reset()
 	{
 		TargetGraphNodes.Empty();
@@ -118,7 +125,7 @@ public:
 
 void FKismetDebugUtilities::EndOfScriptExecution()
 {
-	
+#if DO_BLUEPRINT_GUARD
 	FBlueprintExceptionTracker& BlueprintExceptionTracker = FBlueprintExceptionTracker::Get();
 	if(BlueprintExceptionTracker.ScriptEntryTag == 1)
 	{
@@ -127,18 +134,22 @@ void FKismetDebugUtilities::EndOfScriptExecution()
 
 		Data.Reset();
 	}
+#endif // DO_BLUEPRINT_GUARD
 }
 
 void FKismetDebugUtilities::RequestSingleStepIn()
 {
+#if DO_BLUEPRINT_GUARD
 	FKismetDebugUtilitiesData& Data = FKismetDebugUtilitiesData::Get();
 	FBlueprintExceptionTracker& BlueprintExceptionTracker = FBlueprintExceptionTracker::Get();
 
 	Data.bIsSingleStepping = true;
+#endif // DO_BLUEPRINT_GUARD
 }
 
 void FKismetDebugUtilities::RequestStepOver()
 {
+#if DO_BLUEPRINT_GUARD
 	FKismetDebugUtilitiesData& Data = FKismetDebugUtilitiesData::Get();
 	FBlueprintExceptionTracker& BlueprintExceptionTracker = FBlueprintExceptionTracker::Get();
 
@@ -171,10 +182,12 @@ void FKismetDebugUtilities::RequestStepOver()
 			}
 		}
 	}
+#endif // DO_BLUEPRINT_GUARD
 }
 
 void FKismetDebugUtilities::RequestStepOut()
 {
+#if DO_BLUEPRINT_GUARD
 	FKismetDebugUtilitiesData& Data = FKismetDebugUtilitiesData::Get();
 	FBlueprintExceptionTracker& BlueprintExceptionTracker = FBlueprintExceptionTracker::Get();
 
@@ -184,6 +197,7 @@ void FKismetDebugUtilities::RequestStepOut()
 		Data.bIsSteppingOut = true;
 		Data.TargetGraphStackDepth = BlueprintExceptionTracker.ScriptStack.Num() - 1;
 	}
+#endif // DO_BLUEPRINT_GUARD
 }
 
 void FKismetDebugUtilities::OnScriptException(const UObject* ActiveObject, const FFrame& StackFrame, const FBlueprintExceptionInfo& Info)
@@ -501,6 +515,7 @@ UEdGraphNode* FKismetDebugUtilities::FindSourceNodeForCodeLocation(const UObject
 
 void FKismetDebugUtilities::CheckBreakConditions(UEdGraphNode* NodeStoppedAt, bool bHitBreakpoint, int32 BreakpointOffset, bool& InOutBreakExecution)
 {
+#if DO_BLUEPRINT_GUARD
 	FKismetDebugUtilitiesData& Data = FKismetDebugUtilitiesData::Get();
 	FBlueprintExceptionTracker& BlueprintExceptionTracker = FBlueprintExceptionTracker::Get();
 
@@ -599,10 +614,12 @@ void FKismetDebugUtilities::CheckBreakConditions(UEdGraphNode* NodeStoppedAt, bo
 			}
 		}
 	}
+#endif // DO_BLUEPRINT_GUARD
 }
 
 void FKismetDebugUtilities::AttemptToBreakExecution(UBlueprint* BlueprintObj, const UObject* ActiveObject, const FFrame& StackFrame, const FBlueprintExceptionInfo& Info, UEdGraphNode* NodeStoppedAt, int32 DebugOpcodeOffset)
 {
+#if DO_BLUEPRINT_GUARD
 	checkSlow(BlueprintObj->GetObjectBeingDebugged() == ActiveObject);
 
 	FKismetDebugUtilitiesData& Data = FKismetDebugUtilitiesData::Get();
@@ -699,6 +716,7 @@ void FKismetDebugUtilities::AttemptToBreakExecution(UBlueprint* BlueprintObj, co
 		WatchViewer::UpdateInstancedWatchDisplay();
 		FSlateApplication::Get().EnterDebuggingMode();
 	}
+#endif // DO_BLUEPRINT_GUARD
 }
 
 UEdGraphNode* FKismetDebugUtilities::GetCurrentInstruction()
diff --git a/Engine/Source/Runtime/CoreUObject/Private/UObject/ScriptCore.cpp b/Engine/Source/Runtime/CoreUObject/Private/UObject/ScriptCore.cpp
--- a/Engine/Source/Runtime/CoreUObject/Private/UObject/ScriptCore.cpp
+++ b/Engine/Source/Runtime/CoreUObject/Private/UObject/ScriptCore.cpp
@@ -928,6 +928,7 @@ void ClearReturnValue(UProperty* ReturnProp, RESULT_DECL)
 
 DEFINE_FUNCTION(UObject::ProcessInternal)
 {
+#if DO_BLUEPRINT_GUARD
 	// remove later when stable
 	if (P_THIS->GetClass()->HasAnyClassFlags(CLASS_NewerVersionExists))
 	{
@@ -937,6 +938,7 @@ DEFINE_FUNCTION(UObject::ProcessInternal)
 		}
 		return;
 	}
+#endif
 
 	UFunction* Function = (UFunction*)Stack.Node;
 
@@ -1261,11 +1263,35 @@ FBlueprintEventTimer::FScopedNativeTimer::~FScopedNativeTimer()
 
 #endif
 
+// Switch for a lightweight process event counter, useful when disabling the blueprint guard
+// which can taint profiling results:
+#define LIGHTWEIGHT_PROCESS_EVENT_COUNTER 0 && !DO_BLUEPRINT_GUARD
+
+#if LIGHTWEIGHT_PROCESS_EVENT_COUNTER
+thread_local int32 ProcessEventCounter = 0;
+#endif
+
 void UObject::ProcessEvent( UFunction* Function, void* Parms )
 {
 	checkf(!IsUnreachable(),TEXT("%s  Function: '%s'"), *GetFullName(), *Function->GetPathName());
 	checkf(!FUObjectThreadContext::Get().IsRoutingPostLoad, TEXT("Cannot call UnrealScript (%s - %s) while PostLoading objects"), *GetFullName(), *Function->GetFullName());
 
+#if LIGHTWEIGHT_PROCESS_EVENT_COUNTER
+	CONDITIONAL_SCOPE_CYCLE_COUNTER(STAT_BlueprintTime, IsInGameThread() && ProcessEventCounter == 0);
+	TGuardValue<int32> PECounter(ProcessEventCounter, ProcessEventCounter+1);
+#endif
+
+#if DO_BLUEPRINT_GUARD
+	FBlueprintExceptionTracker& BlueprintExceptionTracker = FBlueprintExceptionTracker::Get();
+	BlueprintExceptionTracker.ScriptEntryTag++;
+
+	CONDITIONAL_SCOPE_CYCLE_COUNTER(STAT_BlueprintTime, IsInGameThread() && BlueprintExceptionTracker.ScriptEntryTag == 1);
+#endif
+
+#if TOTAL_OVERHEAD_SCRIPT_STATS
+	FBlueprintEventTimer::FScopedVMTimer VMTime;
+#endif // TOTAL_OVERHEAD_SCRIPT_STATS
+
 	// Reject.
 	if (IsPendingKill())
 	{
@@ -1304,10 +1330,6 @@ void UObject::ProcessEvent( UFunction* Function, void* Parms )
 	}
 	checkSlow((Function->ParmsSize == 0) || (Parms != NULL));
 
-#if TOTAL_OVERHEAD_SCRIPT_STATS
-	FBlueprintEventTimer::FScopedVMTimer VMTime;
-#endif // TOTAL_OVERHEAD_SCRIPT_STATS
-
 #if PER_FUNCTION_SCRIPT_STATS
 	const bool bShouldTrackFunction = Stats::IsThreadCollectingData();
 	FScopeCycleCounterUObject FunctionScope(bShouldTrackFunction ? Function : nullptr);
@@ -1318,13 +1340,6 @@ void UObject::ProcessEvent( UFunction* Function, void* Parms )
 	FScopeCycleCounterUObject ContextScope(bShouldTrackObject ? this : nullptr);
 #endif
 
-#if DO_BLUEPRINT_GUARD
-	FBlueprintExceptionTracker& BlueprintExceptionTracker = FBlueprintExceptionTracker::Get();
-	BlueprintExceptionTracker.ScriptEntryTag++;
-
-	CONDITIONAL_SCOPE_CYCLE_COUNTER(STAT_BlueprintTime, IsInGameThread() && BlueprintExceptionTracker.ScriptEntryTag == 1);
-#endif
-
 #if UE_BLUEPRINT_EVENTGRAPH_FASTCALLS
 	// Fast path for ubergraph calls
 	int32 EventGraphParams;
