Correctly identify root actors when creating a blueprint from components.
Fix attachment hierarchy and relative transforms when creating a blueprint from actors that have an attachment relationship but an intervening actor is not selected. (i.e. A -> B -> C, only A & C are selected).

diff --git a/Engine/Source/Editor/UnrealEd/Private/Kismet2/Kismet2.cpp b/Engine/Source/Editor/UnrealEd/Private/Kismet2/Kismet2.cpp
--- a/Engine/Source/Editor/UnrealEd/Private/Kismet2/Kismet2.cpp
+++ b/Engine/Source/Editor/UnrealEd/Private/Kismet2/Kismet2.cpp
@@ -1307,14 +1307,13 @@ UBlueprint* FKismetEditorUtilities::CreateBlueprintFromActor(const FString& Path
 	return NewBlueprint;
 }
 
-void FKismetEditorUtilities::AddComponentsToBlueprint(UBlueprint* Blueprint, const TArray<UActorComponent*>& Components, bool bHarvesting, USCS_Node* OptionalNewRootNode, bool bKeepMobility /*= false*/)
+void FKismetEditorUtilities::AddComponentsToBlueprint(UBlueprint* Blueprint, TArray<UActorComponent*> Components, bool bHarvesting, USCS_Node* OptionalNewRootNode, bool bKeepMobility /*= false*/)
 {
 	USimpleConstructionScript* SCS = Blueprint->SimpleConstructionScript;
 
 	TArray<UBlueprint*> ParentBPStack;
 	UBlueprint::GetBlueprintHierarchyFromClass(Blueprint->GeneratedClass, ParentBPStack);
 
-	TMap<USceneComponent*, USCS_Node*> SceneComponentsToAdd;
 	TMap<USceneComponent*, USCS_Node*> InstanceComponentToNodeMap;
 
 	auto AddChildToSCSRootNodeLambda = [SCS, OptionalNewRootNode](USCS_Node* InSCSNode)
@@ -1370,116 +1369,150 @@ void FKismetEditorUtilities::AddComponentsToBlueprint(UBlueprint* Blueprint, con
 		}
 	};
 
-	for (UActorComponent* ActorComponent : Components)
+	// Clear out nulls and components we won't be able to create, can swap since we're about to sort anyways
+	Components.RemoveAllSwap([](const UActorComponent* ActorComponent)
 	{
-		if (ActorComponent)
-		{
-			AActor* Actor = ActorComponent->GetOwner();
-			check(Actor);
+		return (ActorComponent == nullptr || !ActorComponent->GetClass()->HasMetaData(FBlueprintMetadata::MD_BlueprintSpawnableComponent));
+	});
+
+	// Sort scene components to the start of the array and those that are children after their parents
+	Components.Sort([](const UActorComponent& A, const UActorComponent& B)
+	{
+		const USceneComponent* SceneA = Cast<const USceneComponent>(&A);
+		const USceneComponent* SceneB = Cast<const USceneComponent>(&B);
 
-			if (!ActorComponent->GetClass()->HasMetaData(FBlueprintMetadata::MD_BlueprintSpawnableComponent))
+		if (SceneA)
+		{
+			if (SceneB)
 			{
-				continue;
+				return SceneB->IsAttachedTo(SceneA);
 			}
+			return true;
+		}
+		return (SceneB == nullptr);
+	});
 
-			USCS_Node* SCSNode = FAddComponentsToBlueprintImpl::MakeComponentCopy(ActorComponent, SCS, InstanceComponentToNodeMap, bKeepMobility);
+	for (int32 CompIndex = 0; CompIndex < Components.Num(); ++CompIndex)
+	{
+		UActorComponent* ActorComponent = Components[CompIndex];
+		AActor* Actor = ActorComponent->GetOwner();
+		check(Actor);
 
-			USceneComponent* SceneComponent = Cast<USceneComponent>(ActorComponent);
-			// The easy part is non-scene component or the Root simply add it
-			if (SceneComponent == nullptr)
-			{
-				SCS->AddNode(SCSNode);
-			}
-			else
+		USCS_Node* SCSNode = FAddComponentsToBlueprintImpl::MakeComponentCopy(ActorComponent, SCS, InstanceComponentToNodeMap, bKeepMobility);
+
+		USceneComponent* SceneComponent = Cast<USceneComponent>(ActorComponent);
+		// The easy part is non-scene component or the Root simply add it
+		if (SceneComponent == nullptr)
+		{
+			SCS->AddNode(SCSNode);
+		}
+		else
+		{
+			// Since there is no guarantee that the attach parent is in the component set, but there could still be attachment hierarchy
+			// which is best to maintain, seek out the closest attached parent in the component set
+			USceneComponent* FirstAttachParent = nullptr;
+			if (SceneComponent->GetAttachParent())
 			{
-				if (SceneComponent == Actor->GetRootComponent() && (SceneComponent->GetAttachParent() == nullptr || !Components.Contains(SceneComponent->GetAttachParent())))
+				const int32 FirstAttachParentIndex = Components.FindLastByPredicate([SceneComponent](const UActorComponent* Other)
 				{
-					if (OptionalNewRootNode != nullptr)
-					{
-						OptionalNewRootNode->AddChildNode(SCSNode);
-					}
-					else
+					return (SceneComponent->IsAttachedTo(CastChecked<const USceneComponent>(Other)));
+				}, CompIndex);
+
+				if (FirstAttachParentIndex != INDEX_NONE)
+				{
+					FirstAttachParent = CastChecked<USceneComponent>(Components[FirstAttachParentIndex]);
+
+					// If the parent we're going to be attached to isn't the original attach parent, then adjust the
+					// relative transform so that the end result will be consistent with current relationship
+					if (FirstAttachParent != SceneComponent->GetAttachParent())
 					{
-						SCS->AddNode(SCSNode);
+						USceneComponent* SceneComponentTemplate = CastChecked<USceneComponent>(SCSNode->ComponentTemplate);
+
+						const FTransform ComponentToWorld = SceneComponent->GetComponentTransform();
+						const FTransform RelativeTransform = ComponentToWorld.GetRelativeTransform(FirstAttachParent->GetComponentTransform());
+						if (!SceneComponent->bAbsoluteLocation)
+						{
+							SceneComponentTemplate->RelativeLocation = RelativeTransform.GetLocation();
+						}
+						if (!SceneComponent->bAbsoluteRotation)
+						{
+							SceneComponentTemplate->RelativeRotation = RelativeTransform.GetRotation().Rotator();
+						}
+						if (!SceneComponent->bAbsoluteScale)
+						{
+							SceneComponentTemplate->RelativeScale3D = RelativeTransform.GetScale3D();
+						}
 					}
 				}
-				// If we're not attached to a blueprint component, add ourself to the root node or the SCS root component:
-				else if (SceneComponent->GetAttachParent() == nullptr)
+			}
+
+			if (SceneComponent == Actor->GetRootComponent() && (SceneComponent->GetAttachParent() == nullptr || FirstAttachParent == nullptr))
+			{
+				if (OptionalNewRootNode != nullptr)
+				{
+					OptionalNewRootNode->AddChildNode(SCSNode);
+				}
+				else
 				{
+					SCS->AddNode(SCSNode);
+				}
+			}
+			// If we're not attached to a blueprint component, add ourself to the root node or the SCS root component:
+			else if (SceneComponent->GetAttachParent() == nullptr)
+			{
 					AddChildToSCSRootNodeLambda(SCSNode);
+			}
+			// If we're attached to a blueprint component look it up as the variable name is the component name
+			else if (SceneComponent->GetAttachParent()->IsCreatedByConstructionScript())
+			{
+				USCS_Node* ParentSCSNode = nullptr;
+				if (FirstAttachParent)
+				{
+					ParentSCSNode = InstanceComponentToNodeMap.FindChecked(FirstAttachParent);
 				}
-				// If we're attached to a blueprint component look it up as the variable name is the component name
-				else if (SceneComponent->GetAttachParent()->IsCreatedByConstructionScript())
+				else
 				{
-					USCS_Node* ParentSCSNode = nullptr;
-					if (USCS_Node** ParentSCSNodePtr = InstanceComponentToNodeMap.Find(SceneComponent->GetAttachParent()))
+					for (UBlueprint* ParentBlueprint : ParentBPStack)
 					{
-						ParentSCSNode = *ParentSCSNodePtr;
-					}
-					else if (Components.Contains(SceneComponent->GetAttachParent()))
-					{
-						// since you cannot rely on the order of the supplied  
-						// Components array, we might be looking for a parent 
-						// that hasn't been added yet
-						ParentSCSNode = FAddComponentsToBlueprintImpl::MakeComponentCopy(SceneComponent->GetAttachParent(), SCS, InstanceComponentToNodeMap, bKeepMobility);
-					}
-					else
-					{
-						for (UBlueprint* ParentBlueprint : ParentBPStack)
+						ParentSCSNode = ParentBlueprint->SimpleConstructionScript->FindSCSNode(SceneComponent->GetAttachParent()->GetFName());
+						if (ParentSCSNode)
 						{
-							ParentSCSNode = ParentBlueprint->SimpleConstructionScript->FindSCSNode(SceneComponent->GetAttachParent()->GetFName());
-							if (ParentSCSNode)
-							{
-								break;
-							}
+							break;
 						}
 					}
-					check(ParentSCSNode);
-
-					if (ParentSCSNode->GetSCS() != SCS)
-					{
-						SCS->AddNode(SCSNode);
-						SCSNode->SetParent(ParentSCSNode);
-					}
-					else
-					{
-						ParentSCSNode->AddChildNode(SCSNode);
-					}
 				}
-				else if ((SceneComponent->GetAttachParent()->CreationMethod == EComponentCreationMethod::Native) && !bHarvesting)
+				check(ParentSCSNode);
+
+				if (ParentSCSNode->GetSCS() != SCS)
 				{
-					// If we're attached to a component that will be native in the new blueprint
 					SCS->AddNode(SCSNode);
-					SCSNode->SetParent(SceneComponent->GetAttachParent());
+					SCSNode->SetParent(ParentSCSNode);
 				}
 				else
 				{
-					// Otherwise check if we've already created the parents' new SCS node and attach to that or cache it off to do next pass
-					USCS_Node** ParentSCSNode = InstanceComponentToNodeMap.Find(SceneComponent->GetAttachParent());
-					if (ParentSCSNode)
-					{
-						(*ParentSCSNode)->AddChildNode(SCSNode);
-					}
-					else
-					{
-						SceneComponentsToAdd.Add(SceneComponent, SCSNode);
-					}
+					ParentSCSNode->AddChildNode(SCSNode);
+				}
+			}
+			else if ((SceneComponent->GetAttachParent()->CreationMethod == EComponentCreationMethod::Native) && !bHarvesting)
+			{
+				// If we're attached to a component that will be native in the new blueprint
+				SCS->AddNode(SCSNode);
+				SCSNode->SetParent(SceneComponent->GetAttachParent());
+			}
+			else
+			{
+				// Otherwise we will already have created the parents' new SCS node, so attach to that
+				if (USCS_Node** ParentSCSNode = InstanceComponentToNodeMap.Find(FirstAttachParent))
+				{
+					(*ParentSCSNode)->AddChildNode(SCSNode);
+				}
+				else
+				{
+					// Unsure what this case is really for when harvesting, but to be consistent with previous behavior still do this
+					ensure(bHarvesting);
+					AddChildToSCSRootNodeLambda(SCSNode);
 				}
 			}
-		}
-	}
-
-	// Hook up the remaining components nodes that the parent's node was missing when it was processed
-	for (auto ComponentIt = SceneComponentsToAdd.CreateConstIterator(); ComponentIt; ++ComponentIt)
-	{
-		// The AttachParent may or may not be BP-spawnable; if it's not, then we won't have created the parent node, so just add it as a child of the root node in that case.
-		if (USCS_Node** ParentSCSNode = InstanceComponentToNodeMap.Find(ComponentIt.Key()->GetAttachParent()))
-		{
-			(*ParentSCSNode)->AddChildNode(ComponentIt.Value());
-		}
-		else
-		{
-			AddChildToSCSRootNodeLambda(ComponentIt.Value());
 		}
 	}
 }
@@ -1604,41 +1637,39 @@ public:
 			SCS = Blueprint->SimpleConstructionScript;
 
 			// Create a common root if necessary
-			TArray<AActor*> RootActors;
-			USCS_Node* RootNodeOverride = nullptr;
+			TArray<AActor*> RootActors = SelectedActors;
 
-			// First sort the selected actors such that any actor that has its root attached to one of the other
-			// selected actors comes later in the array.
-			SelectedActors.Sort([](AActor& A, AActor& B)
+			int32 ConsideringActorIndex = 0;
+			while (ConsideringActorIndex < RootActors.Num())
 			{
-				return (B.IsAttachedTo(&A));
-			});
-
-			// Determine how many of the selected actors are not attached to one another
-			for (int32 OuterIndex = SelectedActors.Num() - 1; OuterIndex >= 0; --OuterIndex)
-			{
-				AActor* SelectedActor = SelectedActors[OuterIndex];
-				if (SelectedActor->GetRootComponent())
+				AActor* ConsideringActor = RootActors[ConsideringActorIndex];
+				bool bIsRoot = true;
+				for (int32 CheckIndex = RootActors.Num()-1; CheckIndex > ConsideringActorIndex; --CheckIndex)
 				{
-					if (OuterIndex == 0)
+					AActor* ActorToCheck = RootActors[CheckIndex];
+					if (ActorToCheck->IsAttachedTo(ConsideringActor))
 					{
-						RootActors.Add(SelectedActor);
+						RootActors.RemoveAtSwap(CheckIndex);
 					}
-					else
+					else if (ConsideringActor->IsAttachedTo(ActorToCheck))
 					{
-						for (int32 InnerIndex = OuterIndex - 1; InnerIndex >= 0; --InnerIndex)
-						{
-							if (!SelectedActor->IsAttachedTo(SelectedActors[InnerIndex]))
-							{
-								RootActors.Add(SelectedActor);
-								break;
-							}
-						}
+						bIsRoot = false;
+						break;
 					}
 				}
-			}
+
+				if (bIsRoot)
+				{
+					++ConsideringActorIndex;
+				}
+				else
+				{
+					RootActors.RemoveAtSwap(ConsideringActorIndex);
+				}
+			}		
 
 			// If there is not one unique root actor then create a new scene component to serve as the shared root node
+			USCS_Node* RootNodeOverride = nullptr;
 			if (RootActors.Num() != 1)
 			{
 				RootNodeOverride = SCS->CreateNode(USceneComponent::StaticClass(), TEXT("SharedRoot"));
@@ -1657,7 +1688,7 @@ public:
 					}
 				}
 			}
-			FKismetEditorUtilities::AddComponentsToBlueprint(Blueprint, AllSelectedComponents, /*bHarvesting=*/ true, RootNodeOverride);
+			FKismetEditorUtilities::AddComponentsToBlueprint(Blueprint, MoveTemp(AllSelectedComponents), /*bHarvesting=*/ true, RootNodeOverride);
 
 			FTransform NewActorTransform = FTransform::Identity;
 			if (RootActors.Num() == 1)
diff --git a/Engine/Source/Editor/UnrealEd/Public/Kismet2/KismetEditorUtilities.h b/Engine/Source/Editor/UnrealEd/Public/Kismet2/KismetEditorUtilities.h
--- a/Engine/Source/Editor/UnrealEd/Public/Kismet2/KismetEditorUtilities.h
+++ b/Engine/Source/Editor/UnrealEd/Public/Kismet2/KismetEditorUtilities.h
@@ -155,7 +155,7 @@ public:
 	static bool CanCreateBlueprintOfClass(const UClass* Class);
 
 	/** Take a list of components that belong to a single Actor and add them to a blueprint as SCSNodes */
-	static void AddComponentsToBlueprint(UBlueprint* Blueprint, const TArray<UActorComponent*>& Components, bool bHarvesting = false, class USCS_Node* OptionalNewRootNode = nullptr, bool bKeepMobility = false);
+	static void AddComponentsToBlueprint(UBlueprint* Blueprint, TArray<UActorComponent*> Components, bool bHarvesting = false, class USCS_Node* OptionalNewRootNode = nullptr, bool bKeepMobility = false);
 
 	/** 
 	 * Take an Actor and generate a blueprint based on it. Uses the Actors type as the parent class. 
