(4.20.2) Fix crash when dealing with transient nodes that reference a user defined enum that is changing
#rnx

diff --git a/Engine/Source/Editor/UnrealEd/Private/Kismet2/EnumEditorUtils.cpp b/Engine/Source/Editor/UnrealEd/Private/Kismet2/EnumEditorUtils.cpp
--- a/Engine/Source/Editor/UnrealEd/Private/Kismet2/EnumEditorUtils.cpp
+++ b/Engine/Source/Editor/UnrealEd/Private/Kismet2/EnumEditorUtils.cpp
@@ -338,13 +338,13 @@ void FEnumEditorUtils::BroadcastChanges(const UUserDefinedEnum* Enum, const TArr
 			INodeDependingOnEnumInterface* NodeDependingOnEnum = Cast<INodeDependingOnEnumInterface>(Node);
 			if (FNodeValidatorHelper::IsValid(Node) && NodeDependingOnEnum && (Enum == NodeDependingOnEnum->GetEnum()))
 			{
-				if (UBlueprint* Blueprint = Node->GetBlueprint())
+				if (Node->HasValidBlueprint())
 				{
 					if (NodeDependingOnEnum->ShouldBeReconstructedAfterEnumChanged())
 					{
 						Node->ReconstructNode();
 					}
-					BlueprintsToRefresh.Add(Blueprint);
+					BlueprintsToRefresh.Add(Node->GetBlueprint());
 				}
 			}
 		}
