Allow ALT-drag and CTRL-W to duplicate non-instance components when selected in the level editor's viewport. Also address a few QoL issues with cut/deletion of selected components via context menu or hotkey.

Change summary:
- Modified UUnrealEdEngine::edactDuplicateSelected() to remove the check that previously limited it to 'Instance' type components. Possibly was the result of a cut & paste from the "delete" command, where that condition is applicable.
- Modified UUnrealEdEngine::edactDuplicateSelected() and edactCopySelected() to iterate over all selected components when handling the "copy" and "duplicate" commands from the level editor's context menu when components are selected. It previously iterated over only editable components; this was inconsistent with the Actor Details panel which supports copy/duplicate commands on all selected components.
- Modified UUnrealEdEngine::Exec_Edit() to wrap the "cut" command in a scoped transaction when components are selected (adds undo support).
- Modified UUnrealEdEngine::Exec_Actor() to add an appropriate scoped transaction label for the "delete" and "duplicate" commands when components are selected.
- Modified FComponentEditorUtils::DeleteComponents() to update the selected component reference after reconstructing the Actor instance. Previously, the reference was left pointing to a trashed component instance, which resulted in an invalid selection.
- Modified UUnrealEdEngine::edactDeleteSelected() to notify the level editor of a selection change after selected component deletion via the context menu (or CTRL-X and DEL hotkeys). This properly snaps the selection and gizmo to the new component selection that is chosen by FComponentEditorUtils::DeleteComponents().
- Modified SActorDetails::PostUndo() to update the pivot (gizmo) along with the selection after an undo. This properly snaps the pivot to follow a restored component selection after undo of a selected component cut or delete operation.

diff --git a/Engine/Source/Editor/LevelEditor/Private/SActorDetails.cpp b/Engine/Source/Editor/LevelEditor/Private/SActorDetails.cpp
--- a/Engine/Source/Editor/LevelEditor/Private/SActorDetails.cpp
+++ b/Engine/Source/Editor/LevelEditor/Private/SActorDetails.cpp
@@ -282,6 +282,9 @@ void SActorDetails::PostUndo(bool bSuccess)
 	if (SelectedActor)
 	{
 		GUnrealEd->SetActorSelectionFlags(SelectedActor);
+
+		// Update the pivot (widget) as the current selection may be a component within the Actor instance
+		GUnrealEd->UpdatePivotLocationForSelection();
 	}
 }
 
diff --git a/Engine/Source/Editor/UnrealEd/Private/EditorActor.cpp b/Engine/Source/Editor/UnrealEd/Private/EditorActor.cpp
--- a/Engine/Source/Editor/UnrealEd/Private/EditorActor.cpp
+++ b/Engine/Source/Editor/UnrealEd/Private/EditorActor.cpp
@@ -133,9 +133,9 @@ void UUnrealEdEngine::edactCopySelected( UWorld* InWorld, FString* DestinationDa
 {
 	if (GetSelectedComponentCount() > 0)
 	{
-		// Copy components
+		// Copy selected components
 		TArray<UActorComponent*> SelectedComponents;
-		for (FSelectedEditableComponentIterator It(GetSelectedEditableComponentIterator()); It; ++It)
+		for (FSelectionIterator It(GetSelectedComponentIterator()); It; ++It)
 		{
 			SelectedComponents.Add(CastChecked<UActorComponent>(*It));
 		}
@@ -469,17 +469,13 @@ void UUnrealEdEngine::edactDuplicateSelected( ULevel* InLevel, bool bOffsetLocat
 		TArray<UActorComponent*> NewComponentClones;
 		NewComponentClones.Reserve(NumSelectedComponents);
 
-		// Duplicate selected components if they are an Instance component
-		for (FSelectedEditableComponentIterator It(GetSelectedEditableComponentIterator()); It; ++It)
+		// Duplicate selected components
+		for (FSelectionIterator It(GetSelectedComponentIterator()); It; ++It)
 		{
 			UActorComponent* Component = CastChecked<UActorComponent>(*It);
-			if (Component->CreationMethod == EComponentCreationMethod::Instance)
+			if (UActorComponent* Clone = FComponentEditorUtils::DuplicateComponent(Component))
 			{
-				UActorComponent* Clone = FComponentEditorUtils::DuplicateComponent(Component);
-				if (Clone)
-				{
-					NewComponentClones.Add(Clone);
-				}
+				NewComponentClones.Add(Clone);
 			}
 		}
 
@@ -715,6 +711,9 @@ bool UUnrealEdEngine::edactDeleteSelected( UWorld* InWorld, bool bVerifyDeletion
 					// Make sure the selection changed event fires so the SCS trees can update their selection
 					ComponentSelection->MarkBatchDirty();
 					ComponentSelection->EndBatchSelectOperation(true);
+
+					// Notify the level editor of the new component selection
+					NoteSelectionChange();
 				}
 
 				UE_LOG(LogEditorActor, Log, TEXT("Deleted %d Components (%3.3f secs)"), NumDeletedComponents, FPlatformTime::Seconds() - StartSeconds);
diff --git a/Engine/Source/Editor/UnrealEd/Private/Kismet2/ComponentEditorUtils.cpp b/Engine/Source/Editor/UnrealEd/Private/Kismet2/ComponentEditorUtils.cpp
--- a/Engine/Source/Editor/UnrealEd/Private/Kismet2/ComponentEditorUtils.cpp
+++ b/Engine/Source/Editor/UnrealEd/Private/Kismet2/ComponentEditorUtils.cpp
@@ -572,10 +572,39 @@ int32 FComponentEditorUtils::DeleteComponents(const TArray<UActorComponent*>& Co
 		NumDeletedComponents++;
 	}
 
+	// Non-native components will be reinstanced, so we have to update the ptr after reconstruction
+	// in order to avoid pointing at an invalid (trash) instance after re-running construction scripts.
+	FName ComponentToSelectName;
+	const AActor* ComponentToSelectOwner = nullptr;
+	if (OutComponentToSelect && OutComponentToSelect->CreationMethod != EComponentCreationMethod::Native)
+	{
+		// Keep track of the pending selection's name and owner
+		ComponentToSelectName = OutComponentToSelect->GetFName();
+		ComponentToSelectOwner = OutComponentToSelect->GetOwner();
+
+		// Reset the ptr value - we'll reassign it after reconstruction
+		OutComponentToSelect = nullptr;
+	}
+
 	// Reconstruct owner instance(s) after deletion
 	for(AActor* ActorToReconstruct : ActorsToReconstruct)
 	{
 		ActorToReconstruct->RerunConstructionScripts();
+
+		// If this actor matches the owner of the component to be selected, find the new instance of the component in the actor
+		if (ComponentToSelectName != NAME_None && OutComponentToSelect == nullptr && ActorToReconstruct == ComponentToSelectOwner)
+		{
+			TInlineComponentArray<UActorComponent*> ActorComponents;
+			ActorToReconstruct->GetComponents(ActorComponents);
+			for (UActorComponent* ActorComponent : ActorComponents)
+			{
+				if (ActorComponent->GetFName() == ComponentToSelectName)
+				{
+					OutComponentToSelect = ActorComponent;
+					break;
+				}
+			}
+		}
 	}
 
 	return NumDeletedComponents;
diff --git a/Engine/Source/Editor/UnrealEd/Private/UnrealEdSrv.cpp b/Engine/Source/Editor/UnrealEd/Private/UnrealEdSrv.cpp
--- a/Engine/Source/Editor/UnrealEd/Private/UnrealEdSrv.cpp
+++ b/Engine/Source/Editor/UnrealEd/Private/UnrealEdSrv.cpp
@@ -1598,6 +1598,9 @@ bool UUnrealEdEngine::Exec_Edit( UWorld* InWorld, const TCHAR* Str, FOutputDevic
 
 		if (bComponentsSelected)
 		{
+			// Same transaction language used in CopySelectedActorsToClipboard below
+			const FScopedTransaction Transaction(NSLOCTEXT("UnrealEd", "Cut", "Cut"));
+
 			edactCopySelected(InWorld);
 			edactDeleteSelected(InWorld);
 		}
@@ -2001,6 +2004,9 @@ bool UUnrealEdEngine::Exec_Actor( UWorld* InWorld, const TCHAR* Str, FOutputDevi
 	// Keep a pointer to the beginning of the string to use for message displaying purposes
 	const TCHAR* const FullStr = Str;
 
+	// Determine whether or not components are selected (used to properly label transaction names)
+	const bool bComponentsSelected = GetSelectedComponentCount() > 0;
+
 	if( FParse::Command(&Str,TEXT("ADD")) )
 	{
 		UClass* Class;
@@ -2563,7 +2569,7 @@ bool UUnrealEdEngine::Exec_Actor( UWorld* InWorld, const TCHAR* Str, FOutputDevi
 		// if not specially handled by the current editing mode,
 		if (!bHandled)
 		{
-			const FScopedTransaction Transaction( NSLOCTEXT("UnrealEd", "DeleteActors", "Delete Actors") );
+			const FScopedTransaction Transaction( bComponentsSelected ? NSLOCTEXT("UnrealEd", "DeleteComponents", "Delete Components") : NSLOCTEXT("UnrealEd", "DeleteActors", "Delete Actors") );
 			edactDeleteSelected( InWorld );
 		}
 		return true;
@@ -2746,7 +2752,7 @@ bool UUnrealEdEngine::Exec_Actor( UWorld* InWorld, const TCHAR* Str, FOutputDevi
 		if (!bHandled)
 		{
 			//@todo locked levels - if all actor levels are locked, cancel the transaction
-			const FScopedTransaction Transaction( NSLOCTEXT("UnrealEd", "DuplicateActors", "Duplicate Actors") );
+			const FScopedTransaction Transaction( bComponentsSelected ? NSLOCTEXT("UnrealEd", "DuplicateComponents", "Duplicate Components") : NSLOCTEXT("UnrealEd", "DuplicateActors", "Duplicate Actors") );
 
 			// duplicate selected
 			ABrush::SetSuppressBSPRegeneration(true);
