Fix SA warning (CIS).

#rnx

diff --git a/Engine/Source/Editor/UnrealEd/Private/Kismet2/ComponentEditorUtils.cpp b/Engine/Source/Editor/UnrealEd/Private/Kismet2/ComponentEditorUtils.cpp
--- a/Engine/Source/Editor/UnrealEd/Private/Kismet2/ComponentEditorUtils.cpp
+++ b/Engine/Source/Editor/UnrealEd/Private/Kismet2/ComponentEditorUtils.cpp
@@ -589,6 +589,7 @@ int32 FComponentEditorUtils::DeleteComponents(const TArray<UActorComponent*>& Co
 	// Reconstruct owner instance(s) after deletion
 	for(AActor* ActorToReconstruct : ActorsToReconstruct)
 	{
+		check(ActorToReconstruct != nullptr);
 		ActorToReconstruct->RerunConstructionScripts();
 
 		// If this actor matches the owner of the component to be selected, find the new instance of the component in the actor
