Fix a regression in which the DefaultSceneRoot could be duplicated via ALT-drag (for consistency, as its duplication is disabled elsewhere).

Change summary:
- Added FComponentEditorUtils::CanCopyComponent().
- Modified FComponentEditorUtils::CanCopyComponents() to call CanCopyComponent() (where the test logic has been moved).
- Modified UUnrealEdEngine::edactDuplicateSelected() to call FComponentEditorUtils::CanCopyComponent() before duplication of selected component(s).

Notes:
- This was introduced by 4266581 (UE-61704).
- For ALT-drag, duplication is covered by the move transaction, so whether or not we actually duplicated a component doesn't tie in to the label we use for undo atm.

#rnx

diff --git a/Engine/Source/Editor/UnrealEd/Private/EditorActor.cpp b/Engine/Source/Editor/UnrealEd/Private/EditorActor.cpp
--- a/Engine/Source/Editor/UnrealEd/Private/EditorActor.cpp
+++ b/Engine/Source/Editor/UnrealEd/Private/EditorActor.cpp
@@ -473,9 +473,13 @@ void UUnrealEdEngine::edactDuplicateSelected( ULevel* InLevel, bool bOffsetLocat
 		for (FSelectionIterator It(GetSelectedComponentIterator()); It; ++It)
 		{
 			UActorComponent* Component = CastChecked<UActorComponent>(*It);
-			if (UActorComponent* Clone = FComponentEditorUtils::DuplicateComponent(Component))
+
+			if (FComponentEditorUtils::CanCopyComponent(Component))
 			{
-				NewComponentClones.Add(Clone);
+				if (UActorComponent* Clone = FComponentEditorUtils::DuplicateComponent(Component))
+				{
+					NewComponentClones.Add(Clone);
+				}
 			}
 		}
 
diff --git a/Engine/Source/Editor/UnrealEd/Private/Kismet2/ComponentEditorUtils.cpp b/Engine/Source/Editor/UnrealEd/Private/Kismet2/ComponentEditorUtils.cpp
--- a/Engine/Source/Editor/UnrealEd/Private/Kismet2/ComponentEditorUtils.cpp
+++ b/Engine/Source/Editor/UnrealEd/Private/Kismet2/ComponentEditorUtils.cpp
@@ -296,6 +296,21 @@ USceneComponent* FComponentEditorUtils::FindClosestParentInList(UActorComponent*
 	return nullptr;
 }
 
+bool FComponentEditorUtils::CanCopyComponent(const UActorComponent* ComponentToCopy)
+{
+	if (ComponentToCopy != nullptr && ComponentToCopy->GetFName() != USceneComponent::GetDefaultSceneRootVariableName())
+	{
+		UClass* ComponentClass = ComponentToCopy->GetClass();
+		check(ComponentClass != nullptr);
+
+		// Component class cannot be abstract and must also be tagged as BlueprintSpawnable
+		return !ComponentClass->HasAnyClassFlags(CLASS_Abstract)
+			&& ComponentClass->HasMetaData(FBlueprintMetadata::MD_BlueprintSpawnableComponent);
+	}
+
+	return false;
+}
+
 bool FComponentEditorUtils::CanCopyComponents(const TArray<UActorComponent*>& ComponentsToCopy)
 {
 	bool bCanCopy = ComponentsToCopy.Num() > 0;
@@ -305,16 +320,7 @@ bool FComponentEditorUtils::CanCopyComponents(const TArray<UActorComponent*>& Co
 		{
 			// Check for the default scene root; that cannot be copied/duplicated
 			UActorComponent* Component = ComponentsToCopy[i];
-			bCanCopy = Component != nullptr && Component->GetFName() != USceneComponent::GetDefaultSceneRootVariableName();
-			if (bCanCopy)
-			{
-				UClass* ComponentClass = Component->GetClass();
-				check(ComponentClass != nullptr);
-
-				// Component class cannot be abstract and must also be tagged as BlueprintSpawnable
-				bCanCopy = !ComponentClass->HasAnyClassFlags(CLASS_Abstract)
-					&& ComponentClass->HasMetaData(FBlueprintMetadata::MD_BlueprintSpawnableComponent);
-			}
+			bCanCopy = CanCopyComponent(Component);
 		}
 	}
 
diff --git a/Engine/Source/Editor/UnrealEd/Public/Kismet2/ComponentEditorUtils.h b/Engine/Source/Editor/UnrealEd/Public/Kismet2/ComponentEditorUtils.h
--- a/Engine/Source/Editor/UnrealEd/Public/Kismet2/ComponentEditorUtils.h
+++ b/Engine/Source/Editor/UnrealEd/Public/Kismet2/ComponentEditorUtils.h
@@ -33,6 +33,13 @@ public:
 	/** Generate a valid variable name string for the given component instance based on the name of the asset referenced by the component */
 	static FString GenerateValidVariableNameFromAsset(UObject* Asset, AActor* ComponentOwner);
 
+	/**
+	* Checks whether it is valid to copy the given component
+	* @param ComponentToCopy The component to check
+	* @return Whether the given component can be copied
+	*/
+	static bool CanCopyComponent(const UActorComponent* ComponentToCopy);
+
 	/**
 	 * Checks whether it is valid to copy the indicated components
 	 * @param ComponentsToCopy The list of components to check
