[NO NON-PII COMMIT MESSAGE]

diff --git a/Engine/Source/Editor/UnrealEd/Private/EditorEngine.cpp b/Engine/Source/Editor/UnrealEd/Private/EditorEngine.cpp
--- a/Engine/Source/Editor/UnrealEd/Private/EditorEngine.cpp
+++ b/Engine/Source/Editor/UnrealEd/Private/EditorEngine.cpp
@@ -7062,6 +7062,8 @@ void UEditorEngine::InitializeNewlyCreatedInactiveWorld(UWorld* World)
 	check(World);
 	if (!World->bIsWorldInitialized && World->WorldType == EWorldType::Inactive)
 	{
+		const bool bOldDirtyState = World->GetOutermost()->IsDirty();
+
 		// Create the world without a physics scene because creating too many physics scenes causes deadlock issues in PhysX. The scene will be created when it is opened in the level editor.
 		// Also, don't create an FXSystem because it consumes too much video memory. This is also created when the level editor opens this world.
 		World->InitWorld(GetEditorWorldInitializationValues()
@@ -7071,6 +7073,12 @@ void UEditorEngine::InitializeNewlyCreatedInactiveWorld(UWorld* World)
 
 		// Update components so the scene is populated
 		World->UpdateWorldComponents(true, true);
+
+		// Need to restore the dirty state as registering components dirties the world
+		if (!bOldDirtyState)
+		{
+			World->GetOutermost()->SetDirtyFlag(bOldDirtyState);
+		}
 	}
 }
 
