Remove unused code
#rnx

diff --git a/Engine/Source/Editor/UnrealEd/Private/Factories/ActorFactory.cpp b/Engine/Source/Editor/UnrealEd/Private/Factories/ActorFactory.cpp
--- a/Engine/Source/Editor/UnrealEd/Private/Factories/ActorFactory.cpp
+++ b/Engine/Source/Editor/UnrealEd/Private/Factories/ActorFactory.cpp
@@ -526,7 +526,7 @@ void UActorFactoryDeferredDecal::PostSpawnActor(UObject* Asset, AActor* NewActor
 
 	UMaterialInterface* Material = GetMaterial( Asset );
 
-	if ( Material != NULL )
+	if (Material != NULL )
 	{
 		// Change properties
 		TInlineComponentArray<UDecalComponent*> DecalComponents;
@@ -534,7 +534,7 @@ void UActorFactoryDeferredDecal::PostSpawnActor(UObject* Asset, AActor* NewActor
 
 		UDecalComponent* DecalComponent = NULL;
 		for (int32 Idx = 0; Idx < DecalComponents.Num() && DecalComponent == NULL; Idx++)
-		{
+			{
 			DecalComponent = DecalComponents[Idx];
 		}
 
@@ -1119,21 +1119,6 @@ UActorFactoryCameraActor::UActorFactoryCameraActor(const FObjectInitializer& Obj
 	NewActorClass = ACameraActor::StaticClass();
 }
 
-static UBillboardComponent* CreateEditorOnlyBillboardComponent(AActor* ActorOwner, USceneComponent* AttachParent)
-{
-	// Create a new billboard component to serve as a visualization of the actor until there is another primitive component
-	UBillboardComponent* BillboardComponent = NewObject<UBillboardComponent>(ActorOwner, NAME_None, RF_Transactional);
-
-	BillboardComponent->Sprite = LoadObject<UTexture2D>(nullptr, TEXT("/Engine/EditorResources/EmptyActor.EmptyActor"));
-	BillboardComponent->RelativeScale3D = FVector(0.5f, 0.5f, 0.5f);
-	BillboardComponent->Mobility = EComponentMobility::Movable;
-	BillboardComponent->bIsEditorOnly = true;
-
-	BillboardComponent->SetupAttachment(AttachParent);
-
-	return BillboardComponent;
-}
-
 /*-----------------------------------------------------------------------------
 UActorFactoryEmptyActor
 -----------------------------------------------------------------------------*/
