Fix inverted condition
#rnx

diff --git a/Engine/Source/Editor/UnrealEd/Private/Kismet2/ComponentEditorUtils.cpp b/Engine/Source/Editor/UnrealEd/Private/Kismet2/ComponentEditorUtils.cpp
--- a/Engine/Source/Editor/UnrealEd/Private/Kismet2/ComponentEditorUtils.cpp
+++ b/Engine/Source/Editor/UnrealEd/Private/Kismet2/ComponentEditorUtils.cpp
@@ -554,7 +554,7 @@ int32 FComponentEditorUtils::DeleteComponents(const TArray<UActorComponent*>& Co
 					// For a non-scene component, try to select the preceding non-scene component
 					for (UActorComponent* Component : Owner->GetComponents())
 					{
-						if (Component == nullptr)
+						if (Component != nullptr)
 						{
 							if (Component == ComponentToDelete)
 							{
