Control rig : invalid FOV issue

diff --git a/Engine/Source/Editor/UnrealEd/Classes/Preferences/PersonaOptions.h b/Engine/Source/Editor/UnrealEd/Classes/Preferences/PersonaOptions.h
--- a/Engine/Source/Editor/UnrealEd/Classes/Preferences/PersonaOptions.h
+++ b/Engine/Source/Editor/UnrealEd/Classes/Preferences/PersonaOptions.h
@@ -47,6 +47,12 @@ struct FViewportConfigOptions
 
 	UPROPERTY(config)
 	FName CameraFollowBoneName;
+
+	FViewportConfigOptions()
+		: ViewModeIndex(EViewModeIndex::VMI_Lit)
+		, ViewFOV(53.43)
+		, CameraFollowMode(EAnimationViewportCameraFollowMode::None)
+	{}
 };
 
 /** Options that should be unique per asset editor (like skeletal mesh or anim sequence editors) */
