Autoload Animation Blueprint compiler module in editor so that we can set up logging category (following how normal blueprint compiler works)

diff --git a/Engine/Source/Editor/UnrealEd/Private/EditorEngine.cpp b/Engine/Source/Editor/UnrealEd/Private/EditorEngine.cpp
--- a/Engine/Source/Editor/UnrealEd/Private/EditorEngine.cpp
+++ b/Engine/Source/Editor/UnrealEd/Private/EditorEngine.cpp
@@ -966,6 +966,7 @@ void UEditorEngine::Init(IEngineLoop* InEngineLoop)
 			TEXT("KismetCompiler"),
 			TEXT("Kismet"),
 			TEXT("Persona"),
+			TEXT("AnimationBlueprintEditor"),
 			TEXT("LevelEditor"),
 			TEXT("MainFrame"),
 			TEXT("PropertyEditor"),
