Import and Export support for Weighted Tangents

Main tricky thing was the difference between how we store weighted tangents (length of hypotenuse) and FBX stores weighted, in addtion to the silly FBX fact that the current key stores the next key's arrive tangent.

Also had to utilize RichCurve since that has support for storing the weight values and the old Matinee InterpFloatValue does not.  Therefore as suggested set a bunch of functions as deprecated and since they call each other warnings appears.

Also fixed some issues with import and export not supporting auto tangents.

Originally was going to add weighted tangent interpolation to FRichCurve also but we don't actually touch the SkeletalMesh code here, it's all within the Sequencer code path, so that can go into a seperate changelist, with it's own set of tests.

diff --git a/Engine/Source/Editor/MovieSceneTools/Private/MatineeImportTools.cpp b/Engine/Source/Editor/MovieSceneTools/Private/MatineeImportTools.cpp
--- a/Engine/Source/Editor/MovieSceneTools/Private/MatineeImportTools.cpp
+++ b/Engine/Source/Editor/MovieSceneTools/Private/MatineeImportTools.cpp
@@ -114,7 +114,8 @@ bool FMatineeImportTools::TryConvertMatineeToggleToOutParticleKey( ETrackToggleA
 }
 
 
-void FMatineeImportTools::SetOrAddKey(TMovieSceneChannel<FMovieSceneFloatValue>& ChannelInterface, FFrameNumber Time, float Value, float ArriveTangent, float LeaveTangent, EInterpCurveMode MatineeInterpMode )
+void FMatineeImportTools::SetOrAddKey(TMovieSceneChannel<FMovieSceneFloatValue>& ChannelInterface, FFrameNumber Time, float Value, float ArriveTangent, float LeaveTangent, EInterpCurveMode MatineeInterpMode
+	, ERichCurveTangentWeightMode WeightedMode, float ArriveTangentWeight, float LeaveTangentWeight )
 {
 	if (ChannelInterface.FindKey(Time) == INDEX_NONE)
 	{
@@ -124,6 +125,9 @@ void FMatineeImportTools::SetOrAddKey(TMovieSceneChannel<FMovieSceneFloatValue>&
 		NewKey.TangentMode = MatineeInterpolationToRichCurveTangent( MatineeInterpMode );
 		NewKey.Tangent.ArriveTangent = ArriveTangent;
 		NewKey.Tangent.LeaveTangent = LeaveTangent;
+		NewKey.Tangent.TangentWeightMode = WeightedMode;
+		NewKey.Tangent.ArriveTangentWeight = ArriveTangentWeight;
+		NewKey.Tangent.LeaveTangentWeight = LeaveTangentWeight;
 		ChannelInterface.AddKey( Time, NewKey );
 	}
 }
diff --git a/Engine/Source/Editor/MovieSceneTools/Private/MovieSceneToolHelpers.cpp b/Engine/Source/Editor/MovieSceneTools/Private/MovieSceneToolHelpers.cpp
--- a/Engine/Source/Editor/MovieSceneTools/Private/MovieSceneToolHelpers.cpp
+++ b/Engine/Source/Editor/MovieSceneTools/Private/MovieSceneToolHelpers.cpp
@@ -909,31 +909,40 @@ bool ImportFBXProperty(FString NodeName, FString AnimatedPropertyName, FGuid Obj
 
 				const int32 ChannelIndex = 0;
 				const int32 CompositeIndex = 0;
-				FInterpCurveFloat CurveHandle;
+				FRichCurve Source;
 				const bool bNegative = false;
-				CurveAPI.GetCurveData(NodeName, AnimatedPropertyName, ChannelIndex, CompositeIndex, CurveHandle, bNegative);
+				CurveAPI.GetCurveData(NodeName, AnimatedPropertyName, ChannelIndex, CompositeIndex, Source, bNegative);
 
 				FMovieSceneFloatChannel* Channel = FloatSection->GetChannelProxy().GetChannel<FMovieSceneFloatChannel>(0);
 				TMovieSceneChannel<FMovieSceneFloatValue> ChannelInterface = Channel->GetInterface();
 
 				ChannelInterface.Reset();
 				double DecimalRate = FrameRate.AsDecimal();
-				for (int32 KeyIndex = 0; KeyIndex < CurveHandle.Points.Num(); ++KeyIndex)
+
+				for (auto SourceIt = Source.GetKeyHandleIterator(); SourceIt; ++SourceIt)
 				{
-					float ArriveTangent = CurveHandle.Points[KeyIndex].ArriveTangent;
-					if (KeyIndex > 0)
+					FRichCurveKey &Key = Source.GetKey(SourceIt.Key());
+					float ArriveTangent = Key.ArriveTangent;
+					FKeyHandle PrevKeyHandle = Source.GetPreviousKey(SourceIt.Key());
+					if (Source.IsKeyHandleValid(PrevKeyHandle))
 					{
-						ArriveTangent = ArriveTangent / ((CurveHandle.Points[KeyIndex].InVal - CurveHandle.Points[KeyIndex - 1].InVal) * DecimalRate);
+						FRichCurveKey &PrevKey = Source.GetKey(PrevKeyHandle);
+						ArriveTangent = ArriveTangent / ((Key.Time - PrevKey.Time) * DecimalRate);
+
 					}
-					
-					float LeaveTangent = CurveHandle.Points[KeyIndex].LeaveTangent;
-					if (KeyIndex < CurveHandle.Points.Num() - 1)
+					float LeaveTangent = Key.LeaveTangent;
+					FKeyHandle NextKeyHandle = Source.GetNextKey(SourceIt.Key());
+					if (Source.IsKeyHandleValid(NextKeyHandle))
 					{
-						LeaveTangent = LeaveTangent / ((CurveHandle.Points[KeyIndex+1].InVal - CurveHandle.Points[KeyIndex].InVal) * DecimalRate);
+						FRichCurveKey &NextKey = Source.GetKey(NextKeyHandle);
+						LeaveTangent = LeaveTangent / ((NextKey.Time - Key.Time) * DecimalRate);
 					}
 
-					FFrameNumber KeyTime = (CurveHandle.Points[KeyIndex].InVal * FrameRate).RoundToFrame();
-					FMatineeImportTools::SetOrAddKey(ChannelInterface, KeyTime, CurveHandle.Points[KeyIndex].OutVal, ArriveTangent, LeaveTangent, CurveHandle.Points[KeyIndex].InterpMode);
+					FFrameNumber KeyTime = (Key.Time * FrameRate).RoundToFrame();
+					FMatineeImportTools::SetOrAddKey(ChannelInterface, KeyTime, Key.Value, ArriveTangent, LeaveTangent,
+						MovieSceneToolHelpers::RichCurveInterpolationToMatineeInterpolation(Key.InterpMode, Key.TangentMode), Key.TangentWeightMode,
+						Key.ArriveTangentWeight, Key.LeaveTangentWeight);
+
 				}
 
 				if (ImportFBXSettings->bReduceKeys)
@@ -958,23 +967,29 @@ bool ImportFBXProperty(FString NodeName, FString AnimatedPropertyName, FGuid Obj
 	return false;
 }
 
-void ImportTransformChannel(const FInterpCurveFloat& Source, FMovieSceneFloatChannel* Dest, FFrameRate DestFrameRate, bool bNegateTangents)
+void ImportTransformChannel(const FRichCurve& Source, FMovieSceneFloatChannel* Dest, FFrameRate DestFrameRate, bool bNegateTangents)
 {
 	TMovieSceneChannel<FMovieSceneFloatValue> ChannelInterface = Dest->GetInterface();
 	ChannelInterface.Reset();
 	double DecimalRate = DestFrameRate.AsDecimal();
-	for (int32 KeyIndex = 0; KeyIndex < Source.Points.Num(); ++KeyIndex)
+
+	for (auto SourceIt = Source.GetKeyHandleIterator(); SourceIt; ++SourceIt)
 	{
-		float ArriveTangent = Source.Points[KeyIndex].ArriveTangent;
-		if (KeyIndex > 0)
+		const FRichCurveKey Key = Source.GetKey(SourceIt.Key());
+		float ArriveTangent = Key.ArriveTangent;
+		FKeyHandle PrevKeyHandle = Source.GetPreviousKey(SourceIt.Key());
+		if (Source.IsKeyHandleValid(PrevKeyHandle))
 		{
-			ArriveTangent = ArriveTangent / ((Source.Points[KeyIndex].InVal - Source.Points[KeyIndex-1].InVal) * DecimalRate);
+			const FRichCurveKey PrevKey = Source.GetKey(PrevKeyHandle);
+			ArriveTangent = ArriveTangent / ((Key.Time - PrevKey.Time) * DecimalRate);
+
 		}
-		
-		float LeaveTangent = Source.Points[KeyIndex].LeaveTangent;
-		if (KeyIndex < Source.Points.Num() - 1)
+		float LeaveTangent = Key.LeaveTangent;
+		FKeyHandle NextKeyHandle = Source.GetNextKey(SourceIt.Key());
+		if (Source.IsKeyHandleValid(NextKeyHandle))
 		{
-			LeaveTangent = LeaveTangent / ((Source.Points[KeyIndex+1].InVal - Source.Points[KeyIndex].InVal) * DecimalRate);
+			const FRichCurveKey NextKey = Source.GetKey(NextKeyHandle);
+			LeaveTangent = LeaveTangent / ((NextKey.Time - Key.Time) * DecimalRate);
 		}
 
 		if (bNegateTangents)
@@ -983,9 +998,14 @@ void ImportTransformChannel(const FInterpCurveFloat& Source, FMovieSceneFloatCha
 			LeaveTangent = -LeaveTangent;
 		}
 
-		FFrameNumber KeyTime = (Source.Points[KeyIndex].InVal * DestFrameRate).RoundToFrame();
-		FMatineeImportTools::SetOrAddKey(ChannelInterface, KeyTime, Source.Points[KeyIndex].OutVal, ArriveTangent, LeaveTangent, Source.Points[KeyIndex].InterpMode);
+		FFrameNumber KeyTime = (Key.Time * DestFrameRate).RoundToFrame();
+		FMatineeImportTools::SetOrAddKey(ChannelInterface, KeyTime, Key.Value, ArriveTangent, LeaveTangent,
+			MovieSceneToolHelpers::RichCurveInterpolationToMatineeInterpolation(Key.InterpMode, Key.TangentMode), Key.TangentWeightMode,
+			Key.ArriveTangentWeight, Key.LeaveTangentWeight);
+
 	}
+
+
 	const UMovieSceneUserImportFBXSettings* ImportFBXSettings = GetDefault<UMovieSceneUserImportFBXSettings>();
 	if (ImportFBXSettings->bReduceKeys)
 	{
@@ -1002,13 +1022,13 @@ bool ImportFBXTransform(FString NodeName, FGuid ObjectBinding, UnFbx::FFbxCurves
 	const UMovieSceneUserImportFBXSettings* ImportFBXSettings = GetDefault<UMovieSceneUserImportFBXSettings>();
 
 	// Look for transforms explicitly
-	FInterpCurveFloat Translation[3];
-	FInterpCurveFloat EulerRotation[3];
-	FInterpCurveFloat Scale[3];
+	FRichCurve Translation[3];
+	FRichCurve EulerRotation[3];
+	FRichCurve Scale[3];
 	FTransform DefaultTransform;
 	CurveAPI.GetConvertedTransformCurveData(NodeName, Translation[0], Translation[1], Translation[2], EulerRotation[0], EulerRotation[1], EulerRotation[2], Scale[0], Scale[1], Scale[2], DefaultTransform);
 
-	UMovieScene3DTransformTrack* TransformTrack = InMovieScene->FindTrack<UMovieScene3DTransformTrack>(ObjectBinding); 
+ 	UMovieScene3DTransformTrack* TransformTrack = InMovieScene->FindTrack<UMovieScene3DTransformTrack>(ObjectBinding); 
 	if (!TransformTrack)
 	{
 		InMovieScene->Modify();
@@ -1634,14 +1654,22 @@ bool MovieSceneToolHelpers::ImportFBX(UMovieScene* InMovieScene, ISequencer& InS
 }
 
 
-EInterpCurveMode MovieSceneToolHelpers::RichCurveInterpolationToMatineeInterpolation( ERichCurveInterpMode InterpMode )
+EInterpCurveMode MovieSceneToolHelpers::RichCurveInterpolationToMatineeInterpolation( ERichCurveInterpMode InterpMode, ERichCurveTangentMode TangentMode)
 {
 	switch ( InterpMode )
 	{
 	case ERichCurveInterpMode::RCIM_Constant:
 		return CIM_Constant;
 	case ERichCurveInterpMode::RCIM_Cubic:
-		return CIM_CurveAuto;
+		if (TangentMode == RCTM_Auto)
+		{
+			return CIM_CurveAuto;
+		}
+		else if (TangentMode == RCTM_Break)
+		{
+			return CIM_CurveBreak;
+		}
+		return CIM_CurveUser;  
 	case ERichCurveInterpMode::RCIM_Linear:
 		return CIM_Linear;
 	default:
@@ -1669,7 +1697,7 @@ void MovieSceneToolHelpers::CopyKeyDataToMoveAxis(const TMovieSceneChannel<FMovi
 		FInterpCurvePoint<float>& Point = MoveAxis->FloatTrack.Points[PointIndex];
 		Point.ArriveTangent = Value.Tangent.ArriveTangent;
 		Point.LeaveTangent = Value.Tangent.LeaveTangent;
-		Point.InterpMode = RichCurveInterpolationToMatineeInterpolation(Value.InterpMode);
+		Point.InterpMode = RichCurveInterpolationToMatineeInterpolation(Value.InterpMode, Value.TangentMode);
 	}
 }
 
diff --git a/Engine/Source/Editor/MovieSceneTools/Public/MatineeImportTools.h b/Engine/Source/Editor/MovieSceneTools/Public/MatineeImportTools.h
--- a/Engine/Source/Editor/MovieSceneTools/Public/MatineeImportTools.h
+++ b/Engine/Source/Editor/MovieSceneTools/Public/MatineeImportTools.h
@@ -47,7 +47,8 @@ public:
 	static bool TryConvertMatineeToggleToOutParticleKey( ETrackToggleAction ToggleAction, EParticleKey& OutParticleKey );
 
 	/** Adds a key to a rich curve based on matinee curve key data. */
-	static void SetOrAddKey( TMovieSceneChannel<FMovieSceneFloatValue>& Curve, FFrameNumber Time, float Value, float ArriveTangent, float LeaveTangent, EInterpCurveMode MatineeInterpMode );
+	static void SetOrAddKey( TMovieSceneChannel<FMovieSceneFloatValue>& Curve, FFrameNumber Time, float Value, float ArriveTangent, float LeaveTangent, EInterpCurveMode MatineeInterpMode 
+	, ERichCurveTangentWeightMode WeightedMode = RCTWM_WeightedNone, float ArriveTangentWeight = 0.0f, float LeaveTangentWeight = 0.0f);
 
 	/** Copies keys from a matinee bool track to a sequencer bool track. */
 	static bool CopyInterpBoolTrack( UInterpTrackBoolProp* MatineeBoolTrack, UMovieSceneBoolTrack* BoolTrack );
diff --git a/Engine/Source/Editor/MovieSceneTools/Public/MovieSceneToolHelpers.h b/Engine/Source/Editor/MovieSceneTools/Public/MovieSceneToolHelpers.h
--- a/Engine/Source/Editor/MovieSceneTools/Public/MovieSceneToolHelpers.h
+++ b/Engine/Source/Editor/MovieSceneTools/Public/MovieSceneToolHelpers.h
@@ -202,7 +202,7 @@ public:
 	 * @param InterpMode The rich curve interpolation to convert
 	 * @return The converted matinee interpolation
 	 */
-	static EInterpCurveMode RichCurveInterpolationToMatineeInterpolation( ERichCurveInterpMode InterpMode );
+	static EInterpCurveMode RichCurveInterpolationToMatineeInterpolation( ERichCurveInterpMode InterpMode, ERichCurveTangentMode TangentMode);
 
 	/*
 	 * Copy key data to move axis
diff --git a/Engine/Source/Editor/UnrealEd/Private/Fbx/FbxCurveDataImport.cpp b/Engine/Source/Editor/UnrealEd/Private/Fbx/FbxCurveDataImport.cpp
--- a/Engine/Source/Editor/UnrealEd/Private/Fbx/FbxCurveDataImport.cpp
+++ b/Engine/Source/Editor/UnrealEd/Private/Fbx/FbxCurveDataImport.cpp
@@ -2,12 +2,13 @@
 
 #include "CoreMinimal.h"
 #include "FbxImporter.h"
+#include "Curves/RichCurve.h"
 
 namespace UnFbx {
 
 	FbxNode* GetNodeFromName(const FString& NodeName, FbxNode* NodeToQuery)
 	{
-		if ( !FCString::Strcmp(*NodeName, UTF8_TO_TCHAR(NodeToQuery->GetName())))
+		if (!FCString::Strcmp(*NodeName, UTF8_TO_TCHAR(NodeToQuery->GetName())))
 		{
 			return NodeToQuery;
 		}
@@ -27,7 +28,7 @@ namespace UnFbx {
 
 	FbxNode* GetNodeFromUniqueID(uint64 UniqueID, FbxNode* NodeToQuery)
 	{
-		if ( UniqueID == NodeToQuery->GetUniqueID())
+		if (UniqueID == NodeToQuery->GetUniqueID())
 		{
 			return NodeToQuery;
 		}
@@ -44,7 +45,7 @@ namespace UnFbx {
 
 		return nullptr;
 	}
-	
+
 	void FFbxCurvesAPI::GetAllNodeNameArray(TArray<FString> &AllNodeNames) const
 	{
 		AllNodeNames.Empty(TransformData.Num());
@@ -65,7 +66,7 @@ namespace UnFbx {
 			AnimatedNodeNames.Add(AnimNodeKvp.Value.Name);
 		}
 	}
-	
+
 	void FFbxCurvesAPI::GetNodeAnimatedPropertyNameArray(const FString &NodeName, TArray<FString> &AnimatedPropertyNames) const
 	{
 		AnimatedPropertyNames.Empty();
@@ -162,6 +163,7 @@ namespace UnFbx {
 		}
 	}
 
+	//deprecated
 	void FFbxCurvesAPI::GetCurveData(const FFbxAnimCurveHandle &CurveHandle, FInterpCurveFloat& CurveData, bool bNegative) const
 	{
 		if (CurveHandle.AnimCurve == nullptr)
@@ -206,6 +208,145 @@ namespace UnFbx {
 		}
 	}
 
+	//Similar to function UnFbx::FFbxImporter::ImportCurve in SkeletalMeshEdit but with weighted tangent support.
+	void FFbxCurvesAPI::GetCurveData(const FFbxAnimCurveHandle &CurveHandle, FRichCurve& RichCurve, bool bNegative) const
+	{
+		static float DefaultCurveWeight = FbxAnimCurveDef::sDEFAULT_WEIGHT;
+		FbxAnimCurve* FbxCurve = CurveHandle.AnimCurve;
+		if (FbxCurve)
+		{
+			RichCurve.Reset();
+			for (int32 KeyIndex = 0; KeyIndex < FbxCurve->KeyGetCount(); ++KeyIndex)
+			{
+				FbxAnimCurveKey Key = FbxCurve->KeyGet(KeyIndex);
+				FbxTime KeyTime = Key.GetTime();
+				float Value = bNegative ? -Key.GetValue() : Key.GetValue();
+				FKeyHandle NewKeyHandle = RichCurve.AddKey(KeyTime.GetSecondDouble(), Value, false);
+
+				FbxAnimCurveDef::ETangentMode KeyTangentMode = Key.GetTangentMode();
+				FbxAnimCurveDef::EInterpolationType KeyInterpMode = Key.GetInterpolation();
+				FbxAnimCurveDef::EWeightedMode KeyTangentWeightMode = Key.GetTangentWeightMode();
+
+				ERichCurveInterpMode NewInterpMode = RCIM_Linear;
+				ERichCurveTangentMode NewTangentMode = RCTM_Auto;
+				ERichCurveTangentWeightMode NewTangentWeightMode = RCTWM_WeightedNone;
+
+				float LeaveTangent = 0.f;
+				float ArriveTangent = 0.f;
+				float LeaveTangentWeight = 0.f;
+				float ArriveTangentWeight = 0.f;
+				float ArriveTimeDiff = 0.f;
+				float LeaveTimeDiff = 0.f;
+
+				switch (KeyInterpMode)
+				{
+				case FbxAnimCurveDef::eInterpolationConstant://! Constant value until next key.
+					NewInterpMode = RCIM_Constant;
+					break;
+				case FbxAnimCurveDef::eInterpolationLinear://! Linear progression to next key.
+					NewInterpMode = RCIM_Linear;
+					break;
+				case FbxAnimCurveDef::eInterpolationCubic://! Cubic progression to next key.
+					NewInterpMode = RCIM_Cubic;
+					// get tangents
+					{
+						FbxAnimCurveKey CurKey = CurveHandle.AnimCurve->KeyGet(KeyIndex);
+						float LeftTangent = FbxCurve->KeyGetLeftDerivative(KeyIndex);
+						float RightTangent = FbxCurve->KeyGetRightDerivative(KeyIndex);
+						
+						if (KeyIndex > 0)
+						{
+							ArriveTimeDiff = CurKey.GetTime().GetSecondDouble() - FbxCurve->KeyGetTime(KeyIndex - 1).GetSecondDouble();
+							ArriveTangent = LeftTangent * (ArriveTimeDiff);
+						}
+
+						if (KeyIndex < FbxCurve->KeyGetCount() - 1)
+						{
+							LeaveTimeDiff = FbxCurve->KeyGetTime(KeyIndex + 1).GetSecondDouble() - CurKey.GetTime().GetSecondDouble();
+							LeaveTangent = RightTangent * (LeaveTimeDiff);
+						}
+					}
+					break;
+				}
+				if (KeyTangentMode & FbxAnimCurveDef::eTangentGenericBreak)
+				{
+				 	NewTangentMode = RCTM_Break;
+				}
+				else if (KeyTangentMode &  FbxAnimCurveDef::eTangentAuto) //break and auto are exclusive
+				{
+					NewTangentMode = RCTM_Auto;
+				}
+				else
+				{
+				 	NewTangentMode = RCTM_User;
+				}
+
+				switch (KeyTangentWeightMode)
+				{
+				case FbxAnimCurveDef::eWeightedNone://! Tangent has default weights of 0.333; we define this state as not weighted.
+					LeaveTangentWeight = ArriveTangentWeight = DefaultCurveWeight;
+					NewTangentWeightMode = RCTWM_WeightedNone;
+					break;
+				case FbxAnimCurveDef::eWeightedRight: //! Right tangent is weighted.
+					NewTangentWeightMode = RCTWM_WeightedLeave;
+					LeaveTangentWeight = Key.GetDataFloat(FbxAnimCurveDef::eRightWeight);
+					ArriveTangentWeight = DefaultCurveWeight;
+					break;
+				case FbxAnimCurveDef::eWeightedNextLeft://! Left tangent is weighted.
+					NewTangentWeightMode = RCTWM_WeightedArrive;
+					LeaveTangentWeight = DefaultCurveWeight;
+					if (KeyIndex > 0)
+					{
+						FbxAnimCurveKey PrevKey = FbxCurve->KeyGet(KeyIndex - 1);
+						ArriveTangentWeight = PrevKey.GetDataFloat(FbxAnimCurveDef::eNextLeftWeight);              
+					}
+					else
+					{
+						ArriveTangentWeight = 0.f;
+					}
+					break;
+				case FbxAnimCurveDef::eWeightedAll://! Both left and right tangents are weighted.
+					NewTangentWeightMode = RCTWM_WeightedBoth;
+					LeaveTangentWeight = Key.GetDataFloat(FbxAnimCurveDef::eRightWeight);
+					if (KeyIndex > 0)
+					{
+ 						FbxAnimCurveKey PrevKey = FbxCurve->KeyGet(KeyIndex - 1);
+						ArriveTangentWeight = PrevKey.GetDataFloat(FbxAnimCurveDef::eNextLeftWeight);
+					}
+					else
+				 	{
+						ArriveTangentWeight = 0.f;
+					}
+					break;
+				}
+				RichCurve.SetKeyInterpMode(NewKeyHandle, NewInterpMode);
+				RichCurve.SetKeyTangentMode(NewKeyHandle, NewTangentMode);
+				RichCurve.SetKeyTangentWeightMode(NewKeyHandle, NewTangentWeightMode);
+
+				FRichCurveKey& NewKey = RichCurve.GetKey(NewKeyHandle);
+				NewKey.ArriveTangent = ArriveTangent;
+				NewKey.LeaveTangent = LeaveTangent;
+				//Tangent Weights in FBX/Maya are normalized X (Time) values.
+				//Our weights are the length of hypontenuse. So here we do the
+				//conversion. Note that Specificed Tangent is already Tangent * Time_Difference;
+				//so we just need to scale it by the normalized weight value.
+				if (!FMath::IsNearlyZero(ArriveTangentWeight))
+				{
+					const float Y = ArriveTangent * ArriveTangentWeight;
+					ArriveTangentWeight = FMath::Sqrt(Y*Y + ArriveTimeDiff * ArriveTimeDiff);
+				}
+				NewKey.ArriveTangentWeight = ArriveTangentWeight;
+				if (!FMath::IsNearlyZero(LeaveTangentWeight))
+				{
+					const float Y = LeaveTangent * LeaveTangentWeight;
+					LeaveTangentWeight = FMath::Sqrt(Y*Y + LeaveTimeDiff * LeaveTimeDiff);
+				}
+				NewKey.LeaveTangentWeight = LeaveTangentWeight;
+			}
+		}
+	}
+
+
 	void FFbxCurvesAPI::GetBakeCurveData(const FFbxAnimCurveHandle &CurveHandle, TArray<float>& CurveData, float PeriodTime, float StartTime /*= 0.0f*/, float StopTime /*= -1.0f*/, bool bNegative /*= false*/) const
 	{
 		//Make sure the parameter are ok
@@ -232,6 +373,7 @@ namespace UnFbx {
 		}
 	}
 
+	//deprecrated
 	void FFbxCurvesAPI::GetCurveData(const FString& NodeName, const FString& PropertyName, int32 ChannelIndex, int32 CompositeIndex, FInterpCurveFloat& CurveData, bool bNegative) const
 	{
 		FFbxAnimCurveHandle CurveHandle;
@@ -246,6 +388,20 @@ namespace UnFbx {
 		}
 	}
 
+	void FFbxCurvesAPI::GetCurveData(const FString& NodeName, const FString& PropertyName, int32 ChannelIndex, int32 CompositeIndex, FRichCurve& RichCurve, bool bNegative) const
+	{
+		FFbxAnimCurveHandle CurveHandle;
+		GetCurveHandle(NodeName, PropertyName, ChannelIndex, CompositeIndex, CurveHandle);
+		if (CurveHandle.AnimCurve != nullptr)
+		{
+			GetCurveData(CurveHandle, RichCurve, bNegative);
+		}
+		else
+		{
+			RichCurve.Reset();
+		}
+	}
+
 	void FFbxCurvesAPI::GetBakeCurveData(const FString& NodeName, const FString& PropertyName, int32 ChannelIndex, int32 CompositeIndex, TArray<float>& CurveData, float PeriodTime, float StartTime /*= 0.0f*/, float StopTime /*= -1.0f*/, bool bNegative /*= false*/) const
 	{
 		FFbxAnimCurveHandle CurveHandle;
@@ -442,6 +598,86 @@ namespace UnFbx {
 			DefaultTransform = TransformData[Node->GetUniqueID()];
 		}
 	}
+
+
+
+	void FFbxCurvesAPI::GetConvertedTransformCurveData(const FString& NodeName, FRichCurve& TranslationX, FRichCurve& TranslationY, FRichCurve& TranslationZ,
+		FRichCurve& EulerRotationX, FRichCurve& EulerRotationY, FRichCurve& EulerRotationZ,
+		FRichCurve& ScaleX, FRichCurve& ScaleY, FRichCurve& ScaleZ,
+		FTransform& DefaultTransform) const
+	{
+		
+		for (TPair< uint64, FFbxAnimNodeHandle> AnimNodeKvp : CurvesData)
+		{
+			const FFbxAnimNodeHandle& AnimNodeHandle = AnimNodeKvp.Value;
+			if (AnimNodeHandle.Name.Compare(NodeName) == 0)
+			{
+				bool bIsCamera = AnimNodeHandle.AttributeType == FbxNodeAttribute::eCamera;
+				bool bIsLight = AnimNodeHandle.AttributeType == FbxNodeAttribute::eLight;
+				FFbxAnimCurveHandle TransformCurves[9];
+				for (auto NodePropertyKvp : AnimNodeHandle.NodeProperties)
+				{
+					FFbxAnimPropertyHandle& AnimPropertyHandle = NodePropertyKvp.Value;
+					for (FFbxAnimCurveHandle& CurveHandle : AnimPropertyHandle.CurveHandles)
+					{
+						if (CurveHandle.CurveType != FFbxAnimCurveHandle::NotTransform)
+						{
+							TransformCurves[(int32)(CurveHandle.CurveType)] = CurveHandle;
+						}
+					}
+				}
+
+				GetCurveData(TransformCurves[0], TranslationX, false);
+				GetCurveData(TransformCurves[1], TranslationY, true);
+				GetCurveData(TransformCurves[2], TranslationZ, false);
+
+				GetCurveData(TransformCurves[3], EulerRotationX, false);
+				GetCurveData(TransformCurves[4], EulerRotationY, true);
+				GetCurveData(TransformCurves[5], EulerRotationZ, true);
+
+				GetCurveData(TransformCurves[6], ScaleX, false);
+				GetCurveData(TransformCurves[7], ScaleY, false);
+				GetCurveData(TransformCurves[8], ScaleZ, false);
+
+				if (bIsCamera || bIsLight)
+				{
+					{//extra scope since we can't reset Key Iterators
+						//need to convert rotations to unreal space.  Uses previous FInterpCurvePoint implementation that goes through the minimal number
+						//of curve keys and sets them together. Obviously if the keys are not at the same times exactly this won't work.
+						auto EulerRotXIt = EulerRotationX.GetKeyHandleIterator();
+						auto EulerRotYIt = EulerRotationY.GetKeyHandleIterator();
+						auto EulerRotZIt = EulerRotationZ.GetKeyHandleIterator();
+
+
+						while (EulerRotXIt && EulerRotYIt && EulerRotZIt)
+						{
+
+							float Pitch = EulerRotationY.GetKeyValue(EulerRotYIt.Key());
+							float Yaw = EulerRotationZ.GetKeyValue(EulerRotZIt.Key());
+							float Roll = EulerRotationX.GetKeyValue(EulerRotXIt.Key());;
+							ConvertRotationToUnreal(Roll, Pitch, Yaw, bIsCamera, bIsLight);
+							EulerRotationX.SetKeyValue(EulerRotXIt.Key(), Roll, false);
+							EulerRotationY.SetKeyValue(EulerRotYIt.Key(), Pitch, false);
+							EulerRotationZ.SetKeyValue(EulerRotZIt.Key(), Yaw, false);
+
+							++EulerRotXIt;
+							++EulerRotYIt;
+							++EulerRotZIt;
+						}
+					}
+			
+				}
+					
+			}
+		}
+
+		FbxNode* Node = GetNodeFromName(NodeName, Scene->GetRootNode());
+		if (Node)
+		{
+			DefaultTransform = TransformData[Node->GetUniqueID()];
+		}
+		
+	}
 	
 	//////////////////////////////////////////////////////////////////////////
 	// FFbxImporter: Curve Extraction Implementation
diff --git a/Engine/Source/Editor/UnrealEd/Private/Fbx/FbxMainExport.cpp b/Engine/Source/Editor/UnrealEd/Private/Fbx/FbxMainExport.cpp
--- a/Engine/Source/Editor/UnrealEd/Private/Fbx/FbxMainExport.cpp
+++ b/Engine/Source/Editor/UnrealEd/Private/Fbx/FbxMainExport.cpp
@@ -2298,7 +2298,8 @@ void FFbxExporter::ExportAnimatedFloat(FbxProperty* FbxProperty, FInterpCurveFlo
 }
 
 
-void RichCurveInterpolationToFbxInterpolation(ERichCurveInterpMode InInterpolation, ERichCurveTangentMode InTangentMode, FbxAnimCurveDef::EInterpolationType &OutInterpolation, FbxAnimCurveDef::ETangentMode &OutTangentMode)
+void RichCurveInterpolationToFbxInterpolation(ERichCurveInterpMode InInterpolation, ERichCurveTangentMode InTangentMode, ERichCurveTangentWeightMode InTangentWeightMode,
+	FbxAnimCurveDef::EInterpolationType &OutInterpolation, FbxAnimCurveDef::ETangentMode &OutTangentMode, FbxAnimCurveDef::EWeightedMode &OutTangentWeightMode)
 {
 	if (InInterpolation == ERichCurveInterpMode::RCIM_Cubic)
 	{
@@ -2310,28 +2311,53 @@ void RichCurveInterpolationToFbxInterpolation(ERichCurveInterpMode InInterpolati
 		{
 			OutTangentMode = FbxAnimCurveDef::eTangentBreak;
 		}
+		else if (InTangentMode == RCTM_Auto)
+		{
+			OutTangentMode = FbxAnimCurveDef::eTangentAuto;
+		}
+		
+		switch (InTangentWeightMode)
+		{
+		case ERichCurveTangentWeightMode::RCTWM_WeightedBoth:
+			OutTangentWeightMode = FbxAnimCurveDef::eWeightedAll;
+			break;
+		case ERichCurveTangentWeightMode::RCTWM_WeightedArrive:
+			OutTangentWeightMode = FbxAnimCurveDef::eWeightedNextLeft;
+			break;
+		case ERichCurveTangentWeightMode::RCTWM_WeightedLeave:
+			OutTangentWeightMode = FbxAnimCurveDef::eWeightedRight;
+			break;
+		case ERichCurveTangentWeightMode::RCTWM_WeightedNone:
+		default:
+			OutTangentWeightMode = FbxAnimCurveDef::eWeightedNone;
+			break;
+
+		};
 	}
 	else if (InInterpolation == ERichCurveInterpMode::RCIM_Linear)
 	{
 		OutInterpolation = FbxAnimCurveDef::eInterpolationLinear;
 		OutTangentMode = FbxAnimCurveDef::eTangentUser;
+		OutTangentWeightMode = FbxAnimCurveDef::eWeightedNone;
 	}
 	else if (InInterpolation == ERichCurveInterpMode::RCIM_Constant)
 	{
 		OutInterpolation = FbxAnimCurveDef::eInterpolationConstant;
 		OutTangentMode = (FbxAnimCurveDef::ETangentMode)FbxAnimCurveDef::eConstantStandard;
+		OutTangentWeightMode = FbxAnimCurveDef::eWeightedNone;
 	}
 	else
 	{
 		OutInterpolation = FbxAnimCurveDef::eInterpolationCubic;
 		OutTangentMode = FbxAnimCurveDef::eTangentUser;
+		OutTangentWeightMode = FbxAnimCurveDef::eWeightedNone;
 	}
 }
 
 void FFbxExporter::ExportChannelToFbxCurve(FbxAnimCurve& InFbxCurve, const FMovieSceneFloatChannel& InChannel, FFrameRate TickResolution, ERichCurveValueMode ValueMode, bool bNegative)
 {
 	const float NegateFactor = bNegative ? -1.f : 1.f;
-
+	const float kOneThird = 1.f / 3.f;
 	InFbxCurve.KeyModifyBegin();
 
 	TArrayView<const FFrameNumber>          Times  = InChannel.GetTimes();
@@ -2339,40 +2365,59 @@ void FFbxExporter::ExportChannelToFbxCurve(FbxAnimCurve& InFbxCurve, const FMovi
 
 	for (int32 Index = 0; Index < Times.Num(); ++Index)
 	{
-		FFrameNumber          KeyTime  = Times[Index];
-		FMovieSceneFloatValue KeyValue = Values[Index];
+		const FFrameNumber          KeyTime  = Times[Index];
+		const FMovieSceneFloatValue KeyValue = Values[Index];
 
 		const float Value = (ValueMode == ERichCurveValueMode::Fov ? DefaultCamera->ComputeFocalLength( KeyValue.Value ) : KeyValue.Value) * NegateFactor;
 
 		FbxTime FbxTime;
 		FbxAnimCurveKey FbxKey;
-		FbxTime.SetSecondDouble(KeyTime / TickResolution);
+		const double KeyTimeSeconds = KeyTime / TickResolution;
+		FbxTime.SetSecondDouble(KeyTimeSeconds);
 
-		int FbxKeyIndex = InFbxCurve.KeyAdd(FbxTime);
+		const int FbxKeyIndex = InFbxCurve.KeyAdd(FbxTime);
 
 		FbxAnimCurveDef::EInterpolationType Interpolation = FbxAnimCurveDef::eInterpolationCubic;
 		FbxAnimCurveDef::ETangentMode Tangent = FbxAnimCurveDef::eTangentAuto;
-
-		RichCurveInterpolationToFbxInterpolation(KeyValue.InterpMode, KeyValue.TangentMode, Interpolation, Tangent);
+		FbxAnimCurveDef::EWeightedMode WeightedMode = FbxAnimCurveDef::eWeightedNone;
+		RichCurveInterpolationToFbxInterpolation(KeyValue.InterpMode, KeyValue.TangentMode, KeyValue.Tangent.TangentWeightMode, Interpolation, Tangent, WeightedMode);
 
 		if (Interpolation == FbxAnimCurveDef::eInterpolationCubic)
 		{
 			if (Index < Times.Num() - 1)
 			{
+				float LeaveTangentWeight = kOneThird;
+				float NextArriveTangentWeight = kOneThird;
+				const double  NextTime = Times[Index + 1] / TickResolution;
+				const float TimeDiff = static_cast<float>(NextTime - KeyTimeSeconds);
+				
 				float LeaveTangent = KeyValue.Tangent.LeaveTangent * TickResolution.AsDecimal();
-				float NextArriveTangent = Values[Index+1].Tangent.ArriveTangent * TickResolution.AsDecimal();
+				float NextArriveTangent = Values[Index + 1].Tangent.ArriveTangent * TickResolution.AsDecimal();
 
+				//Need to convert Ue4 tangent weight which is the length of the hypotenuse to FBX normalized X(time) weight
+				if (WeightedMode == FbxAnimCurveDef::eWeightedAll || FbxAnimCurveDef::eWeightedRight)
+				{
+					const float XVal = FMath::Sqrt((KeyValue.Tangent.LeaveTangentWeight * KeyValue.Tangent.LeaveTangentWeight) / (1.0f + LeaveTangent * LeaveTangent));
+					LeaveTangentWeight = XVal / TimeDiff;
+				}
+
+				//make sure next tangent is weighted else use default weight
+				if ((Values[Index + 1].Tangent.TangentWeightMode == ERichCurveTangentWeightMode::RCTWM_WeightedBoth || Values[Index + 1].Tangent.TangentWeightMode == ERichCurveTangentWeightMode::RCTWM_WeightedArrive) )
+				{
+					const float XVal = FMath::Sqrt((Values[Index + 1].Tangent.ArriveTangentWeight * Values[Index + 1].Tangent.ArriveTangentWeight) / (1.0f + NextArriveTangent * NextArriveTangent));
+					NextArriveTangentWeight = XVal / TimeDiff;
+				}
 				if (bNegative)
 				{
 					LeaveTangent = -LeaveTangent;
 					NextArriveTangent = -NextArriveTangent;
 				}
 
-				InFbxCurve.KeySet(FbxKeyIndex, FbxTime, Value, Interpolation, Tangent, LeaveTangent, NextArriveTangent );
+				InFbxCurve.KeySet(FbxKeyIndex, FbxTime, Value, Interpolation, Tangent, LeaveTangent, NextArriveTangent, WeightedMode, LeaveTangentWeight, NextArriveTangentWeight );
 			}
 			else
 			{
-				InFbxCurve.KeySet(FbxKeyIndex, FbxTime, Value, Interpolation, Tangent);
+				InFbxCurve.KeySet(FbxKeyIndex, FbxTime, Value, Interpolation, Tangent, 0.f, 0.f, WeightedMode);
 			}
 		}
 		else
diff --git a/Engine/Source/Editor/UnrealEd/Public/FbxImporter.h b/Engine/Source/Editor/UnrealEd/Public/FbxImporter.h
--- a/Engine/Source/Editor/UnrealEd/Public/FbxImporter.h
+++ b/Engine/Source/Editor/UnrealEd/Public/FbxImporter.h
@@ -340,20 +340,37 @@ public:
 	UNREALED_API void GetAllNodeNameArray(TArray<FString> &AllNodeNames) const;
 	UNREALED_API void GetAnimatedNodeNameArray(TArray<FString> &AnimatedNodeNames) const;
 	UNREALED_API void GetNodeAnimatedPropertyNameArray(const FString &NodeName, TArray<FString> &AnimatedPropertyNames) const;
+
+	DEPRECATED(4.21, "Please use FRichCurve version instead to get tangent weight support")
 	UNREALED_API void GetCurveData(const FString& NodeName, const FString& PropertyName, int32 ChannelIndex, int32 CompositeIndex, FInterpCurveFloat& CurveData, bool bNegative) const;
+	
+	UNREALED_API void GetCurveData(const FString& NodeName, const FString& PropertyName, int32 ChannelIndex, int32 CompositeIndex, FRichCurve& CurveData, bool bNegative) const;
+
+	
 	UNREALED_API void GetBakeCurveData(const FString& NodeName, const FString& PropertyName, int32 ChannelIndex, int32 CompositeIndex, TArray<float>& CurveData, float PeriodTime, float StartTime = 0.0f, float StopTime= -1.0f, bool bNegative = false) const;
 
 	//Handle API
 	UNREALED_API void GetAllNodePropertyCurveHandles(const FString& NodeName, const FString& PropertyName, TArray<FFbxAnimCurveHandle> &PropertyCurveHandles) const;
 	UNREALED_API void GetCurveHandle(const FString& NodeName, const FString& PropertyName, int32 ChannelIndex, int32 CompositeIndex, FFbxAnimCurveHandle &CurveHandle) const;
+	
+	DEPRECATED(4.21, "Please use FRichCurve version instead to get tangent weight support")
 	UNREALED_API void GetCurveData(const FFbxAnimCurveHandle &CurveHandle, FInterpCurveFloat& CurveData, bool bNegative) const;
+	
+	UNREALED_API void GetCurveData(const FFbxAnimCurveHandle &CurveHandle, FRichCurve& CurveData, bool bNegative) const;
+
 	UNREALED_API void GetBakeCurveData(const FFbxAnimCurveHandle &CurveHandle, TArray<float>& CurveData, float PeriodTime, float StartTime = 0.0f, float StopTime = -1.0f, bool bNegative = false) const;
 
 	//Conversion API
+	DEPRECATED(4.21, "Please use FRichCurve version instead to get tangent weight support")
 	UNREALED_API void GetConvertedTransformCurveData(const FString& NodeName, FInterpCurveFloat& TranslationX, FInterpCurveFloat& TranslationY, FInterpCurveFloat& TranslationZ,
 													 FInterpCurveFloat& EulerRotationX, FInterpCurveFloat& EulerRotationY, FInterpCurveFloat& EulerRotationZ, 
 													 FInterpCurveFloat& ScaleX, FInterpCurveFloat& ScaleY, FInterpCurveFloat& ScaleZ,
 													 FTransform& DefaultTransform) const;
+	
+	UNREALED_API void GetConvertedTransformCurveData(const FString& NodeName, FRichCurve& TranslationX, FRichCurve& TranslationY, FRichCurve& TranslationZ,
+		FRichCurve& EulerRotationX, FRichCurve& EulerRotationY, FRichCurve& EulerRotationZ,
+		FRichCurve& ScaleX, FRichCurve& ScaleY, FRichCurve& ScaleZ,
+		FTransform& DefaultTransform) const;
 
 	FbxScene* Scene;
 	TMap<uint64, FFbxAnimNodeHandle> CurvesData;
