Fix for deprecrated warnings caused by deprecated functions calling other deprecated functions. Added pragma's to disable the warnings. Noticed it in other parts of the code also.

diff --git a/Engine/Source/Editor/UnrealEd/Private/Fbx/FbxCurveDataImport.cpp b/Engine/Source/Editor/UnrealEd/Private/Fbx/FbxCurveDataImport.cpp
--- a/Engine/Source/Editor/UnrealEd/Private/Fbx/FbxCurveDataImport.cpp
+++ b/Engine/Source/Editor/UnrealEd/Private/Fbx/FbxCurveDataImport.cpp
@@ -380,7 +380,9 @@ namespace UnFbx {
 		GetCurveHandle(NodeName, PropertyName, ChannelIndex, CompositeIndex, CurveHandle);
 		if (CurveHandle.AnimCurve != nullptr)
 		{
+#pragma warning(disable : 4996) // 'function' was declared deprecated
 			GetCurveData(CurveHandle, CurveData, bNegative);
+#pragma warning(default : 4996) // 'function' was declared deprecated
 		}
 		else
 		{
@@ -516,6 +518,8 @@ namespace UnFbx {
 					}
 				}
 
+#pragma warning(disable : 4996) // 'function' was declared deprecated
+
 				GetCurveData(TransformCurves[0], TranslationX, false);
 				GetCurveData(TransformCurves[1], TranslationY, true);
 				GetCurveData(TransformCurves[2], TranslationZ, false);
@@ -527,7 +531,7 @@ namespace UnFbx {
 				GetCurveData(TransformCurves[6], ScaleX, false);
 				GetCurveData(TransformCurves[7], ScaleY, false);
 				GetCurveData(TransformCurves[8], ScaleZ, false);
-
+#pragma warning(default : 4996) // 'function' was declared deprecated
 				if (bIsCamera || bIsLight)
 				{
 					int32 CurvePointNum = FMath::Min3<int32>(EulerRotationX.Points.Num(), EulerRotationY.Points.Num(), EulerRotationZ.Points.Num());
