Fix static analysis warning

diff --git a/Engine/Source/Editor/UnrealEd/Private/Fbx/FbxMainExport.cpp b/Engine/Source/Editor/UnrealEd/Private/Fbx/FbxMainExport.cpp
--- a/Engine/Source/Editor/UnrealEd/Private/Fbx/FbxMainExport.cpp
+++ b/Engine/Source/Editor/UnrealEd/Private/Fbx/FbxMainExport.cpp
@@ -2395,7 +2395,7 @@ void FFbxExporter::ExportChannelToFbxCurve(FbxAnimCurve& InFbxCurve, const FMovi
 				float NextArriveTangent = Values[Index + 1].Tangent.ArriveTangent * TickResolution.AsDecimal();
 
 				//Need to convert Ue4 tangent weight which is the length of the hypotenuse to FBX normalized X(time) weight
-				if (WeightedMode == FbxAnimCurveDef::eWeightedAll || FbxAnimCurveDef::eWeightedRight)
+				if (WeightedMode == FbxAnimCurveDef::eWeightedAll || WeightedMode == FbxAnimCurveDef::eWeightedRight)
 				{
 					const float XVal = FMath::Sqrt((KeyValue.Tangent.LeaveTangentWeight * KeyValue.Tangent.LeaveTangentWeight) / (1.0f + LeaveTangent * LeaveTangent));
 					LeaveTangentWeight = XVal / TimeDiff;
