Fix for regression when we added  imported tangent weights.

There was an Euler Filter to make sure the angles stayed consistent. It seems to still be needed since we could get flips from FBX resampling.
Was removed from changelist 4062604. Just putting it back in.  Need this in before the next for for the bad maya import.

diff --git a/Engine/Source/Editor/UnrealEd/Private/Fbx/FbxCurveDataImport.cpp b/Engine/Source/Editor/UnrealEd/Private/Fbx/FbxCurveDataImport.cpp
--- a/Engine/Source/Editor/UnrealEd/Private/Fbx/FbxCurveDataImport.cpp
+++ b/Engine/Source/Editor/UnrealEd/Private/Fbx/FbxCurveDataImport.cpp
@@ -610,7 +610,7 @@ namespace UnFbx {
 		FRichCurve& ScaleX, FRichCurve& ScaleY, FRichCurve& ScaleZ,
 		FTransform& DefaultTransform) const
 	{
-		
+
 		for (TPair< uint64, FFbxAnimNodeHandle> AnimNodeKvp : CurvesData)
 		{
 			const FFbxAnimNodeHandle& AnimNodeHandle = AnimNodeKvp.Value;
@@ -655,7 +655,6 @@ namespace UnFbx {
 
 						while (EulerRotXIt && EulerRotYIt && EulerRotZIt)
 						{
-
 							float Pitch = EulerRotationY.GetKeyValue(EulerRotYIt.Key());
 							float Yaw = EulerRotationZ.GetKeyValue(EulerRotZIt.Key());
 							float Roll = EulerRotationX.GetKeyValue(EulerRotXIt.Key());;
@@ -669,9 +668,62 @@ namespace UnFbx {
 							++EulerRotZIt;
 						}
 					}
-			
 				}
-					
+				if (bIsCamera)
+				{
+					// The RichCurve code doesn't differentiate between angles and other data, so an interpolation from 179 to -179
+					// will cause the camera to rotate all the way around through 0 degrees.  So here we make a second pass over the 
+					// Euler track to convert the angles into a more interpolation-friendly format.  
+					float CurrentAngleOffset[3] = { 0.f, 0.f, 0.f };
+				
+					auto EulerRotXIt = EulerRotationX.GetKeyHandleIterator();
+					auto EulerRotYIt = EulerRotationY.GetKeyHandleIterator();
+					auto EulerRotZIt = EulerRotationZ.GetKeyHandleIterator();
+
+					FVector PreviousOutVal;
+					FVector CurrentOutVal;
+					bool bFirst = true;
+					while (EulerRotXIt && EulerRotYIt && EulerRotZIt)
+					{
+						float X = EulerRotationX.GetKeyValue(EulerRotXIt.Key());;
+						float Y = EulerRotationY.GetKeyValue(EulerRotYIt.Key());
+						float Z = EulerRotationZ.GetKeyValue(EulerRotZIt.Key());
+
+						if (!bFirst)
+						{
+							PreviousOutVal = CurrentOutVal;
+							CurrentOutVal = FVector(X, Y, Z);
+						}
+						else
+						{
+							CurrentOutVal = FVector(X, Y, Z);
+							bFirst = false;
+						}
+						
+						for (int32 AxisIndex = 0; AxisIndex < 3; ++AxisIndex)
+						{
+							float DeltaAngle = (CurrentOutVal[AxisIndex] + CurrentAngleOffset[AxisIndex]) - PreviousOutVal[AxisIndex];
+
+							if (DeltaAngle >= 180)
+							{
+								CurrentAngleOffset[AxisIndex] -= 360;
+							}
+							else if (DeltaAngle <= -180)
+							{
+								CurrentAngleOffset[AxisIndex] += 360;
+							}
+
+							CurrentOutVal[AxisIndex] += CurrentAngleOffset[AxisIndex];
+						}
+						EulerRotationX.SetKeyValue(EulerRotXIt.Key(), CurrentOutVal.X, false);
+						EulerRotationY.SetKeyValue(EulerRotYIt.Key(), CurrentOutVal.Y, false);
+						EulerRotationZ.SetKeyValue(EulerRotZIt.Key(), CurrentOutVal.Z, false);
+
+						++EulerRotXIt;
+						++EulerRotYIt;
+						++EulerRotZIt;
+					}
+				}
 			}
 		}
 
