Sequencer: Export manual focus distance and ensure that the property is animatable.

diff --git a/Engine/Source/Editor/UnrealEd/Private/Fbx/FbxMainExport.cpp b/Engine/Source/Editor/UnrealEd/Private/Fbx/FbxMainExport.cpp
--- a/Engine/Source/Editor/UnrealEd/Private/Fbx/FbxMainExport.cpp
+++ b/Engine/Source/Editor/UnrealEd/Private/Fbx/FbxMainExport.cpp
@@ -2602,6 +2602,10 @@ void FFbxExporter::ExportLevelSequenceFloatTrack( FbxNode& FbxActor, UMovieScene
 	{
 		Property = FbxActor.FindProperty( "UE_MotionBlur_Amount", false );
 	}
+	else if ( PropertyName == "FocusSettings.ManualFocusDistance" && FbxCamera )
+	{
+		Property = FbxCamera->FocusDistance;
+	}
 
 	if (Property == 0)
 	{
@@ -2615,6 +2619,9 @@ void FFbxExporter::ExportLevelSequenceFloatTrack( FbxNode& FbxActor, UMovieScene
 		return;
 	}
 
+	// Ensure that the property is animatable so that GetCurveNode succeeds
+	Property.ModifyFlag(FbxPropertyFlags::eAnimatable, true);
+
 	FbxAnimCurve* AnimCurve = FbxAnimCurve::Create( Scene, "" );
 	FbxAnimCurveNode* CurveNode = Property.GetCurveNode( true );
 	if ( !CurveNode )
