Sequencer: Don't pass auto tangent mode to FBX since it has a different implementation. Default to user as before.

diff --git a/Engine/Source/Editor/UnrealEd/Private/Fbx/FbxMainExport.cpp b/Engine/Source/Editor/UnrealEd/Private/Fbx/FbxMainExport.cpp
--- a/Engine/Source/Editor/UnrealEd/Private/Fbx/FbxMainExport.cpp
+++ b/Engine/Source/Editor/UnrealEd/Private/Fbx/FbxMainExport.cpp
@@ -2312,10 +2312,6 @@ void RichCurveInterpolationToFbxInterpolation(ERichCurveInterpMode InInterpolati
 		{
 			OutTangentMode = FbxAnimCurveDef::eTangentBreak;
 		}
-		else if (InTangentMode == RCTM_Auto)
-		{
-			OutTangentMode = FbxAnimCurveDef::eTangentAuto;
-		}
 		
 		switch (InTangentWeightMode)
 		{
