Sequencer: Added ability to turn off saving animations sequences when exporting level meshes so we don't save the animation sequence twice when saving FBX from the sequencer, which was causes a bug in this file since the times weren't in sync.

diff --git a/Engine/Source/Editor/Sequencer/Private/Sequencer.cpp b/Engine/Source/Editor/Sequencer/Private/Sequencer.cpp
--- a/Engine/Source/Editor/Sequencer/Private/Sequencer.cpp
+++ b/Engine/Source/Editor/Sequencer/Private/Sequencer.cpp
@@ -7972,7 +7972,8 @@ void FSequencer::ExportFBXInternal(const FString& ExportFilename, TArray<FGuid>&
 
 			// Export the persistent level and all of it's actors
 			UWorld* World = Cast<UWorld>(GetPlaybackContext());
-			Exporter->ExportLevelMesh(World->PersistentLevel, bSelectedOnly, NodeNameAdapter);
+			const bool bSaveAnimSeq = false; //force off saving any AnimSequences since this can conflict when we export the level sequence animations.
+			Exporter->ExportLevelMesh(World->PersistentLevel, bSelectedOnly, NodeNameAdapter, bSaveAnimSeq);
 
 			// Export streaming levels and actors
 			for (int32 CurLevelIndex = 0; CurLevelIndex < World->GetNumLevels(); ++CurLevelIndex)
@@ -7980,7 +7981,7 @@ void FSequencer::ExportFBXInternal(const FString& ExportFilename, TArray<FGuid>&
 				ULevel* CurLevel = World->GetLevel(CurLevelIndex);
 				if (CurLevel != NULL && CurLevel != (World->PersistentLevel))
 				{
-					Exporter->ExportLevelMesh(CurLevel, bSelectedOnly, NodeNameAdapter);
+					Exporter->ExportLevelMesh(CurLevel, bSelectedOnly, NodeNameAdapter, bSaveAnimSeq);
 				}
 			}
 
diff --git a/Engine/Source/Editor/UnrealEd/Private/Fbx/FbxMainExport.cpp b/Engine/Source/Editor/UnrealEd/Private/Fbx/FbxMainExport.cpp
--- a/Engine/Source/Editor/UnrealEd/Private/Fbx/FbxMainExport.cpp
+++ b/Engine/Source/Editor/UnrealEd/Private/Fbx/FbxMainExport.cpp
@@ -416,7 +416,7 @@ static void SortActorsHierarchy(TArray<AActor*>& Actors)
 /**
  * Exports the basic scene information to the FBX document.
  */
-void FFbxExporter::ExportLevelMesh( ULevel* InLevel, bool bSelectedOnly, INodeNameAdapter& NodeNameAdapter )
+void FFbxExporter::ExportLevelMesh( ULevel* InLevel, bool bSelectedOnly, INodeNameAdapter& NodeNameAdapter, bool bSaveAnimSeq)
 {
 	if (InLevel == NULL)
 	{
@@ -476,7 +476,7 @@ void FFbxExporter::ExportLevelMesh( ULevel* InLevel, bool bSelectedOnly, INodeNa
 		if (bIsBlueprintClass)
 		{
 			// Export blueprint actors and all their components.
-			ExportActor(Actor, true, NodeNameAdapter);
+			ExportActor(Actor, true, NodeNameAdapter, bSaveAnimSeq);
 		}
 		else if (Actor->IsA(ALight::StaticClass()))
 		{
@@ -497,7 +497,7 @@ void FFbxExporter::ExportLevelMesh( ULevel* InLevel, bool bSelectedOnly, INodeNa
 		}
 		else if (Actor->IsA(AEmitter::StaticClass()))
 		{
-			ExportActor( Actor, false, NodeNameAdapter ); // Just export the placement of the particle emitter.
+			ExportActor( Actor, false, NodeNameAdapter, bSaveAnimSeq ); // Just export the placement of the particle emitter.
 		}
 		else if(Actor->IsA(ACameraActor::StaticClass()))
 		{
@@ -506,7 +506,7 @@ void FFbxExporter::ExportLevelMesh( ULevel* InLevel, bool bSelectedOnly, INodeNa
 		else
 		{
 			// Export any other type of actors and all their components.
-			ExportActor(Actor, true, NodeNameAdapter);
+			ExportActor(Actor, true, NodeNameAdapter, bSaveAnimSeq);
 		}
 	}
 }
@@ -1580,7 +1580,7 @@ bool FFbxExporter::ExportLevelSequence( UMovieScene* MovieScene, const TArray<FG
  * Exports a scene node with the placement indicated by a given actor.
  * This scene node will always have two transformations: one translation vector and one Euler rotation.
  */
-FbxNode* FFbxExporter::ExportActor(AActor* Actor, bool bExportComponents, INodeNameAdapter& NodeNameAdapter )
+FbxNode* FFbxExporter::ExportActor(AActor* Actor, bool bExportComponents, INodeNameAdapter& NodeNameAdapter, bool bSaveAnimSeq )
 {
 	// Verify that this actor isn't already exported, create a structure for it
 	// and buffer it.
@@ -1770,7 +1770,7 @@ FbxNode* FFbxExporter::ExportActor(AActor* Actor, bool bExportComponents, INodeN
 				}
 				else if (SkelMeshComp && SkelMeshComp->SkeletalMesh)
 				{
-					ExportSkeletalMeshComponent(SkelMeshComp, *SkelMeshComp->GetName(), ExportNode);
+					ExportSkeletalMeshComponent(SkelMeshComp, *SkelMeshComp->GetName(), ExportNode, bSaveAnimSeq);
 				}
 				else if (Component->IsA(UCameraComponent::StaticClass()))
 				{
@@ -1786,7 +1786,7 @@ FbxNode* FFbxExporter::ExportActor(AActor* Actor, bool bExportComponents, INodeN
 				}
 				else if (ChildActorComp && ChildActorComp->GetChildActor())
 				{
-					FbxNode* ChildActorNode = ExportActor(ChildActorComp->GetChildActor(), true, NodeNameAdapter);
+					FbxNode* ChildActorNode = ExportActor(ChildActorComp->GetChildActor(), true, NodeNameAdapter, bSaveAnimSeq);
 					FbxActors.Add(ChildActorComp->GetChildActor(), ChildActorNode);
 				}
 			}
diff --git a/Engine/Source/Editor/UnrealEd/Private/Fbx/FbxSkeletalMeshExport.cpp b/Engine/Source/Editor/UnrealEd/Private/Fbx/FbxSkeletalMeshExport.cpp
--- a/Engine/Source/Editor/UnrealEd/Private/Fbx/FbxSkeletalMeshExport.cpp
+++ b/Engine/Source/Editor/UnrealEd/Private/Fbx/FbxSkeletalMeshExport.cpp
@@ -455,11 +455,12 @@ void FFbxExporter::CreateBindPose(FbxNode* MeshRootNode)
 	}
 }
 
-void FFbxExporter::ExportSkeletalMeshComponent(USkeletalMeshComponent* SkelMeshComp, const TCHAR* MeshName, FbxNode* ActorRootNode)
+void FFbxExporter::ExportSkeletalMeshComponent(USkeletalMeshComponent* SkelMeshComp, const TCHAR* MeshName, FbxNode* ActorRootNode, bool bSaveAnimSeq)
 {
 	if (SkelMeshComp && SkelMeshComp->SkeletalMesh)
 	{
-		UAnimSequence* AnimSeq = (SkelMeshComp->GetAnimationMode() == EAnimationMode::AnimationSingleNode)? Cast<UAnimSequence>(SkelMeshComp->AnimationData.AnimToPlay) : NULL;
+		UAnimSequence* AnimSeq = (bSaveAnimSeq && SkelMeshComp->GetAnimationMode() == EAnimationMode::AnimationSingleNode) ? 
+			Cast<UAnimSequence>(SkelMeshComp->AnimationData.AnimToPlay) : NULL;
 		FbxNode* SkeletonRootNode = ExportSkeletalMeshToFbx(SkelMeshComp->SkeletalMesh, AnimSeq, MeshName, ActorRootNode);
 		if(SkeletonRootNode)
 		{
diff --git a/Engine/Source/Editor/UnrealEd/Private/FbxExporter.h b/Engine/Source/Editor/UnrealEd/Private/FbxExporter.h
--- a/Engine/Source/Editor/UnrealEd/Private/FbxExporter.h
+++ b/Engine/Source/Editor/UnrealEd/Private/FbxExporter.h
@@ -110,7 +110,7 @@ public:
 	/**
 	 * Exports the basic scene information to the FBX document.
 	 */
-	virtual void ExportLevelMesh( ULevel* InLevel, bool bSelectedOnly, INodeNameAdapter& NodeNameAdapter );
+	virtual void ExportLevelMesh( ULevel* InLevel, bool bSelectedOnly, INodeNameAdapter& NodeNameAdapter , bool bSaveAnimSeq = true);
 
 	/**
 	 * Exports the given Matinee sequence information into a FBX document.
@@ -215,7 +215,7 @@ public:
 	 * Exports the basic information about an actor and buffers it.
 	 * This function creates one FBX node for the actor with its placement.
 	 */
-	FbxNode* ExportActor(AActor* Actor, bool bExportComponents, INodeNameAdapter& NodeNameAdapter);
+	FbxNode* ExportActor(AActor* Actor, bool bExportComponents, INodeNameAdapter& NodeNameAdapter, bool bSaveAnimSeq = true);
 
 private:
 	FFbxExporter();
@@ -374,7 +374,7 @@ private:
 	FbxNode* ExportSkeletalMeshToFbx(const USkeletalMesh* SkelMesh, const UAnimSequence* AnimSeq, const TCHAR* MeshName, FbxNode* ActorRootNode);
 
 	/** Export SkeletalMeshComponent */
-	void ExportSkeletalMeshComponent(USkeletalMeshComponent* SkelMeshComp, const TCHAR* MeshName, FbxNode* ActorRootNode);
+	void ExportSkeletalMeshComponent(USkeletalMeshComponent* SkelMeshComp, const TCHAR* MeshName, FbxNode* ActorRootNode, bool bSaveAnimSeq = true);
 
 	/**
 	 * Add the given animation sequence as rotation and translation tracks to the given list of bone nodes
diff --git a/Engine/Source/Editor/UnrealEd/Public/MatineeExporter.h b/Engine/Source/Editor/UnrealEd/Public/MatineeExporter.h
--- a/Engine/Source/Editor/UnrealEd/Public/MatineeExporter.h
+++ b/Engine/Source/Editor/UnrealEd/Public/MatineeExporter.h
@@ -60,7 +60,7 @@ public:
 	/**
 	 * Exports the basic scene information to a file.
 	 */
-	virtual void ExportLevelMesh( ULevel* Level, bool bSelectedOnly, INodeNameAdapter& NodeNameAdapter ) = 0;
+	virtual void ExportLevelMesh( ULevel* Level, bool bSelectedOnly, INodeNameAdapter& NodeNameAdapter, bool bSaveAnimSeq = true) = 0;
 
 	/**
 	 * Exports the light-specific information for a light actor.
