Added missing API export

diff --git a/Engine/Source/Editor/UnrealEd/Classes/Commandlets/GatherTextCommandletBase.h b/Engine/Source/Editor/UnrealEd/Classes/Commandlets/GatherTextCommandletBase.h
--- a/Engine/Source/Editor/UnrealEd/Classes/Commandlets/GatherTextCommandletBase.h
+++ b/Engine/Source/Editor/UnrealEd/Classes/Commandlets/GatherTextCommandletBase.h
@@ -51,14 +51,11 @@ private:
 /**
  *	UGatherTextCommandletBase: Base class for localization commandlets. Just to force certain behaviors and provide helper functionality. 
  */
-class FJsonValue;
-class FJsonObject;
 UCLASS()
-class UGatherTextCommandletBase : public UCommandlet
+class UNREALED_API UGatherTextCommandletBase : public UCommandlet
 {
 	GENERATED_UCLASS_BODY()
 
-
 public:
 	virtual void Initialize( const TSharedPtr< FLocTextHelper >& InGatherManifestHelper, const TSharedPtr< FLocalizationSCC >& InSourceControlInfo );
 
@@ -75,7 +72,5 @@ protected:
 	TSharedPtr< FLocalizationSCC > SourceControlInfo;
 
 private:
-
-	virtual void CreateCustomEngine(const FString& Params) override ; //Disallow other text commandlets to make their own engine.
-	
+	virtual void CreateCustomEngine(const FString& Params) override ; //Disallow other text commandlets to make their own engine.	
 };
