Fix display name for Display UI Extension Points

diff --git a/Engine/Source/Editor/UnrealEd/Classes/Editor/EditorPerProjectUserSettings.h b/Engine/Source/Editor/UnrealEd/Classes/Editor/EditorPerProjectUserSettings.h
--- a/Engine/Source/Editor/UnrealEd/Classes/Editor/EditorPerProjectUserSettings.h
+++ b/Engine/Source/Editor/UnrealEd/Classes/Editor/EditorPerProjectUserSettings.h
@@ -17,7 +17,7 @@ class UEditorPerProjectUserSettings : public UObject
 	// The following options are exposed in the Preferences Editor 
 
 	/** If enabled, any newly opened UI menus, menu bars, and toolbars will show the developer hooks that would accept extensions */
-	UPROPERTY(EditAnywhere, config, Category=DeveloperTools)
+	UPROPERTY(EditAnywhere, config, Category=DeveloperTools, meta = (DisplayName = "Display UI Extension Points"))
 	uint32 bDisplayUIExtensionPoints:1;
 
 	/** If enabled, tooltips linked to documentation will show the developer the link bound to that UI item */
