Add to the fbx importer the possibility to use different sample rate when importing an animation.

diff --git a/Engine/Source/Editor/UnrealEd/Classes/Factories/FbxAnimSequenceImportData.h b/Engine/Source/Editor/UnrealEd/Classes/Factories/FbxAnimSequenceImportData.h
--- a/Engine/Source/Editor/UnrealEd/Classes/Factories/FbxAnimSequenceImportData.h
+++ b/Engine/Source/Editor/UnrealEd/Classes/Factories/FbxAnimSequenceImportData.h
@@ -62,6 +62,10 @@ class UNREALED_API UFbxAnimSequenceImportData : public UFbxAssetImportData
 	UPROPERTY(EditAnywhere, AdvancedDisplay, config, Category = ImportSettings, meta = (ToolTip = "If enabled, samples all animation curves to 30 FPS"))
 	bool bUseDefaultSampleRate;
 
+	/** Use this option to specify a sample rate for the imported animation, a value of 0 use the best matching samplerate. */
+	UPROPERTY(EditAnywhere, AdvancedDisplay, config, Category = ImportSettings, meta = (EditCondition = "!bUseDefaultSampleRate", ToolTip = "Sample fbx animation data at the specified sample rate, 0 find automaticaly the best sample rate"))
+	int32 CustomSampleRate;
+
 	/** Name of source animation that was imported, used to reimport correct animation from the FBX file */
 	UPROPERTY()
 	FString SourceAnimationName;
diff --git a/Engine/Source/Editor/UnrealEd/Classes/Factories/FbxSceneImportOptionsSkeletalMesh.h b/Engine/Source/Editor/UnrealEd/Classes/Factories/FbxSceneImportOptionsSkeletalMesh.h
--- a/Engine/Source/Editor/UnrealEd/Classes/Factories/FbxSceneImportOptionsSkeletalMesh.h
+++ b/Engine/Source/Editor/UnrealEd/Classes/Factories/FbxSceneImportOptionsSkeletalMesh.h
@@ -73,6 +73,10 @@ class UFbxSceneImportOptionsSkeletalMesh : public UObject
 	UPROPERTY(EditAnywhere, AdvancedDisplay, config, Category = Animation, meta = (ToolTip = "If enabled, samples all animation curves to 30 FPS"))
 	bool bUseDefaultSampleRate;
 
+	/** Use this option to specify a sample rate for the imported animation, a value of 0 use the best matching sample rate. */
+	UPROPERTY(EditAnywhere, AdvancedDisplay, config, Category = Animation, meta = (EditCondition = "!bUseDefaultSampleRate", ToolTip = "Sample fbx animation data at the specified sample rate, 0 find automaticaly the best sample rate"))
+	int32 CustomSampleRate;
+
 	/** Import if custom attribute as a curve within the animation **/
 	UPROPERTY(EditAnywhere, AdvancedDisplay, config, Category = Animation)
 	bool bImportCustomAttribute;
diff --git a/Engine/Source/Editor/UnrealEd/Private/Fbx/FbxAnimSequenceImportData.cpp b/Engine/Source/Editor/UnrealEd/Private/Fbx/FbxAnimSequenceImportData.cpp
--- a/Engine/Source/Editor/UnrealEd/Private/Fbx/FbxAnimSequenceImportData.cpp
+++ b/Engine/Source/Editor/UnrealEd/Private/Fbx/FbxAnimSequenceImportData.cpp
@@ -15,6 +15,7 @@ UFbxAnimSequenceImportData::UFbxAnimSequenceImportData(const FObjectInitializer&
 	FrameImportRange.Max = 0;
 
 	MaterialCurveSuffixes.Add(TEXT("_mat"));
+	CustomSampleRate = 0;
 }
 
 UFbxAnimSequenceImportData* UFbxAnimSequenceImportData::GetImportDataForAnimSequence(UAnimSequence* AnimSequence, UFbxAnimSequenceImportData* TemplateForCreation)
diff --git a/Engine/Source/Editor/UnrealEd/Private/Fbx/FbxMainImport.cpp b/Engine/Source/Editor/UnrealEd/Private/Fbx/FbxMainImport.cpp
--- a/Engine/Source/Editor/UnrealEd/Private/Fbx/FbxMainImport.cpp
+++ b/Engine/Source/Editor/UnrealEd/Private/Fbx/FbxMainImport.cpp
@@ -340,6 +340,7 @@ void ApplyImportUIToImportOptions(UFbxImportUI* ImportUI, FBXImportOptions& InOu
 	InOutImportOptions.AnimationName = ImportUI->OverrideAnimationName;
 	// only re-sample if they don't want to use default sample rate
 	InOutImportOptions.bResample = !ImportUI->AnimSequenceImportData->bUseDefaultSampleRate;
+	InOutImportOptions.ResampleRate = ImportUI->AnimSequenceImportData->CustomSampleRate;
 	InOutImportOptions.bPreserveLocalTransform = ImportUI->AnimSequenceImportData->bPreserveLocalTransform;
 	InOutImportOptions.bDeleteExistingMorphTargetCurves = ImportUI->AnimSequenceImportData->bDeleteExistingMorphTargetCurves;
 	InOutImportOptions.bRemoveRedundantKeys = ImportUI->AnimSequenceImportData->bRemoveRedundantKeys;
diff --git a/Engine/Source/Editor/UnrealEd/Private/Fbx/FbxSceneImportOptionsSkeletalMesh.cpp b/Engine/Source/Editor/UnrealEd/Private/Fbx/FbxSceneImportOptionsSkeletalMesh.cpp
--- a/Engine/Source/Editor/UnrealEd/Private/Fbx/FbxSceneImportOptionsSkeletalMesh.cpp
+++ b/Engine/Source/Editor/UnrealEd/Private/Fbx/FbxSceneImportOptionsSkeletalMesh.cpp
@@ -21,6 +21,7 @@ UFbxSceneImportOptionsSkeletalMesh::UFbxSceneImportOptionsSkeletalMesh(const FOb
 	, AnimationLength(EFbxSceneVertexColorImportOption::Replace)
 	, FrameImportRange(0, 0)
 	, bUseDefaultSampleRate(false)
+	, CustomSampleRate(0)
 	, bImportCustomAttribute(true)
 	, bPreserveLocalTransform(false)
 	, bDeleteExistingMorphTargetCurves(false)
@@ -54,6 +55,7 @@ void UFbxSceneImportOptionsSkeletalMesh::FillSkeletalMeshInmportData(UFbxSkeleta
 	AnimSequenceImportData->bImportCustomAttribute = bImportCustomAttribute;
 	AnimSequenceImportData->bPreserveLocalTransform = bPreserveLocalTransform;
 	AnimSequenceImportData->bUseDefaultSampleRate = bUseDefaultSampleRate;
+	AnimSequenceImportData->CustomSampleRate = CustomSampleRate;
 	AnimSequenceImportData->FrameImportRange = FrameImportRange;
 
 	AnimSequenceImportData->bImportAsScene = true;
diff --git a/Engine/Source/Editor/UnrealEd/Private/Fbx/ReimportFbxSceneFactory.cpp b/Engine/Source/Editor/UnrealEd/Private/Fbx/ReimportFbxSceneFactory.cpp
--- a/Engine/Source/Editor/UnrealEd/Private/Fbx/ReimportFbxSceneFactory.cpp
+++ b/Engine/Source/Editor/UnrealEd/Private/Fbx/ReimportFbxSceneFactory.cpp
@@ -1274,10 +1274,17 @@ EReimportResult::Type UReimportFbxSceneFactory::ReimportSkeletalMesh(void* VoidF
 				int32 ResampleRate = DEFAULT_SAMPLERATE;
 				if (GlobalImportSettings->bResample)
 				{
-					int32 MaxStackResampleRate = FbxImporter->GetMaxSampleRate(SortedLinks, FBXMeshNodeArray);
-					if (MaxStackResampleRate != 0)
+					if(FbxImporter->ImportOptions->ResampleRate > 0)
 					{
-						ResampleRate = MaxStackResampleRate;
+						ResampleRate = FbxImporter->ImportOptions->ResampleRate;
+					}
+					else
+					{
+						int32 BestResampleRate = FbxImporter->GetMaxSampleRate(SortedLinks, FBXMeshNodeArray);
+						if (BestResampleRate > 0)
+						{
+							ResampleRate = BestResampleRate;
+						}
 					}
 				}
 				int32 ValidTakeCount = 0;
@@ -1336,7 +1343,18 @@ EReimportResult::Type UReimportFbxSceneFactory::ReimportSkeletalMesh(void* VoidF
 						ResampleRate = DEFAULT_SAMPLERATE;
 						if (FbxImporter->ImportOptions->bResample)
 						{
-							ResampleRate = FbxImporter->GetMaxSampleRate(SortedLinks, FBXMeshNodeArray);
+							if(FbxImporter->ImportOptions->ResampleRate > 0)
+							{
+								ResampleRate = FbxImporter->ImportOptions->ResampleRate;
+							}
+							else
+							{
+								int32 BestResampleRate = FbxImporter->GetMaxSampleRate(SortedLinks, FBXMeshNodeArray);
+								if (BestResampleRate > 0)
+								{
+									ResampleRate = BestResampleRate;
+								}
+							}
 						}
 						FbxTimeSpan AnimTimeSpan = FbxImporter->GetAnimationTimeSpan(SortedLinks[0], CurAnimStack, ResampleRate);
 
diff --git a/Engine/Source/Editor/UnrealEd/Private/Fbx/SFbxSceneOptionWindow.cpp b/Engine/Source/Editor/UnrealEd/Private/Fbx/SFbxSceneOptionWindow.cpp
--- a/Engine/Source/Editor/UnrealEd/Private/Fbx/SFbxSceneOptionWindow.cpp
+++ b/Engine/Source/Editor/UnrealEd/Private/Fbx/SFbxSceneOptionWindow.cpp
@@ -1506,7 +1506,8 @@ void SFbxSceneOptionWindow::CopySkeletalMeshOptionsToFbxOptions(UnFbx::FBXImport
 	ImportSettings->bDeleteExistingMorphTargetCurves = SkeletalMeshOptions->bDeleteExistingMorphTargetCurves;
 	ImportSettings->bImportCustomAttribute = SkeletalMeshOptions->bImportCustomAttribute;
 	ImportSettings->bPreserveLocalTransform = SkeletalMeshOptions->bPreserveLocalTransform;
-	ImportSettings->bResample = SkeletalMeshOptions->bUseDefaultSampleRate;
+	ImportSettings->bResample = !SkeletalMeshOptions->bUseDefaultSampleRate;
+	ImportSettings->ResampleRate = SkeletalMeshOptions->CustomSampleRate;
 	ImportSettings->AnimationRange.X = SkeletalMeshOptions->FrameImportRange.Min;
 	ImportSettings->AnimationRange.Y = SkeletalMeshOptions->FrameImportRange.Max;
 }
@@ -1528,7 +1529,8 @@ void SFbxSceneOptionWindow::CopyFbxOptionsToSkeletalMeshOptions(UnFbx::FBXImport
 	SkeletalMeshOptions->bDeleteExistingMorphTargetCurves = ImportSettings->bDeleteExistingMorphTargetCurves;
 	SkeletalMeshOptions->bImportCustomAttribute = ImportSettings->bImportCustomAttribute;
 	SkeletalMeshOptions->bPreserveLocalTransform = ImportSettings->bPreserveLocalTransform;
-	SkeletalMeshOptions->bUseDefaultSampleRate = ImportSettings->bResample;
+	SkeletalMeshOptions->bUseDefaultSampleRate = !ImportSettings->bResample;
+	SkeletalMeshOptions->CustomSampleRate = ImportSettings->ResampleRate;
 	SkeletalMeshOptions->FrameImportRange.Min = ImportSettings->AnimationRange.X;
 	SkeletalMeshOptions->FrameImportRange.Max = ImportSettings->AnimationRange.Y;
 }
diff --git a/Engine/Source/Editor/UnrealEd/Private/SkeletalMeshEdit.cpp b/Engine/Source/Editor/UnrealEd/Private/SkeletalMeshEdit.cpp
--- a/Engine/Source/Editor/UnrealEd/Private/SkeletalMeshEdit.cpp
+++ b/Engine/Source/Editor/UnrealEd/Private/SkeletalMeshEdit.cpp
@@ -240,7 +240,18 @@ bool UEditorEngine::ReimportFbxAnimation( USkeleton* Skeleton, UAnimSequence* An
 					int32 ResampleRate = DEFAULT_SAMPLERATE;
 					if (FbxImporter->ImportOptions->bResample)
 					{
-						ResampleRate = FbxImporter->GetMaxSampleRate(SortedLinks, FBXMeshNodeArray);
+						if(FbxImporter->ImportOptions->ResampleRate > 0)
+						{
+							ResampleRate = FbxImporter->ImportOptions->ResampleRate;
+						}
+						else
+						{
+							int32 BestResampleRate = FbxImporter->GetMaxSampleRate(SortedLinks, FBXMeshNodeArray);
+							if(BestResampleRate > 0)
+							{
+								ResampleRate = BestResampleRate;
+							}
+						}
 					}
 					FbxTimeSpan AnimTimeSpan = FbxImporter->GetAnimationTimeSpan(SortedLinks[0], CurAnimStack, ResampleRate);
 					// for now it's not importing morph - in the future, this should be optional or saved with asset
@@ -531,22 +542,28 @@ UAnimSequence * UnFbx::FFbxImporter::ImportAnimations(USkeleton* Skeleton, UObje
 	int32 ResampleRate = DEFAULT_SAMPLERATE;
 	if ( ImportOptions->bResample )
 	{
-		// For FBX data, "Frame Rate" is just the speed at which the animation is played back.  It can change
-		// arbitrarily, and the underlying data can stay the same.  What we really want here is the Sampling Rate,
-		// ie: the number of animation keys per second.  These are the individual animation curve keys
-		// on the FBX nodes of the skeleton.  So we loop through the nodes of the skeleton and find the maximum number 
-		// of keys that any node has, then divide this by the total length (in seconds) of the animation to find the 
-		// sampling rate of this set of data 
-
-		// we want the maximum resample rate, so that we don't lose any precision of fast anims,
-		// and don't mind creating lerped frames for slow anims
-		int32 MaxStackResampleRate = GetMaxSampleRate(SortedLinks, NodeArray);
-
-		if(MaxStackResampleRate != 0)
+		if (ImportOptions->ResampleRate > 0)
 		{
-			ResampleRate = MaxStackResampleRate;
+			ResampleRate = ImportOptions->ResampleRate;
+		}
+		else
+		{
+			// For FBX data, "Frame Rate" is just the speed at which the animation is played back.  It can change
+			// arbitrarily, and the underlying data can stay the same.  What we really want here is the Sampling Rate,
+			// ie: the number of animation keys per second.  These are the individual animation curve keys
+			// on the FBX nodes of the skeleton.  So we loop through the nodes of the skeleton and find the maximum number 
+			// of keys that any node has, then divide this by the total length (in seconds) of the animation to find the 
+			// sampling rate of this set of data 
+
+			// we want the maximum resample rate, so that we don't lose any precision of fast anims,
+			// and don't mind creating lerped frames for slow anims
+			int32 BestResampleRate = GetMaxSampleRate(SortedLinks, NodeArray);
+
+			if (BestResampleRate > 0)
+			{
+				ResampleRate = BestResampleRate;
+			}
 		}
-
 	}
 
 	int32 AnimStackCount = Scene->GetSrcObjectCount<FbxAnimStack>();
diff --git a/Engine/Source/Editor/UnrealEd/Public/FbxImporter.h b/Engine/Source/Editor/UnrealEd/Public/FbxImporter.h
--- a/Engine/Source/Editor/UnrealEd/Public/FbxImporter.h
+++ b/Engine/Source/Editor/UnrealEd/Public/FbxImporter.h
@@ -161,6 +161,7 @@ struct FBXImportOptions
 	bool bImportAnimations;
 	bool bUpdateSkeletonReferencePose;
 	bool bResample;
+	int32 ResampleRate;
 	bool bImportRigidMesh;
 	bool bUseT0AsRefPose;
 	bool bPreserveSmoothingGroups;
