PR #4771: Fix access violation for ImportAsset commandlet fbx reimport. (Contributed by [REDACTED_PERSON])

diff --git a/Engine/Source/Editor/UnrealEd/Private/Factories/EditorFactories.cpp b/Engine/Source/Editor/UnrealEd/Private/Factories/EditorFactories.cpp
--- a/Engine/Source/Editor/UnrealEd/Private/Factories/EditorFactories.cpp
+++ b/Engine/Source/Editor/UnrealEd/Private/Factories/EditorFactories.cpp
@@ -5688,7 +5688,7 @@ EReimportResult::Type UReimportFbxSkeletalMeshFactory::Reimport( UObject* Obj )
 		CleanUp();
 
 		// Reimporting can have dangerous effects if the mesh is still in the transaction buffer.  Reset the transaction buffer if this is the case
-		if( GEditor->IsObjectInTransactionBuffer( SkeletalMesh ) )
+		if(!IsRunningCommandlet() && GEditor->IsObjectInTransactionBuffer( SkeletalMesh ) )
 		{
 			GEditor->ResetTransaction( LOCTEXT("ReimportSkeletalMeshTransactionReset", "Reimporting a skeletal mesh which was in the undo buffer") );
 		}
