Fix tiny fonts from appearing in slow task dialogs

diff --git a/Engine/Source/Editor/UnrealEd/Private/FeedbackContextEditor.cpp b/Engine/Source/Editor/UnrealEd/Private/FeedbackContextEditor.cpp
--- a/Engine/Source/Editor/UnrealEd/Private/FeedbackContextEditor.cpp
+++ b/Engine/Source/Editor/UnrealEd/Private/FeedbackContextEditor.cpp
@@ -270,7 +270,7 @@ private:
 		const int32 MaxTextWidth = FixedWidth - FixedPaddingH*2;
 		while( FontInfo.Size > 9 && MeasureService->Measure(MainText, FontInfo).X > MaxTextWidth )
 		{
-			FontInfo.Size -= 4;
+			FontInfo.Size -= 2;
 		}
 
 		return FontInfo;
