Fix specifying relative paths for the auto-import commandlet not working

diff --git a/Engine/Source/Editor/UnrealEd/Private/AutomatedAssetImportData.cpp b/Engine/Source/Editor/UnrealEd/Private/AutomatedAssetImportData.cpp
--- a/Engine/Source/Editor/UnrealEd/Private/AutomatedAssetImportData.cpp
+++ b/Engine/Source/Editor/UnrealEd/Private/AutomatedAssetImportData.cpp
@@ -5,6 +5,7 @@
 #include "UObject/Package.h"
 #include "Misc/PackageName.h"
 #include "UObject/UObjectIterator.h"
+#include "Misc/Paths.h"
 
 DEFINE_LOG_CATEGORY(LogAutomatedImport);
 
@@ -28,6 +29,11 @@ void UAutomatedAssetImportData::Initialize(TSharedPtr<FJsonObject> InImportGroup
 	ImportGroupJsonData = InImportGroupJsonData;
 	if(Filenames.Num() > 0)
 	{
+		for (FString& Filename : Filenames)
+		{
+			// Make sure all filenames are absolute path
+			Filename = FPaths::ConvertRelativePathToFull(Filename);
+		}
 		// If a factory doesn't have a vaild full path assume it is an unreal internal factory
 		if(!FactoryName.IsEmpty() && !FactoryName.StartsWith("/Script/"))
 		{
