Make sure when converting relative path filenames in automated import that we convert them relative to the project directory.

diff --git a/Engine/Source/Editor/UnrealEd/Private/AutomatedAssetImportData.cpp b/Engine/Source/Editor/UnrealEd/Private/AutomatedAssetImportData.cpp
--- a/Engine/Source/Editor/UnrealEd/Private/AutomatedAssetImportData.cpp
+++ b/Engine/Source/Editor/UnrealEd/Private/AutomatedAssetImportData.cpp
@@ -32,7 +32,7 @@ void UAutomatedAssetImportData::Initialize(TSharedPtr<FJsonObject> InImportGroup
 		for (FString& Filename : Filenames)
 		{
 			// Make sure all filenames are absolute path
-			Filename = FPaths::ConvertRelativePathToFull(Filename);
+			Filename = FPaths::ConvertRelativePathToFull(FPaths::ProjectDir(), Filename);
 		}
 		// If a factory doesn't have a vaild full path assume it is an unreal internal factory
 		if(!FactoryName.IsEmpty() && !FactoryName.StartsWith("/Script/"))
