Fix preview mesh actor becoming stuck to the cursor when the editor or viewport loses focus

diff --git a/Engine/Source/Editor/UnrealEd/Private/EditorEngine.cpp b/Engine/Source/Editor/UnrealEd/Private/EditorEngine.cpp
--- a/Engine/Source/Editor/UnrealEd/Private/EditorEngine.cpp
+++ b/Engine/Source/Editor/UnrealEd/Private/EditorEngine.cpp
@@ -1475,6 +1475,11 @@ void UEditorEngine::Tick( float DeltaSeconds, bool bIdleMode )
 		}
 	}
 
+	if (!bHasFocus)
+	{
+		SetPreviewMeshMode(false);
+	}
+
 	// Tick any editor FTickableEditorObject dervived classes
 	FTickableEditorObject::TickObjects( DeltaSeconds );
 
diff --git a/Engine/Source/Editor/UnrealEd/Private/LevelEditorViewport.cpp b/Engine/Source/Editor/UnrealEd/Private/LevelEditorViewport.cpp
--- a/Engine/Source/Editor/UnrealEd/Private/LevelEditorViewport.cpp
+++ b/Engine/Source/Editor/UnrealEd/Private/LevelEditorViewport.cpp
@@ -1936,6 +1936,14 @@ void FLevelEditorViewportClient::ReceivedFocus(FViewport* InViewport)
 	FEditorViewportClient::ReceivedFocus(InViewport);
 }
 
+void FLevelEditorViewportClient::LostFocus(FViewport* InViewport)
+{
+	FEditorViewportClient::LostFocus(InViewport);
+
+	GEditor->SetPreviewMeshMode(false);
+}
+
+
 //
 //	FLevelEditorViewportClient::ProcessClick
 //
diff --git a/Engine/Source/Editor/UnrealEd/Public/LevelEditorViewport.h b/Engine/Source/Editor/UnrealEd/Public/LevelEditorViewport.h
--- a/Engine/Source/Editor/UnrealEd/Public/LevelEditorViewport.h
+++ b/Engine/Source/Editor/UnrealEd/Public/LevelEditorViewport.h
@@ -173,7 +173,8 @@ public:
 	virtual void SetCameraSpeedSetting(int32 SpeedSetting) override;
 	virtual float GetCameraSpeedScalar() const override;
 	virtual void SetCameraSpeedScalar(float SpeedScalar) override;
-	virtual void ReceivedFocus(FViewport* Viewport) override;
+	virtual void ReceivedFocus(FViewport* InViewport) override;
+	virtual void LostFocus(FViewport* InViewport) override;
 	virtual void ProcessClick(FSceneView& View, HHitProxy* HitProxy, FKey Key, EInputEvent Event, uint32 HitX, uint32 HitY) override;
 	virtual UWorld* GetWorld() const override;
 	virtual void BeginCameraMovement(bool bHasMovement) override;
