Fix crash opening level after removing as sublevel
jira: UE-61305

diff --git a/Engine/Source/Editor/UnrealEd/Private/EditorLevelUtils.cpp b/Engine/Source/Editor/UnrealEd/Private/EditorLevelUtils.cpp
--- a/Engine/Source/Editor/UnrealEd/Private/EditorLevelUtils.cpp
+++ b/Engine/Source/Editor/UnrealEd/Private/EditorLevelUtils.cpp
@@ -624,6 +624,8 @@ bool UEditorLevelUtils::RemoveLevelFromWorld(ULevel* InLevel)
 
 		FEditorSupportDelegates::PrepareToCleanseEditorObject.Broadcast(InLevel);
 
+		GEditor->Trans->Reset(LOCTEXT("RemoveLevelTransReset", "Removing Levels from World"));
+
 		EditorDestroyLevel(InLevel);
 
 		// Redraw the main editor viewports.
diff --git a/Engine/Source/Editor/UnrealEd/Private/EditorServer.cpp b/Engine/Source/Editor/UnrealEd/Private/EditorServer.cpp
--- a/Engine/Source/Editor/UnrealEd/Private/EditorServer.cpp
+++ b/Engine/Source/Editor/UnrealEd/Private/EditorServer.cpp
@@ -2502,6 +2502,14 @@ bool UEditorEngine::Map_Load(const TCHAR* Str, FOutputDevice& Ar)
 					if ( ExistingPackage )
 					{
 						WorldPackage = ExistingPackage;
+
+						if (!ExistingPackage->IsFullyLoaded())
+						{
+							const FName WorldPackageFName = FName(*LongTempFname);
+							UWorld::WorldTypePreLoadMap.FindOrAdd(WorldPackageFName) = EWorldType::Editor;
+							ExistingPackage->FullyLoad();
+							UWorld::WorldTypePreLoadMap.Remove(WorldPackageFName);
+						}
 					}
 					else
 					{
