Fix the material ID assignation when re-importing static mesh

diff --git a/Engine/Source/Editor/UnrealEd/Private/StaticMeshEdit.cpp b/Engine/Source/Editor/UnrealEd/Private/StaticMeshEdit.cpp
--- a/Engine/Source/Editor/UnrealEd/Private/StaticMeshEdit.cpp
+++ b/Engine/Source/Editor/UnrealEd/Private/StaticMeshEdit.cpp
@@ -1667,14 +1667,12 @@ void RestoreExistingMeshData(ExistingStaticMeshData* ExistingMeshDataPtr, UStati
 
 				FMeshSectionInfo NewSectionInfo =  NewMesh->SectionInfoMap.Get(i, SectionIndex);
 				bool bFoundOldMatch = false;
-				bool bKeepOldSectionMaterialIndex = false;
 				int32 OriginalSectionMaterialIndex = INDEX_NONE;
 				if (RemapMaterial.IsValidIndex(NewSectionInfo.MaterialIndex) && NewMesh->StaticMaterials.IsValidIndex(NewSectionInfo.MaterialIndex))
 				{
 					//Find the matching old index
 					for (int32 ExistSectionIndex = 0; ExistSectionIndex < OldSectionNumber; ++ExistSectionIndex)
 					{
-						bKeepOldSectionMaterialIndex = false;
 						OriginalSectionMaterialIndex = INDEX_NONE;
 						FMeshSectionInfo OldSectionInfo = ExistingMeshDataPtr->ExistingSectionInfoMap.Get(i, ExistSectionIndex);
 						if (ExistingMeshDataPtr->UseMaterialNameSlotWorkflow)
@@ -1684,7 +1682,6 @@ void RestoreExistingMeshData(ExistingStaticMeshData* ExistingMeshDataPtr, UStati
 								ExistingMeshDataPtr->LastImportMeshLodSectionMaterialData[i].IsValidIndex(ExistSectionIndex))
 							{
 								//Keep the old section material index only if the user has change the section mapping
-								bKeepOldSectionMaterialIndex = ExistingMeshDataPtr->LastImportMeshLodSectionMaterialData[i][ExistSectionIndex] != ExistingMeshDataPtr->ExistingMaterials[OldSectionInfo.MaterialIndex].ImportedMaterialSlotName;
 								for (int32 ExistMaterialIndex = 0; ExistMaterialIndex < ExistingMeshDataPtr->ExistingMaterials.Num(); ++ExistMaterialIndex)
 								{
 									if (ExistingMeshDataPtr->LastImportMeshLodSectionMaterialData[i][ExistSectionIndex] == ExistingMeshDataPtr->ExistingMaterials[ExistMaterialIndex].ImportedMaterialSlotName)
@@ -1698,11 +1695,6 @@ void RestoreExistingMeshData(ExistingStaticMeshData* ExistingMeshDataPtr, UStati
 						int32 OldSectionMatchIndex = OriginalSectionMaterialIndex != INDEX_NONE ? OriginalSectionMaterialIndex : OldSectionInfo.MaterialIndex;
 						if (RemapMaterial[NewSectionInfo.MaterialIndex] == OldSectionMatchIndex)
 						{
-							//Set the remap section
-							if (!bKeepOldSectionMaterialIndex)
-							{
-								OldSectionInfo.MaterialIndex = NewSectionInfo.MaterialIndex;
-							}
 							NewMesh->SectionInfoMap.Set(i, SectionIndex, OldSectionInfo);
 							bFoundOldMatch = true;
 							break;
