Fix brush actors not showing the correct icon in scene outliner.
- Actors can now supply their own icon if needed

diff --git a/Engine/Source/Editor/SceneOutliner/Private/SSceneOutliner.cpp b/Engine/Source/Editor/SceneOutliner/Private/SSceneOutliner.cpp
--- a/Engine/Source/Editor/SceneOutliner/Private/SSceneOutliner.cpp
+++ b/Engine/Source/Editor/SceneOutliner/Private/SSceneOutliner.cpp
@@ -2378,7 +2378,7 @@ namespace SceneOutliner
 		}
 	}
 
-	FSlateBrush* SSceneOutliner::GetCachedIconForClass(FName InClassName) const
+	const FSlateBrush* SSceneOutliner::GetCachedIconForClass(FName InClassName) const
 	{ 
 		if (CachedIcons.Find(InClassName))
 		{
@@ -2390,7 +2390,7 @@ namespace SceneOutliner
 		}
 	}
 
-	void SSceneOutliner::CacheIconForClass(FName InClassName, FSlateBrush* InSlateBrush)
+	void SSceneOutliner::CacheIconForClass(FName InClassName, const FSlateBrush* InSlateBrush)
 	{
 		CachedIcons.Emplace(InClassName, InSlateBrush);
 	}
diff --git a/Engine/Source/Editor/SceneOutliner/Private/SSceneOutliner.h b/Engine/Source/Editor/SceneOutliner/Private/SSceneOutliner.h
--- a/Engine/Source/Editor/SceneOutliner/Private/SSceneOutliner.h
+++ b/Engine/Source/Editor/SceneOutliner/Private/SSceneOutliner.h
@@ -125,10 +125,10 @@ namespace SceneOutliner
 		virtual void SetKeyboardFocus() override;
 
 		/** Gets the cached icon for this class name */
-		virtual FSlateBrush* GetCachedIconForClass(FName InClassName) const override;
+		virtual const FSlateBrush* GetCachedIconForClass(FName InClassName) const override;
 		
 		/** Sets the cached icon for this class name */
-		virtual void CacheIconForClass(FName InClassName, FSlateBrush* InSlateBrush) override;
+		virtual void CacheIconForClass(FName InClassName, const FSlateBrush* InSlateBrush) override;
 		
 	private:
 		/** Methods that implement structural modification logic for the tree */
@@ -511,7 +511,7 @@ namespace SceneOutliner
 		/** The tree item that is currently pending a rename */
 		TWeakPtr<ITreeItem> PendingRenameItem;
 
-		TMap<FName, FSlateBrush*> CachedIcons;
+		TMap<FName, const FSlateBrush*> CachedIcons;
 
 	private:
 		/** Functions relating to sorting */
diff --git a/Engine/Source/Editor/SceneOutliner/Private/SceneOutlinerItemLabelColumn.cpp b/Engine/Source/Editor/SceneOutliner/Private/SceneOutlinerItemLabelColumn.cpp
--- a/Engine/Source/Editor/SceneOutliner/Private/SceneOutlinerItemLabelColumn.cpp
+++ b/Engine/Source/Editor/SceneOutliner/Private/SceneOutlinerItemLabelColumn.cpp
@@ -220,15 +220,22 @@ private:
 		{
 			if (WeakSceneOutliner.IsValid())
 			{
-				FSlateBrush* CachedBrush = WeakSceneOutliner.Pin()->GetCachedIconForClass(Actor->GetClass()->GetFName());
+				FName IconName = Actor->GetCustomIconName();
+				if (IconName == NAME_None)
+				{
+					IconName = Actor->GetClass()->GetFName();
+				}
+
+				const FSlateBrush* CachedBrush = WeakSceneOutliner.Pin()->GetCachedIconForClass(IconName);
 				if (CachedBrush != nullptr)
 				{
 					return CachedBrush;
 				}
 				else
 				{
-					const FSlateBrush* FoundSlateBrush = FClassIconFinder::FindIconForActor(const_cast<AActor*>(Actor));
-					WeakSceneOutliner.Pin()->CacheIconForClass(Actor->GetClass()->GetFName(), const_cast<FSlateBrush*>(FoundSlateBrush));
+
+					const FSlateBrush* FoundSlateBrush = FClassIconFinder::FindIconForActor(Actor);
+					WeakSceneOutliner.Pin()->CacheIconForClass(IconName, FoundSlateBrush);
 					return FoundSlateBrush;
 				}
 			}
diff --git a/Engine/Source/Editor/SceneOutliner/Public/ISceneOutliner.h b/Engine/Source/Editor/SceneOutliner/Public/ISceneOutliner.h
--- a/Engine/Source/Editor/SceneOutliner/Public/ISceneOutliner.h
+++ b/Engine/Source/Editor/SceneOutliner/Public/ISceneOutliner.h
@@ -32,8 +32,8 @@ public:
 	virtual void SetKeyboardFocus() = 0;
 
 	/** Gets the cached icon for this class name */
-	virtual FSlateBrush* GetCachedIconForClass(FName InClassName) const = 0;
+	virtual const FSlateBrush* GetCachedIconForClass(FName InClassName) const = 0;
 
 	/** Sets the cached icon for this class name */
-	virtual void CacheIconForClass(FName InClassName, FSlateBrush* InSlateBrush) = 0;
+	virtual void CacheIconForClass(FName InClassName, const FSlateBrush* InSlateBrush) = 0;
 };
diff --git a/Engine/Source/Editor/UMGEditor/Private/UMGEditorModule.cpp b/Engine/Source/Editor/UMGEditor/Private/UMGEditorModule.cpp
--- a/Engine/Source/Editor/UMGEditor/Private/UMGEditorModule.cpp
+++ b/Engine/Source/Editor/UMGEditor/Private/UMGEditorModule.cpp
@@ -40,14 +40,6 @@
 
 const FName UMGEditorAppIdentifier = FName(TEXT("UMGEditorApp"));
 
-PRAGMA_DISABLE_DEPRECATION_WARNINGS
-const FSlateBrush* GetEditorIcon_Deprecated(UWidget* Widget)
-{
-	const FSlateBrush* Brush = Widget->GetEditorIcon();
-	return Brush ? Brush : FClassIconFinder::FindIconForClass(Widget->GetClass());
-}
-PRAGMA_ENABLE_DEPRECATION_WARNINGS
-
 class FUMGEditorModule : public IUMGEditorModule, public IBlueprintCompiler, public FGCObject
 {
 public:
diff --git a/Engine/Source/Editor/UnrealEd/Private/ClassIconFinder.cpp b/Engine/Source/Editor/UnrealEd/Private/ClassIconFinder.cpp
--- a/Engine/Source/Editor/UnrealEd/Private/ClassIconFinder.cpp
+++ b/Engine/Source/Editor/UnrealEd/Private/ClassIconFinder.cpp
@@ -15,7 +15,7 @@ const FSlateBrush* FClassIconFinder::FindIconForActors(const TArray< TWeakObject
 
 	for( int32 ActorIdx = 0; ActorIdx < InActors.Num(); ++ActorIdx )
 	{
-		TWeakObjectPtr<AActor> Actor = InActors[ActorIdx];
+		TWeakObjectPtr<const AActor> Actor = InActors[ActorIdx];
 		UClass* ObjClass = Actor->GetClass();
 		check(ObjClass);
 
@@ -44,26 +44,19 @@ const FSlateBrush* FClassIconFinder::FindIconForActors(const TArray< TWeakObject
 	return CommonIcon.GetOptionalIcon();
 }
 
-FSlateIcon FClassIconFinder::FindSlateIconForActor( const TWeakObjectPtr<AActor>& InActor )
+FSlateIcon FClassIconFinder::FindSlateIconForActor( const TWeakObjectPtr<const AActor>& InActor )
 {
 	// Actor specific overrides to normal per-class icons
 	const AActor* Actor = InActor.Get();
 
 	if ( Actor )
 	{
-		const ABrush* Brush = Cast< ABrush >( Actor );
-		if ( Brush )
+		FName IconName = Actor->GetCustomIconName();
+		if (IconName != NAME_None)
 		{
-			if (Brush_Add == Brush->BrushType)
-			{
-				return FSlateIconFinder::FindIcon("ClassIcon.BrushAdditive");
-			}
-			else if (Brush_Subtract == Brush->BrushType)
-			{
-				return FSlateIconFinder::FindIcon("ClassIcon.BrushSubtractive");
-			}
+			return FSlateIconFinder::FindIcon(IconName);
 		}
-
+	
 		// Actor didn't specify an icon - fallback on the class icon
 		return FSlateIconFinder::FindIconForClass(Actor->GetClass());
 	}
@@ -74,7 +67,7 @@ FSlateIcon FClassIconFinder::FindSlateIconForActor( const TWeakObjectPtr<AActor>
 	}
 }
 
-const FSlateBrush* FClassIconFinder::FindIconForActor( const TWeakObjectPtr<AActor>& InActor )
+const FSlateBrush* FClassIconFinder::FindIconForActor( const TWeakObjectPtr<const AActor>& InActor )
 {
 	return FindSlateIconForActor(InActor).GetOptionalIcon();
 }
diff --git a/Engine/Source/Editor/UnrealEd/Public/ClassIconFinder.h b/Engine/Source/Editor/UnrealEd/Public/ClassIconFinder.h
--- a/Engine/Source/Editor/UnrealEd/Public/ClassIconFinder.h
+++ b/Engine/Source/Editor/UnrealEd/Public/ClassIconFinder.h
@@ -20,10 +20,10 @@ public:
 	UNREALED_API static const FSlateBrush* FindIconForActors(const TArray< TWeakObjectPtr<AActor> >& InActors, UClass*& CommonBaseClass);
 
 	/** Find the small icon to use for the supplied actor */
-	UNREALED_API static const FSlateBrush* FindIconForActor( const TWeakObjectPtr<AActor>& InActor );
+	UNREALED_API static const FSlateBrush* FindIconForActor( const TWeakObjectPtr<const AActor>& InActor );
 	
 	/** Find the small icon to use for the supplied actor */
-	UNREALED_API static FSlateIcon FindSlateIconForActor( const TWeakObjectPtr<AActor>& InActor );
+	UNREALED_API static FSlateIcon FindSlateIconForActor( const TWeakObjectPtr<const AActor>& InActor );
 
 	/** Utility function to convert a Blueprint into the most suitable class possible for use by the icon finder */
 	UNREALED_API static const UClass* GetIconClassForBlueprint(const UBlueprint* InBlueprint);
@@ -37,44 +37,6 @@ public:
 		return FSlateIconFinder::FindCustomIconBrushForClass(InClass, TEXT("ClassThumbnail"), InDefaultName);
 	}
 
-public:
-
-	/** Registers a new style set to use to try to find icons. */
-	DEPRECATED(4.13, "This function is no longer required.")
-	static void RegisterIconSource(const ISlateStyle* StyleSet){}
-
-	/** Unregisters an style set. */
-	DEPRECATED(4.13, "This function is no longer required.")
-	static void UnregisterIconSource(const ISlateStyle* StyleSet){}
-
-	/** Find the small icon to use for the supplied class */
-	DEPRECATED(4.13, "Please use FSlateIconFinder::FindIconForClass directly (followed by FSlateIcon::GetIcon().")
-	static const FSlateBrush* FindIconForClass(const UClass* InClass, const FName& InDefaultName = FName() )
-	{
-		return FSlateIconFinder::FindIconBrushForClass(InClass, InDefaultName);
-	}
-
-	/** Find the small icon name to use for the supplied class */
-	DEPRECATED(4.13, "Please use FSlateIconFinder::FindIconForClass directly (followed by FSlateIcon::GetStyleName().")
-	static FName FindIconNameForClass(const UClass* InClass, const FName& InDefaultName = FName() )
-	{
-		return FSlateIconFinder::FindIconForClass(InClass, InDefaultName).GetStyleName();
-	}
-
-	/** Find the large thumbnail name to use for the supplied class */
-	DEPRECATED(4.13, "Please use FSlateIconFinder::FindCustomIconForClass directly (followed by FSlateIcon::GetStyleName().")
-	static FName FindThumbnailNameForClass(const UClass* InClass, const FName& InDefaultName = FName() )
-	{
-		return FSlateIconFinder::FindCustomIconForClass(InClass, TEXT("ClassThumbnail"), InDefaultName).GetStyleName();
-	}
-
-	/** Find the small icon name to use for the supplied actor */
-	DEPRECATED(4.13, "Please use FindSlateIconForActor directly.")
-	static FName FindIconNameForActor( const TWeakObjectPtr<AActor>& InActor )
-	{
-		return FindSlateIconForActor(InActor).GetStyleName();
-	}
-
 private:
 	UNREALED_API static const FSlateBrush* LookupBrush(FName IconName);
 };
diff --git a/Engine/Source/Runtime/Engine/Classes/Engine/Brush.h b/Engine/Source/Runtime/Engine/Classes/Engine/Brush.h
--- a/Engine/Source/Runtime/Engine/Classes/Engine/Brush.h
+++ b/Engine/Source/Runtime/Engine/Classes/Engine/Brush.h
@@ -144,6 +144,7 @@ public:
 	virtual void PostEditMove(bool bFinished) override;
 	virtual void PostEditChangeProperty(FPropertyChangedEvent& PropertyChangedEvent) override;
 	virtual void PostEditChangeChainProperty(FPropertyChangedChainEvent& PropertyChangedEvent) override;
+	virtual FName GetCustomIconName() const override;
 #endif // WITH_EDITOR
 
 	virtual bool Modify(bool bAlwaysMarkDirty = false) override;
diff --git a/Engine/Source/Runtime/Engine/Classes/GameFramework/Actor.h b/Engine/Source/Runtime/Engine/Classes/GameFramework/Actor.h
--- a/Engine/Source/Runtime/Engine/Classes/GameFramework/Actor.h
+++ b/Engine/Source/Runtime/Engine/Classes/GameFramework/Actor.h
@@ -1809,6 +1809,8 @@ public:
 	/** Returns NumUncachedStaticLightingInteractions for this actor */
 	const int32 GetNumUncachedStaticLightingInteractions() const;
 
+	/** Returns a custom brush icon name to use in place of the automatic class icon where actors are represented via 2d icons in the editor (e.g scene outliner and menus) */
+	virtual FName GetCustomIconName() const { return NAME_None; }
 #endif		// WITH_EDITOR
 
 	/**
diff --git a/Engine/Source/Runtime/Engine/Private/Brush.cpp b/Engine/Source/Runtime/Engine/Private/Brush.cpp
--- a/Engine/Source/Runtime/Engine/Private/Brush.cpp
+++ b/Engine/Source/Runtime/Engine/Private/Brush.cpp
@@ -74,6 +74,22 @@ void ABrush::PostEditChangeChainProperty(FPropertyChangedChainEvent& PropertyCha
 	Super::PostEditChangeChainProperty(PropertyChangedEvent);
 }
 
+FName ABrush::GetCustomIconName() const
+{
+	if (BrushType == Brush_Add)
+	{
+		static const FName AdditiveIconName("ClassIcon.BrushAdditive");
+		return AdditiveIconName;
+	}
+	else if (BrushType == Brush_Subtract)
+	{
+		static const FName SubtactiveIconName("ClassIcon.BrushSubtractive");
+		return SubtactiveIconName;
+	}
+
+	return NAME_None;
+}
+
 void ABrush::PostEditChangeProperty(FPropertyChangedEvent& PropertyChangedEvent)
 {
 	if(Brush)
