Make the namespace an import option

diff --git a/Engine/Config/BaseEditorPerProjectUserSettings.ini b/Engine/Config/BaseEditorPerProjectUserSettings.ini
--- a/Engine/Config/BaseEditorPerProjectUserSettings.ini
+++ b/Engine/Config/BaseEditorPerProjectUserSettings.ini
@@ -48,6 +48,7 @@ PropertyMatrix_NumberOfPasteOperationsBeforeWarning=20
 bSCSEditorShowGrid=true
 
 ; Import
+bKeepFbxNamespace=False
 bShowImportDialogAtReimport=False
 
 ; Export
diff --git a/Engine/Source/Editor/UnrealEd/Classes/Editor/EditorPerProjectUserSettings.h b/Engine/Source/Editor/UnrealEd/Classes/Editor/EditorPerProjectUserSettings.h
--- a/Engine/Source/Editor/UnrealEd/Classes/Editor/EditorPerProjectUserSettings.h
+++ b/Engine/Source/Editor/UnrealEd/Classes/Editor/EditorPerProjectUserSettings.h
@@ -71,6 +71,10 @@ class UEditorPerProjectUserSettings : public UObject
 	UPROPERTY(EditAnywhere, config, Category=HotReload)
 	uint32 bShowCompilerLogOnCompileError : 1;
 
+	/** If enabled, the fbx parser will keep the fbx namespaces, otherwise the namespace will be append to fbx node. */
+	UPROPERTY(EditAnywhere, config, Category = Import)
+	uint32 bKeepFbxNamespace : 1;
+
 	/** If enabled, the fbx option dialog will show when user re-import a fbx */
 	UPROPERTY(EditAnywhere, config, Category = Import)
 	uint32 bShowImportDialogAtReimport : 1;
diff --git a/Engine/Source/Editor/UnrealEd/Private/Fbx/FbxMainImport.cpp b/Engine/Source/Editor/UnrealEd/Private/Fbx/FbxMainImport.cpp
--- a/Engine/Source/Editor/UnrealEd/Private/Fbx/FbxMainImport.cpp
+++ b/Engine/Source/Editor/UnrealEd/Private/Fbx/FbxMainImport.cpp
@@ -31,6 +31,7 @@
 #include "ARFilter.h"
 #include "Animation/Skeleton.h"
 #include "HAL/PlatformApplicationMisc.h"
+#include "Editor/EditorPerProjectUserSettings.h"
 
 DEFINE_LOG_CATEGORY(LogFbx);
 
@@ -947,6 +948,8 @@ void FFbxImporter::FixMaterialClashName()
 
 void FFbxImporter::EnsureNodeNameAreValid()
 {
+	const bool bKeepNamespace = GetDefault<UEditorPerProjectUserSettings>()->bKeepFbxNamespace;
+
 	TSet<FString> AllNodeName;
 	int32 CurrentNameIndex = 1;
 	for (int32 NodeIndex = 0; NodeIndex < Scene->GetNodeCount(); ++NodeIndex)
@@ -969,25 +972,28 @@ void FFbxImporter::EnsureNodeNameAreValid()
 					FFbxErrors::Generic_LoadingSceneFailed);
 			}
 		}
-		if (NodeName.Contains(TEXT(":")))
-		{
-			NodeName = NodeName.Replace(TEXT(":"), TEXT("_"));
-			Node->SetName(TCHAR_TO_UTF8(*NodeName));
-		}
-		if (AllNodeName.Contains(NodeName))
+		if (bKeepNamespace)
 		{
-			FString UniqueNodeName;
-			do
+			if (NodeName.Contains(TEXT(":")))
 			{
-				UniqueNodeName = NodeName + FString::FromInt(CurrentNameIndex++);
-			} while (AllNodeName.Contains(UniqueNodeName));
-			Node->SetName(TCHAR_TO_UTF8(*UniqueNodeName));
-			if (!GIsAutomationTesting)
+				NodeName = NodeName.Replace(TEXT(":"), TEXT("_"));
+				Node->SetName(TCHAR_TO_UTF8(*NodeName));
+			}
+			if (AllNodeName.Contains(NodeName))
 			{
-				AddTokenizedErrorMessage(
-					FTokenizedMessage::Create(EMessageSeverity::Warning,
-						FText::Format(LOCTEXT("FbxImport_NodeNameClash", "FBX File Loading: Found name clash, node '{0}' was rename '{1}'"), FText::FromString(NodeName), FText::FromString(UniqueNodeName))),
-					FFbxErrors::Generic_LoadingSceneFailed);
+				FString UniqueNodeName;
+				do
+				{
+					UniqueNodeName = NodeName + FString::FromInt(CurrentNameIndex++);
+				} while (AllNodeName.Contains(UniqueNodeName));
+				Node->SetName(TCHAR_TO_UTF8(*UniqueNodeName));
+				if (!GIsAutomationTesting)
+				{
+					AddTokenizedErrorMessage(
+						FTokenizedMessage::Create(EMessageSeverity::Warning,
+							FText::Format(LOCTEXT("FbxImport_NodeNameClash", "FBX File Loading: Found name clash, node '{0}' was rename '{1}'"), FText::FromString(NodeName), FText::FromString(UniqueNodeName))),
+						FFbxErrors::Generic_LoadingSceneFailed);
+				}
 			}
 		}
 		AllNodeName.Add(NodeName);
