Unshelved from pending changelist '4229008':

Fix fbx importer staticmesh the light map index, the index was check before the build.

diff --git a/Engine/Source/Editor/UnrealEd/Private/Fbx/FbxStaticMeshImport.cpp b/Engine/Source/Editor/UnrealEd/Private/Fbx/FbxStaticMeshImport.cpp
--- a/Engine/Source/Editor/UnrealEd/Private/Fbx/FbxStaticMeshImport.cpp
+++ b/Engine/Source/Editor/UnrealEd/Private/Fbx/FbxStaticMeshImport.cpp
@@ -2018,8 +2018,6 @@ void UnFbx::FFbxImporter::PostImportStaticMesh(UStaticMesh* StaticMesh, TArray<F
 	{
 		MeshDescription->TriangulateMesh();
 	}
-	//Make sure the light map UVChannel is valid
-	StaticMesh->EnforceLightmapRestrictions();
 
 	//Prebuild the static mesh when we use LodGroup and we want to modify the LodNumber
 	if (!ImportOptions->bImportScene)
@@ -2080,6 +2078,9 @@ void UnFbx::FFbxImporter::PostImportStaticMesh(UStaticMesh* StaticMesh, TArray<F
 		AddTokenizedErrorMessage(FTokenizedMessage::Create(EMessageSeverity::Warning, Error), FFbxErrors::StaticMesh_BuildError);
 	}
 
+	//Make sure the light map UVChannel is valid, this must be done after the build
+	StaticMesh->EnforceLightmapRestrictions();
+
 	//Set the specified LOD distances for every LODs we have to do this after the build in case there is a specified Lod Group
 	if (!ImportOptions->bAutoComputeLodDistances && !ImportOptions->bImportScene)
 	{
