Added include to fix non unity builds.

diff --git a/Engine/Source/Editor/UnrealEd/Private/DragAndDrop/AssetDragDropOp.cpp b/Engine/Source/Editor/UnrealEd/Private/DragAndDrop/AssetDragDropOp.cpp
--- a/Engine/Source/Editor/UnrealEd/Private/DragAndDrop/AssetDragDropOp.cpp
+++ b/Engine/Source/Editor/UnrealEd/Private/DragAndDrop/AssetDragDropOp.cpp
@@ -3,6 +3,7 @@
 #include "DragAndDrop/AssetDragDropOp.h"
 #include "Engine/Level.h"
 #include "ActorFactories/ActorFactory.h"
+#include "GameFramework/Actor.h"
 #include "Widgets/SOverlay.h"
 #include "Widgets/Layout/SBox.h"
 #include "AssetThumbnail.h"
