Properly handled FPackageName::TryConvertFilenameToLongPackageName() failing in Cache Thumbnail.

diff --git a/Engine/Source/Editor/UnrealEd/Private/ObjectTools.cpp b/Engine/Source/Editor/UnrealEd/Private/ObjectTools.cpp
--- a/Engine/Source/Editor/UnrealEd/Private/ObjectTools.cpp
+++ b/Engine/Source/Editor/UnrealEd/Private/ObjectTools.cpp
@@ -4046,20 +4046,20 @@ namespace ThumbnailTools
 	/** Searches for an object's thumbnail in memory and returns it if found */
 	FObjectThumbnail* FindCachedThumbnailInPackage( const FString& InPackageFileName, const FName InObjectFullName )
 	{
-		FObjectThumbnail* FoundThumbnail = NULL;
+		FObjectThumbnail* FoundThumbnail = nullptr;
 
 		FString PackageName = InPackageFileName;
-		FPackageName::TryConvertFilenameToLongPackageName(PackageName, PackageName);
-
-		// First check to see if the package is already in memory.  If it is, some or all of the thumbnails
-		// may already be loaded and ready.
-		UObject* PackageOuter = NULL;
-		UPackage* Package = FindPackage( PackageOuter, *PackageName);
-		if( Package != NULL )
+		if (FPackageName::TryConvertFilenameToLongPackageName(PackageName, PackageName))
 		{
-			FoundThumbnail = FindCachedThumbnailInPackage( Package, InObjectFullName );
+			// First check to see if the package is already in memory.  If it is, some or all of the thumbnails
+			// may already be loaded and ready.
+			UObject* PackageOuter = nullptr;
+			UPackage* Package = FindPackage(PackageOuter, *PackageName);
+			if (Package != nullptr)
+			{
+				FoundThumbnail = FindCachedThumbnailInPackage(Package, InObjectFullName);
+			}
 		}
-
 		return FoundThumbnail;
 	}
 
@@ -4069,15 +4069,15 @@ namespace ThumbnailTools
 	const FObjectThumbnail* FindCachedThumbnail( const FString& InFullName )
 	{
 		// Determine the package file path/name for the specified object
-		FString PackageName = GetPackageNameForObject( InFullName );
+		const FString PackageFileName = GetPackageNameForObject(InFullName);
 
-		if ( PackageName.IsEmpty() )
+		if (PackageFileName.IsEmpty())
 		{
 			// Couldn't find the package
 			return nullptr;
 		}
 
-		return FindCachedThumbnailInPackage( PackageName, FName( *InFullName ) );
+		return FindCachedThumbnailInPackage( PackageFileName, FName( *InFullName ) );
 	}
 
 
