Changes to more accurately represent android behavior in PIE and UMG

diff --git a/Engine/Source/Editor/UMGEditor/Private/Designer/SDesignerView.cpp b/Engine/Source/Editor/UMGEditor/Private/Designer/SDesignerView.cpp
--- a/Engine/Source/Editor/UMGEditor/Private/Designer/SDesignerView.cpp
+++ b/Engine/Source/Editor/UMGEditor/Private/Designer/SDesignerView.cpp
@@ -3210,7 +3210,7 @@ void SDesignerView::HandleOnCommonResolutionSelected(FPlayScreenResolution InRes
 
 	ScaleFactor = 1.0f;
 	ULevelEditorPlaySettings* PlayInSettings = GetMutableDefault<ULevelEditorPlaySettings>();
-	UDeviceProfile* DeviceProfile = UDeviceProfileManager::Get().FindProfile(PreviewOverrideName, false);
+	const UDeviceProfile* DeviceProfile = UDeviceProfileManager::Get().FindProfile(PreviewOverrideName, false);
 	if (DeviceProfile)
 	{
 		PlayInSettings->RescaleForMobilePreview(DeviceProfile, PreviewWidth, PreviewHeight, ScaleFactor);
@@ -3530,6 +3530,15 @@ FReply SDesignerView::HandleSwapAspectRatioClicked()
 	int32 OldPreviewHeight = PreviewHeight;
 	PreviewHeight = PreviewWidth;
 	PreviewWidth = OldPreviewHeight;
+
+	ScaleFactor = 1.0f;
+	ULevelEditorPlaySettings* PlayInSettings = GetMutableDefault<ULevelEditorPlaySettings>();
+	const UDeviceProfile* DeviceProfile = UDeviceProfileManager::Get().FindProfile(PreviewOverrideName, false);
+	if (DeviceProfile)
+	{
+		PlayInSettings->RescaleForMobilePreview(DeviceProfile, PreviewWidth, PreviewHeight, ScaleFactor);
+	}
+
 	bPreviewIsPortrait = (PreviewHeight > PreviewWidth);
 	GConfig->SetInt(*ConfigSectionName, TEXT("PreviewWidth"), PreviewWidth, GEditorPerProjectIni);
 	GConfig->SetInt(*ConfigSectionName, TEXT("PreviewHeight"), PreviewHeight, GEditorPerProjectIni);
diff --git a/Engine/Source/Editor/UnrealEd/Classes/Settings/LevelEditorPlaySettings.h b/Engine/Source/Editor/UnrealEd/Classes/Settings/LevelEditorPlaySettings.h
--- a/Engine/Source/Editor/UnrealEd/Classes/Settings/LevelEditorPlaySettings.h
+++ b/Engine/Source/Editor/UnrealEd/Classes/Settings/LevelEditorPlaySettings.h
@@ -474,5 +474,5 @@ public:
 
 	FMargin CalculateCustomUnsafeZones(TArray<FVector2D>& CustomSafeZoneStarts, TArray<FVector2D>& CustomSafeZoneDimensions, FString& DeviceType, FVector2D PreviewSize);
 	FMargin FlipCustomUnsafeZones(TArray<FVector2D>& CustomSafeZoneStarts, TArray<FVector2D>& CustomSafeZoneDimensions, FString& DeviceType, FVector2D PreviewSize);
-	void RescaleForMobilePreview(class UDeviceProfile* DeviceProfile, int32 &PreviewWidth, int32 &PreviewHeight, float &ScaleFactor);
+	void RescaleForMobilePreview(const class UDeviceProfile* DeviceProfile, int32 &PreviewWidth, int32 &PreviewHeight, float &ScaleFactor);
 };
diff --git a/Engine/Source/Editor/UnrealEd/Private/Settings/LevelEditorPlaySettingsCustomization.h b/Engine/Source/Editor/UnrealEd/Private/Settings/LevelEditorPlaySettingsCustomization.h
--- a/Engine/Source/Editor/UnrealEd/Private/Settings/LevelEditorPlaySettingsCustomization.h
+++ b/Engine/Source/Editor/UnrealEd/Private/Settings/LevelEditorPlaySettingsCustomization.h
@@ -347,6 +347,12 @@ protected:
 
 		if (!PlayInSettings->DeviceToEmulate.IsEmpty())
 		{
+			const UDeviceProfile* DeviceProfile = UDeviceProfileManager::Get().FindProfile(PlayInSettings->DeviceToEmulate, false);
+			if(DeviceProfile != nullptr)
+			{
+				float ScaleFactor = 1.0f;
+				PlayInSettings->RescaleForMobilePreview(DeviceProfile, NewWidth, NewHeight, ScaleFactor);
+			}
 			PlayInSettings->PIESafeZoneOverride = PlayInSettings->CalculateCustomUnsafeZones(PlayInSettings->CustomUnsafeZoneStarts, PlayInSettings->CustomUnsafeZoneDimensions, PlayInSettings->DeviceToEmulate, FVector2D(NewWidth, NewHeight));
 		}
 		else
diff --git a/Engine/Source/Editor/UnrealEd/Private/Settings/SettingsClasses.cpp b/Engine/Source/Editor/UnrealEd/Private/Settings/SettingsClasses.cpp
--- a/Engine/Source/Editor/UnrealEd/Private/Settings/SettingsClasses.cpp
+++ b/Engine/Source/Editor/UnrealEd/Private/Settings/SettingsClasses.cpp
@@ -553,39 +553,65 @@ FMargin ULevelEditorPlaySettings::FlipCustomUnsafeZones(TArray<FVector2D>& Custo
 	return CustomSafeZoneOverride;
 }
 
-void ULevelEditorPlaySettings::RescaleForMobilePreview(UDeviceProfile* DeviceProfile, int32 &PreviewWidth, int32 &PreviewHeight, float &ScaleFactor)
+void ULevelEditorPlaySettings::RescaleForMobilePreview(const UDeviceProfile* DeviceProfile, int32 &PreviewWidth, int32 &PreviewHeight, float &ScaleFactor)
 {
 	bool bShouldScale = false;
-	float CVarMobileContentScaleFactor = FCString::Atof(*DeviceProfile->GetCVarValue(TEXT("r.MobileContentScaleFactor")));
-	if (CVarMobileContentScaleFactor != 0)
+	float CVarMobileContentScaleFactor = 0.0f;
+	const FString ScaleFactorString = DeviceProfile->GetCVarValue(TEXT("r.MobileContentScaleFactor"));
+	if (ScaleFactorString != FString())
 	{
-		bShouldScale = true;
+		CVarMobileContentScaleFactor = FCString::Atof(*ScaleFactorString);
+
+		if (!FMath::IsNearlyEqual(CVarMobileContentScaleFactor, 0.0f))
+		{
+			bShouldScale = true;
+			ScaleFactor = CVarMobileContentScaleFactor;
+		}
 	}
 	else
 	{
-		if (DeviceProfile->GetConsolidatedCVarValue(TEXT("r.MobileContentScaleFactor"), CVarMobileContentScaleFactor, true))
+		TMap<FString, FString> ParentValues;
+		DeviceProfile->GatherParentCVarInformationRecursively(ParentValues);
+		const FString* ParentScaleFactorPtr = ParentValues.Find(TEXT("r.MobileContentScaleFactor"));
+		if (ParentScaleFactorPtr != nullptr)
 		{
-			bShouldScale = true;
+			CVarMobileContentScaleFactor = FCString::Atof(**ParentScaleFactorPtr);
+			if (!FMath::IsNearlyEqual(CVarMobileContentScaleFactor, 0.0f))
+			{
+				bShouldScale = true;
+				ScaleFactor = CVarMobileContentScaleFactor;
+			}
 		}
 	}
 	if (bShouldScale)
 	{
 		if (DeviceProfile->DeviceType == TEXT("Android"))
 		{
+			const float OriginalPreviewWidth = PreviewWidth;
+			const float OriginalPreviewHeight = PreviewHeight;
+			float TempPreviewHeight = 0.0f;
+			float TempPreviewWidth = 0.0f;
+			// Portrait
 			if (PreviewHeight > PreviewWidth)
 			{
-				PreviewHeight = 1280;
-				PreviewWidth = 720;
+				TempPreviewHeight = 1280 * ScaleFactor;
+			
 			}
+			// Landscape
 			else
 			{
-				PreviewHeight = 720;
-				PreviewWidth = 1280;
+				TempPreviewHeight = 720 * ScaleFactor;
 			}
+			TempPreviewWidth = TempPreviewHeight * OriginalPreviewWidth / OriginalPreviewHeight + 0.5f;
+			PreviewHeight = (int32)FMath::GridSnap(TempPreviewHeight, 8.0f);
+			PreviewWidth = (int32)FMath::GridSnap(TempPreviewWidth, 8.0f);
+		}
+		else
+		{
+			PreviewWidth *= ScaleFactor;
+			PreviewHeight *= ScaleFactor;
 		}
-		PreviewWidth *= CVarMobileContentScaleFactor;
-		PreviewHeight *= CVarMobileContentScaleFactor;
-		ScaleFactor = CVarMobileContentScaleFactor;
+	
 	}
 }
 
diff --git a/Engine/Source/Runtime/Engine/Classes/DeviceProfiles/DeviceProfile.h b/Engine/Source/Runtime/Engine/Classes/DeviceProfiles/DeviceProfile.h
--- a/Engine/Source/Runtime/Engine/Classes/DeviceProfiles/DeviceProfile.h
+++ b/Engine/Source/Runtime/Engine/Classes/DeviceProfiles/DeviceProfile.h
@@ -94,7 +94,7 @@ public:
 	//~ End UObject Interface
 
 	bool ModifyCVarValue(const FString& CVarName, const FString& CVarValue, bool bAddIfNonExistant = false);
-	FString GetCVarValue(const FString& CVarName);
+	FString GetCVarValue(const FString& CVarName) const;
 
 	/** Lazily generate a consolidated list of CVars, recursing up the device profile hierarchy */
 	const TMap<FString, FString>& GetConsolidatedCVars() const;
diff --git a/Engine/Source/Runtime/Engine/Private/DeviceProfiles/DeviceProfile.cpp b/Engine/Source/Runtime/Engine/Private/DeviceProfiles/DeviceProfile.cpp
--- a/Engine/Source/Runtime/Engine/Private/DeviceProfiles/DeviceProfile.cpp
+++ b/Engine/Source/Runtime/Engine/Private/DeviceProfiles/DeviceProfile.cpp
@@ -223,7 +223,7 @@ bool UDeviceProfile::ModifyCVarValue(const FString& ChangeCVarName, const FStrin
 	return false;
 }
 
-FString UDeviceProfile::GetCVarValue(const FString& CVarName)
+FString UDeviceProfile::GetCVarValue(const FString& CVarName) const
 {
 	auto Index = CVars.IndexOfByPredicate(
 		[&CVarName](const FString& CVar) {
