Fix animation Range import, prevent changing the option when validating the anim range.

diff --git a/Engine/Source/Editor/UnrealEd/Private/SkeletalMeshEdit.cpp b/Engine/Source/Editor/UnrealEd/Private/SkeletalMeshEdit.cpp
--- a/Engine/Source/Editor/UnrealEd/Private/SkeletalMeshEdit.cpp
+++ b/Engine/Source/Editor/UnrealEd/Private/SkeletalMeshEdit.cpp
@@ -496,21 +496,21 @@ FbxTimeSpan UnFbx::FFbxImporter::GetAnimationTimeSpan(FbxNode* RootNode, FbxAnim
 			{
 				FbxTime Duration = AnimTimeSpan.GetDuration();
 
-				ImportOption->AnimationRange.X = FMath::Clamp<int32>(ImportOption->AnimationRange.X, StartFrame, StopFrame);
-				ImportOption->AnimationRange.Y = FMath::Clamp<int32>(ImportOption->AnimationRange.Y, StartFrame, StopFrame);
+				int32 AnimRangeX = FMath::Clamp<int32>(ImportOption->AnimationRange.X, StartFrame, StopFrame);
+				int32 AnimRangeY = FMath::Clamp<int32>(ImportOption->AnimationRange.Y, StartFrame, StopFrame);
 
 				FbxLongLong Interval = EachFrame.Get();
 
 				// now set new time
-				if (StartFrame != ImportOption->AnimationRange.X)
+				if (StartFrame != AnimRangeX)
 				{
-					FbxTime NewTime(ImportOption->AnimationRange.X*Interval);
+					FbxTime NewTime(AnimRangeX*Interval);
 					AnimTimeSpan.SetStart(NewTime);
 				}
 
-				if (StopFrame != ImportOption->AnimationRange.Y)
+				if (StopFrame != AnimRangeY)
 				{
-					FbxTime NewTime(ImportOption->AnimationRange.Y*Interval);
+					FbxTime NewTime(AnimRangeY*Interval);
 					AnimTimeSpan.SetStop(NewTime);
 				}
 			}
