Simplifed the code in the case of a null CurveOwner.

diff --git a/Engine/Source/Editor/UnrealEd/Private/SCurveEditor.cpp b/Engine/Source/Editor/UnrealEd/Private/SCurveEditor.cpp
--- a/Engine/Source/Editor/UnrealEd/Private/SCurveEditor.cpp
+++ b/Engine/Source/Editor/UnrealEd/Private/SCurveEditor.cpp
@@ -2463,6 +2463,13 @@ void SCurveEditor::EmptyAllSelection()
 
 void SCurveEditor::ValidateSelection()
 {
+	//remove all selection if we no longer have a curve interface (Curve Owner)
+	if (!CurveOwner)
+	{
+		EmptyAllSelection();
+		return;
+	}
+
 	//remove any invalid keys
 	for(int32 i = 0;i<SelectedKeys.Num();++i)
 	{
