PR #4935: Prevent crash when clicking use selected game mode multiple times (Contributed by [REDACTED_PERSON])

diff --git a/Engine/Source/Editor/UnrealEd/Public/GameModeInfoCustomizer.h b/Engine/Source/Editor/UnrealEd/Public/GameModeInfoCustomizer.h
--- a/Engine/Source/Editor/UnrealEd/Public/GameModeInfoCustomizer.h
+++ b/Engine/Source/Editor/UnrealEd/Public/GameModeInfoCustomizer.h
@@ -156,7 +156,7 @@ public:
 			+SHorizontalBox::Slot()
 			.AutoWidth()
 			[
-				PropertyCustomizationHelpers::MakeUseSelectedButton(FSimpleDelegate::CreateSP(this, &FGameModeInfoCustomizer::OnUseSelectedGameModeClicked, &LayoutBuilder))
+				PropertyCustomizationHelpers::MakeUseSelectedButton(FSimpleDelegate::CreateSP(this, &FGameModeInfoCustomizer::OnUseSelectedGameModeClicked))
 			]
 			+SHorizontalBox::Slot()
 			.AutoWidth()
@@ -166,7 +166,7 @@ public:
 			+SHorizontalBox::Slot()
 			.AutoWidth()
 			[
-				PropertyCustomizationHelpers::MakeNewBlueprintButton(FSimpleDelegate::CreateSP(this, &FGameModeInfoCustomizer::OnClickNewGameMode, &LayoutBuilder))
+				PropertyCustomizationHelpers::MakeNewBlueprintButton(FSimpleDelegate::CreateSP(this, &FGameModeInfoCustomizer::OnClickNewGameMode))
 			]
 		];
 
@@ -320,7 +320,7 @@ public:
 		}
 	}
 
-	void OnUseSelectedGameModeClicked(IDetailLayoutBuilder* DetailLayout)
+	void OnUseSelectedGameModeClicked()
 	{
 		FEditorDelegates::LoadSelectedAssetsIfNeeded.Broadcast();
 
@@ -329,14 +329,9 @@ public:
 		{
 			DefaultGameModeClassHandle->SetValueFromFormattedString(SelectedClass->GetPathName());
 		}
-
-		if (DetailLayout)
-		{
-			DetailLayout->ForceRefreshDetails();
-		}
 	}
 
-	void OnClickNewGameMode(IDetailLayoutBuilder* DetailLayout)
+	void OnClickNewGameMode()
 	{
 		// Create a new GameMode BP
 		UBlueprint* Blueprint = FKismetEditorUtilities::CreateBlueprintFromClass(LOCTEXT("CreateNewGameMode", "Create New GameMode"), AGameModeBase::StaticClass(), TEXT("NewGameMode"));
@@ -345,11 +340,6 @@ public:
 		{
 			DefaultGameModeClassHandle->SetValueFromFormattedString(Blueprint->GeneratedClass->GetPathName());
 		}
-
-		if (DetailLayout)
-		{
-			DetailLayout->ForceRefreshDetails();
-		}
 	}
 
 	/** Are we allowed to modify the currently selected GameMode */
