Fixed string parsing when looking at parent cvar values

diff --git a/Engine/Source/Editor/UnrealEd/Private/Settings/SettingsClasses.cpp b/Engine/Source/Editor/UnrealEd/Private/Settings/SettingsClasses.cpp
--- a/Engine/Source/Editor/UnrealEd/Private/Settings/SettingsClasses.cpp
+++ b/Engine/Source/Editor/UnrealEd/Private/Settings/SettingsClasses.cpp
@@ -575,7 +575,11 @@ void ULevelEditorPlaySettings::RescaleForMobilePreview(const UDeviceProfile* Dev
 		const FString* ParentScaleFactorPtr = ParentValues.Find(TEXT("r.MobileContentScaleFactor"));
 		if (ParentScaleFactorPtr != nullptr)
 		{
-			CVarMobileContentScaleFactor = FCString::Atof(**ParentScaleFactorPtr);
+			FString CompleteString = *ParentScaleFactorPtr;
+			FString DiscardString;
+			FString ValueString;
+			CompleteString.Split(TEXT("="), &DiscardString, &ValueString);
+			CVarMobileContentScaleFactor = FCString::Atof(*ValueString);
 			if (!FMath::IsNearlyEqual(CVarMobileContentScaleFactor, 0.0f))
 			{
 				bShouldScale = true;
