Fix for resolutions increasing when swapping landscape/portrait

diff --git a/Engine/Source/Editor/UMGEditor/Private/Designer/SDesignerView.cpp b/Engine/Source/Editor/UMGEditor/Private/Designer/SDesignerView.cpp
--- a/Engine/Source/Editor/UMGEditor/Private/Designer/SDesignerView.cpp
+++ b/Engine/Source/Editor/UMGEditor/Private/Designer/SDesignerView.cpp
@@ -3534,7 +3534,9 @@ FReply SDesignerView::HandleSwapAspectRatioClicked()
 	ScaleFactor = 1.0f;
 	ULevelEditorPlaySettings* PlayInSettings = GetMutableDefault<ULevelEditorPlaySettings>();
 	const UDeviceProfile* DeviceProfile = UDeviceProfileManager::Get().FindProfile(PreviewOverrideName, false);
-	if (DeviceProfile)
+
+	// Rescale the swapped sizes if we are on Android
+	if (DeviceProfile && DeviceProfile->DeviceType == TEXT("Android"))
 	{
 		PlayInSettings->RescaleForMobilePreview(DeviceProfile, PreviewWidth, PreviewHeight, ScaleFactor);
 	}
diff --git a/Engine/Source/Editor/UnrealEd/Private/Settings/LevelEditorPlaySettingsCustomization.h b/Engine/Source/Editor/UnrealEd/Private/Settings/LevelEditorPlaySettingsCustomization.h
--- a/Engine/Source/Editor/UnrealEd/Private/Settings/LevelEditorPlaySettingsCustomization.h
+++ b/Engine/Source/Editor/UnrealEd/Private/Settings/LevelEditorPlaySettingsCustomization.h
@@ -348,7 +348,8 @@ protected:
 		if (!PlayInSettings->DeviceToEmulate.IsEmpty())
 		{
 			const UDeviceProfile* DeviceProfile = UDeviceProfileManager::Get().FindProfile(PlayInSettings->DeviceToEmulate, false);
-			if(DeviceProfile != nullptr)
+			// Rescale the swapped sizes if we are on Android
+			if (DeviceProfile && DeviceProfile->DeviceType == TEXT("Android"))
 			{
 				float ScaleFactor = 1.0f;
 				PlayInSettings->RescaleForMobilePreview(DeviceProfile, NewWidth, NewHeight, ScaleFactor);
