Assets that can be dragged into the material graph now indicate that with a checkmark

diff --git a/Engine/Source/Editor/UnrealEd/Classes/MaterialGraph/MaterialGraphSchema.h b/Engine/Source/Editor/UnrealEd/Classes/MaterialGraph/MaterialGraphSchema.h
--- a/Engine/Source/Editor/UnrealEd/Classes/MaterialGraph/MaterialGraphSchema.h
+++ b/Engine/Source/Editor/UnrealEd/Classes/MaterialGraph/MaterialGraphSchema.h
@@ -202,6 +202,7 @@ class UMaterialGraphSchema : public UEdGraphSchema
 	virtual void ForceVisualizationCacheClear() const override;
 	virtual void OnPinConnectionDoubleCicked(UEdGraphPin* PinA, UEdGraphPin* PinB, const FVector2D& GraphPosition) const override;
 	virtual bool SafeDeleteNodeFromGraph(UEdGraph* Graph, UEdGraphNode* NodeToDelete) const override;
+	virtual void GetAssetsGraphHoverMessage(const TArray<FAssetData>& Assets, const UEdGraph* HoverGraph, FString& OutTooltipText, bool& OutOkIcon) const;
 	//~ End UEdGraphSchema Interface
 
 private:
diff --git a/Engine/Source/Editor/UnrealEd/Private/MaterialGraphSchema.cpp b/Engine/Source/Editor/UnrealEd/Private/MaterialGraphSchema.cpp
--- a/Engine/Source/Editor/UnrealEd/Private/MaterialGraphSchema.cpp
+++ b/Engine/Source/Editor/UnrealEd/Private/MaterialGraphSchema.cpp
@@ -950,4 +950,35 @@ bool UMaterialGraphSchema::SafeDeleteNodeFromGraph(UEdGraph* Graph, UEdGraphNode
 	return true;
 }
 
+void UMaterialGraphSchema::GetAssetsGraphHoverMessage(const TArray<FAssetData>& Assets, const UEdGraph* HoverGraph, FString& OutTooltipText, bool& OutOkIcon) const
+{
+	OutOkIcon = false;
+
+	for (int32 AssetIdx = 0; AssetIdx < Assets.Num(); ++AssetIdx)
+	{
+		UObject* Asset = Assets[AssetIdx].GetAsset();
+		UClass* MaterialExpressionClass = Cast<UClass>(Asset);
+		UMaterialFunctionInterface* Func = Cast<UMaterialFunctionInterface>(Asset);
+		UTexture* Tex = Cast<UTexture>(Asset);
+		UMaterialParameterCollection* ParameterCollection = Cast<UMaterialParameterCollection>(Asset);
+
+		if (MaterialExpressionClass && MaterialExpressionClass->IsChildOf(UMaterialExpression::StaticClass()))
+		{
+			OutOkIcon = true;
+		}
+		else if (Func)
+		{
+			OutOkIcon = true;
+		}
+		else if (Tex)
+		{
+			OutOkIcon = true;
+		}
+		else if (ParameterCollection)
+		{
+			OutOkIcon = true;
+		}
+	}
+}
+
 #undef LOCTEXT_NAMESPACE
