Linux: If our default settings are empty help fill in the proper name

diff --git a/Engine/Source/Editor/UnrealEd/Private/Kismet2/DebuggerCommands.cpp b/Engine/Source/Editor/UnrealEd/Private/Kismet2/DebuggerCommands.cpp
--- a/Engine/Source/Editor/UnrealEd/Private/Kismet2/DebuggerCommands.cpp
+++ b/Engine/Source/Editor/UnrealEd/Private/Kismet2/DebuggerCommands.cpp
@@ -282,6 +282,10 @@ FPlayWorldCommands::FPlayWorldCommands()
 		{
 			PlayPlatformName = TEXT("MacNoEditor");
 		}
+		else if (RunningPlatformName == TEXT("Linux"))
+		{
+			PlayPlatformName = TEXT("LinuxNoEditor");
+		}
 
 		if (!PlayPlatformName.IsEmpty())
 		{
