Add imported framerate info to anime sequence

diff --git a/Engine/Plugins/Experimental/AlembicImporter/Source/AlembicLibrary/Private/AbcFile.cpp b/Engine/Plugins/Experimental/AlembicImporter/Source/AlembicLibrary/Private/AbcFile.cpp
--- a/Engine/Plugins/Experimental/AlembicImporter/Source/AlembicLibrary/Private/AbcFile.cpp
+++ b/Engine/Plugins/Experimental/AlembicImporter/Source/AlembicLibrary/Private/AbcFile.cpp
@@ -460,6 +460,11 @@ const float FAbcFile::GetImportLength() const
 	return ImportLength;
 }
 
+const int32 FAbcFile::GetFramerate() const
+{
+	return FramesPerSecond;
+}
+
 const FBoxSphereBounds& FAbcFile::GetArchiveBounds() const
 {
 	return ArchiveBounds;
diff --git a/Engine/Plugins/Experimental/AlembicImporter/Source/AlembicLibrary/Private/AbcImporter.cpp b/Engine/Plugins/Experimental/AlembicImporter/Source/AlembicLibrary/Private/AbcImporter.cpp
--- a/Engine/Plugins/Experimental/AlembicImporter/Source/AlembicLibrary/Private/AbcImporter.cpp
+++ b/Engine/Plugins/Experimental/AlembicImporter/Source/AlembicLibrary/Private/AbcImporter.cpp
@@ -689,6 +689,8 @@ TArray<UObject*> FAbcImporter::ImportAsSkeletalMesh(UObject* InParent, EObjectFl
 		UAnimSequence* Sequence = CreateObjectInstance<UAnimSequence>(InParent, FString::Printf(TEXT("%s_Animation"), *SkeletalMesh->GetName()), Flags);
 		Sequence->SetSkeleton(Skeleton);
 		Sequence->SequenceLength = AbcFile->GetImportLength();
+		Sequence->ImportFileFramerate = AbcFile->GetFramerate();
+		Sequence->ImportResampleFramerate = AbcFile->GetFramerate();
 		int32 ObjectIndex = 0;
 		uint32 TriangleOffset = 0;
 		uint32 WedgeOffset = 0;
diff --git a/Engine/Plugins/Experimental/AlembicImporter/Source/AlembicLibrary/Public/AbcFile.h b/Engine/Plugins/Experimental/AlembicImporter/Source/AlembicLibrary/Public/AbcFile.h
--- a/Engine/Plugins/Experimental/AlembicImporter/Source/AlembicLibrary/Public/AbcFile.h
+++ b/Engine/Plugins/Experimental/AlembicImporter/Source/AlembicLibrary/Public/AbcFile.h
@@ -61,6 +61,7 @@ public:
 	const int32 GetMaxFrameIndex() const;
 	const float GetImportTimeOffset() const;
 	const float GetImportLength() const;
+	const int32 GetFramerate() const;
 	const FBoxSphereBounds& GetArchiveBounds() const;
 	const bool ContainsHeterogeneousMeshes() const;
 
diff --git a/Engine/Source/Editor/UnrealEd/Private/Factories/EditorFactories.cpp b/Engine/Source/Editor/UnrealEd/Private/Factories/EditorFactories.cpp
--- a/Engine/Source/Editor/UnrealEd/Private/Factories/EditorFactories.cpp
+++ b/Engine/Source/Editor/UnrealEd/Private/Factories/EditorFactories.cpp
@@ -5861,6 +5861,7 @@ EReimportResult::Type UReimportFbxAnimSequenceFactory::Reimport( UObject* Obj )
 
 		// update the data in case the file source has changed
 		ImportData->Update(UFactory::CurrentFilename);
+		AnimSequence->ImportFileFramerate = Importer->GetOriginalFbxFramerate();
 
 		// Try to find the outer package so we can dirty it up
 		if (AnimSequence->GetOuter())
diff --git a/Engine/Source/Editor/UnrealEd/Private/Fbx/FbxMainImport.cpp b/Engine/Source/Editor/UnrealEd/Private/Fbx/FbxMainImport.cpp
--- a/Engine/Source/Editor/UnrealEd/Private/Fbx/FbxMainImport.cpp
+++ b/Engine/Source/Editor/UnrealEd/Private/Fbx/FbxMainImport.cpp
@@ -1066,7 +1066,12 @@ bool FFbxImporter::ImportFile(FString Filename, bool bPreventMaterialNameClash /
 		Result = false;
 		CurPhase = NOTSTARTED;
 	}
-	
+
+	const FbxGlobalSettings& GlobalSettings = Scene->GetGlobalSettings();
+	FbxTime::EMode TimeMode = GlobalSettings.GetTimeMode();
+	//Set the original framerate from the current fbx file
+	OriginalFbxFramerate = FbxTime::GetFrameRate(TimeMode);
+
 	return Result;
 }
 
diff --git a/Engine/Source/Editor/UnrealEd/Private/SkeletalMeshEdit.cpp b/Engine/Source/Editor/UnrealEd/Private/SkeletalMeshEdit.cpp
--- a/Engine/Source/Editor/UnrealEd/Private/SkeletalMeshEdit.cpp
+++ b/Engine/Source/Editor/UnrealEd/Private/SkeletalMeshEdit.cpp
@@ -257,6 +257,7 @@ bool UEditorEngine::ReimportFbxAnimation( USkeleton* Skeleton, UAnimSequence* An
 					// for now it's not importing morph - in the future, this should be optional or saved with asset
 					if (FbxImporter->ValidateAnimStack(SortedLinks, FBXMeshNodeArray, CurAnimStack, ResampleRate, bImportMorphTracks, AnimTimeSpan))
 					{
+						AnimSequence->ImportResampleFramerate = ResampleRate;
 						FbxImporter->ImportAnimation(Skeleton, AnimSequence, Filename, SortedLinks, FBXMeshNodeArray, CurAnimStack, ResampleRate, AnimTimeSpan);
 					}
 				}
@@ -620,6 +621,8 @@ UAnimSequence * UnFbx::FFbxImporter::ImportAnimations(USkeleton* Skeleton, UObje
 		UFbxAnimSequenceImportData* ImportData = UFbxAnimSequenceImportData::GetImportDataForAnimSequence(DestSeq, TemplateImportData);
 		ImportData->Update(UFactory::GetCurrentFilename(), &Md5Hash);
 		ImportData->SourceAnimationName = SourceAnimationName;
+		DestSeq->ImportFileFramerate = GetOriginalFbxFramerate();
+		DestSeq->ImportResampleFramerate = ResampleRate;
 
 		ImportAnimation(Skeleton, DestSeq, Name, SortedLinks, NodeArray, CurAnimStack, ResampleRate, AnimTimeSpan);
 
diff --git a/Engine/Source/Editor/UnrealEd/Public/FbxImporter.h b/Engine/Source/Editor/UnrealEd/Public/FbxImporter.h
--- a/Engine/Source/Editor/UnrealEd/Public/FbxImporter.h
+++ b/Engine/Source/Editor/UnrealEd/Public/FbxImporter.h
@@ -1653,6 +1653,8 @@ public:
 	void ClearTokenizedErrorMessages();
 	void FlushToTokenizedErrorMessage(enum EMessageSeverity::Type Severity);
 
+	float GetOriginalFbxFramerate() { return OriginalFbxFramerate; }
+
 private:
 	friend class FFbxLoggerSetter;
 
@@ -1666,6 +1668,8 @@ private:
 	//Cache to create unique name for mesh. This is use to fix name clash
 	TArray<FString> MeshNamesCache;
 
+	float OriginalFbxFramerate;
+
 private:
 
 
diff --git a/Engine/Source/Runtime/Engine/Classes/Animation/AnimSequence.h b/Engine/Source/Runtime/Engine/Classes/Animation/AnimSequence.h
--- a/Engine/Source/Runtime/Engine/Classes/Animation/AnimSequence.h
+++ b/Engine/Source/Runtime/Engine/Classes/Animation/AnimSequence.h
@@ -358,6 +358,16 @@ class ENGINE_API UAnimSequence : public UAnimSequenceBase
 	UPROPERTY(AssetRegistrySearchable, meta=(DisplayName = "Number of Keys"))
 	int32 NumFrames;
 
+#if WITH_EDITORONLY_DATA
+	/** The DCC framerate of the imported file. UI information only, unit are Hz */
+	UPROPERTY(AssetRegistrySearchable, meta = (DisplayName = "Import File Framerate"))
+	float ImportFileFramerate;
+
+	/** The resample framerate that was computed during import. UI information only, unit are Hz */
+	UPROPERTY(AssetRegistrySearchable, meta = (DisplayName = "Import Resample Framerate"))
+	int32 ImportResampleFramerate;
+#endif
+
 protected:
 	/**
 	 * In the future, maybe keeping RawAnimSequenceTrack + TrackMap as one would be good idea to avoid inconsistent array size
diff --git a/Engine/Source/Runtime/Engine/Private/Animation/AnimSequence.cpp b/Engine/Source/Runtime/Engine/Private/Animation/AnimSequence.cpp
--- a/Engine/Source/Runtime/Engine/Private/Animation/AnimSequence.cpp
+++ b/Engine/Source/Runtime/Engine/Private/Animation/AnimSequence.cpp
@@ -135,6 +135,11 @@ UAnimSequence::UAnimSequence(const FObjectInitializer& ObjectInitializer)
 {
 	RateScale = 1.0;
 	CompressedRawDataSize = 0.f;
+
+#if WITH_EDITORONLY_DATA
+	ImportFileFramerate = 0.0f;
+	ImportResampleFramerate = 0;
+#endif
 }
 
 void UAnimSequence::PostInitProperties()
