- Fix the set range feature to always use the file sample rate so the range match what the user see in the DCC
- Also add some fbx file information to the import dialog

diff --git a/Engine/Source/Editor/DetailCustomizations/Private/FbxImportUIDetails.cpp b/Engine/Source/Editor/DetailCustomizations/Private/FbxImportUIDetails.cpp
--- a/Engine/Source/Editor/DetailCustomizations/Private/FbxImportUIDetails.cpp
+++ b/Engine/Source/Editor/DetailCustomizations/Private/FbxImportUIDetails.cpp
@@ -368,6 +368,48 @@ void FFbxImportUIDetails::CustomizeDetails( IDetailLayoutBuilder& DetailBuilder
 			}
 		}
 	}
+
+	//Information category
+	IDetailCategoryBuilder& InformationCategory = DetailBuilder.EditCategory("FbxFileInformation", FText::GetEmpty());
+	CategoryDefaultProperties.Empty();
+	InformationCategory.GetDefaultProperties(CategoryDefaultProperties);
+	for (TSharedRef<IPropertyHandle> Handle : CategoryDefaultProperties)
+	{
+		FString MetaData = Handle->GetMetaData(TEXT("ImportType"));
+		DetailBuilder.HideProperty(Handle);
+		if (IsImportTypeMetaDataValid(ImportType, MetaData))
+		{
+			FDetailWidgetRow& WidgetRow = DetailBuilder.AddCustomRowToCategory(Handle, Handle->GetPropertyDisplayName());
+			FText PropertyValue;
+			Handle->GetValueAsDisplayText(PropertyValue);
+			WidgetRow.NameContent()
+			.HAlign(HAlign_Fill)
+			[
+				SNew(SBox)
+				.HAlign(HAlign_Left)
+				[
+					SNew(STextBlock)
+					.Font(IDetailLayoutBuilder::GetDetailFont())
+					.Margin(FMargin(2.0f, 2.0f))
+					.Text(Handle->GetPropertyDisplayName())
+					.ToolTipText(Handle->GetToolTipText())
+				]
+			];
+			WidgetRow.ValueContent()
+			.HAlign(HAlign_Fill)
+			[
+				SNew(SBox)
+				.HAlign(HAlign_Left)
+				[
+					SNew(STextBlock)
+					.Font(IDetailLayoutBuilder::GetDetailFont())
+					.Margin(FMargin(2.0f, 2.0f))
+					.Text(PropertyValue)
+					.ToolTipText(PropertyValue)
+				]
+			];
+		}
+	}
 }
 
 void FFbxImportUIDetails::AddSubCategory(IDetailLayoutBuilder& DetailBuilder, FName MainCategoryName, TMap<FString, TArray<TSharedPtr<IPropertyHandle>>>& SubCategoriesProperties, TMap<FString, bool >& SubCategoriesAdvanced, TMap<FString, FText >& SubCategoriesTooltip)
diff --git a/Engine/Source/Editor/UnrealEd/Classes/Factories/FbxImportUI.h b/Engine/Source/Editor/UnrealEd/Classes/Factories/FbxImportUI.h
--- a/Engine/Source/Editor/UnrealEd/Classes/Factories/FbxImportUI.h
+++ b/Engine/Source/Editor/UnrealEd/Classes/Factories/FbxImportUI.h
@@ -164,6 +164,43 @@ public:
 
 	/* Whether this UI is construct for a reimport */
 	bool bIsReimport;
+
+
+	//////////////////////////////////////////////////////////////////////////
+	// FBX file informations
+	// Transient value that are set everytime we show the options dialog. These are information only and should be string.
+
+	/* The fbx file version */
+	UPROPERTY(VisibleAnywhere, Transient, Category = FbxFileInformation, meta = (ImportType = "Mesh|Animation", DisplayName = "File Version"))
+	FString FileVersion;
+
+	/* The file creator information */
+	UPROPERTY(VisibleAnywhere, Transient, Category = FbxFileInformation, meta = (ImportType = "Mesh|Animation", DisplayName = "File Creator"))
+	FString FileCreator;
+
+	/* The file vendor information, software name and version that was use to create the file */
+	UPROPERTY(VisibleAnywhere, Transient, Category = FbxFileInformation, meta = (ImportType = "Mesh|Animation", DisplayName = "File Creator Application"))
+	FString FileCreatorApplication;
+
+	/* The file units */
+	UPROPERTY(VisibleAnywhere, Transient, Category = FbxFileInformation, meta = (ImportType = "Mesh|Animation", DisplayName = "File Units"))
+	FString FileUnits;
+
+	/* The file axis direction, up vector and hand */
+	UPROPERTY(VisibleAnywhere, Transient, Category = FbxFileInformation, meta = (ImportType = "Mesh|Animation", DisplayName = "File Axis Direction"))
+	FString FileAxisDirection;
+
+	/* The fbx animation frame rate */
+	UPROPERTY(VisibleAnywhere, Transient, Category = FbxFileInformation, meta = (ImportType = "SkeletalMesh|Animation", DisplayName = "File Frame Rate"))
+	FString FileSampleRate;
+	
+	/* The fbx animation start frame */
+	UPROPERTY(VisibleAnywhere, Transient, Category = FbxFileInformation, meta = (ImportType = "SkeletalMesh|Animation", DisplayName = "Animation Start Frame"))
+	FString AnimStartFrame;
+	
+	/* The fbx animation end frame */
+	UPROPERTY(VisibleAnywhere, Transient, Category = FbxFileInformation, meta = (ImportType = "SkeletalMesh|Animation", DisplayName = "Animation End Frame"))
+	FString AnimEndFrame;
 };
 
 
diff --git a/Engine/Source/Editor/UnrealEd/Private/Factories/EditorFactories.cpp b/Engine/Source/Editor/UnrealEd/Private/Factories/EditorFactories.cpp
--- a/Engine/Source/Editor/UnrealEd/Private/Factories/EditorFactories.cpp
+++ b/Engine/Source/Editor/UnrealEd/Private/Factories/EditorFactories.cpp
@@ -5304,6 +5304,36 @@ EReimportResult::Type UReimportFbxStaticMeshFactory::Reimport( UObject* Obj )
 	const bool IsUnattended = GIsAutomationTesting || FApp::IsUnattended();
 	const bool ShowImportDialogAtReimport = GetDefault<UEditorPerProjectUserSettings>()->bShowImportDialogAtReimport && !IsUnattended;
 
+	if (ImportData == nullptr)
+	{
+		// An existing import data object was not found, make one here and show the options dialog
+		ImportData = UFbxStaticMeshImportData::GetImportDataForStaticMesh(Mesh, ImportUI->StaticMeshImportData);
+		Mesh->AssetImportData = ImportData;
+	}
+
+	//Get the re-import filename
+	const FString Filename = ImportData->GetFirstFilename();
+	const FString FileExtension = FPaths::GetExtension(Filename);
+	const bool bIsValidFile = FileExtension.Equals(TEXT("fbx"), ESearchCase::IgnoreCase) || FileExtension.Equals("obj", ESearchCase::IgnoreCase);
+	if (!bIsValidFile)
+	{
+		return EReimportResult::Failed;
+	}
+	if (!(Filename.Len()))
+	{
+		// Since this is a new system most static meshes don't have paths, so logging has been commented out
+		//UE_LOG(LogEditorFactories, Warning, TEXT("-- cannot reimport: static mesh resource does not have path stored."));
+		return EReimportResult::Failed;
+	}
+	// Ensure that the file provided by the path exists
+	if (IFileManager::Get().FileSize(*Filename) == INDEX_NONE)
+	{
+		UE_LOG(LogEditorFactories, Warning, TEXT("-- cannot reimport: source file cannot be found."));
+		return EReimportResult::Failed;
+	}
+	CurrentFilename = Filename;
+
+
 	if( ImportData  && !ShowImportDialogAtReimport)
 	{
 		// Import data already exists, apply it to the fbx import options
@@ -5312,12 +5342,6 @@ EReimportResult::Type UReimportFbxStaticMeshFactory::Reimport( UObject* Obj )
 	}
 	else
 	{
-		if (ImportData == nullptr)
-		{
-			// An existing import data object was not found, make one here and show the options dialog
-			ImportData = UFbxStaticMeshImportData::GetImportDataForStaticMesh(Mesh, ImportUI->StaticMeshImportData);
-			Mesh->AssetImportData = ImportData;
-		}
 		ReimportUI->bIsReimport = true;
 		ReimportUI->StaticMeshImportData = ImportData;
 		
@@ -5334,7 +5358,8 @@ EReimportResult::Type UReimportFbxStaticMeshFactory::Reimport( UObject* Obj )
 		bool bOutImportAll = false;
 		bool bIsObjFormat = false;
 		bool bIsAutomated = false;
-		GetImportOptions( FFbxImporter, ReimportUI, bShowOptionDialog, bIsAutomated, Obj->GetPathName(), bOperationCanceled, bOutImportAll, bIsObjFormat, bForceImportType, FBXIT_StaticMesh, Mesh);
+
+		GetImportOptions( FFbxImporter, ReimportUI, bShowOptionDialog, bIsAutomated, Obj->GetPathName(), bOperationCanceled, bOutImportAll, bIsObjFormat, Filename, bForceImportType, FBXIT_StaticMesh, Mesh);
 		
 		//Put back the original bAutoGenerateCollision settings since the user cancel the re-import
 		if (bOperationCanceled && Mesh->bCustomizedCollision)
@@ -5351,38 +5376,14 @@ EReimportResult::Type UReimportFbxStaticMeshFactory::Reimport( UObject* Obj )
 
 	if( !bOperationCanceled && ensure(ImportData) )
 	{
-		const FString Filename = ImportData->GetFirstFilename();
-		const FString FileExtension = FPaths::GetExtension(Filename);
-		const bool bIsValidFile = FileExtension.Equals( TEXT("fbx"), ESearchCase::IgnoreCase ) || FileExtension.Equals( "obj",  ESearchCase::IgnoreCase );
-
-		if ( !bIsValidFile )
-		{
-			return EReimportResult::Failed;
-		}
-
-		if(!(Filename.Len()))
-		{
-			// Since this is a new system most static meshes don't have paths, so logging has been commented out
-			//UE_LOG(LogEditorFactories, Warning, TEXT("-- cannot reimport: static mesh resource does not have path stored."));
-			return EReimportResult::Failed;
-		}
-
 		UE_LOG(LogEditorFactories, Log, TEXT("Performing atomic reimport of [%s]"), *Filename);
 
-		// Ensure that the file provided by the path exists
-		if (IFileManager::Get().FileSize(*Filename) == INDEX_NONE)
-		{
-			UE_LOG(LogEditorFactories, Warning, TEXT("-- cannot reimport: source file cannot be found."));
-			return EReimportResult::Failed;
-		}
-
 		//Create a copy of the mesh we re-import
 		if (!IsUnattended)
 		{
 			ImportOptions->OriginalMeshCopy = Cast<UStaticMesh>(StaticDuplicateObject(Mesh, GetTransientPackage(), NAME_None, RF_Standalone));
 		}
 
-		CurrentFilename = Filename;
 		bool bImportSucceed = true;
 		if ( FFbxImporter->ImportFromFile( *Filename, FPaths::GetExtension( Filename ), true ) )
 		{
@@ -5591,7 +5592,25 @@ EReimportResult::Type UReimportFbxSkeletalMeshFactory::Reimport( UObject* Obj )
 	bool bSuccess = false;
 	const bool IsUnattended = GIsAutomationTesting || FApp::IsUnattended();
 	const bool ShowImportDialogAtReimport = GetDefault<UEditorPerProjectUserSettings>()->bShowImportDialogAtReimport && !IsUnattended;
-	
+
+	if (ImportData == nullptr)
+	{
+		// An existing import data object was not found, make one here and show the options dialog
+		ImportData = UFbxSkeletalMeshImportData::GetImportDataForSkeletalMesh(SkeletalMesh, ImportUI->SkeletalMeshImportData);
+		SkeletalMesh->AssetImportData = ImportData;
+	}
+
+	const FString Filename = ImportData->GetFirstFilename();
+	UE_LOG(LogEditorFactories, Log, TEXT("Performing atomic reimport of [%s]"), *Filename);
+
+	// Ensure that the file provided by the path exists
+	if (IFileManager::Get().FileSize(*Filename) == INDEX_NONE)
+	{
+		UE_LOG(LogEditorFactories, Warning, TEXT("-- cannot reimport: source file cannot be found."));
+		return EReimportResult::Failed;
+	}
+	CurrentFilename = Filename;
+
 	if( ImportData && !ShowImportDialogAtReimport)
 	{
 		// Import data already exists, apply it to the fbx import options
@@ -5603,12 +5622,6 @@ EReimportResult::Type UReimportFbxSkeletalMeshFactory::Reimport( UObject* Obj )
 	}
 	else
 	{
-		if (ImportData == nullptr)
-		{
-			// An existing import data object was not found, make one here and show the options dialog
-			ImportData = UFbxSkeletalMeshImportData::GetImportDataForSkeletalMesh(SkeletalMesh, ImportUI->SkeletalMeshImportData);
-			SkeletalMesh->AssetImportData = ImportData;
-		}
 		ReimportUI->bIsReimport = true;
 		ReimportUI->SkeletalMeshImportData = ImportData;
 
@@ -5623,24 +5636,13 @@ EReimportResult::Type UReimportFbxSkeletalMeshFactory::Reimport( UObject* Obj )
 		ImportOptions->bCreatePhysicsAsset = false;
 		ImportOptions->PhysicsAsset = SkeletalMesh->PhysicsAsset;
 
-		ImportOptions = GetImportOptions( FFbxImporter, ReimportUI, bShowOptionDialog, bIsAutomated, Obj->GetPathName(), bOperationCanceled, bOutImportAll, bIsObjFormat, bForceImportType, FBXIT_SkeletalMesh, Obj );
+		ImportOptions = GetImportOptions( FFbxImporter, ReimportUI, bShowOptionDialog, bIsAutomated, Obj->GetPathName(), bOperationCanceled, bOutImportAll, bIsObjFormat, Filename, bForceImportType, FBXIT_SkeletalMesh, Obj );
 	}
 
 	if( !bOperationCanceled && ensure(ImportData) )
 	{
 		ImportOptions->bCanShowDialog = !IsUnattended;
 
-		const FString Filename = ImportData->GetFirstFilename();
-		UE_LOG(LogEditorFactories, Log, TEXT("Performing atomic reimport of [%s]"), *Filename);
-
-		// Ensure that the file provided by the path exists
-		if (IFileManager::Get().FileSize(*Filename) == INDEX_NONE)
-		{
-			UE_LOG(LogEditorFactories, Warning, TEXT("-- cannot reimport: source file cannot be found.") );
-			return EReimportResult::Failed;
-		}
-		CurrentFilename = Filename;
-
 		//Create a copy of the mesh we re-import
 		if (!IsUnattended)
 		{
diff --git a/Engine/Source/Editor/UnrealEd/Private/Fbx/FbxFactory.cpp b/Engine/Source/Editor/UnrealEd/Private/Fbx/FbxFactory.cpp
--- a/Engine/Source/Editor/UnrealEd/Private/Fbx/FbxFactory.cpp
+++ b/Engine/Source/Editor/UnrealEd/Private/Fbx/FbxFactory.cpp
@@ -283,7 +283,8 @@ UObject* UFbxFactory::FactoryCreateFile
 	bool bIsAutomated = IsAutomatedImport();
 	bool bShowImportDialog = bShowOption && !bIsAutomated;
 	bool bImportAll = false;
-	ImportOptions = GetImportOptions(FbxImporter, ImportUI, bShowImportDialog, bIsAutomated, InParent->GetPathName(), bOperationCanceled, bImportAll, bIsObjFormat, bIsObjFormat, ForcedImportType, ExistingObject);
+	
+	ImportOptions = GetImportOptions(FbxImporter, ImportUI, bShowImportDialog, bIsAutomated, InParent->GetPathName(), bOperationCanceled, bImportAll, bIsObjFormat, UFactory::CurrentFilename, false, ForcedImportType, ExistingObject);
 	bOutOperationCanceled = bOperationCanceled;
 	
 	if( bImportAll )
diff --git a/Engine/Source/Editor/UnrealEd/Private/Fbx/FbxMainImport.cpp b/Engine/Source/Editor/UnrealEd/Private/Fbx/FbxMainImport.cpp
--- a/Engine/Source/Editor/UnrealEd/Private/Fbx/FbxMainImport.cpp
+++ b/Engine/Source/Editor/UnrealEd/Private/Fbx/FbxMainImport.cpp
@@ -47,7 +47,7 @@ TSharedPtr<FFbxImporter> FFbxImporter::StaticPreviewInstance;
 
 
 
-FBXImportOptions* GetImportOptions( UnFbx::FFbxImporter* FbxImporter, UFbxImportUI* ImportUI, bool bShowOptionDialog, bool bIsAutomated, const FString& FullPath, bool& OutOperationCanceled, bool& bOutImportAll, bool bIsObjFormat, bool bForceImportType, EFBXImportType ImportType, UObject* ReimportObject)
+FBXImportOptions* GetImportOptions( UnFbx::FFbxImporter* FbxImporter, UFbxImportUI* ImportUI, bool bShowOptionDialog, bool bIsAutomated, const FString& FullPath, bool& OutOperationCanceled, bool& bOutImportAll, bool bIsObjFormat, const FString& InFilename, bool bForceImportType, EFBXImportType ImportType, UObject* ReimportObject)
 {
 	OutOperationCanceled = false;
 
@@ -103,6 +103,58 @@ FBXImportOptions* GetImportOptions( UnFbx::FFbxImporter* FbxImporter, UFbxImport
 		//This option must always be the same value has the skeletalmesh one.
 		ImportUI->AnimSequenceImportData->bImportMeshesInBoneHierarchy = ImportUI->SkeletalMeshImportData->bImportMeshesInBoneHierarchy;
 
+		//////////////////////////////////////////////////////////////////////////
+		// Set the information section data
+		
+		//Make sure the file is open to be able to read the header before showing the options
+		//If the file is already open it will simply return false.
+		FbxImporter->ReadHeaderFromFile(InFilename, true);
+
+		ImportUI->FileVersion = FbxImporter->GetFbxFileVersion();
+		ImportUI->FileCreator = FbxImporter->GetFileCreator();
+		// do analytics on getting Fbx data
+		FbxDocumentInfo* DocInfo = FbxImporter->Scene->GetSceneInfo();
+		if (DocInfo)
+		{
+			FString LastSavedVendor(UTF8_TO_TCHAR(DocInfo->LastSaved_ApplicationVendor.Get().Buffer()));
+			FString LastSavedAppName(UTF8_TO_TCHAR(DocInfo->LastSaved_ApplicationName.Get().Buffer()));
+			FString LastSavedAppVersion(UTF8_TO_TCHAR(DocInfo->LastSaved_ApplicationVersion.Get().Buffer()));
+
+			ImportUI->FileCreatorApplication = LastSavedVendor + TEXT(" ") + LastSavedAppName + TEXT(" ") + LastSavedAppVersion;
+		}
+		else
+		{
+			ImportUI->FileCreatorApplication = TEXT("");
+		}
+
+		ImportUI->FileUnits = FbxImporter->GetFileUnitSystem();
+
+		ImportUI->FileAxisDirection = FbxImporter->GetFileAxisDirection();
+
+		//Set the info original file frame rate
+		ImportUI->FileSampleRate = FString::Printf(TEXT("%.2f"), FbxImporter->GetOriginalFbxFramerate());
+
+		//Set the info start time and the end time
+		ImportUI->AnimStartFrame = TEXT("0");
+		ImportUI->AnimEndFrame = TEXT("0");
+		FbxTimeSpan AnimTimeSpan(FBXSDK_TIME_INFINITE, FBXSDK_TIME_MINUS_INFINITE);
+		int32 AnimStackCount = FbxImporter->Scene->GetSrcObjectCount<FbxAnimStack>();
+		for (int32 AnimStackIndex = 0; AnimStackIndex < AnimStackCount; AnimStackIndex++)
+		{
+			FbxAnimStack* CurAnimStack = FbxImporter->Scene->GetSrcObject<FbxAnimStack>(AnimStackIndex);
+			FbxTimeSpan AnimatedInterval(FBXSDK_TIME_INFINITE, FBXSDK_TIME_MINUS_INFINITE);
+			FbxImporter->Scene->GetRootNode()->GetAnimationInterval(AnimatedInterval, CurAnimStack);
+			// find the most range that covers by both method, that'll be used for clamping
+			AnimTimeSpan.SetStart(FMath::Min<FbxTime>(AnimTimeSpan.GetStart(), AnimatedInterval.GetStart()));
+			AnimTimeSpan.SetStop(FMath::Max<FbxTime>(AnimTimeSpan.GetStop(), AnimatedInterval.GetStop()));
+		}
+		if (AnimTimeSpan.GetStart() != FBXSDK_TIME_INFINITE)
+		{
+			FbxTime EachFrame = FBXSDK_TIME_ONE_SECOND / FbxImporter->GetOriginalFbxFramerate();
+			ImportUI->AnimStartFrame = FString::FromInt(AnimTimeSpan.GetStart().Get() / EachFrame.Get());
+			ImportUI->AnimEndFrame = FString::FromInt(AnimTimeSpan.GetStop().Get() / EachFrame.Get());
+		}
+
 		TSharedPtr<SWindow> ParentWindow;
 
 		if( FModuleManager::Get().IsModuleLoaded( "MainFrame" ) )
@@ -1000,6 +1052,47 @@ void FFbxImporter::EnsureNodeNameAreValid()
 	}
 }
 
+FString FFbxImporter::GetFileAxisDirection()
+{
+	FString AxisDirection;
+	int32 Sign = 1;
+	switch (FileAxisSystem.GetUpVector(Sign))
+	{
+	case FbxAxisSystem::eXAxis:
+		{
+			AxisDirection += TEXT("X");
+		}
+		break;
+	case FbxAxisSystem::eYAxis:
+		{
+			AxisDirection += TEXT("Y");
+		}
+		break;
+	case FbxAxisSystem::eZAxis:
+		{
+			AxisDirection += TEXT("Z");
+		}
+		break;
+	}
+	//Negative sign mean down instead of up
+	AxisDirection += Sign == 1 ? TEXT("-UP") : TEXT("-DOWN");
+		
+	switch (FileAxisSystem.GetCoorSystem())
+	{
+	case FbxAxisSystem::eLeftHanded:
+		{
+			AxisDirection += TEXT(" (LH)");
+		}
+		break;
+	case FbxAxisSystem::eRightHanded:
+		{
+			AxisDirection += TEXT(" (RH)");
+		}
+		break;
+	}
+	return AxisDirection;
+}
+
 #ifdef IOS_REF
 #undef  IOS_REF
 #define IOS_REF (*(SdkManager->GetIOSettings()))
@@ -1050,7 +1143,8 @@ bool FFbxImporter::ImportFile(FString Filename, bool bPreventMaterialNameClash /
 	// Get the version number of the FBX file format.
 	Importer->GetFileVersion(FileMajor, FileMinor, FileRevision);
 	FbxFileVersion = FString::Printf(TEXT("%d.%d.%d"), FileMajor, FileMinor, FileRevision);
-
+	
+	FbxFileCreator = UTF8_TO_TCHAR(Importer->GetFileHeaderInfo()->mCreator.Buffer());
 	// output result
 	if(bStatus)
 	{
@@ -1126,6 +1220,33 @@ void FFbxImporter::ConvertScene()
 	Scene->GetAnimationEvaluator()->Reset();
 }
 
+//-------------------------------------------------------------------------
+//
+//-------------------------------------------------------------------------
+bool FFbxImporter::ReadHeaderFromFile(const FString& Filename, bool bPreventMaterialNameClash /*= false*/)
+{
+	bool Result = true;
+
+
+	switch (CurPhase)
+	{
+	case NOTSTARTED:
+		if (!OpenFile(FString(Filename)))
+		{
+			Result = false;
+			break;
+		}
+	case FILEOPENED:
+		if (!ImportFile(FString(Filename), bPreventMaterialNameClash))
+		{
+			Result = false;
+			CurPhase = NOTSTARTED;
+			break;
+		}
+	}
+	return Result;
+}
+
 //-------------------------------------------------------------------------
 //
 //-------------------------------------------------------------------------
diff --git a/Engine/Source/Editor/UnrealEd/Private/Fbx/ReimportFbxSceneFactory.cpp b/Engine/Source/Editor/UnrealEd/Private/Fbx/ReimportFbxSceneFactory.cpp
--- a/Engine/Source/Editor/UnrealEd/Private/Fbx/ReimportFbxSceneFactory.cpp
+++ b/Engine/Source/Editor/UnrealEd/Private/Fbx/ReimportFbxSceneFactory.cpp
@@ -1293,7 +1293,7 @@ EReimportResult::Type UReimportFbxSceneFactory::ReimportSkeletalMesh(void* VoidF
 				{
 					FbxAnimStack* CurAnimStack = FbxImporter->Scene->GetSrcObject<FbxAnimStack>(AnimStackIndex);
 
-					FbxTimeSpan AnimTimeSpan = FbxImporter->GetAnimationTimeSpan(SortedLinks[0], CurAnimStack, ResampleRate);
+					FbxTimeSpan AnimTimeSpan = FbxImporter->GetAnimationTimeSpan(SortedLinks[0], CurAnimStack);
 					bool bValidAnimStack = FbxImporter->ValidateAnimStack(SortedLinks, FBXMeshNodeArray, CurAnimStack, ResampleRate, GlobalImportSettings->bImportMorph, AnimTimeSpan);
 					// no animation
 					if (!bValidAnimStack)
@@ -1356,7 +1356,7 @@ EReimportResult::Type UReimportFbxSceneFactory::ReimportSkeletalMesh(void* VoidF
 								}
 							}
 						}
-						FbxTimeSpan AnimTimeSpan = FbxImporter->GetAnimationTimeSpan(SortedLinks[0], CurAnimStack, ResampleRate);
+						FbxTimeSpan AnimTimeSpan = FbxImporter->GetAnimationTimeSpan(SortedLinks[0], CurAnimStack);
 
 						if (DestSeq == nullptr)
 						{
diff --git a/Engine/Source/Editor/UnrealEd/Private/FbxAnimUtils.cpp b/Engine/Source/Editor/UnrealEd/Private/FbxAnimUtils.cpp
--- a/Engine/Source/Editor/UnrealEd/Private/FbxAnimUtils.cpp
+++ b/Engine/Source/Editor/UnrealEd/Private/FbxAnimUtils.cpp
@@ -117,7 +117,7 @@ namespace FbxAnimUtils
 				{
 					FbxImporter->MergeAllLayerAnimation(AnimStack, FbxTime::GetFrameRate(FbxImporter->Scene->GetGlobalSettings().GetTimeMode()));
 
-					FbxTimeSpan AnimTimeSpan = FbxImporter->GetAnimationTimeSpan(FbxImporter->Scene->GetRootNode(), AnimStack, DEFAULT_SAMPLERATE);
+					FbxTimeSpan AnimTimeSpan = FbxImporter->GetAnimationTimeSpan(FbxImporter->Scene->GetRootNode(), AnimStack);
 
 					// Grab the start time, as we might have a preroll
 					OutPreRoll = FMath::Abs((float)AnimTimeSpan.GetStart().GetSecondDouble());
diff --git a/Engine/Source/Editor/UnrealEd/Private/SkeletalMeshEdit.cpp b/Engine/Source/Editor/UnrealEd/Private/SkeletalMeshEdit.cpp
--- a/Engine/Source/Editor/UnrealEd/Private/SkeletalMeshEdit.cpp
+++ b/Engine/Source/Editor/UnrealEd/Private/SkeletalMeshEdit.cpp
@@ -153,7 +153,7 @@ bool UEditorEngine::ReimportFbxAnimation( USkeleton* Skeleton, UAnimSequence* An
 		// @hack to make sure skeleton is set before opening the dialog
 		FbxImporter->ImportOptions->SkeletonForAnimation = Skeleton;
 
-		GetImportOptions(FbxImporter, ReimportUI, bShowOptionDialog, bIsAutomated, AnimSequence->GetPathName(), bImportOperationCanceled, bOutImportAll, bIsObjFormat, bForceImportType, FBXIT_Animation, AnimSequence);
+		GetImportOptions(FbxImporter, ReimportUI, bShowOptionDialog, bIsAutomated, AnimSequence->GetPathName(), bImportOperationCanceled, bOutImportAll, bIsObjFormat, InFilename, bForceImportType, FBXIT_Animation, AnimSequence);
 
 		if (bImportOperationCanceled)
 		{
@@ -253,7 +253,7 @@ bool UEditorEngine::ReimportFbxAnimation( USkeleton* Skeleton, UAnimSequence* An
 							}
 						}
 					}
-					FbxTimeSpan AnimTimeSpan = FbxImporter->GetAnimationTimeSpan(SortedLinks[0], CurAnimStack, ResampleRate);
+					FbxTimeSpan AnimTimeSpan = FbxImporter->GetAnimationTimeSpan(SortedLinks[0], CurAnimStack);
 					// for now it's not importing morph - in the future, this should be optional or saved with asset
 					if (FbxImporter->ValidateAnimStack(SortedLinks, FBXMeshNodeArray, CurAnimStack, ResampleRate, bImportMorphTracks, AnimTimeSpan))
 					{
@@ -354,10 +354,8 @@ bool UnFbx::FFbxImporter::IsValidAnimationData(TArray<FbxNode*>& SortedLinks, TA
 			UE_LOG(LogFbx, Log, TEXT("SortedLinks :(%d) %s"), BoneIndex, *BoneName );
 		}
 
-		//@note: the reason we give default sample rate is because we just want to make sure it has duration
-		// we don't want to accept input of [20, 20], but the sample rate should be recalculated after this verification
-		// and proper timeline will be calculated
-		FbxTimeSpan AnimTimeSpan = GetAnimationTimeSpan(SortedLinks[0], CurAnimStack, DEFAULT_SAMPLERATE);
+		//The animation timespan must use the original fbx framerate so the frame number match the DCC frame number
+		FbxTimeSpan AnimTimeSpan = GetAnimationTimeSpan(SortedLinks[0], CurAnimStack);
 		if (AnimTimeSpan.GetDuration() <= 0)
 		{
 			AddTokenizedErrorMessage(FTokenizedMessage::Create(EMessageSeverity::Warning, FText::Format(LOCTEXT("FBXImport_ZeroLength", "Animation Stack {0} does not contain any valid key. Try different time options when import."), FText::FromString(UTF8_TO_TCHAR(CurAnimStack->GetName())))), FFbxErrors::Animation_ZeroLength);
@@ -460,13 +458,13 @@ void UnFbx::FFbxImporter::FillAndVerifyBoneNames(USkeleton* Skeleton, TArray<Fbx
 //
 //-------------------------------------------------------------------------
 
-FbxTimeSpan UnFbx::FFbxImporter::GetAnimationTimeSpan(FbxNode* RootNode, FbxAnimStack* AnimStack, int32 ResampleRate)
+FbxTimeSpan UnFbx::FFbxImporter::GetAnimationTimeSpan(FbxNode* RootNode, FbxAnimStack* AnimStack)
 {
 	FBXImportOptions* ImportOption = GetImportOptions();
 	FbxTimeSpan AnimTimeSpan(FBXSDK_TIME_INFINITE, FBXSDK_TIME_MINUS_INFINITE);
 	if (ImportOption)
 	{
-		bool bUseDefault = ImportOption->AnimationLengthImportType == FBXALIT_ExportedTime || ResampleRate == 0;
+		bool bUseDefault = ImportOption->AnimationLengthImportType == FBXALIT_ExportedTime || FMath::IsNearlyZero(OriginalFbxFramerate, KINDA_SMALL_NUMBER);
 		if  (bUseDefault)
 		{
 			AnimTimeSpan = AnimStack->GetLocalTimeSpan();
@@ -490,7 +488,7 @@ FbxTimeSpan UnFbx::FFbxImporter::GetAnimationTimeSpan(FbxNode* RootNode, FbxAnim
 			AnimTimeSpan.SetStart(StartTime);
 			AnimTimeSpan.SetStop(StopTime);
 
-			FbxTime EachFrame = FBXSDK_TIME_ONE_SECOND/ResampleRate;
+			FbxTime EachFrame = FBXSDK_TIME_ONE_SECOND/OriginalFbxFramerate;
 			int32 StartFrame = StartTime.Get()/EachFrame.Get();
 			int32 StopFrame = StopTime.Get()/EachFrame.Get();
 			if (StartFrame != StopFrame)
@@ -572,7 +570,7 @@ UAnimSequence * UnFbx::FFbxImporter::ImportAnimations(USkeleton* Skeleton, UObje
 	{
 		FbxAnimStack* CurAnimStack = Scene->GetSrcObject<FbxAnimStack>(AnimStackIndex);
 
-		FbxTimeSpan AnimTimeSpan = GetAnimationTimeSpan(SortedLinks[0], CurAnimStack, ResampleRate);
+		FbxTimeSpan AnimTimeSpan = GetAnimationTimeSpan(SortedLinks[0], CurAnimStack);
 		bool bValidAnimStack = ValidateAnimStack(SortedLinks, NodeArray, CurAnimStack, ResampleRate, ImportOptions->bImportMorph, AnimTimeSpan);
 		// no animation
 		if (!bValidAnimStack)
@@ -732,10 +730,7 @@ int32 UnFbx::FFbxImporter::GetMaxSampleRate(TArray<FbxNode*>& SortedLinks, TArra
 	{
 		FbxAnimStack* CurAnimStack = Scene->GetSrcObject<FbxAnimStack>(AnimStackIndex);
 
-		// @note: here we iterate through all timeline to figure out sample rate, not just in range
-		// we have chicken/egg problem if we don't. We need samplerate to figure out time range for the (start, end)
-		// so when you get time range for the sample rate, we just walk through all range
-		FbxTimeSpan AnimStackTimeSpan = GetAnimationTimeSpan(SortedLinks[0], CurAnimStack, 0);
+		FbxTimeSpan AnimStackTimeSpan = GetAnimationTimeSpan(SortedLinks[0], CurAnimStack);
 
 		double AnimStackStart = AnimStackTimeSpan.GetStart().GetSecondDouble();
 		double AnimStackStop = AnimStackTimeSpan.GetStop().GetSecondDouble();
@@ -856,7 +851,7 @@ bool UnFbx::FFbxImporter::ValidateAnimStack(TArray<FbxNode*>& SortedLinks, TArra
 
 	bool bValidAnimStack = true;
 
-	AnimTimeSpan = GetAnimationTimeSpan(SortedLinks[0], CurAnimStack, ResampleRate);
+	AnimTimeSpan = GetAnimationTimeSpan(SortedLinks[0], CurAnimStack);
 	
 	// if no duration is found, return false
 	if (AnimTimeSpan.GetDuration() <= 0)
diff --git a/Engine/Source/Editor/UnrealEd/Public/FbxImporter.h b/Engine/Source/Editor/UnrealEd/Public/FbxImporter.h
--- a/Engine/Source/Editor/UnrealEd/Public/FbxImporter.h
+++ b/Engine/Source/Editor/UnrealEd/Public/FbxImporter.h
@@ -483,7 +483,7 @@ private:
 	static FbxAMatrix JointPostConversionMatrix;
 };
 
-FBXImportOptions* GetImportOptions( class FFbxImporter* FbxImporter, UFbxImportUI* ImportUI, bool bShowOptionDialog, bool bIsAutomated, const FString& FullPath, bool& OutOperationCanceled, bool& OutImportAll, bool bIsObjFormat, bool bForceImportType = false, EFBXImportType ImportType = FBXIT_StaticMesh, UObject* ReimportObject = nullptr);
+FBXImportOptions* GetImportOptions( class FFbxImporter* FbxImporter, UFbxImportUI* ImportUI, bool bShowOptionDialog, bool bIsAutomated, const FString& FullPath, bool& OutOperationCanceled, bool& OutImportAll, bool bIsObjFormat, const FString& InFilename, bool bForceImportType = false, EFBXImportType ImportType = FBXIT_StaticMesh, UObject* ReimportObject = nullptr);
 UNREALED_API void ApplyImportUIToImportOptions(UFbxImportUI* ImportUI, FBXImportOptions& InOutImportOptions);
 
 struct FImportedMaterialData
@@ -606,6 +606,11 @@ public:
 	 * @return bool
 	 */
 	bool OpenFile(FString Filename);
+
+	/*
+	Make sure the file header is read
+	*/
+	bool ReadHeaderFromFile(const FString& Filename, bool bPreventMaterialNameClash = false);
 	
 	/**
 	 * Import Fbx file.
@@ -813,7 +818,7 @@ public:
 	/**
 	 * Get Animation Time Span - duration of the animation
 	 */
-	FbxTimeSpan GetAnimationTimeSpan(FbxNode* RootNode, FbxAnimStack* AnimStack, int32 ResampleRate);
+	FbxTimeSpan GetAnimationTimeSpan(FbxNode* RootNode, FbxAnimStack* AnimStack);
 
 	/**
 	 * Import one animation from CurAnimStack
@@ -1174,6 +1179,11 @@ public:
 
 	FbxGeometryConverter* GetGeometryConverter() { return GeometryConverter; }
 
+	FString GetFbxFileVersion() { return FbxFileVersion; }
+	FString GetFileCreator() { return FbxFileCreator; }
+	FString GetFileUnitSystem() { return FString(UTF8_TO_TCHAR(FileUnitSystem.GetScaleFactorAsString(false).Buffer())); }
+	FString GetFileAxisDirection();
+
 protected:
 	enum IMPORTPHASE
 	{
@@ -1199,6 +1209,7 @@ protected:
 	FString FileBasePath;
 	TWeakObjectPtr<UObject> Parent;
 	FString FbxFileVersion;
+	FString FbxFileCreator;
 
 	//Original File Info
 	FbxAxisSystem FileAxisSystem;
