Fix for crash on opening map built data

diff --git a/Engine/Source/Editor/UnrealEd/Private/Toolkits/AssetEditorManager.cpp b/Engine/Source/Editor/UnrealEd/Private/Toolkits/AssetEditorManager.cpp
--- a/Engine/Source/Editor/UnrealEd/Private/Toolkits/AssetEditorManager.cpp
+++ b/Engine/Source/Editor/UnrealEd/Private/Toolkits/AssetEditorManager.cpp
@@ -23,6 +23,7 @@
 #include "Widgets/Notifications/SNotificationList.h"
 #include "ContentBrowserModule.h"
 #include "MRUFavoritesList.h"
+#include "Engine/MapBuildDataRegistry.h"
 
 #define LOCTEXT_NAMESPACE "AssetEditorManager"
 
@@ -394,7 +395,8 @@ bool FAssetEditorManager::OpenEditorForAsset(UObject* Asset, const EToolkitMode:
 	{
 		GWarn->EndSlowTask();
 	}
-	if (Asset->IsAsset())
+	// Do not add to recently opened asset list if this is a level-associated asset like Level Blueprint or Built Data. Their naming is not compatible
+	if (Asset->IsAsset() && !Asset->IsA(UMapBuildDataRegistry::StaticClass()))
 	{
 		FString AssetPath = Asset->GetOuter()->GetPathName();
 		FContentBrowserModule& CBModule = FModuleManager::LoadModuleChecked<FContentBrowserModule>(TEXT("ContentBrowser"));
