Don't allow the user to attempt to import a csv unless they've specified the rowstruct for datatables or the interpolation for curve tables.

diff --git a/Engine/Source/Editor/UnrealEd/Private/SCSVImportOptions.cpp b/Engine/Source/Editor/UnrealEd/Private/SCSVImportOptions.cpp
--- a/Engine/Source/Editor/UnrealEd/Private/SCSVImportOptions.cpp
+++ b/Engine/Source/Editor/UnrealEd/Private/SCSVImportOptions.cpp
@@ -25,7 +25,7 @@ void SCSVImportOptions::Construct(const FArguments& InArgs)
 
 	// Find table row struct info
 	UScriptStruct* TableRowStruct = FindObjectChecked<UScriptStruct>(ANY_PACKAGE, TEXT("TableRowBase"));
-	if (TableRowStruct != NULL)
+	if (TableRowStruct != nullptr)
 	{
 		// Make combo of table rowstruct options
 		for (TObjectIterator<UScriptStruct> It; It; ++It)
@@ -151,6 +151,7 @@ void SCSVImportOptions::Construct(const FArguments& InArgs)
 					SNew(SButton)
 					.Text(LOCTEXT("OK", "OK"))
 					.OnClicked( this, &SCSVImportOptions::OnImport )
+					.IsEnabled( this, &SCSVImportOptions::CanImport )
 				]
 				+SHorizontalBox::Slot()
 				.AutoWidth()
@@ -179,7 +180,7 @@ void SCSVImportOptions::Construct(const FArguments& InArgs)
 	/** If we should import */
 bool SCSVImportOptions::ShouldImport()
 {
-	return ((SelectedStruct != NULL) || GetSelectedImportType() != ECSVImportType::ECSV_DataTable) && bImport;
+	return ((SelectedStruct != nullptr) || GetSelectedImportType() != ECSVImportType::ECSV_DataTable) && bImport;
 }
 
 /** Get the row struct we selected */
@@ -244,7 +245,7 @@ TSharedRef<SWidget> SCSVImportOptions::MakeImportTypeItemWidget(TSharedPtr<ECSVI
 /** Called to create a widget for each struct */
 TSharedRef<SWidget> SCSVImportOptions::MakeRowStructItemWidget(UScriptStruct* Struct)
 {
-	check( Struct != NULL );
+	check(Struct != nullptr);
 	return	SNew(STextBlock)
 			.Text(FText::FromString(Struct->GetName()));
 }
@@ -282,23 +283,44 @@ FReply SCSVImportOptions::OnImport()
 {
 	SelectedStruct = RowStructCombo->GetSelectedItem();
 	SelectedImportType = *ImportTypeCombo->GetSelectedItem();
-	if(CurveInterpCombo->GetSelectedItem().IsValid())
+	if (CurveInterpCombo->GetSelectedItem().IsValid())
 	{
 		SelectedCurveInterpMode = *CurveInterpCombo->GetSelectedItem();
 	}
 	bImport = true;
-	if ( WidgetWindow.IsValid() )
+	if (WidgetWindow.IsValid())
 	{
 		WidgetWindow.Pin()->RequestDestroyWindow();
 	}
 	return FReply::Handled();
 }
 
+bool SCSVImportOptions::CanImport() const
+{
+	ECSVImportType ImportType = *ImportTypeCombo->GetSelectedItem();
+
+	switch (ImportType)
+	{
+	case ECSVImportType::ECSV_DataTable:
+		return RowStructCombo->GetSelectedItem() != nullptr;
+		break;
+	case ECSVImportType::ECSV_CurveTable:
+		return CurveInterpCombo->GetSelectedItem().IsValid();
+		break;
+	case ECSVImportType::ECSV_CurveFloat:
+	case ECSVImportType::ECSV_CurveVector:
+	case ECSVImportType::ECSV_CurveLinearColor:
+		return true;
+	}
+	
+	return false;
+}
+
 /** Called when 'Cancel' button is pressed */
 FReply SCSVImportOptions::OnCancel()
 {
 	bImport = false;
-	if ( WidgetWindow.IsValid() )
+	if (WidgetWindow.IsValid())
 	{
 		WidgetWindow.Pin()->RequestDestroyWindow();
 	}
diff --git a/Engine/Source/Editor/UnrealEd/Public/SCSVImportOptions.h b/Engine/Source/Editor/UnrealEd/Public/SCSVImportOptions.h
--- a/Engine/Source/Editor/UnrealEd/Public/SCSVImportOptions.h
+++ b/Engine/Source/Editor/UnrealEd/Public/SCSVImportOptions.h
@@ -73,6 +73,9 @@ public:
 	/** Called when 'OK' button is pressed */
 	FReply OnImport();
 
+	/** Do we have all of the data we need to import this asset? */
+	bool CanImport() const;
+
 	/** Called when 'Cancel' button is pressed */
 	FReply OnCancel();
 
@@ -84,7 +87,7 @@ public:
 
 private:
 	/** Whether we should go ahead with import */
-	bool										bImport;
+	uint8										bImport : 1;
 
 	/** Window that owns us */
 	TWeakPtr< SWindow >							WidgetWindow;
