Avoid crashing when running the help command for the GenerateBlueprintAPICommandlet

diff --git a/Engine/Source/Editor/UnrealEd/Private/Commandlets/GenerateBlueprintAPICommandlet.cpp b/Engine/Source/Editor/UnrealEd/Private/Commandlets/GenerateBlueprintAPICommandlet.cpp
--- a/Engine/Source/Editor/UnrealEd/Private/Commandlets/GenerateBlueprintAPICommandlet.cpp
+++ b/Engine/Source/Editor/UnrealEd/Private/Commandlets/GenerateBlueprintAPICommandlet.cpp
@@ -1160,9 +1160,13 @@ int32 UGenerateBlueprintAPICommandlet::Main(FString const& Params)
 		}
 	}
 
-	GenerateBlueprintAPIUtils::DumpCategoryInfo(1, FileOut);
+	if(FileOut)
+	{
+		GenerateBlueprintAPIUtils::DumpCategoryInfo(1, FileOut);
+
+		CloseFileStream(&FileOut);
+	}
 
-	CloseFileStream(&FileOut);
 	return 0;
 }
 
