Don't reorder the rowmap when removing a row.

diff --git a/Engine/Source/Editor/UnrealEd/Private/DataTableEditorUtils.cpp b/Engine/Source/Editor/UnrealEd/Private/DataTableEditorUtils.cpp
--- a/Engine/Source/Editor/UnrealEd/Private/DataTableEditorUtils.cpp
+++ b/Engine/Source/Editor/UnrealEd/Private/DataTableEditorUtils.cpp
@@ -30,7 +30,7 @@ bool FDataTableEditorUtils::RemoveRow(UDataTable* DataTable, FName Name)
 
 		BroadcastPreChange(DataTable, EDataTableChangeInfo::RowList);
 		DataTable->Modify();
-		uint8* RowData = NULL;
+		uint8* RowData = nullptr;
 		const bool bRemoved = DataTable->RowMap.RemoveAndCopyValue(Name, RowData);
 		if (bRemoved && RowData)
 		{
@@ -39,7 +39,7 @@ bool FDataTableEditorUtils::RemoveRow(UDataTable* DataTable, FName Name)
 			bResult = true;
 
 			// Compact the map so that a subsequent add goes at the end of the table
-			DataTable->RowMap.Compact();
+			DataTable->RowMap.CompactStable();
 		}
 		BroadcastPostChange(DataTable, EDataTableChangeInfo::RowList);
 	}
@@ -48,9 +48,9 @@ bool FDataTableEditorUtils::RemoveRow(UDataTable* DataTable, FName Name)
 
 uint8* FDataTableEditorUtils::AddRow(UDataTable* DataTable, FName RowName)
 {
-	if (!DataTable || (RowName == NAME_None) || (DataTable->RowMap.Find(RowName) != NULL) || !DataTable->RowStruct)
+	if (!DataTable || (RowName == NAME_None) || (DataTable->RowMap.Find(RowName) != nullptr) || !DataTable->RowStruct)
 	{
-		return NULL;
+		return nullptr;
 	}
 
 	const FScopedTransaction Transaction(LOCTEXT("AddDataTableRow", "Add Data Table Row"));
@@ -78,7 +78,7 @@ bool FDataTableEditorUtils::RenameRow(UDataTable* DataTable, FName OldName, FNam
 		BroadcastPreChange(DataTable, EDataTableChangeInfo::RowList);
 		DataTable->Modify();
 
-		uint8* RowData = NULL;
+		uint8* RowData = nullptr;
 		const bool bValidnewName = (NewName != NAME_None) && !DataTable->RowMap.Find(NewName);
 		const bool bRemoved = bValidnewName && DataTable->RowMap.RemoveAndCopyValue(OldName, RowData);
 		if (bRemoved)
