Get rid of the failed to import message when cancelling the import of a csv or json file.

diff --git a/Engine/Source/Editor/UnrealEd/Classes/Factories/CSVImportFactory.h b/Engine/Source/Editor/UnrealEd/Classes/Factories/CSVImportFactory.h
--- a/Engine/Source/Editor/UnrealEd/Classes/Factories/CSVImportFactory.h
+++ b/Engine/Source/Editor/UnrealEd/Classes/Factories/CSVImportFactory.h
@@ -11,6 +11,7 @@
 #include "Curves/RichCurve.h"
 #include "Factories/Factory.h"
 #include "Factories/ImportSettings.h"
+#include "EditorReimportHandler.h"
 #include "CSVImportFactory.generated.h"
 
 DECLARE_LOG_CATEGORY_EXTERN(LogCSVImportFactory, Log, All);
@@ -55,14 +56,14 @@ class UNREALED_API UCSVImportFactory : public UFactory, public IImportSettingsPa
 
 	//~ Begin UFactory Interface
 	virtual FText GetDisplayName() const override;
-	virtual UObject* FactoryCreateText( UClass* InClass, UObject* InParent, FName InName, EObjectFlags Flags, UObject* Context, const TCHAR* Type, const TCHAR*& Buffer, const TCHAR* BufferEnd, FFeedbackContext* Warn ) override;
+	virtual UObject* FactoryCreateText(UClass* InClass, UObject* InParent, FName InName, EObjectFlags Flags, UObject* Context, const TCHAR* Type, const TCHAR*& Buffer, const TCHAR* BufferEnd, FFeedbackContext* Warn, bool& bOutOperationCanceled) override;
 	virtual bool DoesSupportClass(UClass * Class) override;
 	virtual bool FactoryCanImport(const FString& Filename) override;
 	virtual	IImportSettingsParser* GetImportSettingsParser() override { return this; }
 
 
 	/* Reimport an object that was created based on a CSV */
-	bool ReimportCSV(UObject* Obj);
+	EReimportResult::Type ReimportCSV(UObject* Obj);
 	
 	/**
 	 * IImportSettings interface
@@ -76,7 +77,7 @@ protected:
 
 private:
 	/* Reimport object from the given path*/
-	bool Reimport(UObject* Obj, const FString& Path  );
+	EReimportResult::Type Reimport(UObject* Obj, const FString& Path);
 
 	UPROPERTY()
 	FCSVImportSettings AutomatedImportSettings;
diff --git a/Engine/Source/Editor/UnrealEd/Classes/Factories/Factory.h b/Engine/Source/Editor/UnrealEd/Classes/Factories/Factory.h
--- a/Engine/Source/Editor/UnrealEd/Classes/Factories/Factory.h
+++ b/Engine/Source/Editor/UnrealEd/Classes/Factories/Factory.h
@@ -274,6 +274,26 @@ protected:
 		return nullptr;
 	}
 
+	/**
+	* Create a new object by importing it from a text buffer.
+	*
+	* @param InClass
+	* @param InParent
+	* @param InName
+	* @param Flags
+	* @param Context
+	* @param Type (must not be nullptr, i.e. TEXT("TGA"))
+	* @param Buffer
+	* @param BufferEnd
+	* @param Warn
+	* @param bOutOperationCanceled Will indicate whether the user canceled the import.
+	* @return The new object.
+	*/
+	virtual UObject* FactoryCreateText(UClass* InClass, UObject* InParent, FName InName, EObjectFlags Flags, UObject* Context, const TCHAR* Type, const TCHAR*& Buffer, const TCHAR* BufferEnd, FFeedbackContext* Warn, bool& bOutOperationCanceled)
+	{
+		return FactoryCreateText(InClass, InParent, InName, Flags, Context, Type, Buffer, BufferEnd, Warn);
+	}
+
 	/**
 	 * Create a new object by importing it from a binary buffer.
 	 *
diff --git a/Engine/Source/Editor/UnrealEd/Private/Factories/CSVImportFactory.cpp b/Engine/Source/Editor/UnrealEd/Private/Factories/CSVImportFactory.cpp
--- a/Engine/Source/Editor/UnrealEd/Private/Factories/CSVImportFactory.cpp
+++ b/Engine/Source/Editor/UnrealEd/Private/Factories/CSVImportFactory.cpp
@@ -40,9 +40,9 @@ FCSVImportSettings::FCSVImportSettings()
 }
 
 
-static UClass* GetCurveClass( ECSVImportType ImportType )
+static UClass* GetCurveClass(ECSVImportType ImportType)
 {
-	switch( ImportType )
+	switch(ImportType)
 	{
 	case ECSVImportType::ECSV_CurveFloat:
 		return UCurveFloat::StaticClass();
@@ -84,24 +84,26 @@ FText UCSVImportFactory::GetDisplayName() const
 
 bool UCSVImportFactory::DoesSupportClass(UClass * Class)
 {
-	return (Class == UDataTable::StaticClass() || Class == UCurveTable::StaticClass() || Class == UCurveFloat::StaticClass() || Class == UCurveVector::StaticClass() || Class == UCurveLinearColor::StaticClass() );
+	return (Class == UDataTable::StaticClass() || Class == UCurveTable::StaticClass() || Class == UCurveFloat::StaticClass() || Class == UCurveVector::StaticClass() || Class == UCurveLinearColor::StaticClass());
 }
 
 bool UCSVImportFactory::FactoryCanImport(const FString& Filename)
 {
 	const FString Extension = FPaths::GetExtension(Filename);
 
-	if(Extension == TEXT("csv"))
+	if (Extension == TEXT("csv"))
 	{
 		return true;
 	}
 	return false;
 }
 
-UObject* UCSVImportFactory::FactoryCreateText(UClass* InClass, UObject* InParent, FName InName, EObjectFlags Flags, UObject* Context, const TCHAR* Type, const TCHAR*& Buffer, const TCHAR* BufferEnd, FFeedbackContext* Warn)
+UObject* UCSVImportFactory::FactoryCreateText(UClass* InClass, UObject* InParent, FName InName, EObjectFlags Flags, UObject* Context, const TCHAR* Type, const TCHAR*& Buffer, const TCHAR* BufferEnd, FFeedbackContext* Warn, bool& bOutOperationCanceled)
 {
 	FEditorDelegates::OnAssetPreImport.Broadcast(this, InClass, InParent, InName, Type);
 
+	bOutOperationCanceled = false;
+
 	// See if table/curve already exists
 	UDataTable* ExistingTable = FindObject<UDataTable>(InParent, *InName.ToString());
 	UCurveTable* ExistingCurveTable = FindObject<UCurveTable>(InParent, *InName.ToString());
@@ -109,11 +111,11 @@ UObject* UCSVImportFactory::FactoryCreateText(UClass* InClass, UObject* InParent
 
 	// Save off information if so
 	bool bHaveInfo = false;
-	UScriptStruct* ImportRowStruct = NULL;
+	UScriptStruct* ImportRowStruct = nullptr;
 	ERichCurveInterpMode ImportCurveInterpMode = RCIM_Linear;
 	ECSVImportType ImportType = ECSVImportType::ECSV_DataTable;
 
-	if( IsAutomatedImport() )
+	if (IsAutomatedImport())
 	{
 		ImportRowStruct = AutomatedImportSettings.ImportRowStruct;
 		ImportCurveInterpMode = AutomatedImportSettings.ImportCurveInterpMode;
@@ -122,17 +124,17 @@ UObject* UCSVImportFactory::FactoryCreateText(UClass* InClass, UObject* InParent
 		// For automated import to work a row struct must be specified for a datatable type or a curve type must be specified
 		bHaveInfo = ImportRowStruct != nullptr || ImportType != ECSVImportType::ECSV_DataTable;
 	}
-	else if(ExistingTable != NULL)
+	else if (ExistingTable != nullptr)
 	{
 		ImportRowStruct = ExistingTable->RowStruct;
 		bHaveInfo = true;
 	}
-	else if(ExistingCurveTable != NULL)
+	else if (ExistingCurveTable != nullptr)
 	{
 		ImportType = ECSVImportType::ECSV_CurveTable;
 		bHaveInfo = true;
 	}
-	else if(ExistingCurve != NULL)
+	else if (ExistingCurve != nullptr)
 	{
 		ImportType = ExistingCurve->IsA(UCurveFloat::StaticClass()) ? ECSVImportType::ECSV_CurveFloat : ECSVImportType::ECSV_CurveVector;
 		bHaveInfo = true;
@@ -141,21 +143,21 @@ UObject* UCSVImportFactory::FactoryCreateText(UClass* InClass, UObject* InParent
 	bool bDoImport = true;
 
 	// If we do not have the info we need, pop up window to ask for things
-	if(!bHaveInfo && !IsAutomatedImport())
+	if (!bHaveInfo && !IsAutomatedImport())
 	{
 		TSharedPtr<SWindow> ParentWindow;
 		// Check if the main frame is loaded.  When using the old main frame it may not be.
-		if( FModuleManager::Get().IsModuleLoaded( "MainFrame" ) )
+		if (FModuleManager::Get().IsModuleLoaded("MainFrame"))
 		{
-			IMainFrameModule& MainFrame = FModuleManager::LoadModuleChecked<IMainFrameModule>( "MainFrame" );
+			IMainFrameModule& MainFrame = FModuleManager::LoadModuleChecked<IMainFrameModule>("MainFrame");
 			ParentWindow = MainFrame.GetParentWindow();
 		}
 
 		TSharedPtr<SCSVImportOptions> ImportOptionsWindow;
 
 		TSharedRef<SWindow> Window = SNew(SWindow)
-			.Title( LOCTEXT("DataTableOptionsWindowTitle", "DataTable Options" ))
-			.SizingRule( ESizingRule::Autosized );
+			.Title(LOCTEXT("DataTableOptionsWindowTitle", "DataTable Options"))
+			.SizingRule(ESizingRule::Autosized);
 		
 		FString ParentFullPath;
 
@@ -177,18 +179,19 @@ UObject* UCSVImportFactory::FactoryCreateText(UClass* InClass, UObject* InParent
 		ImportRowStruct = ImportOptionsWindow->GetSelectedRowStruct();
 		ImportCurveInterpMode = ImportOptionsWindow->GetSelectedCurveIterpMode();
 		bDoImport = ImportOptionsWindow->ShouldImport();
+		bOutOperationCanceled = !bDoImport;
 	}
-	else if(!bHaveInfo && IsAutomatedImport())
+	else if (!bHaveInfo && IsAutomatedImport())
 	{
-		if(ImportType == ECSVImportType::ECSV_DataTable && !ImportRowStruct)
+		if (ImportType == ECSVImportType::ECSV_DataTable && !ImportRowStruct)
 		{
 			UE_LOG(LogCSVImportFactory, Error, TEXT("A Data table row type must be specified in the import settings json file for automated import"));
 		}
 		bDoImport = false;
 	}
 
-	UObject* NewAsset = NULL;
-	if(bDoImport)
+	UObject* NewAsset = nullptr;
+	if (bDoImport)
 	{
 		// Convert buffer to an FString (will this be slow with big tables?)
 		FString String;
@@ -208,7 +211,7 @@ UObject* UCSVImportFactory::FactoryCreateText(UClass* InClass, UObject* InParent
 
 			// If there is an existing table, need to call this to free data memory before recreating object
 			bool bStripFromClientBuilds = false;
-			if(ExistingTable != NULL)
+			if (ExistingTable != nullptr)
 			{
 				DataTableClass = ExistingTable->GetClass();
 				ExistingTable->EmptyTable();
@@ -232,7 +235,7 @@ UObject* UCSVImportFactory::FactoryCreateText(UClass* InClass, UObject* InParent
 			UClass* CurveTableClass = UCurveTable::StaticClass();
 
 			// If there is an existing table, need to call this to free data memory before recreating object
-			if(ExistingCurveTable != NULL)
+			if (ExistingCurveTable != nullptr)
 			{
 				CurveTableClass = ExistingCurveTable->GetClass();
 				ExistingCurveTable->EmptyTable();
@@ -251,7 +254,7 @@ UObject* UCSVImportFactory::FactoryCreateText(UClass* InClass, UObject* InParent
 		}
 		else if (ImportType == ECSVImportType::ECSV_CurveFloat || ImportType == ECSVImportType::ECSV_CurveVector || ImportType == ECSVImportType::ECSV_CurveLinearColor)
 		{
-			UClass* CurveClass = ExistingCurve ? ExistingCurve->GetClass() : GetCurveClass( ImportType );
+			UClass* CurveClass = ExistingCurve ? ExistingCurve->GetClass() : GetCurveClass(ImportType);
 
 			// Create/reset curve
 			UCurveBase* NewCurve = NewObject<UCurveBase>(InParent, CurveClass, InName, Flags);
@@ -263,11 +266,11 @@ UObject* UCSVImportFactory::FactoryCreateText(UClass* InClass, UObject* InParent
 			NewAsset = NewCurve;
 		}
 		
-		if(Problems.Num() > 0)
+		if (Problems.Num() > 0)
 		{
 			FString AllProblems;
 
-			for(int32 ProbIdx=0; ProbIdx<Problems.Num(); ProbIdx++)
+			for (int32 ProbIdx=0; ProbIdx<Problems.Num(); ProbIdx++)
 			{
 				// Output problems to log
 				UE_LOG(LogCSVImportFactory, Log, TEXT("%d:%s"), ProbIdx, *Problems[ProbIdx]);
@@ -275,7 +278,7 @@ UObject* UCSVImportFactory::FactoryCreateText(UClass* InClass, UObject* InParent
 				AllProblems += TEXT("\n");
 			}
 
-			if(!IsAutomatedImport())
+			if (!IsAutomatedImport())
 			{
 				// Pop up any problems for user
 				FMessageDialog::Open(EAppMsgType::Ok, FText::FromString(AllProblems));
@@ -288,22 +291,23 @@ UObject* UCSVImportFactory::FactoryCreateText(UClass* InClass, UObject* InParent
 	return NewAsset;
 }
 
-bool UCSVImportFactory::ReimportCSV( UObject* Obj )
+EReimportResult::Type UCSVImportFactory::ReimportCSV(UObject* Obj)
 {
-	bool bHandled = false;
-	if(UCurveBase* Curve = Cast<UCurveBase>(Obj))
+	EReimportResult::Type Result = EReimportResult::Failed;
+
+	if (UCurveBase* Curve = Cast<UCurveBase>(Obj))
 	{
-		bHandled = Reimport(Curve, Curve->AssetImportData->GetFirstFilename());
+		Result = Reimport(Curve, Curve->AssetImportData->GetFirstFilename());
 	}
-	else if(UCurveTable* CurveTable = Cast<UCurveTable>(Obj))
+	else if (UCurveTable* CurveTable = Cast<UCurveTable>(Obj))
 	{
-		bHandled = Reimport(CurveTable, CurveTable->AssetImportData->GetFirstFilename());
+		Result = Reimport(CurveTable, CurveTable->AssetImportData->GetFirstFilename());
 	}
-	else if(UDataTable* DataTable = Cast<UDataTable>(Obj))
+	else if (UDataTable* DataTable = Cast<UDataTable>(Obj))
 	{
-		bHandled = Reimport(DataTable, DataTable->AssetImportData->GetFirstFilename());
+		Result = Reimport(DataTable, DataTable->AssetImportData->GetFirstFilename());
 	}
-	return bHandled;
+	return Result;
 }
 
 void UCSVImportFactory::ParseFromJson(TSharedRef<class FJsonObject> ImportSettingsJson)
@@ -311,22 +315,27 @@ void UCSVImportFactory::ParseFromJson(TSharedRef<class FJsonObject> ImportSettin
 	FJsonObjectConverter::JsonObjectToUStruct(ImportSettingsJson, FCSVImportSettings::StaticStruct(), &AutomatedImportSettings, 0, 0);
 }
 
-bool UCSVImportFactory::Reimport(UObject* Obj, const FString& Path)
+EReimportResult::Type UCSVImportFactory::Reimport(UObject* Obj, const FString& Path)
 {
-	if(Path.IsEmpty() == false)
+	if (Path.IsEmpty() == false)
 	{ 
 		FString FilePath = IFileManager::Get().ConvertToRelativePath(*Path);
 
 		FString Data;
-		if( FFileHelper::LoadFileToString( Data, *FilePath) )
+		if (FFileHelper::LoadFileToString(Data, *FilePath))
 		{
 			const TCHAR* Ptr = *Data;
 			CurrentFilename = FilePath; //not thread safe but seems to be how it is done..
-			auto Result = FactoryCreateText( Obj->GetClass(), Obj->GetOuter(), Obj->GetFName(), Obj->GetFlags(), NULL, *FPaths::GetExtension(FilePath), Ptr, Ptr+Data.Len(), NULL );
-			return true;
+			bool bWasCancelled = false;
+			UObject* Result = FactoryCreateText(Obj->GetClass(), Obj->GetOuter(), Obj->GetFName(), Obj->GetFlags(), nullptr, *FPaths::GetExtension(FilePath), Ptr, Ptr+Data.Len(), nullptr, bWasCancelled);
+			if (bWasCancelled)
+			{
+				return EReimportResult::Cancelled;
+			}
+			return Result ? EReimportResult::Succeeded : EReimportResult::Failed;
 		}
 	}
-	return false;
+	return EReimportResult::Failed;
 }
 
 TArray<FString> UCSVImportFactory::DoImportDataTable(UDataTable* TargetDataTable, const FString& DataToImport)
@@ -375,15 +384,15 @@ UReimportDataTableFactory::UReimportDataTableFactory(const FObjectInitializer& O
 	Formats.Add(TEXT("json;JavaScript Object Notation"));
 }
 
-bool UReimportDataTableFactory::FactoryCanImport( const FString& Filename )
+bool UReimportDataTableFactory::FactoryCanImport(const FString& Filename)
 {
 	return true;
 }
 
-bool UReimportDataTableFactory::CanReimport( UObject* Obj, TArray<FString>& OutFilenames )
+bool UReimportDataTableFactory::CanReimport(UObject* Obj, TArray<FString>& OutFilenames)
 {	
 	UDataTable* DataTable = Cast<UDataTable>(Obj);
-	if(DataTable)
+	if (DataTable)
 	{
 		DataTable->AssetImportData->ExtractFilenames(OutFilenames);
 		return true;
@@ -391,19 +400,19 @@ bool UReimportDataTableFactory::CanReimport( UObject* Obj, TArray<FString>& OutF
 	return false;
 }
 
-void UReimportDataTableFactory::SetReimportPaths( UObject* Obj, const TArray<FString>& NewReimportPaths )
+void UReimportDataTableFactory::SetReimportPaths(UObject* Obj, const TArray<FString>& NewReimportPaths)
 {	
 	UDataTable* DataTable = Cast<UDataTable>(Obj);
-	if(DataTable && ensure(NewReimportPaths.Num() == 1))
+	if (DataTable && ensure(NewReimportPaths.Num() == 1))
 	{
 		DataTable->AssetImportData->UpdateFilenameOnly(NewReimportPaths[0]);
 	}
 }
 
-EReimportResult::Type UReimportDataTableFactory::Reimport( UObject* Obj )
+EReimportResult::Type UReimportDataTableFactory::Reimport(UObject* Obj)
 {	
-	auto Result = EReimportResult::Failed;
-	if(auto DataTable = Cast<UDataTable>(Obj))
+	EReimportResult::Type Result = EReimportResult::Failed;
+	if (UDataTable* DataTable = Cast<UDataTable>(Obj))
 	{
 		FDataTableEditorUtils::BroadcastPreChange(DataTable, FDataTableEditorUtils::EDataTableChangeInfo::RowList);
 		Result = UCSVImportFactory::ReimportCSV(DataTable) ? EReimportResult::Succeeded : EReimportResult::Failed;
@@ -425,10 +434,10 @@ UReimportCurveTableFactory::UReimportCurveTableFactory(const FObjectInitializer&
 	Formats.Add(TEXT("json;JavaScript Object Notation"));
 }
 
-bool UReimportCurveTableFactory::CanReimport( UObject* Obj, TArray<FString>& OutFilenames )
+bool UReimportCurveTableFactory::CanReimport(UObject* Obj, TArray<FString>& OutFilenames)
 {	
 	UCurveTable* CurveTable = Cast<UCurveTable>(Obj);
-	if(CurveTable)
+	if (CurveTable)
 	{
 		CurveTable->AssetImportData->ExtractFilenames(OutFilenames);
 		return true;
@@ -436,18 +445,18 @@ bool UReimportCurveTableFactory::CanReimport( UObject* Obj, TArray<FString>& Out
 	return false;
 }
 
-void UReimportCurveTableFactory::SetReimportPaths( UObject* Obj, const TArray<FString>& NewReimportPaths )
+void UReimportCurveTableFactory::SetReimportPaths(UObject* Obj, const TArray<FString>& NewReimportPaths)
 {	
 	UCurveTable* CurveTable = Cast<UCurveTable>(Obj);
-	if(CurveTable && ensure(NewReimportPaths.Num() == 1))
+	if (CurveTable && ensure(NewReimportPaths.Num() == 1))
 	{
 		CurveTable->AssetImportData->UpdateFilenameOnly(NewReimportPaths[0]);
 	}
 }
 
-EReimportResult::Type UReimportCurveTableFactory::Reimport( UObject* Obj )
+EReimportResult::Type UReimportCurveTableFactory::Reimport(UObject* Obj)
 {	
-	if(Cast<UCurveTable>(Obj))
+	if (Cast<UCurveTable>(Obj))
 	{
 		return UCSVImportFactory::ReimportCSV(Obj) ? EReimportResult::Succeeded : EReimportResult::Failed;
 	}
@@ -467,10 +476,10 @@ UReimportCurveFactory::UReimportCurveFactory(const FObjectInitializer& ObjectIni
 	SupportedClass = UCurveBase::StaticClass();
 }
 
-bool UReimportCurveFactory::CanReimport( UObject* Obj, TArray<FString>& OutFilenames )
+bool UReimportCurveFactory::CanReimport(UObject* Obj, TArray<FString>& OutFilenames)
 {	
 	UCurveBase* CurveBase = Cast<UCurveBase>(Obj);
-	if(CurveBase)
+	if (CurveBase)
 	{
 		CurveBase->AssetImportData->ExtractFilenames(OutFilenames);
 		return true;
@@ -478,18 +487,18 @@ bool UReimportCurveFactory::CanReimport( UObject* Obj, TArray<FString>& OutFilen
 	return false;
 }
 
-void UReimportCurveFactory::SetReimportPaths( UObject* Obj, const TArray<FString>& NewReimportPaths )
+void UReimportCurveFactory::SetReimportPaths(UObject* Obj, const TArray<FString>& NewReimportPaths)
 {	
 	UCurveBase* CurveBase = Cast<UCurveBase>(Obj);
-	if(CurveBase && ensure(NewReimportPaths.Num() == 1))
+	if (CurveBase && ensure(NewReimportPaths.Num() == 1))
 	{
 		CurveBase->AssetImportData->UpdateFilenameOnly(NewReimportPaths[0]);
 	}
 }
 
-EReimportResult::Type UReimportCurveFactory::Reimport( UObject* Obj )
+EReimportResult::Type UReimportCurveFactory::Reimport(UObject* Obj)
 {	
-	if(Cast<UCurveBase>(Obj))
+	if (Cast<UCurveBase>(Obj))
 	{
 		return UCSVImportFactory::ReimportCSV(Obj) ? EReimportResult::Succeeded : EReimportResult::Failed;
 	}
diff --git a/Engine/Source/Editor/UnrealEd/Private/Factories/Factory.cpp b/Engine/Source/Editor/UnrealEd/Private/Factories/Factory.cpp
--- a/Engine/Source/Editor/UnrealEd/Private/Factories/Factory.cpp
+++ b/Engine/Source/Editor/UnrealEd/Private/Factories/Factory.cpp
@@ -80,7 +80,7 @@ UObject* UFactory::FactoryCreateFile(UClass* InClass, UObject* InParent, FName I
 		ParseParms(Parms);
 		const TCHAR* Ptr = *Data;
 			
-		return FactoryCreateText(InClass, InParent, InName, Flags, nullptr, *FileExtension, Ptr, Ptr + Data.Len(), Warn);
+		return FactoryCreateText(InClass, InParent, InName, Flags, nullptr, *FileExtension, Ptr, Ptr + Data.Len(), Warn, bOutOperationCanceled);
 	}
 
 	// load as binary
diff --git a/Engine/Source/Editor/UnrealEd/Public/SCSVImportOptions.h b/Engine/Source/Editor/UnrealEd/Public/SCSVImportOptions.h
--- a/Engine/Source/Editor/UnrealEd/Public/SCSVImportOptions.h
+++ b/Engine/Source/Editor/UnrealEd/Public/SCSVImportOptions.h
@@ -34,7 +34,7 @@ public:
 	SCSVImportOptions()
 		: bImport(false)
 		, SelectedImportType(ECSVImportType::ECSV_DataTable)
-		, SelectedStruct(NULL)
+		, SelectedStruct(nullptr)
 		{}
 
 	void Construct(const FArguments& InArgs);
