Warnings and errors generated by macro expansion now include a link back to the macro instance that created the suspect nodes

diff --git a/Engine/Source/Editor/KismetCompiler/Private/KismetCompiler.cpp b/Engine/Source/Editor/KismetCompiler/Private/KismetCompiler.cpp
--- a/Engine/Source/Editor/KismetCompiler/Private/KismetCompiler.cpp
+++ b/Engine/Source/Editor/KismetCompiler/Private/KismetCompiler.cpp
@@ -3396,6 +3396,7 @@ void FKismetCompilerContext::ExpandTunnelsAndMacros(UEdGraph* SourceGraph)
 			for (int32 I = 0; I < ClonedGraph->Nodes.Num(); ++I)
 			{
 				MacroGeneratedNodes.Add(ClonedGraph->Nodes[I], CurrentNode);
+				MessageLog.NotifyIntermediateMacroNode(CurrentNode, ClonedGraph->Nodes[I]);
 			}
 
 			TArray<UEdGraphNode*> MacroNodes(ClonedGraph->Nodes);
diff --git a/Engine/Source/Editor/UnrealEd/Private/EdGraphToken.cpp b/Engine/Source/Editor/UnrealEd/Private/EdGraphToken.cpp
--- a/Engine/Source/Editor/UnrealEd/Private/EdGraphToken.cpp
+++ b/Engine/Source/Editor/UnrealEd/Private/EdGraphToken.cpp
@@ -3,30 +3,24 @@
 #include "EdGraphToken.h"
 #include "Kismet2/CompilerResultsLog.h"
 #include "Kismet2/BlueprintEditorUtils.h"
+#include "Misc/UObjectToken.h"
 
-TSharedRef<IMessageToken> FEdGraphToken::Create(const UObject* InObject, const FCompilerResultsLog* Log, UEdGraphNode*& OutSourceNode)
+void FEdGraphToken::Create(const UObject* InObject, FCompilerResultsLog* Log, FTokenizedMessage &OutMessage, TArray<UEdGraphNode*>& OutSourceNodes)
 {
-	const UObject* SourceNode = Log->FindSourceObject(InObject);
-	if (OutSourceNode == nullptr)
-	{
-		OutSourceNode = const_cast<UEdGraphNode*>(Cast<UEdGraphNode>(SourceNode));
-	}
-	return MakeShareable(new FEdGraphToken(Log->FindSourceObject(InObject), nullptr));
+	return CreateInternal(InObject, Log, OutMessage, OutSourceNodes, nullptr);
 }
 
-TSharedRef<IMessageToken> FEdGraphToken::Create(const UEdGraphPin* InPin, const FCompilerResultsLog* Log, UEdGraphNode*& OutSourceNode)
+void FEdGraphToken::Create(const UEdGraphPin* InPin, FCompilerResultsLog* Log, FTokenizedMessage &OutMessage, TArray<UEdGraphNode*>& OutSourceNodes)
 {
-	const UObject* SourceNode = InPin ? Log->FindSourceObject(InPin->GetOwningNode()) : nullptr;
-	if (OutSourceNode == nullptr)
+	if(InPin && InPin->GetOwningNode())
 	{
-		OutSourceNode = const_cast<UEdGraphNode*>(Cast<UEdGraphNode>(SourceNode));
+		return CreateInternal(InPin->GetOwningNode(), Log, OutMessage, OutSourceNodes, InPin);
 	}
-	return MakeShareable(new FEdGraphToken(SourceNode, Log->FindSourcePin(InPin)));
 }
 
-TSharedRef<IMessageToken> FEdGraphToken::Create(const TCHAR* String, const FCompilerResultsLog* Log, UEdGraphNode*& OutSourceNode)
+void FEdGraphToken::Create(const TCHAR* String, FCompilerResultsLog* Log, FTokenizedMessage &OutMessage, TArray<UEdGraphNode*>& OutSourceNode)
 {
-	return FTextToken::Create(FText::FromString(FString(String)));
+	OutMessage.AddToken( FTextToken::Create(FText::FromString(FString(String))) );
 }
 
 const UEdGraphPin* FEdGraphToken::GetPin() const
@@ -76,3 +70,39 @@ FEdGraphToken::FEdGraphToken(const UObject* InObject, const UEdGraphPin* InPin)
 		CachedText = NSLOCTEXT("MessageLog", "NoneObjectToken", "<None>");
 	}
 }
+
+void FEdGraphToken::CreateInternal(const UObject* InObject, FCompilerResultsLog* Log, FTokenizedMessage &OutMessage, TArray<UEdGraphNode*>& OutSourceNodes, const UEdGraphPin* Pin)
+{
+	UObject* SourceObject = Log->FindSourceObject(const_cast<UObject*>(InObject));
+	OutMessage.AddToken( MakeShareable(new FEdGraphToken(SourceObject, Pin ? Log->FindSourcePin(Pin) : nullptr)));
+	if (UEdGraphNode* SourceNode = Cast<UEdGraphNode>(SourceObject))
+	{
+		OutSourceNodes.Add(SourceNode);
+		
+		// If this node came from a macro it actually has two source nodes, look up the other source node and add that as well:
+		if(const UEdGraph* OwningGraph = SourceNode->GetGraph())
+		{
+			if(const UEdGraphSchema* Schema = OwningGraph->GetSchema())
+			{
+				if(Schema->GetGraphType(OwningGraph) == GT_Macro)
+				{
+					if(UObject* MacroSourceObject = Log->FindSourceMacroInstance(Cast<const UEdGraphNode>(InObject)))	
+					{
+						OutMessage.AddToken( FTextToken::Create( NSLOCTEXT("EdGraphToken", "FromMacroInstance", "generated from expanding")) );
+						OutMessage.AddToken( MakeShareable(new FEdGraphToken(MacroSourceObject, nullptr)) );
+						if(UEdGraphNode* MacroSourceNode = Cast<UEdGraphNode>(MacroSourceObject))
+						{
+							OutSourceNodes.Add(MacroSourceNode);
+						}
+					}
+				}
+			}
+		}
+	}
+	
+	if(SourceObject)
+	{
+		// The message link is only used when the user double clicks on the line, we'll jump to the first source object by default:
+		OutMessage.SetMessageLink(FUObjectToken::Create(SourceObject));
+	}
+}
diff --git a/Engine/Source/Editor/UnrealEd/Private/Kismet2/CompilerResultsLog.cpp b/Engine/Source/Editor/UnrealEd/Private/Kismet2/CompilerResultsLog.cpp
--- a/Engine/Source/Editor/UnrealEd/Private/Kismet2/CompilerResultsLog.cpp
+++ b/Engine/Source/Editor/UnrealEd/Private/Kismet2/CompilerResultsLog.cpp
@@ -236,6 +236,39 @@ void FCompilerResultsLog::InternalLogSummary()
 	}
 }
 
+bool FCompilerResultsLog::CommitPotentialMessages(UEdGraphNode* Source)
+{
+	TArray<TSharedRef<FTokenizedMessage>> FoundMessages;
+	if (PotentialMessages.RemoveAndCopyValue(Source, FoundMessages))
+	{
+		for (const TSharedRef<FTokenizedMessage>& Message : FoundMessages)
+		{
+			switch (Message->GetSeverity())
+			{
+			case EMessageSeverity::Error:
+				++NumErrors;
+				break;
+
+			case EMessageSeverity::Warning:
+				++NumWarnings;
+				break;
+
+			default:
+				break;
+			}
+
+			// build nodes list:
+			TArray<UEdGraphNode*> Nodes;
+			GetNodesFromTokens(Message->GetMessageTokens(), Nodes);
+
+			Nodes.Add(Source);
+			InternalLogMessage(Message, Nodes);
+		}
+		return true;
+	}
+	return false;
+}
+
 /** Update the source backtrack map to note that NewObject was most closely generated/caused by the SourceObject */
 void FCompilerResultsLog::NotifyIntermediateObjectCreation(UObject* NewObject, UObject* SourceObject)
 {
@@ -258,11 +291,38 @@ UObject const* FCompilerResultsLog::FindSourceObject(UObject const* PossiblyDupl
 	return SourceBacktrackMap.FindSourceObject(PossiblyDuplicatedObject);
 }
 
+UObject* FCompilerResultsLog::FindSourceMacroInstance(const UEdGraphNode* IntermediateNode) const
+{
+	UObject* Result = nullptr;
+	const UEdGraphNode* Iter = IntermediateNode;
+	while (const TWeakObjectPtr<UEdGraphNode>* SourceInstanceNode = FullMacroBacktrackMap.Find(Iter))
+	{
+		Iter = SourceInstanceNode->Get();
+	}
+
+	Result = const_cast<UEdGraphNode*>(Iter);
+
+	if(Result)
+	{
+		Result = const_cast<FCompilerResultsLog*>(this)->FindSourceObject(Result);
+	}
+
+	return Result;
+}
+
 void FCompilerResultsLog::NotifyIntermediateTunnelNode(const UEdGraphNode* Node, const UEdGraphNode* OuterTunnelInstance)
 {
 	IntermediateTunnelNodeToTunnelInstanceMap.Add(Node, OuterTunnelInstance);
 }
 
+void FCompilerResultsLog::NotifyIntermediateMacroNode(UEdGraphNode* SourceNode, const UEdGraphNode* IntermediateNode)
+{
+	if(ensure(SourceNode != IntermediateNode))
+	{
+		FullMacroBacktrackMap.Add(IntermediateNode, SourceNode);
+	}
+}
+
 const UEdGraphNode* FCompilerResultsLog::GetIntermediateTunnelInstance(const UEdGraphNode* IntermediateNode) const
 {
 	TWeakObjectPtr<const UEdGraphNode> Result;
@@ -345,11 +405,11 @@ const UEdGraphPin* FCompilerResultsLog::FindSourcePin(const UEdGraphPin* Possibl
 	return SourceBacktrackMap.FindSourcePin(PossiblyDuplicatedPin);
 }
 
-void FCompilerResultsLog::InternalLogMessage(const TSharedRef<FTokenizedMessage>& Message, UEdGraphNode* SourceNode)
+void FCompilerResultsLog::InternalLogMessage(const TSharedRef<FTokenizedMessage>& Message, const TArray<UEdGraphNode*>& SourceNodes)
 {
 	const EMessageSeverity::Type Severity = Message->GetSeverity();
 	Messages.Add(Message);
-	AnnotateNode(SourceNode, Message);
+	AnnotateNode(SourceNodes, Message);
 
 	if (!bSilentMode && (!bLogInfoOnly || (Severity == EMessageSeverity::Info)))
 	{
@@ -376,12 +436,10 @@ void FCompilerResultsLog::InternalLogMessage(const TSharedRef<FTokenizedMessage>
 	}
 }
 
-void FCompilerResultsLog::AnnotateNode(class UEdGraphNode* Node, TSharedRef<FTokenizedMessage> LogLine)
+void FCompilerResultsLog::AnnotateNode(const TArray<UEdGraphNode*>& Nodes, TSharedRef<FTokenizedMessage> LogLine)
 {
-	if (Node != nullptr)
+	for(UEdGraphNode* Node : Nodes)
 	{
-		LogLine->SetMessageLink(FUObjectToken::Create(Node));
-
 		if (bAnnotateMentionedNodes)
 		{
 			// Determine if this message is the first or more important than the previous one (only showing one error/warning per node for now)
@@ -532,6 +590,43 @@ void FCompilerResultsLog::GetGlobalModuleCompilerDump(const FString& LogDump, EC
 
 	MessageLog.AddMessages(ParseCompilerLogDump(LogDump));
 }
+		
+void FCompilerResultsLog::GetNodesFromTokens(const TArray<TSharedRef<IMessageToken> >& MessageTokens, TArray<UEdGraphNode*>& OutOwnerNodes)
+{
+	for (TSharedRef<IMessageToken> const& Token : MessageTokens)
+	{
+		if (Token->GetType() == EMessageToken::Object)
+		{
+			FWeakObjectPtr ObjectPtr = ((FUObjectToken&)Token.Get()).GetObject();
+			if (!ObjectPtr.IsValid())
+			{
+				continue;
+			}
+			UObject* ObjectArgument = ObjectPtr.Get();
+			if(UEdGraphNode* Node = Cast<UEdGraphNode>(ObjectArgument))
+			{
+				OutOwnerNodes.Add(Node);
+			}
+		}
+		else if (Token->GetType() == EMessageToken::EdGraph)
+		{
+			FEdGraphToken& AsEdGraphToken = ((FEdGraphToken&)Token.Get());
+			const UEdGraphPin* PinBeingReferenced = AsEdGraphToken.GetPin();
+			UEdGraphNode* OwnerNode = Cast<UEdGraphNode>((UObject*)AsEdGraphToken.GetGraphObject());
+			if (OwnerNode == nullptr)
+			{
+				if (PinBeingReferenced)
+				{
+					OutOwnerNodes.Add(Cast<UEdGraphNode>(PinBeingReferenced->GetOwningNodeUnchecked()));
+				}
+			}
+			else
+			{
+				OutOwnerNodes.Add(OwnerNode);
+			}
+		}
+	}
+}
 
 void FCompilerResultsLog::Append(FCompilerResultsLog const& Other)
 {
@@ -558,37 +653,9 @@ void FCompilerResultsLog::Append(FCompilerResultsLog const& Other)
 		}
 		Messages.Add(Message);
 		
-		UEdGraphNode* OwnerNode = nullptr;
-		for (TSharedRef<IMessageToken> const& Token : Message->GetMessageTokens())
-		{
-			if (Token->GetType() == EMessageToken::Object)
-			{
-				FWeakObjectPtr ObjectPtr = ((FUObjectToken&)Token.Get()).GetObject();
-				if (!ObjectPtr.IsValid())
-				{
-					continue;
-				}
-				UObject* ObjectArgument = ObjectPtr.Get();
-
-				OwnerNode = Cast<UEdGraphNode>(ObjectArgument);
-				break;
-			}
-			else if (Token->GetType() == EMessageToken::EdGraph)
-			{
-				FEdGraphToken& AsEdGraphToken = ((FEdGraphToken&)Token.Get());
-				const UEdGraphPin* PinBeingReferenced = AsEdGraphToken.GetPin();
-				OwnerNode = Cast<UEdGraphNode>((UObject*)AsEdGraphToken.GetGraphObject());
-				if (OwnerNode == nullptr)
-				{
-					if (PinBeingReferenced)
-					{
-						OwnerNode = PinBeingReferenced->GetOwningNodeUnchecked();
-					}
-				}
-				break;
-			}
-		}
-		AnnotateNode(OwnerNode, Message);
+		TArray<UEdGraphNode*> OwnerNodes;
+		GetNodesFromTokens(Message->GetMessageTokens(), OwnerNodes);
+		AnnotateNode(OwnerNodes, Message);
 	}
 }
 
diff --git a/Engine/Source/Editor/UnrealEd/Public/EdGraphToken.h b/Engine/Source/Editor/UnrealEd/Public/EdGraphToken.h
--- a/Engine/Source/Editor/UnrealEd/Public/EdGraphToken.h
+++ b/Engine/Source/Editor/UnrealEd/Public/EdGraphToken.h
@@ -16,9 +16,9 @@ class FEdGraphToken : public IMessageToken
 {
 public:
 	/** Factory method, tokens can only be constructed as shared refs */
-	UNREALED_API static TSharedRef<IMessageToken> Create(const UObject* InObject, const FCompilerResultsLog* Log, UEdGraphNode*& OutSourceNode);
-	UNREALED_API static TSharedRef<IMessageToken> Create(const UEdGraphPin* InPin, const FCompilerResultsLog* Log, UEdGraphNode*& OutSourceNode);
-	UNREALED_API static TSharedRef<IMessageToken> Create(const TCHAR* String, const FCompilerResultsLog* Log, UEdGraphNode*& OutSourceNode);
+	UNREALED_API static void Create(const UObject* InObject, FCompilerResultsLog* Log, FTokenizedMessage &OutMessage, TArray<UEdGraphNode*>& OutSourceNodes);
+	UNREALED_API static void Create(const UEdGraphPin* InPin, FCompilerResultsLog* Log, FTokenizedMessage &OutMessage, TArray<UEdGraphNode*>& OutSourceNodes);
+	UNREALED_API static void Create(const TCHAR* String, FCompilerResultsLog* Log, FTokenizedMessage &OutMessage, TArray<UEdGraphNode*>& OutSourceNodes);
 
 	/** Begin IMessageToken interface */
 	virtual EMessageToken::Type GetType() const override
@@ -32,6 +32,8 @@ public:
 private:
 	/** Private constructor */
 	FEdGraphToken( const UObject* InObject, const UEdGraphPin* InPin );
+	/** Helper to facilitate code reuse between ::Create overloads */
+	static void CreateInternal(const UObject* InObject, FCompilerResultsLog* Log, FTokenizedMessage &OutMessage, TArray<UEdGraphNode*>& OutSourceNodes, const UEdGraphPin* Pin);
 
 	/** An object being referenced by this token, if any */
 	FWeakObjectPtr ObjectBeingReferenced;
diff --git a/Engine/Source/Editor/UnrealEd/Public/Kismet2/CompilerResultsLog.h b/Engine/Source/Editor/UnrealEd/Public/Kismet2/CompilerResultsLog.h
--- a/Engine/Source/Editor/UnrealEd/Public/Kismet2/CompilerResultsLog.h
+++ b/Engine/Source/Editor/UnrealEd/Public/Kismet2/CompilerResultsLog.h
@@ -117,6 +117,9 @@ protected:
 	// Map to track intermediate tunnel nodes back to the intermediate expansion tunnel instance.
 	TMap<TWeakObjectPtr<const UEdGraphNode>, TWeakObjectPtr<const UEdGraphNode>> IntermediateTunnelNodeToTunnelInstanceMap;
 
+	// Map to track intermediate nodes back to the source macro instance nodes
+	TMap< TWeakObjectPtr<const UEdGraphNode>, TWeakObjectPtr<UEdGraphNode> > FullMacroBacktrackMap;
+
 public:
 	FCompilerResultsLog(bool bIsCompatibleWithEvents = true);
 	virtual ~FCompilerResultsLog();
@@ -182,8 +185,8 @@ public:
 	TSharedRef<FTokenizedMessage> StorePotentialError(const UEdGraphNode* Source, const TCHAR* Format, Args... args)
 	{
 		TSharedRef<FTokenizedMessage> Line = FTokenizedMessage::Create(EMessageSeverity::Error);
-		UEdGraphNode* SourceNode = nullptr;
-		Tokenize(Format, *Line, SourceNode, args...);
+		TArray<UEdGraphNode*> SourceNodes;
+		Tokenize(Format, *Line, SourceNodes, args...);
 		PotentialMessages.FindOrAdd(Source).Add(Line);
 		return Line;
 	}
@@ -196,8 +199,8 @@ public:
 	TSharedRef<FTokenizedMessage> StorePotentialWarning(const UEdGraphNode* Source, const TCHAR* Format, Args... args)
 	{
 		TSharedRef<FTokenizedMessage> Line = FTokenizedMessage::Create(EMessageSeverity::Warning);
-		UEdGraphNode* SourceNode = nullptr;
-		Tokenize(Format, *Line, SourceNode, args...);
+		TArray<UEdGraphNode*> SourceNodes;
+		Tokenize(Format, *Line, SourceNodes, args...);
 		PotentialMessages.FindOrAdd(Source).Add(Line);
 		return Line;
 	}
@@ -210,8 +213,8 @@ public:
 	TSharedRef<FTokenizedMessage> StorePotentialNote(const UEdGraphNode* Source, const TCHAR* Format, Args... args)
 	{
 		TSharedRef<FTokenizedMessage> Line = FTokenizedMessage::Create(EMessageSeverity::Info);
-		UEdGraphNode* SourceNode = nullptr;
-		Tokenize(Format, *Line, SourceNode, args...);
+		TArray<UEdGraphNode*> SourceNodes;
+		Tokenize(Format, *Line, SourceNodes, args...);
 		PotentialMessages.FindOrAdd(Source).Add(Line);
 		return Line;
 	}
@@ -219,33 +222,7 @@ public:
 	/**
 	 * Commit all stored potential messages for a given node. Returns true if any messages were written.
 	 */
-	bool CommitPotentialMessages(UEdGraphNode* Source)
-	{
-		TArray<TSharedRef<FTokenizedMessage>> FoundMessages;
-		if (PotentialMessages.RemoveAndCopyValue(Source, FoundMessages))
-		{
-			for (const TSharedRef<FTokenizedMessage>& Message : FoundMessages)
-			{
-				switch (Message->GetSeverity())
-				{
-				case EMessageSeverity::Error:
-					++NumErrors;
-					break;
-
-				case EMessageSeverity::Warning:
-					++NumWarnings;
-					break;
-
-				default:
-					break;
-				}
-
-				InternalLogMessage(Message, Source);
-			}
-			return true;
-		}
-		return false;
-	}
+	bool CommitPotentialMessages(UEdGraphNode* Source);
 
 	/** Update the source backtrack map to note that NewObject was most closely generated/caused by the SourceObject */
 	void NotifyIntermediateObjectCreation(UObject* NewObject, UObject* SourceObject);
@@ -254,9 +231,13 @@ public:
 	/** Update the expansion map to note that Node was expanded from OuterTunnelInstance, both the node and tunnel instance should be intermediate nodes */
 	void NotifyIntermediateTunnelNode(const UEdGraphNode* Node, const UEdGraphNode* OuterTunnelInstance);
 
+	/** Update the map that tracks nodes created by macro instance nodes */
+	void NotifyIntermediateMacroNode(UEdGraphNode* SourceNode, const UEdGraphNode* IntermediateNode);
+
 	/** Returns the true source object for the passed in object */
 	UObject* FindSourceObject(UObject* PossiblyDuplicatedObject);
 	UObject const* FindSourceObject(UObject const* PossiblyDuplicatedObject) const;
+	UObject* FindSourceMacroInstance(const UEdGraphNode* IntermediateNode) const;
 
 	/** Returns the intermediate tunnel instance that generated the node */
 	const UEdGraphNode* GetIntermediateTunnelInstance(const UEdGraphNode* IntermediateNode) const;
@@ -301,21 +282,22 @@ protected:
 	/** Helper method to add a child event to the given parent event scope */
 	void AddChildEvent(TSharedPtr<FCompilerEvent>& ParentEventScope, TSharedRef<FCompilerEvent>& ChildEventScope);
 
-	void InternalLogMessage(const TSharedRef<FTokenizedMessage>& Message, UEdGraphNode* SourceNode );
+	void InternalLogMessage(const TSharedRef<FTokenizedMessage>& Message, const TArray<UEdGraphNode*>& SourceNodes );
 
-	void Tokenize(const TCHAR* Text, FTokenizedMessage &OutMessage, UEdGraphNode*& OutSourceNode)
+	void Tokenize(const TCHAR* Text, FTokenizedMessage &OutMessage, TArray<UEdGraphNode*>& OutSourceNode)
 	{
 		OutMessage.AddToken(FTextToken::Create(FText::FromString(Text)));
 	}
 
 	template<typename T, typename... Args>
-	void Tokenize(const TCHAR* Format, FTokenizedMessage &OutMessage, UEdGraphNode*& OutSourceNode, T First, Args... Rest)
+	void Tokenize(const TCHAR* Format, FTokenizedMessage &OutMessage, TArray<UEdGraphNode*>& OutSourceNode, T First, Args... Rest)
 	{
 		// read to next "@@":
 		if (const TCHAR* DelimiterStr = FCString::Strstr(Format, TEXT("@@")))
 		{
 			OutMessage.AddToken(FTextToken::Create(FText::FromString(FString(DelimiterStr - Format, Format))));
-			OutMessage.AddToken(FEdGraphToken::Create(First, this, OutSourceNode));
+			FEdGraphToken::Create(First, this, OutMessage, OutSourceNode);
+
 			const TCHAR* NextChunk = DelimiterStr + FCString::Strlen(TEXT("@@"));
 			if (*NextChunk)
 			{
@@ -334,13 +316,13 @@ protected:
 		// Convention for SourceNode established by the original version of the compiler results log
 		// was to annotate the error on the first node we can find. I am preserving that behavior
 		// for this type safe, variadic version:
-		UEdGraphNode* SourceNode = nullptr;
-		Tokenize(Format, *Message, SourceNode, args...);
-		InternalLogMessage(Message, SourceNode);
+		TArray<UEdGraphNode*> SourceNodes;
+		Tokenize(Format, *Message, SourceNodes, args...);
+		InternalLogMessage(Message, SourceNodes);
 	}
 
-	/** */
-	void AnnotateNode(class UEdGraphNode* Node, TSharedRef<FTokenizedMessage> LogLine);
+	/** Links the UEdGraphNode with the LogLine: */
+	void AnnotateNode(const TArray<UEdGraphNode*>& Nodes, TSharedRef<FTokenizedMessage> LogLine);
 
 	/** Internal method to append the final compiler results summary to the MessageLog */
 	void InternalLogSummary();
@@ -361,6 +343,9 @@ private:
 
 	/** Callback function for binding the global compiler dump to open the static compiler log */
 	static void GetGlobalModuleCompilerDump(const FString& LogDump, ECompilationResult::Type CompilationResult, bool bShowLog);
+	
+	/** Searches a token list for referenced UEdGraphNodes, used to update the nodes when a log is committed */
+	static void GetNodesFromTokens(const TArray<TSharedRef<IMessageToken> >& MessageTokens, TArray<UEdGraphNode*>& OutOwnerNodes);
 
 	/** The log's name, for easy re-use */
 	static const FName Name;
