Don't crash if an ICH is (somewhat surprisingly) missing

diff --git a/Engine/Source/Editor/UnrealEd/Private/Kismet2/BlueprintEditorUtils.cpp b/Engine/Source/Editor/UnrealEd/Private/Kismet2/BlueprintEditorUtils.cpp
--- a/Engine/Source/Editor/UnrealEd/Private/Kismet2/BlueprintEditorUtils.cpp
+++ b/Engine/Source/Editor/UnrealEd/Private/Kismet2/BlueprintEditorUtils.cpp
@@ -9077,10 +9077,10 @@ void FBlueprintEditorUtils::HandleDisableEditableWhenInherited(UObject* Modified
 			UBlueprintGeneratedClass* BPGC = Cast<UBlueprintGeneratedClass>(ArchetypeInstance->GetOuter());
 			if (BPGC)
 			{
-				UInheritableComponentHandler* ICH = BPGC->GetInheritableComponentHandler(false);
-				check(ICH);
-
-				ICH->RemoveOverridenComponentTemplate(ICH->FindKey(CastChecked<UActorComponent>(ArchetypeInstance)));
+				if (UInheritableComponentHandler* ICH = BPGC->GetInheritableComponentHandler(false))
+				{
+					ICH->RemoveOverridenComponentTemplate(ICH->FindKey(CastChecked<UActorComponent>(ArchetypeInstance)));
+				}
 			}
 		}
 	}
