Add optional message identifier to tokenized message, use message identifier to facilitate suppressing a new warning in the blueprint compiler, remove the unused (and unusable) FTokenizedMiscData

diff --git a/Engine/Source/Developer/MessageLog/Private/Model/MessageLogListingModel.cpp b/Engine/Source/Developer/MessageLog/Private/Model/MessageLogListingModel.cpp
--- a/Engine/Source/Developer/MessageLog/Private/Model/MessageLogListingModel.cpp
+++ b/Engine/Source/Developer/MessageLog/Private/Model/MessageLogListingModel.cpp
@@ -22,24 +22,6 @@ const TSharedPtr< FTokenizedMessage >  FMessageLogListingModel::GetMessageAtInde
 	return FoundMessage;
 }
 
-const TSharedPtr< FTokenizedMessage > FMessageLogListingModel::GetMessageFromData( const FTokenizedMiscData& MessageData ) const
-{
-	// Go through all the messages looking for a data match
-	for(uint32 PageIndex = 0; PageIndex < NumPages(); PageIndex++)
-	{
-		for( MessageContainer::TConstIterator It( GetMessageIterator(PageIndex) ); It; ++It )
-		{
-			const TSharedPtr< FTokenizedMessage > Message = *It;
-
-			if( Message->GetMessageData() == &MessageData )
-			{
-				return Message;
-			}
-		}
-	}
-	return nullptr;
-}
-
 FText FMessageLogListingModel::GetAllMessagesAsText( const uint32 PageIndex ) const
 {
 	FText CompiledText;
diff --git a/Engine/Source/Developer/MessageLog/Private/Model/MessageLogListingModel.h b/Engine/Source/Developer/MessageLog/Private/Model/MessageLogListingModel.h
--- a/Engine/Source/Developer/MessageLog/Private/Model/MessageLogListingModel.h
+++ b/Engine/Source/Developer/MessageLog/Private/Model/MessageLogListingModel.h
@@ -52,9 +52,6 @@ public:
 	/** Returns the message at the specified index */
 	const TSharedPtr<FTokenizedMessage> GetMessageAtIndex( const uint32 PageIndex, const int32 MessageIndex ) const;
 
-	/** Tries to find the first message that matches the message data */
-	const TSharedPtr<FTokenizedMessage> GetMessageFromData( const FTokenizedMiscData& MessageData ) const;
-
 	/** Gets all messages as a string */
 	FText GetAllMessagesAsText(const uint32 PageIndex) const;
 
diff --git a/Engine/Source/Developer/MessageLog/Private/Presentation/MessageLogListingViewModel.cpp b/Engine/Source/Developer/MessageLog/Private/Presentation/MessageLogListingViewModel.cpp
--- a/Engine/Source/Developer/MessageLog/Private/Presentation/MessageLogListingViewModel.cpp
+++ b/Engine/Source/Developer/MessageLog/Private/Presentation/MessageLogListingViewModel.cpp
@@ -162,11 +162,6 @@ void FMessageLogListingViewModel::ClearMessages()
 	MessageLogListingModel->ClearMessages();
 }
 
-TSharedPtr<FTokenizedMessage> FMessageLogListingViewModel::GetMessageFromData(const struct FTokenizedMiscData& MessageData) const
-{
-	return MessageLogListingModel->GetMessageFromData(MessageData);
-}
-
 void FMessageLogListingViewModel::SelectMessages( const TArray< TSharedRef<FTokenizedMessage> >& InSelectedMessages )
 {
 	SelectedFilteredMessages.Empty();
diff --git a/Engine/Source/Developer/MessageLog/Private/Presentation/MessageLogListingViewModel.h b/Engine/Source/Developer/MessageLog/Private/Presentation/MessageLogListingViewModel.h
--- a/Engine/Source/Developer/MessageLog/Private/Presentation/MessageLogListingViewModel.h
+++ b/Engine/Source/Developer/MessageLog/Private/Presentation/MessageLogListingViewModel.h
@@ -41,7 +41,6 @@ public:
 	virtual void AddMessage( const TSharedRef< class FTokenizedMessage >& NewMessage, bool bMirrorToOutputLog ) override;
 	virtual void AddMessages( const TArray< TSharedRef< class FTokenizedMessage > >& NewMessages, bool bMirrorToOutputLog ) override;
 	virtual void ClearMessages() override;
-	virtual TSharedPtr<FTokenizedMessage> GetMessageFromData(const struct FTokenizedMiscData& MessageData) const override;
 	virtual const TArray< TSharedRef<class FTokenizedMessage> >& GetSelectedMessages() const override;
 	virtual void SelectMessages( const TArray< TSharedRef<class FTokenizedMessage> >& InSelectedMessages ) override;
 	virtual const TArray< TSharedRef<class FTokenizedMessage> >& GetFilteredMessages() const override;
diff --git a/Engine/Source/Developer/MessageLog/Public/IMessageLogListing.h b/Engine/Source/Developer/MessageLog/Public/IMessageLogListing.h
--- a/Engine/Source/Developer/MessageLog/Public/IMessageLogListing.h
+++ b/Engine/Source/Developer/MessageLog/Public/IMessageLogListing.h
@@ -25,13 +25,6 @@ public:
 	/** Clears messages (if paged, in the current page) */
 	virtual void ClearMessages() = 0;
 
-	/** 
-	 * Tries to find the first message that matches the message data.
-	 * 
-	 * @param	MessageData	The message data to test.
-	 */
-	virtual TSharedPtr<class FTokenizedMessage> GetMessageFromData(const struct FTokenizedMiscData& MessageData) const = 0;
-
 	/** 
 	 * Gets a list of the selected messages for the specific log listing
 	 * 
diff --git a/Engine/Source/Editor/KismetCompiler/Private/KismetCompiler.cpp b/Engine/Source/Editor/KismetCompiler/Private/KismetCompiler.cpp
--- a/Engine/Source/Editor/KismetCompiler/Private/KismetCompiler.cpp
+++ b/Engine/Source/Editor/KismetCompiler/Private/KismetCompiler.cpp
@@ -1504,7 +1504,7 @@ void FKismetCompilerContext::PruneIsolatedNodes(const TArray<UEdGraphNode*>& Roo
 		if(bNeighborsNotPruned)
 		{
 			// Warn the user if they are attempting to read an output value from a pruned exec node:
-			MessageLog.Warning(*LOCTEXT("PrunedExecNodeAttemptedUse", "@@ was pruned because its Exec pin is not connected, the connected value is not available and will instead be read as default").ToString(), PrunedExecNodeWithNeighbors.Key);
+			MessageLog.Warning(FName(TEXT("PrunedExecInUse")), *LOCTEXT("PrunedExecNodeAttemptedUse", "@@ was pruned because its Exec pin is not connected, the connected value is not available and will instead be read as default").ToString(), PrunedExecNodeWithNeighbors.Key);
 		}
 	}
 }
diff --git a/Engine/Source/Editor/UnrealEd/Private/Kismet2/CompilerResultsLog.cpp b/Engine/Source/Editor/UnrealEd/Private/Kismet2/CompilerResultsLog.cpp
--- a/Engine/Source/Editor/UnrealEd/Private/Kismet2/CompilerResultsLog.cpp
+++ b/Engine/Source/Editor/UnrealEd/Private/Kismet2/CompilerResultsLog.cpp
@@ -13,6 +13,7 @@
 #include "EngineLogs.h"
 #include "Engine/Blueprint.h"
 #include "IMessageLogListing.h"
+#include "Settings/EditorProjectSettings.h"
 
 #if WITH_EDITOR
 
@@ -192,6 +193,33 @@ void FCompilerResultsLog::InternalLogEvent(const FCompilerEvent& InEvent, int32
 	}
 }
 
+bool FCompilerResultsLog::IsMessageEnabled(FName ID)
+{
+	if(ID == NAME_None)
+	{
+		return true;
+	}
+
+	const UBlueprintEditorProjectSettings* EditorProjectSettings = GetDefault<UBlueprintEditorProjectSettings>();
+	if(IsRunningCommandlet())
+	{
+		if(	EditorProjectSettings->DisabledCompilerMessagesHeadless.Contains(ID) ||
+			EditorProjectSettings->DisabledCompilerMessages.Contains(ID))
+		{
+			return false;
+		}
+	}
+	else
+	{
+		if( EditorProjectSettings->DisabledCompilerMessages.Contains(ID) )
+		{
+			return false;
+		}
+	}
+
+	return true;
+}
+
 void FCompilerResultsLog::InternalLogSummary()
 {
 	if(CurrentEventScope.IsValid())
@@ -262,7 +290,7 @@ bool FCompilerResultsLog::CommitPotentialMessages(UEdGraphNode* Source)
 			GetNodesFromTokens(Message->GetMessageTokens(), Nodes);
 
 			Nodes.Add(Source);
-			InternalLogMessage(Message, Nodes);
+			InternalLogMessage(Message->GetIdentifier(), Message, Nodes);
 		}
 		return true;
 	}
@@ -405,11 +433,12 @@ const UEdGraphPin* FCompilerResultsLog::FindSourcePin(const UEdGraphPin* Possibl
 	return SourceBacktrackMap.FindSourcePin(PossiblyDuplicatedPin);
 }
 
-void FCompilerResultsLog::InternalLogMessage(const TSharedRef<FTokenizedMessage>& Message, const TArray<UEdGraphNode*>& SourceNodes)
+void FCompilerResultsLog::InternalLogMessage(FName MessageID, const TSharedRef<FTokenizedMessage>& Message, const TArray<UEdGraphNode*>& SourceNodes)
 {
 	const EMessageSeverity::Type Severity = Message->GetSeverity();
 	Messages.Add(Message);
 	AnnotateNode(SourceNodes, Message);
+	Message->SetIdentifier(MessageID);
 
 	if (!bSilentMode && (!bLogInfoOnly || (Severity == EMessageSeverity::Info)))
 	{
diff --git a/Engine/Source/Editor/UnrealEd/Public/Kismet2/CompilerResultsLog.h b/Engine/Source/Editor/UnrealEd/Public/Kismet2/CompilerResultsLog.h
--- a/Engine/Source/Editor/UnrealEd/Public/Kismet2/CompilerResultsLog.h
+++ b/Engine/Source/Editor/UnrealEd/Public/Kismet2/CompilerResultsLog.h
@@ -148,7 +148,7 @@ public:
 	{
 		++NumErrors;
 		TSharedRef<FTokenizedMessage> Line = FTokenizedMessage::Create(EMessageSeverity::Error);
-		InternalLogMessage(Format, Line, args...);
+		InternalLogMessage(NAME_None, Format, Line, args...);
 		return Line;
 	}
 
@@ -161,10 +161,28 @@ public:
 	{
 		++NumWarnings;
 		TSharedRef<FTokenizedMessage> Line = FTokenizedMessage::Create(EMessageSeverity::Warning);
-		InternalLogMessage(Format, Line, args...);
+		InternalLogMessage(NAME_None, Format, Line, args...);
 		return Line;
 	}
 
+	/**
+	 * Write a warning in to the compiler log.
+	 * Note: @@ will be replaced by node or pin links for nodes/pins passed via varargs 
+	 */
+	template<typename... Args>
+	void Warning(FName ID, const TCHAR* Format, Args... args)
+	{
+		if(!IsMessageEnabled(ID))
+		{
+			return;
+		}
+
+		++NumWarnings;
+		TSharedRef<FTokenizedMessage> Line = FTokenizedMessage::Create(EMessageSeverity::Warning);
+		InternalLogMessage(ID, Format, Line, args...);
+		return;
+	}
+
 	/**
 	 * Write a note in to the compiler log.
 	 * Note: @@ will be replaced by node or pin links for nodes/pins passed via varargs 
@@ -173,7 +191,7 @@ public:
 	TSharedRef<FTokenizedMessage> Note(const TCHAR* Format, Args... args)
 	{
 		TSharedRef<FTokenizedMessage> Line = FTokenizedMessage::Create(EMessageSeverity::Info);
-		InternalLogMessage(Format, Line, args...);
+		InternalLogMessage(NAME_None, Format, Line, args...);
 		return Line;
 	}
 
@@ -282,7 +300,7 @@ protected:
 	/** Helper method to add a child event to the given parent event scope */
 	void AddChildEvent(TSharedPtr<FCompilerEvent>& ParentEventScope, TSharedRef<FCompilerEvent>& ChildEventScope);
 
-	void InternalLogMessage(const TSharedRef<FTokenizedMessage>& Message, const TArray<UEdGraphNode*>& SourceNodes );
+	void InternalLogMessage(FName MessageID, const TSharedRef<FTokenizedMessage>& Message, const TArray<UEdGraphNode*>& SourceNodes );
 
 	void Tokenize(const TCHAR* Text, FTokenizedMessage &OutMessage, TArray<UEdGraphNode*>& OutSourceNode)
 	{
@@ -311,14 +329,14 @@ protected:
 	}
 
 	template<typename... Args>
-	void InternalLogMessage(const TCHAR* Format, const TSharedRef<FTokenizedMessage>& Message, Args... args)
+	void InternalLogMessage(FName MessageID, const TCHAR* Format, const TSharedRef<FTokenizedMessage>& Message, Args... args)
 	{
 		// Convention for SourceNode established by the original version of the compiler results log
 		// was to annotate the error on the first node we can find. I am preserving that behavior
 		// for this type safe, variadic version:
 		TArray<UEdGraphNode*> SourceNodes;
 		Tokenize(Format, *Message, SourceNodes, args...);
-		InternalLogMessage(Message, SourceNodes);
+		InternalLogMessage(MessageID, Message, SourceNodes);
 	}
 
 	/** Links the UEdGraphNode with the LogLine: */
@@ -330,6 +348,8 @@ protected:
 	/** Internal helper method to recursively append event details into the MessageLog */
 	void InternalLogEvent(const FCompilerEvent& InEvent, int32 InDepth = 0);
 
+	/** Returns true if the user has requested this compiler message be suppressed */
+	bool IsMessageEnabled(FName ID);
 private:
 
 	/** Map of stored potential messages indexed by a node. Can be committed to the results log by calling CommitPotentialMessages for that node. */
diff --git a/Engine/Source/Editor/UnrealEd/Public/Settings/EditorProjectSettings.h b/Engine/Source/Editor/UnrealEd/Public/Settings/EditorProjectSettings.h
--- a/Engine/Source/Editor/UnrealEd/Public/Settings/EditorProjectSettings.h
+++ b/Engine/Source/Editor/UnrealEd/Public/Settings/EditorProjectSettings.h
@@ -183,5 +183,21 @@ public:
 	/** If enabled, the editor will load packages to look for soft references to actors when deleting/renaming them. This can be slow in large projects so disable this to improve performance but increase the chance of breaking blueprints/sequences that use soft actor references */
 	UPROPERTY(EditAnywhere, config, Category=Actors)
 	uint32 bValidateUnloadedSoftActorReferences : 1;
+
+	/** 
+	 * List of compiler messages that have been suppressed during cook (or other commandlets) for 
+	 * the current project - useful for silencing warnings that were added to the engine after 
+	 * project inception and are going to be addressed as they are found by content authors
+	 */
+	UPROPERTY(EditAnywhere, config, Category= Blueprints, DisplayName = "Compiler Messages Disabled during Cook")
+	TArray<FName> DisabledCompilerMessagesHeadless;
+	
+	/** 
+	 * List of compiler messages that have been suppressed completely - message suppression is only 
+	 * advisable when using blueprints that you cannot update and are raising innocuous warnings. 
+	 * If useless messages are being raised prefer to contact support rather than disabling messages
+	 */
+	UPROPERTY(EditAnywhere, config, Category= Blueprints, DisplayName = "Compiler Messages Disabled Entirely")
+	TArray<FName> DisabledCompilerMessages;
 };
 
diff --git a/Engine/Source/Runtime/Core/Private/Logging/TokenizedMessage.cpp b/Engine/Source/Runtime/Core/Private/Logging/TokenizedMessage.cpp
--- a/Engine/Source/Runtime/Core/Private/Logging/TokenizedMessage.cpp
+++ b/Engine/Source/Runtime/Core/Private/Logging/TokenizedMessage.cpp
@@ -102,6 +102,16 @@ FName FTokenizedMessage::GetSeverityIconName(EMessageSeverity::Type InSeverity)
 	return SeverityIconName;
 }
 
+FName FTokenizedMessage::GetIdentifier() const
+{
+	return Identifier;
+}
+
+void FTokenizedMessage::SetIdentifier(FName InIdentifier)
+{
+	Identifier = InIdentifier;
+}
+
 TSharedRef<FTokenizedMessage> FTokenizedMessage::AddToken( const TSharedRef<IMessageToken>& InToken )
 {
 	MessageTokens.Add( InToken );
@@ -124,17 +134,6 @@ EMessageSeverity::Type FTokenizedMessage::GetSeverity() const
 	return Severity;
 }
 
-TSharedRef<FTokenizedMessage> FTokenizedMessage::SetMessageData( FTokenizedMiscData* InMessageData )
-{
-	MessageData = InMessageData;
-	return AsShared();
-}
-
-FTokenizedMiscData* FTokenizedMessage::GetMessageData() const
-{
-	return MessageData;
-}
-
 const TArray<TSharedRef<IMessageToken> >& FTokenizedMessage::GetMessageTokens() const
 {
 	return MessageTokens;
diff --git a/Engine/Source/Runtime/Core/Public/Logging/TokenizedMessage.h b/Engine/Source/Runtime/Core/Public/Logging/TokenizedMessage.h
--- a/Engine/Source/Runtime/Core/Public/Logging/TokenizedMessage.h
+++ b/Engine/Source/Runtime/Core/Public/Logging/TokenizedMessage.h
@@ -102,10 +102,6 @@ protected:
 	FText CachedText;
 };
 
-
-/** Dummy struct that other structs need to inherit from if they wish to assign to MessageData */
-struct FTokenizedMiscData {};
-
 /** This class represents a rich tokenized message, such as would be used for compiler output with 'hyperlinks' to source file locations */
 class FTokenizedMessage : public TSharedFromThis<FTokenizedMessage>
 {
@@ -150,21 +146,6 @@ public:
 	 */
 	CORE_API EMessageSeverity::Type GetSeverity() const;
 
-	/** 
-	 * Sets the custom message data of this message
-	 * 
-	 * @param	InMessageData	The custom message data to set
-	 * @returns this message, so calls can be chained together
-	 */
-	CORE_API TSharedRef<FTokenizedMessage> SetMessageData( FTokenizedMiscData* InMessageData );
-
-	/** 
-	 * Gets the custom message data of this message
-	 * 
-	 * @returns the custom message data contained in this message
-	 */
-	CORE_API FTokenizedMiscData* GetMessageData() const;
-
 	/**
 	 * Get the tokens in this message
 	 * 
@@ -203,12 +184,18 @@ public:
 	 * @returns the name of the icon for the severity
 	 */
 	CORE_API static FName GetSeverityIconName(EMessageSeverity::Type InSeverity);
+	
+	/** @returns Identifier for the message, if set, else NAME_None */
+	CORE_API FName GetIdentifier() const;
+	
+	/** Assigns Identifier for the message to the provided name */
+	CORE_API void SetIdentifier(FName InIdentifier);
 
 private:
 	/** Private constructor - we want to only create these structures as shared references via Create() */
 	FTokenizedMessage()
 		: Severity( EMessageSeverity::Info )
-		, MessageData( nullptr )
+		, Identifier(NAME_None)
 	{ }
 
 protected:
@@ -224,8 +211,8 @@ private:
 	/** The severity of this message */
 	EMessageSeverity::Type Severity;
 
-	/** A payload for this message */
-	FTokenizedMiscData* MessageData;	// @todo Don't use TSharedPtr as it'll get free'd mid-sort! (refcount == 0)
+	/** Identifier for the message */
+	FName Identifier;
 };
 
 /** Basic message token with a localized text payload */
