Copying //UE4/Dev-Framework to //UE4/Dev-Main (Source: //UE4/Dev-Framework @ 4344132)
#rnx

============================
  MAJOR FEATURES & CHANGES
============================

Change 4289396 by [REDACTED_PERSON]

	Fix get data table row result pin not properly updating if it was split.

Change 4289407 by [REDACTED_PERSON]

	Don't allow set member input pin to be split at all

Change 4291871 by [REDACTED_PERSON]

	Exporting EnvQueryGenerator_SimpleGrid with AIMODULE_API #UE4

Change 4292078 by [REDACTED_PERSON]

	Fix GetDataTableRow title not updating when changing type

Change 4292770 by [REDACTED_PERSON]

	Interpreter Optimization 7 of N: CPF_ZeroConstructor assigned for struct properties that that are User Defined Structs

Change 4293207 by [REDACTED_PERSON]

	Interpreter Optimization 8 of N: CPF_IsPlainOldData and CPF_NoDestructor assigned for struct properties that that are User Defined Structs

Change 4297023 by [REDACTED_PERSON]

	Remove dumb returns

Change 4298303 by [REDACTED_PERSON]

	Fix for issues compiling blueprints after changing a UserDefinedStructure

Change 4301603 by [REDACTED_PERSON]

	Fix incorrect code generation at cook time when Blueprint nativization is enabled for objects instanced through the editor and assigned to an instanced object reference property.

	Change summary:
	- Modified FEmitDefaultValueHelper::HandleSpecialTypes() to emit runtime code to look for an existing instance of a default subobject only if the 'RF_DefaultSubObject' flag is set on the source instance. This works around instances that return 'true' for IsDefaultSubobject() but were not created with CreateDefaultSubobject() (e.g. - instanced in the editor UI).

Change 4302720 by [REDACTED_PERSON]

	Prevent user from entering 'None' as parameter name on a blueprint function

Change 4303695 by [REDACTED_PERSON]

	Don't allow the user to attempt to import a csv unless they've specified the rowstruct for datatables or the interpolation for curve tables.

Change 4304633 by [REDACTED_PERSON]

Change 4313208 by [REDACTED_PERSON]

	Warn users when the outputs of a node that is pruned are used

Change 4313339 by [REDACTED_PERSON]

	Interpreter Optimization 9 of N: Mark ubergraph function as final so that calls to it don't need to search the target object

Change 4313652 by [REDACTED_PERSON]

	Don't warn when a non-pure node that is connected only via its exec is orphaned - fixes overzealous warnings EngineTest

Change 4313661 by [REDACTED_PERSON]

	Comment typo fixes

Change 4313729 by [REDACTED_PERSON]

	Avoid crashing when running the help command for the GenerateBlueprintAPICommandlet

Change 4318175 by [REDACTED_PERSON]

	PR #5024: Require BlueprintAssignable delegate parameters to be supported by blueprints (Contributed by [REDACTED_PERSON])

Change 4318302 by [REDACTED_PERSON]

	Use define instead of duplicating text

Change 4318681 by [REDACTED_PERSON]

	Don't allow users to try to rename rows to emtpy strings, whitespace only, or names that have invalid characters.

Change 4318683 by [REDACTED_PERSON]

	Don't reorder the rowmap when removing a row.

Change 4318686 by [REDACTED_PERSON]

	Get rid of the failed to import message when cancelling the import of a csv or json file.

Change 4318689 by [REDACTED_PERSON]

	Back out changelist 4281972

Change 4318713 by [REDACTED_PERSON]

	Validate CommutativeAssociativeBinaryOperator metadata in UHT and also at the K2Node level in case the function changes to an invalid construct later

Change 4319363 by [REDACTED_PERSON]

	Warnings and errors generated by macro expansion now include a link back to the macro instance that created the suspect nodes

Change 4319559 by [REDACTED_PERSON]

	PR #5028: UE-63225: Use correct source BP when retrieving the tooltip text (Contributed by [REDACTED_PERSON])

Change 4323740 by [REDACTED_PERSON]

	Don't crash if an ICH is (somewhat surprisingly) missing

Change 4328877 by mason.seay

	BP for pin bug

Change 4336999 by [REDACTED_PERSON]

	Unshelved from pending changelist '4299645':

	PR #5003: Changed SendGameplayEvent()'s category from "Abilities" to "Ability" for consistency (Contributed by [REDACTED_PERSON])

Change 4337645 by [REDACTED_PERSON]

	PR #5049: Make GetPlayerControllerID pure (Contributed by [REDACTED_PERSON])

Change 4338211 by [REDACTED_PERSON]

	Fix missing backend support for initializing a container variable (including unconnected function output pins) to a default (literal) value term in nativized Blueprint C++ code.

	Change summary:
	- Added a 'CoerceProperty' argument to FBlueprintCompilerCppBackEnd::TermToText() and FEmitHelper::LiteralTerm(). This is passed through as the LHS term's associated property when compiling 'KCST_Assignment' statements.
	- Modified FEmitHelper::LiteralTerm() to emit an initialization list to C++. This occurs when the RHS of a variable assignment statement is a literal container term, which includes the initialization of K2Node_FunctionResult pins to a default (empty) value.
	- Added FBlueprintCompilerCppBackEnd::TermToText() and FEmitHelper::LiteralTerm() methods that accept a struct value for specifying arguments rather than increasing the number of parameters with default values. Modified the implementations to match.

Change 4338908 by [REDACTED_PERSON]

	Don't check when fetching world settings from IsPaused as it is already set up to handle it coming back null

Change 4339262 by [REDACTED_PERSON]

	Add comments to ERedirectType

Change 4339718 by [REDACTED_PERSON]

	Clearing editor world's EnvironmentQueryManager instace's LocalContextMap whenever a BP finished compiling #UE4

Change 4340548 by [REDACTED_PERSON]

	ChildActor's need to be spawned from the ChildActorTemplate otherwise any non-replicated  customized properties will be incorrect on the Client.

Change 4340756 by [REDACTED_PERSON]

	Don't allow a new pin to be matched twice

Change 4342704 by mason.seay

	Updated Pin Testing blueprints

Change 4319930 by [REDACTED_PERSON]

	Mark FPartyState as a BlueprintType

Change 4329176 by [REDACTED_PERSON]

	Add optional message identifier to tokenized message, use message identifier to facilitate suppressing a new warning in the blueprint compiler, remove the unused (and unusable) FTokenizedMiscData

Change 4337859 by [REDACTED_PERSON]

	Fixing misc restricted gameplay tag UI bugs

diff --git a/Engine/Plugins/Editor/GameplayTagsEditor/Source/GameplayTagsEditor/Private/GameplayTagsEditorModule.cpp b/Engine/Plugins/Editor/GameplayTagsEditor/Source/GameplayTagsEditor/Private/GameplayTagsEditorModule.cpp
--- a/Engine/Plugins/Editor/GameplayTagsEditor/Source/GameplayTagsEditor/Private/GameplayTagsEditorModule.cpp
+++ b/Engine/Plugins/Editor/GameplayTagsEditor/Source/GameplayTagsEditor/Private/GameplayTagsEditorModule.cpp
@@ -280,16 +280,14 @@ public:
 		}
 	}
 
-	bool DeleteTagRedirector(const FString& TagToDelete)
+	bool DeleteTagRedirector(const FName& TagToDelete)
 	{
-		FName TagName = FName(*TagToDelete);
-
 		UGameplayTagsSettings* Settings = GetMutableDefault<UGameplayTagsSettings>();
 		UGameplayTagsManager& Manager = UGameplayTagsManager::Get();
 
 		for (int32 i = 0; i < Settings->GameplayTagRedirects.Num(); i++)
 		{
-			if (Settings->GameplayTagRedirects[i].OldTagName == TagName)
+			if (Settings->GameplayTagRedirects[i].OldTagName == TagToDelete)
 			{
 				Settings->GameplayTagRedirects.RemoveAt(i);
 
@@ -299,7 +297,7 @@ public:
 
 				Manager.EditorRefreshGameplayTagTree();
 
-				ShowNotification(FText::Format(LOCTEXT("RemoveTagRedirect", "Deleted tag redirect {0}"), FText::FromString(TagToDelete)), 5.0f);
+				ShowNotification(FText::Format(LOCTEXT("RemoveTagRedirect", "Deleted tag redirect {0}"), FText::FromName(TagToDelete)), 5.0f);
 
 				return true;
 			}
@@ -325,11 +323,13 @@ public:
 		UGameplayTagsSettings*			Settings = GetMutableDefault<UGameplayTagsSettings>();
 		UGameplayTagsDeveloperSettings* DevSettings = GetMutableDefault<UGameplayTagsDeveloperSettings>();
 
+		FName NewTagName = FName(*NewTag);
+
 		// Delete existing redirector
-		DeleteTagRedirector(NewTag);
+		DeleteTagRedirector(NewTagName);
 
 		// Already in the list as an explicit tag, ignore. Note we want to add if it is in implicit tag. (E.g, someone added A.B.C then someone tries to add A.B)
-		if (Manager.IsDictionaryTag(FName(*NewTag)))
+		if (Manager.IsDictionaryTag(NewTagName))
 		{
 			ShowNotification(FText::Format(LOCTEXT("AddTagFailure_AlreadyExists", "Failed to add gameplay tag {0}, already exists!"), FText::FromString(NewTag)), 10.0f);
 
@@ -468,9 +468,9 @@ public:
 		return true;
 	}
 
-	virtual bool DeleteTagFromINI(const FString& TagToDelete) override
+	virtual bool DeleteTagFromINI(TSharedPtr<FGameplayTagNode> TagNodeToDelete) override
 	{
-		FName TagName = FName(*TagToDelete);
+		FName TagName = TagNodeToDelete->GetCompleteTagName();
 
 		UGameplayTagsManager& Manager = UGameplayTagsManager::Get();
 		UGameplayTagsSettings* Settings = GetMutableDefault<UGameplayTagsSettings>();
@@ -481,31 +481,37 @@ public:
 		bool bTagIsRestricted;
 		bool bTagAllowsNonRestrictedChildren;
 
-		if (DeleteTagRedirector(TagToDelete))
+		if (DeleteTagRedirector(TagName))
 		{
 			return true;
 		}
 		
 		if (!Manager.GetTagEditorData(TagName, Comment, TagSourceName, bTagIsExplicit, bTagIsRestricted, bTagAllowsNonRestrictedChildren))
 		{
-			ShowNotification(FText::Format(LOCTEXT("RemoveTagFailureNoTag", "Cannot delete tag {0}, does not exist!"), FText::FromString(TagToDelete)), 10.0f);
-
+			ShowNotification(FText::Format(LOCTEXT("RemoveTagFailureNoTag", "Cannot delete tag {0}, does not exist!"), FText::FromName(TagName)), 10.0f);
 			return false;
 		}
 
+		ensure(bTagIsRestricted == TagNodeToDelete->IsRestrictedGameplayTag());
+
 		const FGameplayTagSource* TagSource = Manager.FindTagSource(TagSourceName);
 
 		// Check if the tag is implicitly defined
 		if (!bTagIsExplicit || !TagSource)
 		{
-			ShowNotification(FText::Format(LOCTEXT("RemoveTagFailureNoSource", "Cannot delete tag {0} as it is implicit, remove children manually"), FText::FromString(TagToDelete)), 10.0f);
-
+			ShowNotification(FText::Format(LOCTEXT("RemoveTagFailureNoSource", "Cannot delete tag {0} as it is implicit, remove children manually"), FText::FromName(TagName)), 10.0f);
 			return false;
 		}
-		else if (!TagSource->SourceTagList)
+		
+		if (bTagIsRestricted && !TagSource->SourceRestrictedTagList)
 		{
-			ShowNotification(FText::Format(LOCTEXT("RemoveTagFailureBadSource", "Cannot delete tag {0} from source {1}, remove manually"), FText::FromString(TagToDelete), FText::FromName(TagSourceName)), 10.0f);
+			ShowNotification(FText::Format(LOCTEXT("RemoveTagFailureBadSource", "Cannot delete tag {0} from source {1}, remove manually"), FText::FromName(TagName), FText::FromName(TagSourceName)), 10.0f);
+			return false;
+		}
 
+		if (!bTagIsRestricted && !TagSource->SourceTagList)
+		{
+			ShowNotification(FText::Format(LOCTEXT("RemoveTagFailureBadSource", "Cannot delete tag {0} from source {1}, remove manually"), FText::FromName(TagName), FText::FromName(TagSourceName)), 10.0f);
 			return false;
 		}
 
@@ -517,7 +523,7 @@ public:
 		TagsThatWillBeDeleted.Add(TagName);
 
 		FGameplayTag ParentTag = ActualTag.RequestDirectParent();
-		while (ParentTag.IsValid())
+		while (ParentTag.IsValid() && !Manager.FindTagNode(ParentTag)->IsExplicitTag())
 		{
 			// See if there are more children than the one we are about to delete
 			FGameplayTagContainer ParentChildTags = Manager.RequestGameplayTagChildrenInDictionary(ParentTag);
@@ -553,27 +559,45 @@ public:
 		}
 
 		// Passed, delete and save
-		UGameplayTagsList* TagList = TagSource->SourceTagList;
+		const FString& ConfigFileName = bTagIsRestricted ? TagSource->SourceRestrictedTagList->ConfigFileName : TagSource->SourceTagList->ConfigFileName;
+		int32 TagListSize = bTagIsRestricted ? TagSource->SourceRestrictedTagList->RestrictedGameplayTagList.Num() : TagSource->SourceTagList->GameplayTagList.Num();
 
-		for (int32 i = 0; i < TagList->GameplayTagList.Num(); i++)
+		for (int32 i = 0; i < TagListSize; i++)
 		{
-			if (TagList->GameplayTagList[i].Tag == TagName)
+			bool bRemoved = false;
+			if (bTagIsRestricted)
 			{
-				TagList->GameplayTagList.RemoveAt(i);
+				if (TagSource->SourceRestrictedTagList->RestrictedGameplayTagList[i].Tag == TagName)
+				{
+					TagSource->SourceRestrictedTagList->RestrictedGameplayTagList.RemoveAt(i);
+					TagSource->SourceRestrictedTagList->UpdateDefaultConfigFile(ConfigFileName);
+					bRemoved = true;
+				}
+			}
+			else
+			{
+				if (TagSource->SourceTagList->GameplayTagList[i].Tag == TagName)
+				{
+					TagSource->SourceTagList->GameplayTagList.RemoveAt(i);
+					TagSource->SourceTagList->UpdateDefaultConfigFile(ConfigFileName);
+					bRemoved = true;
+				}
+			}
 
-				TagList->UpdateDefaultConfigFile(TagList->ConfigFileName);
-				GameplayTagsUpdateSourceControl(TagList->ConfigFileName);
-				GConfig->LoadFile(TagList->ConfigFileName);
+			if (bRemoved)
+			{
+				GameplayTagsUpdateSourceControl(ConfigFileName);
+				GConfig->LoadFile(ConfigFileName);
 
 				// See if we still live due to child tags
 
 				if (ChildTags.Num() > 0)
 				{
-					ShowNotification(FText::Format(LOCTEXT("RemoveTagChildrenExist", "Deleted explicit tag {0}, still exists implicitly due to children"), FText::FromString(TagToDelete)), 5.0f);
+					ShowNotification(FText::Format(LOCTEXT("RemoveTagChildrenExist", "Deleted explicit tag {0}, still exists implicitly due to children"), FText::FromName(TagName)), 5.0f);
 				}
 				else
 				{
-					ShowNotification(FText::Format(LOCTEXT("RemoveTag", "Deleted tag {0}"), FText::FromString(TagToDelete)), 5.0f);
+					ShowNotification(FText::Format(LOCTEXT("RemoveTag", "Deleted tag {0}"), FText::FromName(TagName)), 5.0f);
 				}
 
 				// This invalidates all local variables, need to return right away
@@ -583,11 +607,95 @@ public:
 			}
 		}
 
-		ShowNotification(FText::Format(LOCTEXT("RemoveTagFailureNoTag", "Cannot delete tag {0}, does not exist!"), FText::FromString(TagToDelete)), 10.0f);
+		ShowNotification(FText::Format(LOCTEXT("RemoveTagFailureNoTag", "Cannot delete tag {0}, does not exist!"), FText::FromName(TagName)), 10.0f);
 		
 		return false;
 	}
 
+	virtual bool UpdateTagInINI(const FString& TagToUpdate, const FString& Comment, bool bIsRestrictedTag, bool bAllowNonRestrictedChildren) override
+	{
+		FName TagName = FName(*TagToUpdate);
+
+		UGameplayTagsManager& Manager = UGameplayTagsManager::Get();
+		UGameplayTagsSettings* Settings = GetMutableDefault<UGameplayTagsSettings>();
+
+		FString OldComment;
+		FName TagSourceName;
+		bool bTagIsExplicit;
+		bool bTagWasRestricted;
+		bool bTagDidAllowNonRestrictedChildren;
+
+		bool bSuccess = false;
+
+		if (Manager.GetTagEditorData(TagName, OldComment, TagSourceName, bTagIsExplicit, bTagWasRestricted, bTagDidAllowNonRestrictedChildren))
+		{
+			if (const FGameplayTagSource* TagSource = Manager.FindTagSource(TagSourceName))
+			{
+				// if we're disallowing non-restricted children make sure we don't already have some
+				if (bTagDidAllowNonRestrictedChildren && !bAllowNonRestrictedChildren)
+				{
+					FGameplayTag ActualTag = Manager.RequestGameplayTag(TagName);
+					FGameplayTagContainer ChildTags = Manager.RequestGameplayTagDirectDescendantsInDictionary(ActualTag, EGameplayTagSelectionType::RestrictedOnly);
+					if (!ChildTags.IsEmpty())
+					{
+						ShowNotification(LOCTEXT("ToggleAllowNonRestrictedChildrenFailure", "Cannot prevent non-restricted children since some already exist! Delete them first."), 10.0f);
+						return false;
+					}
+				}
+
+				UObject* TagListObj = nullptr;
+				FString ConfigFileName;
+
+				if (bIsRestrictedTag && TagSource->SourceRestrictedTagList)
+				{
+					URestrictedGameplayTagsList* RestrictedTagList = TagSource->SourceRestrictedTagList;
+					TagListObj = RestrictedTagList;
+					ConfigFileName = RestrictedTagList->ConfigFileName;
+
+					for (int32 i = 0; i < RestrictedTagList->RestrictedGameplayTagList.Num(); i++)
+					{
+						if (RestrictedTagList->RestrictedGameplayTagList[i].Tag == TagName)
+						{
+							RestrictedTagList->RestrictedGameplayTagList[i].bAllowNonRestrictedChildren = bAllowNonRestrictedChildren;
+							bSuccess = true;
+							break;
+						}
+					}
+				}
+
+				// nothing to do here for non-restricted tags but this is the code that would handle it if we expand this function
+				// 				else if (TagSource->SourceTagList)
+				// 				{
+				// 					UGameplayTagsList* TagList = TagSource->SourceTagList;
+				//					TagListObj = TagList;
+				// 
+				// 					for (int32 i = 0; i < TagList->GameplayTagList.Num(); i++)
+				// 					{
+				// 						if (TagList->GameplayTagList[i].Tag == TagName)
+				// 						{
+				// 
+				// 							bSuccess = true;
+				// 							break;
+				// 						}
+				// 					}
+				// 				}
+
+				if (bSuccess)
+				{
+					// Check source control before and after writing, to make sure it gets created or checked out
+					GameplayTagsUpdateSourceControl(ConfigFileName);
+					TagListObj->UpdateDefaultConfigFile(ConfigFileName);
+					GameplayTagsUpdateSourceControl(ConfigFileName);
+
+					GConfig->LoadFile(ConfigFileName);
+				}
+
+			}
+		}
+
+		return bSuccess;
+	}
+
 	virtual bool RenameTagInINI(const FString& TagToRename, const FString& TagToRenameTo) override
 	{
 		FName OldTagName = FName(*TagToRename);
@@ -603,8 +711,8 @@ public:
 		bool bTagAllowsNonRestrictedChildren;
 
 		// Delete existing redirector
-		DeleteTagRedirector(TagToRenameTo);
-		DeleteTagRedirector(TagToRename);
+		DeleteTagRedirector(NewTagName);
+		DeleteTagRedirector(OldTagName);
 
 		if (Manager.GetTagEditorData(OldTagName, OldComment, OldTagSourceName, bTagIsExplicit, bTagIsRestricted, bTagAllowsNonRestrictedChildren))
 		{
diff --git a/Engine/Plugins/Editor/GameplayTagsEditor/Source/GameplayTagsEditor/Private/SGameplayTagWidget.cpp b/Engine/Plugins/Editor/GameplayTagsEditor/Source/GameplayTagsEditor/Private/SGameplayTagWidget.cpp
--- a/Engine/Plugins/Editor/GameplayTagsEditor/Source/GameplayTagsEditor/Private/SGameplayTagWidget.cpp
+++ b/Engine/Plugins/Editor/GameplayTagsEditor/Source/GameplayTagsEditor/Private/SGameplayTagWidget.cpp
@@ -712,13 +712,18 @@ ECheckBoxState SGameplayTagWidget::IsTagChecked(TSharedPtr<FGameplayTagNode> Nod
 
 void SGameplayTagWidget::OnAllowChildrenTagCheckStatusChanged(ECheckBoxState NewCheckState, TSharedPtr<FGameplayTagNode> NodeChanged)
 {
-	if (NewCheckState == ECheckBoxState::Checked)
-	{
-		NodeChanged->bAllowNonRestrictedChildren = true;
-	}
-	else if (NewCheckState == ECheckBoxState::Unchecked)
+	IGameplayTagsEditorModule& TagsEditor = IGameplayTagsEditorModule::Get();
+
+	if (TagsEditor.UpdateTagInINI(NodeChanged->GetCompleteTagString(), NodeChanged->DevComment, NodeChanged->bIsRestrictedTag, NewCheckState == ECheckBoxState::Checked))
 	{
-		NodeChanged->bAllowNonRestrictedChildren = false;
+		if (NewCheckState == ECheckBoxState::Checked)
+		{
+			NodeChanged->bAllowNonRestrictedChildren = true;
+		}
+		else if (NewCheckState == ECheckBoxState::Unchecked)
+		{
+			NodeChanged->bAllowNonRestrictedChildren = false;
+		}
 	}
 }
 
@@ -898,7 +903,7 @@ void SGameplayTagWidget::OnDeleteTag(TSharedPtr<FGameplayTagNode> InTagNode)
 	{
 		IGameplayTagsEditorModule& TagsEditor = IGameplayTagsEditorModule::Get();
 
-		const bool bDeleted = TagsEditor.DeleteTagFromINI(InTagNode->GetCompleteTagString());
+		const bool bDeleted = TagsEditor.DeleteTagFromINI(InTagNode);
 
 		if (bDeleted)
 		{
diff --git a/Engine/Plugins/Editor/GameplayTagsEditor/Source/GameplayTagsEditor/Public/GameplayTagsEditorModule.h b/Engine/Plugins/Editor/GameplayTagsEditor/Source/GameplayTagsEditor/Public/GameplayTagsEditorModule.h
--- a/Engine/Plugins/Editor/GameplayTagsEditor/Source/GameplayTagsEditor/Public/GameplayTagsEditorModule.h
+++ b/Engine/Plugins/Editor/GameplayTagsEditor/Source/GameplayTagsEditor/Public/GameplayTagsEditorModule.h
@@ -41,11 +41,14 @@ public:
 	GAMEPLAYTAGSEDITOR_API virtual bool AddNewGameplayTagToINI(const FString& NewTag, const FString& Comment = TEXT(""), FName TagSourceName = NAME_None, bool bIsRestrictedTag = false, bool bAllowNonRestrictedChildren = true) = 0;
 
 	/** Tries to delete a tag from the library. This will pop up special UI or error messages as needed. It will also delete redirectors if that is specified. */
-	GAMEPLAYTAGSEDITOR_API virtual bool DeleteTagFromINI(const FString& TagToDelete) = 0;
+	GAMEPLAYTAGSEDITOR_API virtual bool DeleteTagFromINI(TSharedPtr<struct FGameplayTagNode> TagNodeToDelete) = 0;
 
 	/** Tries to rename a tag, leaving a rediretor in the ini, and adding the new tag if it does not exist yet */
 	GAMEPLAYTAGSEDITOR_API virtual bool RenameTagInINI(const FString& TagToRename, const FString& TagToRenameTo) = 0;
 
+	/** Updates info about a tag */
+	GAMEPLAYTAGSEDITOR_API virtual bool UpdateTagInINI(const FString& TagToUpdate, const FString& Comment, bool bIsRestrictedTag, bool bAllowNonRestrictedChildren) = 0;
+
 	/** Adds a transient gameplay tag (only valid for the current editor session) */
 	GAMEPLAYTAGSEDITOR_API virtual bool AddTransientEditorGameplayTag(const FString& NewTransientTag) = 0;
 };
diff --git a/Engine/Plugins/Online/OnlineFramework/Source/Party/Public/PartyGameState.h b/Engine/Plugins/Online/OnlineFramework/Source/Party/Public/PartyGameState.h
--- a/Engine/Plugins/Online/OnlineFramework/Source/Party/Public/PartyGameState.h
+++ b/Engine/Plugins/Online/OnlineFramework/Source/Party/Public/PartyGameState.h
@@ -97,7 +97,7 @@ enum class ECrossplayPreference: uint8
 /**
  * Current state of the party
  */
-USTRUCT()
+USTRUCT(BlueprintType)
 struct FPartyState
 {
 	GENERATED_USTRUCT_BODY();
diff --git a/Engine/Plugins/Runtime/GameplayAbilities/Source/GameplayAbilities/Private/AbilitySystemComponent.cpp b/Engine/Plugins/Runtime/GameplayAbilities/Source/GameplayAbilities/Private/AbilitySystemComponent.cpp
--- a/Engine/Plugins/Runtime/GameplayAbilities/Source/GameplayAbilities/Private/AbilitySystemComponent.cpp
+++ b/Engine/Plugins/Runtime/GameplayAbilities/Source/GameplayAbilities/Private/AbilitySystemComponent.cpp
@@ -1084,16 +1084,12 @@ void UAbilitySystemComponent::InvokeGameplayCueEvent(const FGameplayTag Gameplay
 
 void UAbilitySystemComponent::ExecuteGameplayCue(const FGameplayTag GameplayCueTag, FGameplayEffectContextHandle EffectContext)
 {
-	FScopedPredictionWindow ScopedPredictionWindow(this, true);
-
 	// Send to the wrapper on the cue manager
 	UAbilitySystemGlobals::Get().GetGameplayCueManager()->InvokeGameplayCueExecuted(this, GameplayCueTag, ScopedPredictionKey, EffectContext);
 }
 
 void UAbilitySystemComponent::ExecuteGameplayCue(const FGameplayTag GameplayCueTag, const FGameplayCueParameters& GameplayCueParameters)
 {
-	FScopedPredictionWindow ScopedPredictionWindow(this, true);
-
 	// Send to the wrapper on the cue manager
 	UAbilitySystemGlobals::Get().GetGameplayCueManager()->InvokeGameplayCueExecuted_WithParams(this, GameplayCueTag, ScopedPredictionKey, GameplayCueParameters);
 }
diff --git a/Engine/Plugins/Runtime/GameplayAbilities/Source/GameplayAbilities/Public/Abilities/GameplayAbility.h b/Engine/Plugins/Runtime/GameplayAbilities/Source/GameplayAbilities/Public/Abilities/GameplayAbility.h
--- a/Engine/Plugins/Runtime/GameplayAbilities/Source/GameplayAbilities/Public/Abilities/GameplayAbility.h
+++ b/Engine/Plugins/Runtime/GameplayAbilities/Source/GameplayAbilities/Public/Abilities/GameplayAbility.h
@@ -500,7 +500,7 @@ protected:
 	bool bHasBlueprintShouldAbilityRespondToEvent;
 
 	/** Sends a gameplay event, also creates a prediction window */
-	UFUNCTION(BlueprintCallable, Category = "Abilities")
+	UFUNCTION(BlueprintCallable, Category = Ability)
 	virtual void SendGameplayEvent(FGameplayTag EventTag, FGameplayEventData Payload);
 
 	// --------------------------------------
diff --git a/Engine/Source/Developer/BlueprintCompilerCppBackend/Private/BlueprintCompilerCppBackend.cpp b/Engine/Source/Developer/BlueprintCompilerCppBackend/Private/BlueprintCompilerCppBackend.cpp
--- a/Engine/Source/Developer/BlueprintCompilerCppBackend/Private/BlueprintCompilerCppBackend.cpp
+++ b/Engine/Source/Developer/BlueprintCompilerCppBackend/Private/BlueprintCompilerCppBackend.cpp
@@ -156,7 +156,13 @@ struct FSetterExpressionBuilder
 void FBlueprintCompilerCppBackend::EmitAssignmentStatment(FEmitterLocalContext& EmitterContext, FKismetFunctionContext& FunctionContext, FBlueprintCompiledStatement& Statement)
 {
 	check(Statement.LHS && Statement.RHS[0]);
-	const FString SourceExpression = TermToText(EmitterContext, Statement.RHS[0], ENativizedTermUsage::Getter);
+
+	FTermToTextParams SrcTermParams;
+	SrcTermParams.Term = Statement.RHS[0];
+	SrcTermParams.TermUsage = ENativizedTermUsage::Getter;
+	SrcTermParams.CoerceProperty = Statement.LHS->AssociatedVarProperty;
+
+	const FString SourceExpression = TermToText(EmitterContext, SrcTermParams);
 	FSetterExpressionBuilder SetterExpression(*this, EmitterContext, Statement.LHS);
 	FSafeContextScopedEmmitter SafeContextScope(EmitterContext, Statement.LHS->Context, *this);
 
@@ -908,20 +914,29 @@ FString FBlueprintCompilerCppBackend::EmitArrayGetByRef(FEmitterLocalContext& Em
 	return Result;
 }
 
-FString FBlueprintCompilerCppBackend::TermToText(FEmitterLocalContext& EmitterContext, const FBPTerminal* Term, const ENativizedTermUsage TermUsage, const bool bUseSafeContext, FString* EndCustomSetExpression)
+FString FBlueprintCompilerCppBackend::TermToText(FEmitterLocalContext& EmitterContext, const FTermToTextParams& Params)
 {
-	ensure((TermUsage != ENativizedTermUsage::Setter) || !bUseSafeContext);
-	ensure((TermUsage == ENativizedTermUsage::Setter) == (EndCustomSetExpression != nullptr));
-	if (EndCustomSetExpression)
+	const FBPTerminal* Term = Params.Term;
+	
+	ensure((Params.TermUsage != ENativizedTermUsage::Setter) || !Params.bUseSafeContext);
+	ensure((Params.TermUsage == ENativizedTermUsage::Setter) == (Params.EndCustomSetExpression != nullptr));
+	if (Params.EndCustomSetExpression)
 	{
-		EndCustomSetExpression->Reset();
+		Params.EndCustomSetExpression->Reset();
 	}
 
-	const bool bGetter = (TermUsage == ENativizedTermUsage::Getter);
+	const bool bGetter = (Params.TermUsage == ENativizedTermUsage::Getter);
 	const FString PSC_Self(TEXT("self"));
 	if (Term->bIsLiteral)
 	{
-		return FEmitHelper::LiteralTerm(EmitterContext, Term->Type, Term->Name, Term->ObjectLiteral, &Term->TextLiteral);
+		FEmitHelper::FLiteralTermParams LiteralTermParams;
+		LiteralTermParams.Type = Term->Type;
+		LiteralTermParams.CustomValue = Term->Name;
+		LiteralTermParams.LiteralText = Term->TextLiteral;
+		LiteralTermParams.LiteralObject = Term->ObjectLiteral;
+		LiteralTermParams.CoerceProperty = Params.CoerceProperty;
+
+		return FEmitHelper::LiteralTerm(EmitterContext, LiteralTermParams);
 	}
 	else if (Term->InlineGeneratedParameter)
 	{
@@ -953,11 +968,9 @@ FString FBlueprintCompilerCppBackend::TermToText(FEmitterLocalContext& EmitterCo
 				? EmitterContext.ExportCppDeclaration(Term->AssociatedVarProperty, EExportedDeclaration::Local, PropertyExportFlags, FEmitterLocalContext::EPropertyNameInDeclaration::Skip)
 				: FEmitHelper::PinTypeToNativeType(Term->Type);
 
-			const FString DefaultValueConstructor = (!Term->Type.IsContainer())
-				? FEmitHelper::LiteralTerm(EmitterContext, Term->Type, FString(), nullptr, &FText::GetEmpty())
-				: FString::Printf(TEXT("%s{}"), *CppType);
-
-			EmitterContext.AddLine(*FString::Printf(TEXT("%s %s = %s;"), *CppType, *DefaultValueVariable, *DefaultValueConstructor));
+			FEmitHelper::FLiteralTermParams DefaultValueTermParams;
+			DefaultValueTermParams.Type = Term->Type;
+			EmitterContext.AddLine(*FString::Printf(TEXT("%s %s = %s;"), *CppType, *DefaultValueVariable, *FEmitHelper::LiteralTerm(EmitterContext, DefaultValueTermParams)));
 
 			return DefaultValueVariable;
 		};
@@ -1005,7 +1018,7 @@ FString FBlueprintCompilerCppBackend::TermToText(FEmitterLocalContext& EmitterCo
 			bIsAccessible &= !Term->AssociatedVarProperty->HasAnyPropertyFlags(CPF_NativeAccessSpecifierPrivate | CPF_NativeAccessSpecifierProtected);
 			if (!bIsAccessible)
 			{
-				ResultPath = FEmitHelper::AccessInaccessibleProperty(EmitterContext, Term->AssociatedVarProperty, FString(), ContextStr, TEXT("&"), 0, TermUsage, EndCustomSetExpression);
+				ResultPath = FEmitHelper::AccessInaccessibleProperty(EmitterContext, Term->AssociatedVarProperty, FString(), ContextStr, TEXT("&"), 0, Params.TermUsage, Params.EndCustomSetExpression);
 			}
 			else
 			{
@@ -1044,7 +1057,7 @@ FString FBlueprintCompilerCppBackend::TermToText(FEmitterLocalContext& EmitterCo
 					ensure(ContextStr.IsEmpty());
 					ContextStr = TEXT("this");
 				}
-				ResultPath = FEmitHelper::AccessInaccessibleProperty(EmitterContext, Term->AssociatedVarProperty, FString(), ContextStr, FString(), 0, TermUsage, EndCustomSetExpression);
+				ResultPath = FEmitHelper::AccessInaccessibleProperty(EmitterContext, Term->AssociatedVarProperty, FString(), ContextStr, FString(), 0, Params.TermUsage, Params.EndCustomSetExpression);
 			}
 			else
 			{
@@ -1092,7 +1105,7 @@ FString FBlueprintCompilerCppBackend::TermToText(FEmitterLocalContext& EmitterCo
 			ResultPath = FString::Printf(TEXT("const_cast<%s>(%s)"), *CppType, *ResultPath);
 		}
 
-		const FString Conditions = bUseSafeContext ? FSafeContextScopedEmmitter::ValidationChain(EmitterContext, Term->Context, *this) : FString();
+		const FString Conditions = Params.bUseSafeContext ? FSafeContextScopedEmmitter::ValidationChain(EmitterContext, Term->Context, *this) : FString();
 		if (!Conditions.IsEmpty())
 		{
 			const FString DefaultValueVariable = GenerateDefaultLocalVariable(Term);
@@ -1121,23 +1134,25 @@ FString FBlueprintCompilerCppBackend::LatentFunctionInfoTermToText(FEmitterLocal
 
 	check(!FixupTermName.IsEmpty());
 
-	FString StructValues = Term->Name;
+	FEmitHelper::FLiteralTermParams LiteralTermParams;
+	LiteralTermParams.Type = Term->Type;
+	LiteralTermParams.CustomValue = Term->Name;
 
 	// Index 0 is always the ubergraph
 	const int32 TargetStateIndex = StateMapPerFunction[0].StatementToStateIndex(TargetLabel);
-	const int32 LinkageTermStartIdx = StructValues.Find(FixupTermName);
+	const int32 LinkageTermStartIdx = LiteralTermParams.CustomValue.Find(FixupTermName);
 	check(LinkageTermStartIdx != INDEX_NONE);
-	StructValues = StructValues.Replace(TEXT("-1"), *FString::FromInt(TargetStateIndex));
+	LiteralTermParams.CustomValue = LiteralTermParams.CustomValue.Replace(TEXT("-1"), *FString::FromInt(TargetStateIndex));
 
 	int32* ExecutionGroupPtr = UberGraphStatementToExecutionGroup.Find(TargetLabel);
 	if (ExecutionGroupPtr && UberGraphContext)
 	{
 		const FString OldExecutionFunctionName = UEdGraphSchema_K2::FN_ExecuteUbergraphBase.ToString() + TEXT("_") + UberGraphContext->Blueprint->GetName();
 		const FString NewExecutionFunctionName = OldExecutionFunctionName + FString::Printf(TEXT("_%d"), *ExecutionGroupPtr);
-		StructValues = StructValues.Replace(*OldExecutionFunctionName, *NewExecutionFunctionName);
+		LiteralTermParams.CustomValue = LiteralTermParams.CustomValue.Replace(*OldExecutionFunctionName, *NewExecutionFunctionName);
 	}
 
-	return FEmitHelper::LiteralTerm(EmitterContext, Term->Type, StructValues, nullptr);
+	return FEmitHelper::LiteralTerm(EmitterContext, LiteralTermParams);
 }
 
 bool FBlueprintCompilerCppBackend::InnerFunctionImplementation(FKismetFunctionContext& FunctionContext, FEmitterLocalContext& EmitterContext, int32 ExecutionGroup)
diff --git a/Engine/Source/Developer/BlueprintCompilerCppBackend/Private/BlueprintCompilerCppBackend.h b/Engine/Source/Developer/BlueprintCompilerCppBackend/Private/BlueprintCompilerCppBackend.h
--- a/Engine/Source/Developer/BlueprintCompilerCppBackend/Private/BlueprintCompilerCppBackend.h
+++ b/Engine/Source/Developer/BlueprintCompilerCppBackend/Private/BlueprintCompilerCppBackend.h
@@ -68,5 +68,32 @@ protected:
 	FString EmitArrayGetByRef(FEmitterLocalContext& EmitterContext, FBlueprintCompiledStatement& Statement);
 
 public:
-	FString TermToText(FEmitterLocalContext& EmitterContext, const FBPTerminal* Term, const ENativizedTermUsage TermUsage, const bool bUseSafeContext = true, FString* EndCustomSetExpression = nullptr);
+	struct FTermToTextParams
+	{
+		const FBPTerminal* Term;
+		ENativizedTermUsage TermUsage;
+		bool bUseSafeContext;
+		FString* EndCustomSetExpression;
+		const UProperty* CoerceProperty;
+
+		FTermToTextParams()
+			:Term(nullptr)
+			,TermUsage(ENativizedTermUsage::UnspecifiedOrReference)
+			,bUseSafeContext(true)
+			,EndCustomSetExpression(nullptr)
+			,CoerceProperty(nullptr)
+		{}
+	};
+
+	FString TermToText(FEmitterLocalContext& EmitterContext, const FTermToTextParams& Params);
+	FString TermToText(FEmitterLocalContext& EmitterContext, const FBPTerminal* Term, const ENativizedTermUsage TermUsage, const bool bUseSafeContext = true, FString* EndCustomSetExpression = nullptr)
+	{
+		FTermToTextParams Params;
+		Params.Term = Term;
+		Params.TermUsage = TermUsage;
+		Params.bUseSafeContext = bUseSafeContext;
+		Params.EndCustomSetExpression = EndCustomSetExpression;
+
+		return TermToText(EmitterContext, Params);
+	}
 };
diff --git a/Engine/Source/Developer/BlueprintCompilerCppBackend/Private/BlueprintCompilerCppBackendUtils.cpp b/Engine/Source/Developer/BlueprintCompilerCppBackend/Private/BlueprintCompilerCppBackendUtils.cpp
--- a/Engine/Source/Developer/BlueprintCompilerCppBackend/Private/BlueprintCompilerCppBackendUtils.cpp
+++ b/Engine/Source/Developer/BlueprintCompilerCppBackend/Private/BlueprintCompilerCppBackendUtils.cpp
@@ -1021,22 +1021,193 @@ FString FEmitHelper::FloatToString(float Value)
 	return FString::Printf(TEXT("%f"), Value);
 }
 
-FString FEmitHelper::LiteralTerm(FEmitterLocalContext& EmitterContext, const FEdGraphPinType& Type, const FString& CustomValue, UObject* LiteralObject, const FText* OptionalTextLiteral)
+FString FEmitHelper::LiteralTerm(FEmitterLocalContext& EmitterContext, const FLiteralTermParams& Params)
 {
 	auto Schema = GetDefault<UEdGraphSchema_K2>();
 
-	if (UEdGraphSchema_K2::PC_String == Type.PinCategory)
+	class FImportTextErrorContext : public FStringOutputDevice
+	{
+	public:
+		int32 NumErrors;
+
+		FImportTextErrorContext() : FStringOutputDevice(), NumErrors(0) {}
+
+		virtual void Serialize(const TCHAR* V, ELogVerbosity::Type Verbosity, const class FName& Category) override
+		{
+			if (ELogVerbosity::Type::Error == Verbosity)
+			{
+				NumErrors++;
+			}
+			FStringOutputDevice::Serialize(V, Verbosity, Category);
+		}
+	};
+
+	const FEdGraphPinType& Type = Params.Type;
+	const FString& CustomValue = Params.CustomValue;
+
+	if (Type.IsContainer())
+	{
+		FString ContainerInitializerList;
+		FImportTextErrorContext ImportError;
+
+		if (const UArrayProperty* ArrayProperty = Cast<UArrayProperty>(Params.CoerceProperty))
+		{
+			FScriptArray ScriptArray;
+			if (!ArrayProperty->ImportText(*CustomValue, &ScriptArray, PPF_None, nullptr, &ImportError))
+			{
+				UE_LOG(LogK2Compiler, Error, TEXT("FEmitHelper::LiteralTerm cannot parse array value \"%s\" error: %s class: %s"), *CustomValue, *ImportError, *GetPathNameSafe(EmitterContext.GetCurrentlyGeneratedClass()));
+			}
+
+			const int32 NumElements = ScriptArray.Num();
+			FScriptArrayHelper ScriptArrayHelper(ArrayProperty, &ScriptArray);
+
+			FLiteralTermParams InnerTermParams;
+			Schema->ConvertPropertyToPinType(ArrayProperty->Inner, InnerTermParams.Type);
+
+			const UTextProperty* InnerTextProperty = Cast<UTextProperty>(ArrayProperty->Inner);
+			const UObjectPropertyBase* InnerObjectProperty = InnerTextProperty ? nullptr : Cast<UObjectPropertyBase>(ArrayProperty->Inner);
+
+			for (int32 ElementIdx = 0; ElementIdx < NumElements; ++ElementIdx)
+			{
+				uint8* ValuePtr = ScriptArrayHelper.GetRawPtr(ElementIdx);
+				if (ArrayProperty->Inner->ExportText_Direct(InnerTermParams.CustomValue, ValuePtr, ValuePtr, nullptr, PPF_None))
+				{
+					if (InnerTextProperty)
+					{
+						InnerTermParams.LiteralText = InnerTextProperty->GetPropertyValue(ValuePtr);
+						InnerTermParams.CustomValue = InnerTermParams.LiteralText.ToString();
+					}
+					else if (InnerObjectProperty)
+					{
+						InnerTermParams.LiteralObject = InnerObjectProperty->GetObjectPropertyValue(ValuePtr);
+					}
+
+					ContainerInitializerList += LiteralTerm(EmitterContext, InnerTermParams);
+					if (ElementIdx < NumElements - 1)
+					{
+						ContainerInitializerList += TEXT(", ");
+					}
+				}
+			}
+		}
+		else if (const USetProperty* SetProperty = Cast<USetProperty>(Params.CoerceProperty))
+		{
+			FScriptSet ScriptSet;
+			if (!CustomValue.IsEmpty())	// unlike UArrayProperty, USetProperty::ImportText() doesn't allow empty values to pass, so we check for that here.
+			{
+				if (!SetProperty->ImportText(*CustomValue, &ScriptSet, PPF_None, nullptr, &ImportError))
+				{
+					UE_LOG(LogK2Compiler, Error, TEXT("FEmitHelper::LiteralTerm cannot parse set value \"%s\" error: %s class: %s"), *CustomValue, *ImportError, *GetPathNameSafe(EmitterContext.GetCurrentlyGeneratedClass()));
+				}
+			}
+
+			const int32 NumElements = ScriptSet.Num();
+			FScriptSetHelper ScriptSetHelper(SetProperty, &ScriptSet);
+
+			FLiteralTermParams ElementTermParams;
+			Schema->ConvertPropertyToPinType(SetProperty->ElementProp, ElementTermParams.Type);
+
+			const UTextProperty* ElementTextProperty = Cast<UTextProperty>(SetProperty->ElementProp);
+			const UObjectPropertyBase* ElementObjectProperty = ElementTextProperty ? nullptr : Cast<UObjectPropertyBase>(SetProperty->ElementProp);
+
+			for (int32 ElementIdx = 0; ElementIdx < NumElements; ++ElementIdx)
+			{
+				uint8* ValuePtr = ScriptSetHelper.GetElementPtr(ElementIdx);
+				if (SetProperty->ElementProp->ExportText_Direct(ElementTermParams.CustomValue, ValuePtr, ValuePtr, nullptr, PPF_None))
+				{
+					if (ElementTextProperty)
+					{
+						ElementTermParams.LiteralText = ElementTextProperty->GetPropertyValue(ValuePtr);
+						ElementTermParams.CustomValue = ElementTermParams.LiteralText.ToString();
+					}
+					else if (ElementObjectProperty)
+					{
+						ElementTermParams.LiteralObject = ElementObjectProperty->GetObjectPropertyValue(ValuePtr);
+					}
+
+					ContainerInitializerList += LiteralTerm(EmitterContext, ElementTermParams);
+					if (ElementIdx < NumElements - 1)
+					{
+						ContainerInitializerList += TEXT(", ");
+					}
+				}
+			}
+		}
+		else if (const UMapProperty* MapProperty = Cast<UMapProperty>(Params.CoerceProperty))
+		{
+			FScriptSet ScriptMap;
+			if (!CustomValue.IsEmpty())	// unlike UArrayProperty, UMapProperty::ImportText() doesn't allow empty values to pass, so we check for that here.
+			{
+				if (!MapProperty->ImportText(*CustomValue, &ScriptMap, PPF_None, nullptr, &ImportError))
+				{
+					UE_LOG(LogK2Compiler, Error, TEXT("FEmitHelper::LiteralTerm cannot parse map value \"%s\" error: %s class: %s"), *CustomValue, *ImportError, *GetPathNameSafe(EmitterContext.GetCurrentlyGeneratedClass()));
+				}
+			}
+
+			const int32 NumElements = ScriptMap.Num();
+			FScriptMapHelper ScriptMapHelper(MapProperty, &ScriptMap);
+
+			FLiteralTermParams KeyTermParams, ValueTermParams;
+			Schema->ConvertPropertyToPinType(MapProperty->KeyProp, KeyTermParams.Type);
+			Schema->ConvertPropertyToPinType(MapProperty->ValueProp, ValueTermParams.Type);
+
+			const UTextProperty* KeyTextProperty = Cast<UTextProperty>(MapProperty->KeyProp);
+			const UTextProperty* ValueTextProperty = Cast<UTextProperty>(MapProperty->ValueProp);
+			const UObjectPropertyBase* KeyObjectProperty = KeyTextProperty ? nullptr : Cast<UObjectPropertyBase>(MapProperty->KeyProp);
+			const UObjectPropertyBase* ValueObjectProperty = ValueTextProperty ? nullptr : Cast<UObjectPropertyBase>(MapProperty->ValueProp);
+
+			for (int32 ElementIdx = 0, SparseIdx = 0; ElementIdx < NumElements; ++SparseIdx)
+			{
+				if (ScriptMap.IsValidIndex(SparseIdx))
+				{
+					uint8* KeyPtr = ScriptMapHelper.GetKeyPtr(SparseIdx);
+					if (MapProperty->KeyProp->ExportText_Direct(KeyTermParams.CustomValue, KeyPtr, KeyPtr, nullptr, PPF_None))
+					{
+						if (KeyTextProperty)
+						{
+							KeyTermParams.LiteralText = KeyTextProperty->GetPropertyValue(KeyPtr);
+							KeyTermParams.CustomValue = KeyTermParams.LiteralText.ToString();
+						}
+						else if (KeyObjectProperty)
+						{
+							KeyTermParams.LiteralObject = KeyObjectProperty->GetObjectPropertyValue(KeyPtr);
+						}
+
+						uint8* ValuePtr = ScriptMapHelper.GetValuePtr(SparseIdx);
+						if (MapProperty->ValueProp->ExportText_Direct(ValueTermParams.CustomValue , ValuePtr, ValuePtr, nullptr, PPF_None))
+						{
+							if (ValueTextProperty)
+							{
+								ValueTermParams.LiteralText = ValueTextProperty->GetPropertyValue(ValuePtr);
+								ValueTermParams.CustomValue = ValueTermParams.LiteralText.ToString();
+							}
+							else if (ValueObjectProperty)
+							{
+								ValueTermParams.LiteralObject = ValueObjectProperty->GetObjectPropertyValue(ValuePtr);
+							}
+
+							ContainerInitializerList += FString::Printf(TEXT("{%s, %s}"), *LiteralTerm(EmitterContext, KeyTermParams), *LiteralTerm(EmitterContext, ValueTermParams));
+							if (ElementIdx < NumElements - 1)
+							{
+								ContainerInitializerList += TEXT(", ");
+							}
+						}
+					}
+
+					++ElementIdx;
+				}
+			}
+		}
+
+		return FString::Printf(TEXT("{%s}"), *ContainerInitializerList);
+	}
+	else if (UEdGraphSchema_K2::PC_String == Type.PinCategory)
 	{
 		return FString::Printf(TEXT("FString(%s)"), *UStrProperty::ExportCppHardcodedText(CustomValue, EmitterContext.DefaultTarget->Indent));
 	}
 	else if (UEdGraphSchema_K2::PC_Text == Type.PinCategory)
 	{
-		ensure(OptionalTextLiteral);
-		if (OptionalTextLiteral)
-		{
-			return UTextProperty::GenerateCppCodeForTextValue(*OptionalTextLiteral, FString());
-		}
-		return FString::Printf(TEXT("FText::FromString(%s)"), *UStrProperty::ExportCppHardcodedText(CustomValue, EmitterContext.DefaultTarget->Indent));
+		return UTextProperty::GenerateCppCodeForTextValue(Params.LiteralText, FString());
 	}
 	else if (UEdGraphSchema_K2::PC_Float == Type.PinCategory)
 	{
@@ -1151,23 +1322,6 @@ FString FEmitHelper::LiteralTerm(FEmitterLocalContext& EmitterContext, const FEd
 				FStructOnScope StructOnScope(StructType);
 				StructType->InitializeDefaultValue(StructOnScope.GetStructMemory()); // after cl#3098294 only delta (of structure data) against the default value is stored in string. So we need to explicitly provide default values before serialization.
 
-				class FImportTextErrorContext : public FStringOutputDevice
-				{
-				public:
-					int32 NumErrors;
-
-					FImportTextErrorContext() : FStringOutputDevice(), NumErrors(0) {}
-
-					virtual void Serialize(const TCHAR* V, ELogVerbosity::Type Verbosity, const class FName& Category) override
-					{
-						if (ELogVerbosity::Type::Error == Verbosity)
-						{
-							NumErrors++;
-						}
-						FStringOutputDevice::Serialize(V, Verbosity, Category);
-					}
-				};
-
 				FImportTextErrorContext ImportError;
 				const TCHAR* EndOfParsedBuff = StructType->ImportText(*CustomValue, StructOnScope.GetStructMemory(), nullptr, PPF_None, &ImportError, TEXT("FEmitHelper::LiteralTerm"));
 				if (!EndOfParsedBuff || ImportError.NumErrors)
@@ -1205,14 +1359,14 @@ FString FEmitHelper::LiteralTerm(FEmitterLocalContext& EmitterContext, const FEd
 	}
 	else if (UEdGraphSchema_K2::PC_Class == Type.PinCategory)
 	{
-		if (auto FoundClass = Cast<const UClass>(LiteralObject))
+		if (auto FoundClass = Cast<const UClass>(Params.LiteralObject))
 		{
-			const FString MappedObject = EmitterContext.FindGloballyMappedObject(LiteralObject, UClass::StaticClass());
+			const FString MappedObject = EmitterContext.FindGloballyMappedObject(Params.LiteralObject, UClass::StaticClass());
 			if (!MappedObject.IsEmpty())
 			{
 				return MappedObject;
 			}
-			return FString::Printf(TEXT("LoadClass<UClass>(nullptr, TEXT(\"%s\"), nullptr, 0, nullptr)"), *(LiteralObject->GetPathName().ReplaceCharWithEscapedChar()));
+			return FString::Printf(TEXT("LoadClass<UClass>(nullptr, TEXT(\"%s\"), nullptr, 0, nullptr)"), *(Params.LiteralObject->GetPathName().ReplaceCharWithEscapedChar()));
 		}
 		return FString(TEXT("((UClass*)nullptr)"));
 	}
@@ -1236,9 +1390,9 @@ FString FEmitHelper::LiteralTerm(FEmitterLocalContext& EmitterContext, const FEd
 	{
 		UClass* FoundClass = Cast<UClass>(Type.PinSubCategoryObject.Get());
 		UClass* ObjectClassToUse = FoundClass ? EmitterContext.GetFirstNativeOrConvertedClass(FoundClass) : UObject::StaticClass();
-		if (LiteralObject)
+		if (Params.LiteralObject)
 		{
-			const FString MappedObject = EmitterContext.FindGloballyMappedObject(LiteralObject, ObjectClassToUse, true);
+			const FString MappedObject = EmitterContext.FindGloballyMappedObject(Params.LiteralObject, ObjectClassToUse, true);
 			if (!MappedObject.IsEmpty())
 			{
 				return MappedObject;
@@ -1249,25 +1403,12 @@ FString FEmitHelper::LiteralTerm(FEmitterLocalContext& EmitterContext, const FEd
 	}
 	else if (UEdGraphSchema_K2::PC_Interface == Type.PinCategory)
 	{
-		if (!LiteralObject && CustomValue.IsEmpty())
+		if (!Params.LiteralObject && CustomValue.IsEmpty())
 		{
 			return FString(TEXT("nullptr"));
 		}
 	}
-	/*
-	else if (CoerceProperty->IsA(UInterfaceProperty::StaticClass()))
-	{
-		if (Term->Type.PinSubCategory == PSC_Self)
-		{
-			return TEXT("this");
-		}
-		else
-		{
-			ensureMsgf(false, TEXT("It is not possible to express this interface property as a literal value!"));
-			return Term->Name;
-		}
-	}
-	*/
+
 	ensureMsgf(false, TEXT("It is not possible to express this type as a literal value!"));
 	return CustomValue;
 }
diff --git a/Engine/Source/Developer/BlueprintCompilerCppBackend/Private/BlueprintCompilerCppBackendUtils.h b/Engine/Source/Developer/BlueprintCompilerCppBackend/Private/BlueprintCompilerCppBackendUtils.h
--- a/Engine/Source/Developer/BlueprintCompilerCppBackend/Private/BlueprintCompilerCppBackendUtils.h
+++ b/Engine/Source/Developer/BlueprintCompilerCppBackend/Private/BlueprintCompilerCppBackendUtils.h
@@ -287,7 +287,22 @@ struct FEmitHelper
 
 	static void EmitLifetimeReplicatedPropsImpl(FEmitterLocalContext& EmitterContext);
 
-	static FString LiteralTerm(FEmitterLocalContext& EmitterContext, const FEdGraphPinType& Type, const FString& CustomValue, UObject* LiteralObject, const FText* OptionalTextLiteral = nullptr);
+	struct FLiteralTermParams
+	{
+		FEdGraphPinType Type;
+		FString CustomValue;
+		FText LiteralText;
+		UObject* LiteralObject;
+		const UProperty* CoerceProperty;
+
+		FLiteralTermParams()
+			:LiteralObject(nullptr)
+			,CoerceProperty(nullptr)
+		{
+		}
+	};
+
+	static FString LiteralTerm(FEmitterLocalContext& EmitterContext, const FLiteralTermParams& Params);
 
 	static FString PinTypeToNativeType(const FEdGraphPinType& InType);
 
diff --git a/Engine/Source/Developer/BlueprintCompilerCppBackend/Private/BlueprintCompilerCppBackendValueHelper.cpp b/Engine/Source/Developer/BlueprintCompilerCppBackend/Private/BlueprintCompilerCppBackendValueHelper.cpp
--- a/Engine/Source/Developer/BlueprintCompilerCppBackend/Private/BlueprintCompilerCppBackendValueHelper.cpp
+++ b/Engine/Source/Developer/BlueprintCompilerCppBackend/Private/BlueprintCompilerCppBackendValueHelper.cpp
@@ -699,7 +699,7 @@ FString FEmitDefaultValueHelper::HandleSpecialTypes(FEmitterLocalContext& Contex
 	{
 		if (Object)
 		{
-			const bool bIsDefaultSubobject = Object->IsDefaultSubobject();
+			const bool bIsDefaultSubobject = Object->IsDefaultSubobject() && Object->HasAllFlags(RF_DefaultSubObject);
 			const bool bIsInstancedReference = Property->HasAnyPropertyFlags(CPF_InstancedReference);
 
 			UClass* ObjectClassToUse = Context.GetFirstNativeOrConvertedClass(Class);
diff --git a/Engine/Source/Developer/MessageLog/Private/Model/MessageLogListingModel.cpp b/Engine/Source/Developer/MessageLog/Private/Model/MessageLogListingModel.cpp
--- a/Engine/Source/Developer/MessageLog/Private/Model/MessageLogListingModel.cpp
+++ b/Engine/Source/Developer/MessageLog/Private/Model/MessageLogListingModel.cpp
@@ -22,24 +22,6 @@ const TSharedPtr< FTokenizedMessage >  FMessageLogListingModel::GetMessageAtInde
 	return FoundMessage;
 }
 
-const TSharedPtr< FTokenizedMessage > FMessageLogListingModel::GetMessageFromData( const FTokenizedMiscData& MessageData ) const
-{
-	// Go through all the messages looking for a data match
-	for(uint32 PageIndex = 0; PageIndex < NumPages(); PageIndex++)
-	{
-		for( MessageContainer::TConstIterator It( GetMessageIterator(PageIndex) ); It; ++It )
-		{
-			const TSharedPtr< FTokenizedMessage > Message = *It;
-
-			if( Message->GetMessageData() == &MessageData )
-			{
-				return Message;
-			}
-		}
-	}
-	return nullptr;
-}
-
 FText FMessageLogListingModel::GetAllMessagesAsText( const uint32 PageIndex ) const
 {
 	FText CompiledText;
diff --git a/Engine/Source/Developer/MessageLog/Private/Model/MessageLogListingModel.h b/Engine/Source/Developer/MessageLog/Private/Model/MessageLogListingModel.h
--- a/Engine/Source/Developer/MessageLog/Private/Model/MessageLogListingModel.h
+++ b/Engine/Source/Developer/MessageLog/Private/Model/MessageLogListingModel.h
@@ -52,9 +52,6 @@ public:
 	/** Returns the message at the specified index */
 	const TSharedPtr<FTokenizedMessage> GetMessageAtIndex( const uint32 PageIndex, const int32 MessageIndex ) const;
 
-	/** Tries to find the first message that matches the message data */
-	const TSharedPtr<FTokenizedMessage> GetMessageFromData( const FTokenizedMiscData& MessageData ) const;
-
 	/** Gets all messages as a string */
 	FText GetAllMessagesAsText(const uint32 PageIndex) const;
 
diff --git a/Engine/Source/Developer/MessageLog/Private/Presentation/MessageLogListingViewModel.cpp b/Engine/Source/Developer/MessageLog/Private/Presentation/MessageLogListingViewModel.cpp
--- a/Engine/Source/Developer/MessageLog/Private/Presentation/MessageLogListingViewModel.cpp
+++ b/Engine/Source/Developer/MessageLog/Private/Presentation/MessageLogListingViewModel.cpp
@@ -162,11 +162,6 @@ void FMessageLogListingViewModel::ClearMessages()
 	MessageLogListingModel->ClearMessages();
 }
 
-TSharedPtr<FTokenizedMessage> FMessageLogListingViewModel::GetMessageFromData(const struct FTokenizedMiscData& MessageData) const
-{
-	return MessageLogListingModel->GetMessageFromData(MessageData);
-}
-
 void FMessageLogListingViewModel::SelectMessages( const TArray< TSharedRef<FTokenizedMessage> >& InSelectedMessages )
 {
 	SelectedFilteredMessages.Empty();
diff --git a/Engine/Source/Developer/MessageLog/Private/Presentation/MessageLogListingViewModel.h b/Engine/Source/Developer/MessageLog/Private/Presentation/MessageLogListingViewModel.h
--- a/Engine/Source/Developer/MessageLog/Private/Presentation/MessageLogListingViewModel.h
+++ b/Engine/Source/Developer/MessageLog/Private/Presentation/MessageLogListingViewModel.h
@@ -41,7 +41,6 @@ public:
 	virtual void AddMessage( const TSharedRef< class FTokenizedMessage >& NewMessage, bool bMirrorToOutputLog ) override;
 	virtual void AddMessages( const TArray< TSharedRef< class FTokenizedMessage > >& NewMessages, bool bMirrorToOutputLog ) override;
 	virtual void ClearMessages() override;
-	virtual TSharedPtr<FTokenizedMessage> GetMessageFromData(const struct FTokenizedMiscData& MessageData) const override;
 	virtual const TArray< TSharedRef<class FTokenizedMessage> >& GetSelectedMessages() const override;
 	virtual void SelectMessages( const TArray< TSharedRef<class FTokenizedMessage> >& InSelectedMessages ) override;
 	virtual const TArray< TSharedRef<class FTokenizedMessage> >& GetFilteredMessages() const override;
diff --git a/Engine/Source/Developer/MessageLog/Public/IMessageLogListing.h b/Engine/Source/Developer/MessageLog/Public/IMessageLogListing.h
--- a/Engine/Source/Developer/MessageLog/Public/IMessageLogListing.h
+++ b/Engine/Source/Developer/MessageLog/Public/IMessageLogListing.h
@@ -25,13 +25,6 @@ public:
 	/** Clears messages (if paged, in the current page) */
 	virtual void ClearMessages() = 0;
 
-	/** 
-	 * Tries to find the first message that matches the message data.
-	 * 
-	 * @param	MessageData	The message data to test.
-	 */
-	virtual TSharedPtr<class FTokenizedMessage> GetMessageFromData(const struct FTokenizedMiscData& MessageData) const = 0;
-
 	/** 
 	 * Gets a list of the selected messages for the specific log listing
 	 * 
diff --git a/Engine/Source/Editor/BlueprintGraph/Classes/K2Node.h b/Engine/Source/Editor/BlueprintGraph/Classes/K2Node.h
--- a/Engine/Source/Editor/BlueprintGraph/Classes/K2Node.h
+++ b/Engine/Source/Editor/BlueprintGraph/Classes/K2Node.h
@@ -400,8 +400,11 @@ protected:
 
 	enum ERedirectType
 	{
+		/* The pins do not match */
 		ERedirectType_None,
+		/* The pins match by name or redirect and have the same type (or we're ok with the mismatched type) */
 		ERedirectType_Name,
+		/* The pins match via a redirect and the value needs to also be redirected */
 		ERedirectType_Value
 	};
 
diff --git a/Engine/Source/Editor/BlueprintGraph/Classes/K2Node_CommutativeAssociativeBinaryOperator.h b/Engine/Source/Editor/BlueprintGraph/Classes/K2Node_CommutativeAssociativeBinaryOperator.h
--- a/Engine/Source/Editor/BlueprintGraph/Classes/K2Node_CommutativeAssociativeBinaryOperator.h
+++ b/Engine/Source/Editor/BlueprintGraph/Classes/K2Node_CommutativeAssociativeBinaryOperator.h
@@ -41,6 +41,7 @@ public:
 
 	// UEdGraphNode interface
 	virtual void AllocateDefaultPins() override;
+	virtual void ValidateNodeDuringCompilation(class FCompilerResultsLog& MessageLog) const override;
 	// End of UEdGraphNode interface
 
 	// UK2Node interface
diff --git a/Engine/Source/Editor/BlueprintGraph/Private/EdGraphSchema_K2.cpp b/Engine/Source/Editor/BlueprintGraph/Private/EdGraphSchema_K2.cpp
--- a/Engine/Source/Editor/BlueprintGraph/Private/EdGraphSchema_K2.cpp
+++ b/Engine/Source/Editor/BlueprintGraph/Private/EdGraphSchema_K2.cpp
@@ -3345,6 +3345,15 @@ bool UEdGraphSchema_K2::GetPropertyCategoryInfo(const UProperty* TestProperty, F
 	{
 		OutCategory = PC_Struct;
 		OutSubCategoryObject = StructProperty->Struct;
+		// Match IsTypeCompatibleWithProperty and erase REINST_ structs here:
+		if(UUserDefinedStruct* UDS = Cast<UUserDefinedStruct>(StructProperty->Struct))
+		{
+			UUserDefinedStruct* RealStruct = UDS->PrimaryStruct.Get();
+			if(RealStruct)
+			{
+				OutSubCategoryObject = RealStruct;
+			}
+		}
 	}
 	else if (TestProperty->IsA<UFloatProperty>())
 	{
diff --git a/Engine/Source/Editor/BlueprintGraph/Private/K2Node.cpp b/Engine/Source/Editor/BlueprintGraph/Private/K2Node.cpp
--- a/Engine/Source/Editor/BlueprintGraph/Private/K2Node.cpp
+++ b/Engine/Source/Editor/BlueprintGraph/Private/K2Node.cpp
@@ -900,9 +900,12 @@ FString UK2Node::GetPinMetaData(FName InPinName, FName InKey)
 void UK2Node::RewireOldPinsToNewPins(TArray<UEdGraphPin*>& InOldPins, TArray<UEdGraphPin*>& InNewPins, TMap<UEdGraphPin*, UEdGraphPin*>* NewPinToOldPin)
 {
 	TArray<UEdGraphPin*> OrphanedOldPins;
+	TArray<bool> NewPinMatched; // Tracks whether a NewPin has already been matched to an OldPin
 	TMap<UEdGraphPin*, UEdGraphPin*> MatchedPins; // Old to New
 
-	const bool bSaveUnconnectedDefaultPins = (InNewPins.Num() == 0 && (IsA<UK2Node_CallFunction>() || IsA<UK2Node_MacroInstance>()));
+	const int32 NumNewPins = InNewPins.Num();
+	NewPinMatched.AddDefaulted(NumNewPins);
+	const bool bSaveUnconnectedDefaultPins = (NumNewPins == 0 && (IsA<UK2Node_CallFunction>() || IsA<UK2Node_MacroInstance>()));
 
 	// Rewire any connection to pins that are matched by name (O(N^2) right now)
 	// NOTE: we iterate backwards through the list because ReconstructSinglePin()
@@ -917,27 +920,29 @@ void UK2Node::RewireOldPinsToNewPins(TArray<UEdGraphPin*>& InOldPins, TArray<UEd
 
 		// common case is for InOldPins and InNewPins to match, so we start searching from the current index:
 		bool bMatched = false;
-		const int32 NumNewPins = InNewPins.Num();
 		int32 NewPinIndex = (NumNewPins ? OldPinIndex % NumNewPins : 0);
 		for (int32 NewPinCount = NumNewPins - 1; NewPinCount >= 0; --NewPinCount)
 		{
-			// if InNewPins grows in this loop then we may skip entries and fail to find a match:
+			// if InNewPins grows then we may skip entries and fail to find a match or NewPinMatched will not be accurate
 			check(NumNewPins == InNewPins.Num());
-			UEdGraphPin* NewPin = InNewPins[NewPinIndex];
-
-			const ERedirectType RedirectType = DoPinsMatchForReconstruction(NewPin, NewPinIndex, OldPin, OldPinIndex);
-			if (RedirectType != ERedirectType_None)
+			if (!NewPinMatched[NewPinIndex])
 			{
-				ReconstructSinglePin(NewPin, OldPin, RedirectType);
-				MatchedPins.Add(OldPin, NewPin);
-				if (NewPinToOldPin)
+				UEdGraphPin* NewPin = InNewPins[NewPinIndex];
+
+				const ERedirectType RedirectType = DoPinsMatchForReconstruction(NewPin, NewPinIndex, OldPin, OldPinIndex);
+				if (RedirectType != ERedirectType_None)
 				{
-					NewPinToOldPin->Add(NewPin, OldPin);
+					ReconstructSinglePin(NewPin, OldPin, RedirectType);
+					MatchedPins.Add(OldPin, NewPin);
+					if (NewPinToOldPin)
+					{
+						NewPinToOldPin->Add(NewPin, OldPin);
+					}
+					bMatched = true;
+					NewPinMatched[NewPinIndex] = true;
+					break;
 				}
-				bMatched = true;
-				break;
 			}
-
 			NewPinIndex = (NewPinIndex + 1) % InNewPins.Num();
 		}
 
diff --git a/Engine/Source/Editor/BlueprintGraph/Private/K2Node_CommutativeAssociativeBinaryOperator.cpp b/Engine/Source/Editor/BlueprintGraph/Private/K2Node_CommutativeAssociativeBinaryOperator.cpp
--- a/Engine/Source/Editor/BlueprintGraph/Private/K2Node_CommutativeAssociativeBinaryOperator.cpp
+++ b/Engine/Source/Editor/BlueprintGraph/Private/K2Node_CommutativeAssociativeBinaryOperator.cpp
@@ -103,50 +103,57 @@ FEdGraphPinType UK2Node_CommutativeAssociativeBinaryOperator::GetType() const
 	return FEdGraphPinType();
 }
 
-void UK2Node_CommutativeAssociativeBinaryOperator::AllocateDefaultPins()
+void UK2Node_CommutativeAssociativeBinaryOperator::ValidateNodeDuringCompilation(class FCompilerResultsLog& MessageLog) const
 {
-	Super::AllocateDefaultPins();
+	Super::ValidateNodeDuringCompilation(MessageLog);
 
-	const UFunction* Function = GetTargetFunction();
-	if(NULL != Function)
+	if (const UFunction* Function = GetTargetFunction())
 	{
-		check(Function->HasAnyFunctionFlags(FUNC_BlueprintPure));
-	
-#if DO_CHECK
-		ensure(FindOutPin());
-		ensure((FindSelfPin() ? 4 : 3) ==  Pins.Num());
+		if (Function->HasAnyFunctionFlags(FUNC_BlueprintPure))
 		{
-			//VALIDATION
+			UEdGraphPin* SelfPin = FindSelfPin();
 			const FEdGraphPinType InputType = GetType();
-			int32 NativeInputPinsNum = 0;
-			for (int32 PinIt = 0; PinIt < Pins.Num(); PinIt++)
+			int32 OutPinCount = 0;
+			bool bConsistentPinType = true;
+
+			for (UEdGraphPin* Pin : Pins)
 			{
-				const UEdGraphPin* Pin = Pins[PinIt];
-				if (Pin != FindSelfPin())
+				if (Pin->Direction == EEdGraphPinDirection::EGPD_Output)
+				{
+					++OutPinCount;
+				}
+				else if (Pin != SelfPin)
 				{
-					ensure(InputType == Pin->PinType);
-					NativeInputPinsNum += (EEdGraphPinDirection::EGPD_Input ==  Pin->Direction) ? 1 : 0;
+					if (Pin->PinType != InputType)
+					{
+						bConsistentPinType = false;
+					}
 				}
 			}
-			ensure(BinaryOperatorInputsNum == NativeInputPinsNum);
-		}
-#endif // DO_CHECK
 
-		for (int32 i = 0; i < NumAdditionalInputs; ++i)
+			if (OutPinCount != 1)
+			{
+				MessageLog.Error(*LOCTEXT("InvalidOutPin", "@@ - Commutative associative binary operators must have a single output pin.").ToString(), this);
+			}
+			if (!bConsistentPinType)
+			{
+				MessageLog.Error(*LOCTEXT("InconsistentType", "@@ - Commutative associative binary operator inputs must all be of the same type.").ToString(), this);
+			}
+		}
+		else
 		{
-			AddInputPinInner(i);
+			MessageLog.Error(*LOCTEXT("MustBePure", "@@ - Commutative associative binary operators must be pure functions.").ToString(), this);
 		}
 	}
-	else
+}
+
+void UK2Node_CommutativeAssociativeBinaryOperator::AllocateDefaultPins()
+{
+	Super::AllocateDefaultPins();
+
+	for (int32 i = 0; i < NumAdditionalInputs; ++i)
 	{
-		const UClass* FunctionParentClass = FunctionReference.GetMemberParentClass(GetBlueprintClassFromNode());
-		Message_Error(
-			FText::Format(
-				LOCTEXT("NoFunction_ErrorFmt", "CommutativeAssociativeBinaryOperator has no function: '{0}' class: '{1}'"),
-				FText::FromString(FunctionReference.GetMemberName().ToString()),
-				FunctionParentClass ? FText::FromString(FunctionParentClass->GetName()) : LOCTEXT("NoFunction_ErrorNone", "None")
-			).ToString()
-		);
+		AddInputPinInner(i);
 	}
 }
 
diff --git a/Engine/Source/Editor/BlueprintGraph/Private/K2Node_GetDataTableRow.cpp b/Engine/Source/Editor/BlueprintGraph/Private/K2Node_GetDataTableRow.cpp
--- a/Engine/Source/Editor/BlueprintGraph/Private/K2Node_GetDataTableRow.cpp
+++ b/Engine/Source/Editor/BlueprintGraph/Private/K2Node_GetDataTableRow.cpp
@@ -92,10 +92,18 @@ void UK2Node_GetDataTableRow::SetReturnTypeForStruct(UScriptStruct* NewRowStruct
 	if (NewRowStruct != OldRowStruct)
 	{
 		UEdGraphPin* ResultPin = GetResultPin();
+
+		if (ResultPin->SubPins.Num() > 0)
+		{
+			GetSchema()->RecombinePin(ResultPin);
+		}
+
 		// NOTE: purposefully not disconnecting the ResultPin (even though it changed type)... we want the user to see the old
 		//       connections, and incompatible connections will produce an error (plus, some super-struct connections may still be valid)
 		ResultPin->PinType.PinSubCategoryObject = NewRowStruct;
 		ResultPin->PinType.PinCategory = (NewRowStruct == nullptr) ? UEdGraphSchema_K2::PC_Wildcard : UEdGraphSchema_K2::PC_Struct;
+
+		CachedNodeTitle.Clear();
 	}
 }
 
diff --git a/Engine/Source/Editor/BlueprintGraph/Private/K2Node_SetFieldsInStruct.cpp b/Engine/Source/Editor/BlueprintGraph/Private/K2Node_SetFieldsInStruct.cpp
--- a/Engine/Source/Editor/BlueprintGraph/Private/K2Node_SetFieldsInStruct.cpp
+++ b/Engine/Source/Editor/BlueprintGraph/Private/K2Node_SetFieldsInStruct.cpp
@@ -359,13 +359,7 @@ bool UK2Node_SetFieldsInStruct::CanSplitPin(const UEdGraphPin* Pin) const
 		UEdGraphPin* InputPin = FindPinChecked(SetFieldsInStructHelper::StructRefPinName(), EGPD_Input);
 		if (Pin == InputPin)
 		{
-			for (const FOptionalPinFromProperty& OptionalProperty : ShowPinForProperties)
-			{
-				if (OptionalProperty.bShowPin)
-				{
-					return false;
-				}
-			}
+			return false;
 		}
 
 		return true;
diff --git a/Engine/Source/Editor/DataTableEditor/Private/SRowEditor.cpp b/Engine/Source/Editor/DataTableEditor/Private/SRowEditor.cpp
--- a/Engine/Source/Editor/DataTableEditor/Private/SRowEditor.cpp
+++ b/Engine/Source/Editor/DataTableEditor/Private/SRowEditor.cpp
@@ -8,6 +8,7 @@
 #include "Widgets/Input/SButton.h"
 #include "EditorStyleSet.h"
 #include "UObject/StructOnScope.h"
+#include "Misc/MessageDialog.h"
 
 #include "PropertyEditorModule.h"
 #include "IStructureDetailsView.h"
@@ -336,6 +337,14 @@ void SRowEditor::OnRowRenamed(const FText& Text, ETextCommit::Type CommitType)
 {
 	if (!GetCurrentNameAsText().EqualTo(Text) && DataTable.IsValid())
 	{
+		if (Text.IsEmptyOrWhitespace() || !FName::IsValidXName(Text.ToString(), INVALID_NAME_CHARACTERS))
+		{
+			// popup an error dialog here
+			const FText Message = FText::Format(LOCTEXT("InvalidRowName", "'{0}' is not a valid row name"), Text);
+			FMessageDialog::Open(EAppMsgType::Ok, Message);
+
+			return;
+		}
 		const FName NewName = DataTableUtils::MakeValidName(Text.ToString());
 		for (auto Name : CachedRowNames)
 		{
diff --git a/Engine/Source/Editor/Kismet/Private/BlueprintDetailsCustomization.cpp b/Engine/Source/Editor/Kismet/Private/BlueprintDetailsCustomization.cpp
--- a/Engine/Source/Editor/Kismet/Private/BlueprintDetailsCustomization.cpp
+++ b/Engine/Source/Editor/Kismet/Private/BlueprintDetailsCustomization.cpp
@@ -4147,6 +4147,12 @@ bool FBaseBlueprintGraphActionDetails::OnVerifyPinRename(UK2Node_EditablePinBase
 		return false;
 	}
 
+	if (InNewName == TEXT("None"))
+	{
+		OutErrorMessage = FText::Format( LOCTEXT("PinNameNone", "'None' is a reserved name"), FText::AsNumber( NAME_SIZE ) );
+		return false;
+	}
+
 	if (InTargetNode)
 	{
 		// Check if the name conflicts with any of the other internal UFunction's property names (local variables and parameters).
diff --git a/Engine/Source/Editor/Kismet/Private/SBlueprintPalette.cpp b/Engine/Source/Editor/Kismet/Private/SBlueprintPalette.cpp
--- a/Engine/Source/Editor/Kismet/Private/SBlueprintPalette.cpp
+++ b/Engine/Source/Editor/Kismet/Private/SBlueprintPalette.cpp
@@ -135,7 +135,8 @@ static FString GetVarTooltip(UBlueprint* InBlueprint, UClass* VarClass, FName Va
 			}
 			else
 			{
-				FBlueprintEditorUtils::GetBlueprintVariableMetaData(InBlueprint, VarName, nullptr, TEXT("tooltip"), ResultTooltip);
+				const UBlueprint* SourceBlueprint = SourceClass ? Cast<UBlueprint>(SourceClass->ClassGeneratedBy) : nullptr;
+				FBlueprintEditorUtils::GetBlueprintVariableMetaData(SourceBlueprint ? SourceBlueprint : InBlueprint, VarName, nullptr, TEXT("tooltip"), ResultTooltip);
 			}
 		}
 	}
diff --git a/Engine/Source/Editor/KismetCompiler/Private/KismetCompiler.cpp b/Engine/Source/Editor/KismetCompiler/Private/KismetCompiler.cpp
--- a/Engine/Source/Editor/KismetCompiler/Private/KismetCompiler.cpp
+++ b/Engine/Source/Editor/KismetCompiler/Private/KismetCompiler.cpp
@@ -1399,7 +1399,7 @@ void FKismetCompilerContext::PruneIsolatedNodes(const TArray<UEdGraphNode*>& Roo
 	}
 
 	const UEdGraphSchema* const K2Schema = UEdGraphSchema_K2::StaticClass()->GetDefaultObject<UEdGraphSchema_K2>();
-
+	TMap<UEdGraphNode*, TArray<UEdGraphNode*>> PrunedExecNodeNeighbors;
 	for (int32 NodeIndex = 0; NodeIndex < GraphNodes.Num(); ++NodeIndex)
 	{
 		UEdGraphNode* Node = GraphNodes[NodeIndex];
@@ -1434,6 +1434,19 @@ void FKismetCompilerContext::PruneIsolatedNodes(const TArray<UEdGraphNode*>& Roo
 			{
 				if (Node)
 				{
+					// Track nodes that are directly connected to the outputs of the node we are pruning so 
+					// that we can warn if one or more of those neighboring nodes are not also orphaned:
+					Node->ForEachNodeDirectlyConnectedIf(
+						// Consider connections on output pins other than the exec pin:
+						[](const UEdGraphPin* Pin) {
+							if(Pin->Direction == EGPD_Output && Pin->PinType.PinCategory != UEdGraphSchema_K2::PC_Exec) 
+							{ 
+								return true; 
+							}
+							return false;
+						},
+						[&PrunedExecNodeNeighbors, Node](UEdGraphNode* NeighborNode) { PrunedExecNodeNeighbors.FindOrAdd(Node).Add(NeighborNode); }
+					);
 					Node->BreakAllNodeLinks();
 				}
 				GraphNodes.RemoveAtSwap(NodeIndex);
@@ -1476,6 +1489,24 @@ void FKismetCompilerContext::PruneIsolatedNodes(const TArray<UEdGraphNode*>& Roo
 			}
 		}
 	}
+
+	for(const TPair<UEdGraphNode*, TArray<UEdGraphNode*>>& PrunedExecNodeWithNeighbors : PrunedExecNodeNeighbors)
+	{
+		bool bNeighborsNotPruned = false;
+		for(UEdGraphNode* Neighbor : PrunedExecNodeWithNeighbors.Value)
+		{
+			if(GraphNodes.Contains(Neighbor))
+			{
+				bNeighborsNotPruned = true;
+			}
+		}
+
+		if(bNeighborsNotPruned)
+		{
+			// Warn the user if they are attempting to read an output value from a pruned exec node:
+			MessageLog.Warning(FName(TEXT("PrunedExecInUse")), *LOCTEXT("PrunedExecNodeAttemptedUse", "@@ was pruned because its Exec pin is not connected, the connected value is not available and will instead be read as default").ToString(), PrunedExecNodeWithNeighbors.Key);
+		}
+	}
 }
 
 /**
@@ -1739,6 +1770,7 @@ void FKismetCompilerContext::PrecompileFunction(FKismetFunctionContext& Context,
 			ensure(!NewClass->UberGraphFunction);
 			NewClass->UberGraphFunction = Context.Function;
 			NewClass->UberGraphFunction->FunctionFlags |= FUNC_UbergraphFunction;
+			NewClass->UberGraphFunction->FunctionFlags |= FUNC_Final;
 		}
 
 		// Register nets from function entry/exit nodes first, even for skeleton compiles (as they form the signature)
@@ -3364,6 +3396,7 @@ void FKismetCompilerContext::ExpandTunnelsAndMacros(UEdGraph* SourceGraph)
 			for (int32 I = 0; I < ClonedGraph->Nodes.Num(); ++I)
 			{
 				MacroGeneratedNodes.Add(ClonedGraph->Nodes[I], CurrentNode);
+				MessageLog.NotifyIntermediateMacroNode(CurrentNode, ClonedGraph->Nodes[I]);
 			}
 
 			TArray<UEdGraphNode*> MacroNodes(ClonedGraph->Nodes);
diff --git a/Engine/Source/Editor/UnrealEd/Classes/Factories/CSVImportFactory.h b/Engine/Source/Editor/UnrealEd/Classes/Factories/CSVImportFactory.h
--- a/Engine/Source/Editor/UnrealEd/Classes/Factories/CSVImportFactory.h
+++ b/Engine/Source/Editor/UnrealEd/Classes/Factories/CSVImportFactory.h
@@ -11,6 +11,7 @@
 #include "Curves/RichCurve.h"
 #include "Factories/Factory.h"
 #include "Factories/ImportSettings.h"
+#include "EditorReimportHandler.h"
 #include "CSVImportFactory.generated.h"
 
 DECLARE_LOG_CATEGORY_EXTERN(LogCSVImportFactory, Log, All);
@@ -55,14 +56,14 @@ class UNREALED_API UCSVImportFactory : public UFactory, public IImportSettingsPa
 
 	//~ Begin UFactory Interface
 	virtual FText GetDisplayName() const override;
-	virtual UObject* FactoryCreateText( UClass* InClass, UObject* InParent, FName InName, EObjectFlags Flags, UObject* Context, const TCHAR* Type, const TCHAR*& Buffer, const TCHAR* BufferEnd, FFeedbackContext* Warn ) override;
+	virtual UObject* FactoryCreateText(UClass* InClass, UObject* InParent, FName InName, EObjectFlags Flags, UObject* Context, const TCHAR* Type, const TCHAR*& Buffer, const TCHAR* BufferEnd, FFeedbackContext* Warn, bool& bOutOperationCanceled) override;
 	virtual bool DoesSupportClass(UClass * Class) override;
 	virtual bool FactoryCanImport(const FString& Filename) override;
 	virtual	IImportSettingsParser* GetImportSettingsParser() override { return this; }
 
 
 	/* Reimport an object that was created based on a CSV */
-	bool ReimportCSV(UObject* Obj);
+	EReimportResult::Type ReimportCSV(UObject* Obj);
 	
 	/**
 	 * IImportSettings interface
@@ -76,7 +77,7 @@ protected:
 
 private:
 	/* Reimport object from the given path*/
-	bool Reimport(UObject* Obj, const FString& Path  );
+	EReimportResult::Type Reimport(UObject* Obj, const FString& Path);
 
 	UPROPERTY()
 	FCSVImportSettings AutomatedImportSettings;
diff --git a/Engine/Source/Editor/UnrealEd/Classes/Factories/Factory.h b/Engine/Source/Editor/UnrealEd/Classes/Factories/Factory.h
--- a/Engine/Source/Editor/UnrealEd/Classes/Factories/Factory.h
+++ b/Engine/Source/Editor/UnrealEd/Classes/Factories/Factory.h
@@ -274,6 +274,26 @@ protected:
 		return nullptr;
 	}
 
+	/**
+	* Create a new object by importing it from a text buffer.
+	*
+	* @param InClass
+	* @param InParent
+	* @param InName
+	* @param Flags
+	* @param Context
+	* @param Type (must not be nullptr, i.e. TEXT("TGA"))
+	* @param Buffer
+	* @param BufferEnd
+	* @param Warn
+	* @param bOutOperationCanceled Will indicate whether the user canceled the import.
+	* @return The new object.
+	*/
+	virtual UObject* FactoryCreateText(UClass* InClass, UObject* InParent, FName InName, EObjectFlags Flags, UObject* Context, const TCHAR* Type, const TCHAR*& Buffer, const TCHAR* BufferEnd, FFeedbackContext* Warn, bool& bOutOperationCanceled)
+	{
+		return FactoryCreateText(InClass, InParent, InName, Flags, Context, Type, Buffer, BufferEnd, Warn);
+	}
+
 	/**
 	 * Create a new object by importing it from a binary buffer.
 	 *
diff --git a/Engine/Source/Editor/UnrealEd/Private/Commandlets/GenerateBlueprintAPICommandlet.cpp b/Engine/Source/Editor/UnrealEd/Private/Commandlets/GenerateBlueprintAPICommandlet.cpp
--- a/Engine/Source/Editor/UnrealEd/Private/Commandlets/GenerateBlueprintAPICommandlet.cpp
+++ b/Engine/Source/Editor/UnrealEd/Private/Commandlets/GenerateBlueprintAPICommandlet.cpp
@@ -1160,9 +1160,13 @@ int32 UGenerateBlueprintAPICommandlet::Main(FString const& Params)
 		}
 	}
 
-	GenerateBlueprintAPIUtils::DumpCategoryInfo(1, FileOut);
+	if(FileOut)
+	{
+		GenerateBlueprintAPIUtils::DumpCategoryInfo(1, FileOut);
+
+		CloseFileStream(&FileOut);
+	}
 
-	CloseFileStream(&FileOut);
 	return 0;
 }
 
diff --git a/Engine/Source/Editor/UnrealEd/Private/DataTableEditorUtils.cpp b/Engine/Source/Editor/UnrealEd/Private/DataTableEditorUtils.cpp
--- a/Engine/Source/Editor/UnrealEd/Private/DataTableEditorUtils.cpp
+++ b/Engine/Source/Editor/UnrealEd/Private/DataTableEditorUtils.cpp
@@ -30,7 +30,7 @@ bool FDataTableEditorUtils::RemoveRow(UDataTable* DataTable, FName Name)
 
 		BroadcastPreChange(DataTable, EDataTableChangeInfo::RowList);
 		DataTable->Modify();
-		uint8* RowData = NULL;
+		uint8* RowData = nullptr;
 		const bool bRemoved = DataTable->RowMap.RemoveAndCopyValue(Name, RowData);
 		if (bRemoved && RowData)
 		{
@@ -39,7 +39,7 @@ bool FDataTableEditorUtils::RemoveRow(UDataTable* DataTable, FName Name)
 			bResult = true;
 
 			// Compact the map so that a subsequent add goes at the end of the table
-			DataTable->RowMap.Compact();
+			DataTable->RowMap.CompactStable();
 		}
 		BroadcastPostChange(DataTable, EDataTableChangeInfo::RowList);
 	}
@@ -48,9 +48,9 @@ bool FDataTableEditorUtils::RemoveRow(UDataTable* DataTable, FName Name)
 
 uint8* FDataTableEditorUtils::AddRow(UDataTable* DataTable, FName RowName)
 {
-	if (!DataTable || (RowName == NAME_None) || (DataTable->RowMap.Find(RowName) != NULL) || !DataTable->RowStruct)
+	if (!DataTable || (RowName == NAME_None) || (DataTable->RowMap.Find(RowName) != nullptr) || !DataTable->RowStruct)
 	{
-		return NULL;
+		return nullptr;
 	}
 
 	const FScopedTransaction Transaction(LOCTEXT("AddDataTableRow", "Add Data Table Row"));
@@ -78,7 +78,7 @@ bool FDataTableEditorUtils::RenameRow(UDataTable* DataTable, FName OldName, FNam
 		BroadcastPreChange(DataTable, EDataTableChangeInfo::RowList);
 		DataTable->Modify();
 
-		uint8* RowData = NULL;
+		uint8* RowData = nullptr;
 		const bool bValidnewName = (NewName != NAME_None) && !DataTable->RowMap.Find(NewName);
 		const bool bRemoved = bValidnewName && DataTable->RowMap.RemoveAndCopyValue(OldName, RowData);
 		if (bRemoved)
diff --git a/Engine/Source/Editor/UnrealEd/Private/EdGraphToken.cpp b/Engine/Source/Editor/UnrealEd/Private/EdGraphToken.cpp
--- a/Engine/Source/Editor/UnrealEd/Private/EdGraphToken.cpp
+++ b/Engine/Source/Editor/UnrealEd/Private/EdGraphToken.cpp
@@ -3,30 +3,24 @@
 #include "EdGraphToken.h"
 #include "Kismet2/CompilerResultsLog.h"
 #include "Kismet2/BlueprintEditorUtils.h"
+#include "Misc/UObjectToken.h"
 
-TSharedRef<IMessageToken> FEdGraphToken::Create(const UObject* InObject, const FCompilerResultsLog* Log, UEdGraphNode*& OutSourceNode)
+void FEdGraphToken::Create(const UObject* InObject, FCompilerResultsLog* Log, FTokenizedMessage &OutMessage, TArray<UEdGraphNode*>& OutSourceNodes)
 {
-	const UObject* SourceNode = Log->FindSourceObject(InObject);
-	if (OutSourceNode == nullptr)
-	{
-		OutSourceNode = const_cast<UEdGraphNode*>(Cast<UEdGraphNode>(SourceNode));
-	}
-	return MakeShareable(new FEdGraphToken(Log->FindSourceObject(InObject), nullptr));
+	return CreateInternal(InObject, Log, OutMessage, OutSourceNodes, nullptr);
 }
 
-TSharedRef<IMessageToken> FEdGraphToken::Create(const UEdGraphPin* InPin, const FCompilerResultsLog* Log, UEdGraphNode*& OutSourceNode)
+void FEdGraphToken::Create(const UEdGraphPin* InPin, FCompilerResultsLog* Log, FTokenizedMessage &OutMessage, TArray<UEdGraphNode*>& OutSourceNodes)
 {
-	const UObject* SourceNode = InPin ? Log->FindSourceObject(InPin->GetOwningNode()) : nullptr;
-	if (OutSourceNode == nullptr)
+	if(InPin && InPin->GetOwningNode())
 	{
-		OutSourceNode = const_cast<UEdGraphNode*>(Cast<UEdGraphNode>(SourceNode));
+		return CreateInternal(InPin->GetOwningNode(), Log, OutMessage, OutSourceNodes, InPin);
 	}
-	return MakeShareable(new FEdGraphToken(SourceNode, Log->FindSourcePin(InPin)));
 }
 
-TSharedRef<IMessageToken> FEdGraphToken::Create(const TCHAR* String, const FCompilerResultsLog* Log, UEdGraphNode*& OutSourceNode)
+void FEdGraphToken::Create(const TCHAR* String, FCompilerResultsLog* Log, FTokenizedMessage &OutMessage, TArray<UEdGraphNode*>& OutSourceNode)
 {
-	return FTextToken::Create(FText::FromString(FString(String)));
+	OutMessage.AddToken( FTextToken::Create(FText::FromString(FString(String))) );
 }
 
 const UEdGraphPin* FEdGraphToken::GetPin() const
@@ -76,3 +70,39 @@ FEdGraphToken::FEdGraphToken(const UObject* InObject, const UEdGraphPin* InPin)
 		CachedText = NSLOCTEXT("MessageLog", "NoneObjectToken", "<None>");
 	}
 }
+
+void FEdGraphToken::CreateInternal(const UObject* InObject, FCompilerResultsLog* Log, FTokenizedMessage &OutMessage, TArray<UEdGraphNode*>& OutSourceNodes, const UEdGraphPin* Pin)
+{
+	UObject* SourceObject = Log->FindSourceObject(const_cast<UObject*>(InObject));
+	OutMessage.AddToken( MakeShareable(new FEdGraphToken(SourceObject, Pin ? Log->FindSourcePin(Pin) : nullptr)));
+	if (UEdGraphNode* SourceNode = Cast<UEdGraphNode>(SourceObject))
+	{
+		OutSourceNodes.Add(SourceNode);
+		
+		// If this node came from a macro it actually has two source nodes, look up the other source node and add that as well:
+		if(const UEdGraph* OwningGraph = SourceNode->GetGraph())
+		{
+			if(const UEdGraphSchema* Schema = OwningGraph->GetSchema())
+			{
+				if(Schema->GetGraphType(OwningGraph) == GT_Macro)
+				{
+					if(UObject* MacroSourceObject = Log->FindSourceMacroInstance(Cast<const UEdGraphNode>(InObject)))	
+					{
+						OutMessage.AddToken( FTextToken::Create( NSLOCTEXT("EdGraphToken", "FromMacroInstance", "generated from expanding")) );
+						OutMessage.AddToken( MakeShareable(new FEdGraphToken(MacroSourceObject, nullptr)) );
+						if(UEdGraphNode* MacroSourceNode = Cast<UEdGraphNode>(MacroSourceObject))
+						{
+							OutSourceNodes.Add(MacroSourceNode);
+						}
+					}
+				}
+			}
+		}
+	}
+	
+	if(SourceObject)
+	{
+		// The message link is only used when the user double clicks on the line, we'll jump to the first source object by default:
+		OutMessage.SetMessageLink(FUObjectToken::Create(SourceObject));
+	}
+}
diff --git a/Engine/Source/Editor/UnrealEd/Private/Factories/CSVImportFactory.cpp b/Engine/Source/Editor/UnrealEd/Private/Factories/CSVImportFactory.cpp
--- a/Engine/Source/Editor/UnrealEd/Private/Factories/CSVImportFactory.cpp
+++ b/Engine/Source/Editor/UnrealEd/Private/Factories/CSVImportFactory.cpp
@@ -40,9 +40,9 @@ FCSVImportSettings::FCSVImportSettings()
 }
 
 
-static UClass* GetCurveClass( ECSVImportType ImportType )
+static UClass* GetCurveClass(ECSVImportType ImportType)
 {
-	switch( ImportType )
+	switch(ImportType)
 	{
 	case ECSVImportType::ECSV_CurveFloat:
 		return UCurveFloat::StaticClass();
@@ -84,24 +84,26 @@ FText UCSVImportFactory::GetDisplayName() const
 
 bool UCSVImportFactory::DoesSupportClass(UClass * Class)
 {
-	return (Class == UDataTable::StaticClass() || Class == UCurveTable::StaticClass() || Class == UCurveFloat::StaticClass() || Class == UCurveVector::StaticClass() || Class == UCurveLinearColor::StaticClass() );
+	return (Class == UDataTable::StaticClass() || Class == UCurveTable::StaticClass() || Class == UCurveFloat::StaticClass() || Class == UCurveVector::StaticClass() || Class == UCurveLinearColor::StaticClass());
 }
 
 bool UCSVImportFactory::FactoryCanImport(const FString& Filename)
 {
 	const FString Extension = FPaths::GetExtension(Filename);
 
-	if(Extension == TEXT("csv"))
+	if (Extension == TEXT("csv"))
 	{
 		return true;
 	}
 	return false;
 }
 
-UObject* UCSVImportFactory::FactoryCreateText(UClass* InClass, UObject* InParent, FName InName, EObjectFlags Flags, UObject* Context, const TCHAR* Type, const TCHAR*& Buffer, const TCHAR* BufferEnd, FFeedbackContext* Warn)
+UObject* UCSVImportFactory::FactoryCreateText(UClass* InClass, UObject* InParent, FName InName, EObjectFlags Flags, UObject* Context, const TCHAR* Type, const TCHAR*& Buffer, const TCHAR* BufferEnd, FFeedbackContext* Warn, bool& bOutOperationCanceled)
 {
 	FEditorDelegates::OnAssetPreImport.Broadcast(this, InClass, InParent, InName, Type);
 
+	bOutOperationCanceled = false;
+
 	// See if table/curve already exists
 	UDataTable* ExistingTable = FindObject<UDataTable>(InParent, *InName.ToString());
 	UCurveTable* ExistingCurveTable = FindObject<UCurveTable>(InParent, *InName.ToString());
@@ -109,11 +111,11 @@ UObject* UCSVImportFactory::FactoryCreateText(UClass* InClass, UObject* InParent
 
 	// Save off information if so
 	bool bHaveInfo = false;
-	UScriptStruct* ImportRowStruct = NULL;
+	UScriptStruct* ImportRowStruct = nullptr;
 	ERichCurveInterpMode ImportCurveInterpMode = RCIM_Linear;
 	ECSVImportType ImportType = ECSVImportType::ECSV_DataTable;
 
-	if( IsAutomatedImport() )
+	if (IsAutomatedImport())
 	{
 		ImportRowStruct = AutomatedImportSettings.ImportRowStruct;
 		ImportCurveInterpMode = AutomatedImportSettings.ImportCurveInterpMode;
@@ -122,17 +124,17 @@ UObject* UCSVImportFactory::FactoryCreateText(UClass* InClass, UObject* InParent
 		// For automated import to work a row struct must be specified for a datatable type or a curve type must be specified
 		bHaveInfo = ImportRowStruct != nullptr || ImportType != ECSVImportType::ECSV_DataTable;
 	}
-	else if(ExistingTable != NULL)
+	else if (ExistingTable != nullptr)
 	{
 		ImportRowStruct = ExistingTable->RowStruct;
 		bHaveInfo = true;
 	}
-	else if(ExistingCurveTable != NULL)
+	else if (ExistingCurveTable != nullptr)
 	{
 		ImportType = ECSVImportType::ECSV_CurveTable;
 		bHaveInfo = true;
 	}
-	else if(ExistingCurve != NULL)
+	else if (ExistingCurve != nullptr)
 	{
 		ImportType = ExistingCurve->IsA(UCurveFloat::StaticClass()) ? ECSVImportType::ECSV_CurveFloat : ECSVImportType::ECSV_CurveVector;
 		bHaveInfo = true;
@@ -141,21 +143,21 @@ UObject* UCSVImportFactory::FactoryCreateText(UClass* InClass, UObject* InParent
 	bool bDoImport = true;
 
 	// If we do not have the info we need, pop up window to ask for things
-	if(!bHaveInfo && !IsAutomatedImport())
+	if (!bHaveInfo && !IsAutomatedImport())
 	{
 		TSharedPtr<SWindow> ParentWindow;
 		// Check if the main frame is loaded.  When using the old main frame it may not be.
-		if( FModuleManager::Get().IsModuleLoaded( "MainFrame" ) )
+		if (FModuleManager::Get().IsModuleLoaded("MainFrame"))
 		{
-			IMainFrameModule& MainFrame = FModuleManager::LoadModuleChecked<IMainFrameModule>( "MainFrame" );
+			IMainFrameModule& MainFrame = FModuleManager::LoadModuleChecked<IMainFrameModule>("MainFrame");
 			ParentWindow = MainFrame.GetParentWindow();
 		}
 
 		TSharedPtr<SCSVImportOptions> ImportOptionsWindow;
 
 		TSharedRef<SWindow> Window = SNew(SWindow)
-			.Title( LOCTEXT("DataTableOptionsWindowTitle", "DataTable Options" ))
-			.SizingRule( ESizingRule::Autosized );
+			.Title(LOCTEXT("DataTableOptionsWindowTitle", "DataTable Options"))
+			.SizingRule(ESizingRule::Autosized);
 		
 		FString ParentFullPath;
 
@@ -177,18 +179,19 @@ UObject* UCSVImportFactory::FactoryCreateText(UClass* InClass, UObject* InParent
 		ImportRowStruct = ImportOptionsWindow->GetSelectedRowStruct();
 		ImportCurveInterpMode = ImportOptionsWindow->GetSelectedCurveIterpMode();
 		bDoImport = ImportOptionsWindow->ShouldImport();
+		bOutOperationCanceled = !bDoImport;
 	}
-	else if(!bHaveInfo && IsAutomatedImport())
+	else if (!bHaveInfo && IsAutomatedImport())
 	{
-		if(ImportType == ECSVImportType::ECSV_DataTable && !ImportRowStruct)
+		if (ImportType == ECSVImportType::ECSV_DataTable && !ImportRowStruct)
 		{
 			UE_LOG(LogCSVImportFactory, Error, TEXT("A Data table row type must be specified in the import settings json file for automated import"));
 		}
 		bDoImport = false;
 	}
 
-	UObject* NewAsset = NULL;
-	if(bDoImport)
+	UObject* NewAsset = nullptr;
+	if (bDoImport)
 	{
 		// Convert buffer to an FString (will this be slow with big tables?)
 		FString String;
@@ -208,7 +211,7 @@ UObject* UCSVImportFactory::FactoryCreateText(UClass* InClass, UObject* InParent
 
 			// If there is an existing table, need to call this to free data memory before recreating object
 			bool bStripFromClientBuilds = false;
-			if(ExistingTable != NULL)
+			if (ExistingTable != nullptr)
 			{
 				DataTableClass = ExistingTable->GetClass();
 				ExistingTable->EmptyTable();
@@ -232,7 +235,7 @@ UObject* UCSVImportFactory::FactoryCreateText(UClass* InClass, UObject* InParent
 			UClass* CurveTableClass = UCurveTable::StaticClass();
 
 			// If there is an existing table, need to call this to free data memory before recreating object
-			if(ExistingCurveTable != NULL)
+			if (ExistingCurveTable != nullptr)
 			{
 				CurveTableClass = ExistingCurveTable->GetClass();
 				ExistingCurveTable->EmptyTable();
@@ -251,7 +254,7 @@ UObject* UCSVImportFactory::FactoryCreateText(UClass* InClass, UObject* InParent
 		}
 		else if (ImportType == ECSVImportType::ECSV_CurveFloat || ImportType == ECSVImportType::ECSV_CurveVector || ImportType == ECSVImportType::ECSV_CurveLinearColor)
 		{
-			UClass* CurveClass = ExistingCurve ? ExistingCurve->GetClass() : GetCurveClass( ImportType );
+			UClass* CurveClass = ExistingCurve ? ExistingCurve->GetClass() : GetCurveClass(ImportType);
 
 			// Create/reset curve
 			UCurveBase* NewCurve = NewObject<UCurveBase>(InParent, CurveClass, InName, Flags);
@@ -263,11 +266,11 @@ UObject* UCSVImportFactory::FactoryCreateText(UClass* InClass, UObject* InParent
 			NewAsset = NewCurve;
 		}
 		
-		if(Problems.Num() > 0)
+		if (Problems.Num() > 0)
 		{
 			FString AllProblems;
 
-			for(int32 ProbIdx=0; ProbIdx<Problems.Num(); ProbIdx++)
+			for (int32 ProbIdx=0; ProbIdx<Problems.Num(); ProbIdx++)
 			{
 				// Output problems to log
 				UE_LOG(LogCSVImportFactory, Log, TEXT("%d:%s"), ProbIdx, *Problems[ProbIdx]);
@@ -275,7 +278,7 @@ UObject* UCSVImportFactory::FactoryCreateText(UClass* InClass, UObject* InParent
 				AllProblems += TEXT("\n");
 			}
 
-			if(!IsAutomatedImport())
+			if (!IsAutomatedImport())
 			{
 				// Pop up any problems for user
 				FMessageDialog::Open(EAppMsgType::Ok, FText::FromString(AllProblems));
@@ -288,22 +291,23 @@ UObject* UCSVImportFactory::FactoryCreateText(UClass* InClass, UObject* InParent
 	return NewAsset;
 }
 
-bool UCSVImportFactory::ReimportCSV( UObject* Obj )
+EReimportResult::Type UCSVImportFactory::ReimportCSV(UObject* Obj)
 {
-	bool bHandled = false;
-	if(UCurveBase* Curve = Cast<UCurveBase>(Obj))
+	EReimportResult::Type Result = EReimportResult::Failed;
+
+	if (UCurveBase* Curve = Cast<UCurveBase>(Obj))
 	{
-		bHandled = Reimport(Curve, Curve->AssetImportData->GetFirstFilename());
+		Result = Reimport(Curve, Curve->AssetImportData->GetFirstFilename());
 	}
-	else if(UCurveTable* CurveTable = Cast<UCurveTable>(Obj))
+	else if (UCurveTable* CurveTable = Cast<UCurveTable>(Obj))
 	{
-		bHandled = Reimport(CurveTable, CurveTable->AssetImportData->GetFirstFilename());
+		Result = Reimport(CurveTable, CurveTable->AssetImportData->GetFirstFilename());
 	}
-	else if(UDataTable* DataTable = Cast<UDataTable>(Obj))
+	else if (UDataTable* DataTable = Cast<UDataTable>(Obj))
 	{
-		bHandled = Reimport(DataTable, DataTable->AssetImportData->GetFirstFilename());
+		Result = Reimport(DataTable, DataTable->AssetImportData->GetFirstFilename());
 	}
-	return bHandled;
+	return Result;
 }
 
 void UCSVImportFactory::ParseFromJson(TSharedRef<class FJsonObject> ImportSettingsJson)
@@ -311,22 +315,27 @@ void UCSVImportFactory::ParseFromJson(TSharedRef<class FJsonObject> ImportSettin
 	FJsonObjectConverter::JsonObjectToUStruct(ImportSettingsJson, FCSVImportSettings::StaticStruct(), &AutomatedImportSettings, 0, 0);
 }
 
-bool UCSVImportFactory::Reimport(UObject* Obj, const FString& Path)
+EReimportResult::Type UCSVImportFactory::Reimport(UObject* Obj, const FString& Path)
 {
-	if(Path.IsEmpty() == false)
+	if (Path.IsEmpty() == false)
 	{ 
 		FString FilePath = IFileManager::Get().ConvertToRelativePath(*Path);
 
 		FString Data;
-		if( FFileHelper::LoadFileToString( Data, *FilePath) )
+		if (FFileHelper::LoadFileToString(Data, *FilePath))
 		{
 			const TCHAR* Ptr = *Data;
 			CurrentFilename = FilePath; //not thread safe but seems to be how it is done..
-			auto Result = FactoryCreateText( Obj->GetClass(), Obj->GetOuter(), Obj->GetFName(), Obj->GetFlags(), NULL, *FPaths::GetExtension(FilePath), Ptr, Ptr+Data.Len(), NULL );
-			return true;
+			bool bWasCancelled = false;
+			UObject* Result = FactoryCreateText(Obj->GetClass(), Obj->GetOuter(), Obj->GetFName(), Obj->GetFlags(), nullptr, *FPaths::GetExtension(FilePath), Ptr, Ptr+Data.Len(), nullptr, bWasCancelled);
+			if (bWasCancelled)
+			{
+				return EReimportResult::Cancelled;
+			}
+			return Result ? EReimportResult::Succeeded : EReimportResult::Failed;
 		}
 	}
-	return false;
+	return EReimportResult::Failed;
 }
 
 TArray<FString> UCSVImportFactory::DoImportDataTable(UDataTable* TargetDataTable, const FString& DataToImport)
@@ -375,15 +384,15 @@ UReimportDataTableFactory::UReimportDataTableFactory(const FObjectInitializer& O
 	Formats.Add(TEXT("json;JavaScript Object Notation"));
 }
 
-bool UReimportDataTableFactory::FactoryCanImport( const FString& Filename )
+bool UReimportDataTableFactory::FactoryCanImport(const FString& Filename)
 {
 	return true;
 }
 
-bool UReimportDataTableFactory::CanReimport( UObject* Obj, TArray<FString>& OutFilenames )
+bool UReimportDataTableFactory::CanReimport(UObject* Obj, TArray<FString>& OutFilenames)
 {	
 	UDataTable* DataTable = Cast<UDataTable>(Obj);
-	if(DataTable)
+	if (DataTable)
 	{
 		DataTable->AssetImportData->ExtractFilenames(OutFilenames);
 		return true;
@@ -391,19 +400,19 @@ bool UReimportDataTableFactory::CanReimport( UObject* Obj, TArray<FString>& OutF
 	return false;
 }
 
-void UReimportDataTableFactory::SetReimportPaths( UObject* Obj, const TArray<FString>& NewReimportPaths )
+void UReimportDataTableFactory::SetReimportPaths(UObject* Obj, const TArray<FString>& NewReimportPaths)
 {	
 	UDataTable* DataTable = Cast<UDataTable>(Obj);
-	if(DataTable && ensure(NewReimportPaths.Num() == 1))
+	if (DataTable && ensure(NewReimportPaths.Num() == 1))
 	{
 		DataTable->AssetImportData->UpdateFilenameOnly(NewReimportPaths[0]);
 	}
 }
 
-EReimportResult::Type UReimportDataTableFactory::Reimport( UObject* Obj )
+EReimportResult::Type UReimportDataTableFactory::Reimport(UObject* Obj)
 {	
-	auto Result = EReimportResult::Failed;
-	if(auto DataTable = Cast<UDataTable>(Obj))
+	EReimportResult::Type Result = EReimportResult::Failed;
+	if (UDataTable* DataTable = Cast<UDataTable>(Obj))
 	{
 		FDataTableEditorUtils::BroadcastPreChange(DataTable, FDataTableEditorUtils::EDataTableChangeInfo::RowList);
 		Result = UCSVImportFactory::ReimportCSV(DataTable) ? EReimportResult::Succeeded : EReimportResult::Failed;
@@ -425,10 +434,10 @@ UReimportCurveTableFactory::UReimportCurveTableFactory(const FObjectInitializer&
 	Formats.Add(TEXT("json;JavaScript Object Notation"));
 }
 
-bool UReimportCurveTableFactory::CanReimport( UObject* Obj, TArray<FString>& OutFilenames )
+bool UReimportCurveTableFactory::CanReimport(UObject* Obj, TArray<FString>& OutFilenames)
 {	
 	UCurveTable* CurveTable = Cast<UCurveTable>(Obj);
-	if(CurveTable)
+	if (CurveTable)
 	{
 		CurveTable->AssetImportData->ExtractFilenames(OutFilenames);
 		return true;
@@ -436,18 +445,18 @@ bool UReimportCurveTableFactory::CanReimport( UObject* Obj, TArray<FString>& Out
 	return false;
 }
 
-void UReimportCurveTableFactory::SetReimportPaths( UObject* Obj, const TArray<FString>& NewReimportPaths )
+void UReimportCurveTableFactory::SetReimportPaths(UObject* Obj, const TArray<FString>& NewReimportPaths)
 {	
 	UCurveTable* CurveTable = Cast<UCurveTable>(Obj);
-	if(CurveTable && ensure(NewReimportPaths.Num() == 1))
+	if (CurveTable && ensure(NewReimportPaths.Num() == 1))
 	{
 		CurveTable->AssetImportData->UpdateFilenameOnly(NewReimportPaths[0]);
 	}
 }
 
-EReimportResult::Type UReimportCurveTableFactory::Reimport( UObject* Obj )
+EReimportResult::Type UReimportCurveTableFactory::Reimport(UObject* Obj)
 {	
-	if(Cast<UCurveTable>(Obj))
+	if (Cast<UCurveTable>(Obj))
 	{
 		return UCSVImportFactory::ReimportCSV(Obj) ? EReimportResult::Succeeded : EReimportResult::Failed;
 	}
@@ -467,10 +476,10 @@ UReimportCurveFactory::UReimportCurveFactory(const FObjectInitializer& ObjectIni
 	SupportedClass = UCurveBase::StaticClass();
 }
 
-bool UReimportCurveFactory::CanReimport( UObject* Obj, TArray<FString>& OutFilenames )
+bool UReimportCurveFactory::CanReimport(UObject* Obj, TArray<FString>& OutFilenames)
 {	
 	UCurveBase* CurveBase = Cast<UCurveBase>(Obj);
-	if(CurveBase)
+	if (CurveBase)
 	{
 		CurveBase->AssetImportData->ExtractFilenames(OutFilenames);
 		return true;
@@ -478,18 +487,18 @@ bool UReimportCurveFactory::CanReimport( UObject* Obj, TArray<FString>& OutFilen
 	return false;
 }
 
-void UReimportCurveFactory::SetReimportPaths( UObject* Obj, const TArray<FString>& NewReimportPaths )
+void UReimportCurveFactory::SetReimportPaths(UObject* Obj, const TArray<FString>& NewReimportPaths)
 {	
 	UCurveBase* CurveBase = Cast<UCurveBase>(Obj);
-	if(CurveBase && ensure(NewReimportPaths.Num() == 1))
+	if (CurveBase && ensure(NewReimportPaths.Num() == 1))
 	{
 		CurveBase->AssetImportData->UpdateFilenameOnly(NewReimportPaths[0]);
 	}
 }
 
-EReimportResult::Type UReimportCurveFactory::Reimport( UObject* Obj )
+EReimportResult::Type UReimportCurveFactory::Reimport(UObject* Obj)
 {	
-	if(Cast<UCurveBase>(Obj))
+	if (Cast<UCurveBase>(Obj))
 	{
 		return UCSVImportFactory::ReimportCSV(Obj) ? EReimportResult::Succeeded : EReimportResult::Failed;
 	}
diff --git a/Engine/Source/Editor/UnrealEd/Private/Factories/Factory.cpp b/Engine/Source/Editor/UnrealEd/Private/Factories/Factory.cpp
--- a/Engine/Source/Editor/UnrealEd/Private/Factories/Factory.cpp
+++ b/Engine/Source/Editor/UnrealEd/Private/Factories/Factory.cpp
@@ -80,7 +80,7 @@ UObject* UFactory::FactoryCreateFile(UClass* InClass, UObject* InParent, FName I
 		ParseParms(Parms);
 		const TCHAR* Ptr = *Data;
 			
-		return FactoryCreateText(InClass, InParent, InName, Flags, nullptr, *FileExtension, Ptr, Ptr + Data.Len(), Warn);
+		return FactoryCreateText(InClass, InParent, InName, Flags, nullptr, *FileExtension, Ptr, Ptr + Data.Len(), Warn, bOutOperationCanceled);
 	}
 
 	// load as binary
diff --git a/Engine/Source/Editor/UnrealEd/Private/Kismet2/BlueprintEditorUtils.cpp b/Engine/Source/Editor/UnrealEd/Private/Kismet2/BlueprintEditorUtils.cpp
--- a/Engine/Source/Editor/UnrealEd/Private/Kismet2/BlueprintEditorUtils.cpp
+++ b/Engine/Source/Editor/UnrealEd/Private/Kismet2/BlueprintEditorUtils.cpp
@@ -9088,10 +9088,10 @@ void FBlueprintEditorUtils::HandleDisableEditableWhenInherited(UObject* Modified
 			UBlueprintGeneratedClass* BPGC = Cast<UBlueprintGeneratedClass>(ArchetypeInstance->GetOuter());
 			if (BPGC)
 			{
-				UInheritableComponentHandler* ICH = BPGC->GetInheritableComponentHandler(false);
-				check(ICH);
-
-				ICH->RemoveOverridenComponentTemplate(ICH->FindKey(CastChecked<UActorComponent>(ArchetypeInstance)));
+				if (UInheritableComponentHandler* ICH = BPGC->GetInheritableComponentHandler(false))
+				{
+					ICH->RemoveOverridenComponentTemplate(ICH->FindKey(CastChecked<UActorComponent>(ArchetypeInstance)));
+				}
 			}
 		}
 	}
diff --git a/Engine/Source/Editor/UnrealEd/Private/Kismet2/CompilerResultsLog.cpp b/Engine/Source/Editor/UnrealEd/Private/Kismet2/CompilerResultsLog.cpp
--- a/Engine/Source/Editor/UnrealEd/Private/Kismet2/CompilerResultsLog.cpp
+++ b/Engine/Source/Editor/UnrealEd/Private/Kismet2/CompilerResultsLog.cpp
@@ -13,6 +13,7 @@
 #include "EngineLogs.h"
 #include "Engine/Blueprint.h"
 #include "IMessageLogListing.h"
+#include "Settings/EditorProjectSettings.h"
 
 #if WITH_EDITOR
 
@@ -192,6 +193,33 @@ void FCompilerResultsLog::InternalLogEvent(const FCompilerEvent& InEvent, int32
 	}
 }
 
+bool FCompilerResultsLog::IsMessageEnabled(FName ID)
+{
+	if(ID == NAME_None)
+	{
+		return true;
+	}
+
+	const UBlueprintEditorProjectSettings* EditorProjectSettings = GetDefault<UBlueprintEditorProjectSettings>();
+	if(IsRunningCommandlet())
+	{
+		if(	EditorProjectSettings->DisabledCompilerMessagesHeadless.Contains(ID) ||
+			EditorProjectSettings->DisabledCompilerMessages.Contains(ID))
+		{
+			return false;
+		}
+	}
+	else
+	{
+		if( EditorProjectSettings->DisabledCompilerMessages.Contains(ID) )
+		{
+			return false;
+		}
+	}
+
+	return true;
+}
+
 void FCompilerResultsLog::InternalLogSummary()
 {
 	if(CurrentEventScope.IsValid())
@@ -236,6 +264,39 @@ void FCompilerResultsLog::InternalLogSummary()
 	}
 }
 
+bool FCompilerResultsLog::CommitPotentialMessages(UEdGraphNode* Source)
+{
+	TArray<TSharedRef<FTokenizedMessage>> FoundMessages;
+	if (PotentialMessages.RemoveAndCopyValue(Source, FoundMessages))
+	{
+		for (const TSharedRef<FTokenizedMessage>& Message : FoundMessages)
+		{
+			switch (Message->GetSeverity())
+			{
+			case EMessageSeverity::Error:
+				++NumErrors;
+				break;
+
+			case EMessageSeverity::Warning:
+				++NumWarnings;
+				break;
+
+			default:
+				break;
+			}
+
+			// build nodes list:
+			TArray<UEdGraphNode*> Nodes;
+			GetNodesFromTokens(Message->GetMessageTokens(), Nodes);
+
+			Nodes.Add(Source);
+			InternalLogMessage(Message->GetIdentifier(), Message, Nodes);
+		}
+		return true;
+	}
+	return false;
+}
+
 /** Update the source backtrack map to note that NewObject was most closely generated/caused by the SourceObject */
 void FCompilerResultsLog::NotifyIntermediateObjectCreation(UObject* NewObject, UObject* SourceObject)
 {
@@ -258,11 +319,38 @@ UObject const* FCompilerResultsLog::FindSourceObject(UObject const* PossiblyDupl
 	return SourceBacktrackMap.FindSourceObject(PossiblyDuplicatedObject);
 }
 
+UObject* FCompilerResultsLog::FindSourceMacroInstance(const UEdGraphNode* IntermediateNode) const
+{
+	UObject* Result = nullptr;
+	const UEdGraphNode* Iter = IntermediateNode;
+	while (const TWeakObjectPtr<UEdGraphNode>* SourceInstanceNode = FullMacroBacktrackMap.Find(Iter))
+	{
+		Iter = SourceInstanceNode->Get();
+	}
+
+	Result = const_cast<UEdGraphNode*>(Iter);
+
+	if(Result)
+	{
+		Result = const_cast<FCompilerResultsLog*>(this)->FindSourceObject(Result);
+	}
+
+	return Result;
+}
+
 void FCompilerResultsLog::NotifyIntermediateTunnelNode(const UEdGraphNode* Node, const UEdGraphNode* OuterTunnelInstance)
 {
 	IntermediateTunnelNodeToTunnelInstanceMap.Add(Node, OuterTunnelInstance);
 }
 
+void FCompilerResultsLog::NotifyIntermediateMacroNode(UEdGraphNode* SourceNode, const UEdGraphNode* IntermediateNode)
+{
+	if(ensure(SourceNode != IntermediateNode))
+	{
+		FullMacroBacktrackMap.Add(IntermediateNode, SourceNode);
+	}
+}
+
 const UEdGraphNode* FCompilerResultsLog::GetIntermediateTunnelInstance(const UEdGraphNode* IntermediateNode) const
 {
 	TWeakObjectPtr<const UEdGraphNode> Result;
@@ -345,11 +433,12 @@ const UEdGraphPin* FCompilerResultsLog::FindSourcePin(const UEdGraphPin* Possibl
 	return SourceBacktrackMap.FindSourcePin(PossiblyDuplicatedPin);
 }
 
-void FCompilerResultsLog::InternalLogMessage(const TSharedRef<FTokenizedMessage>& Message, UEdGraphNode* SourceNode)
+void FCompilerResultsLog::InternalLogMessage(FName MessageID, const TSharedRef<FTokenizedMessage>& Message, const TArray<UEdGraphNode*>& SourceNodes)
 {
 	const EMessageSeverity::Type Severity = Message->GetSeverity();
 	Messages.Add(Message);
-	AnnotateNode(SourceNode, Message);
+	AnnotateNode(SourceNodes, Message);
+	Message->SetIdentifier(MessageID);
 
 	if (!bSilentMode && (!bLogInfoOnly || (Severity == EMessageSeverity::Info)))
 	{
@@ -376,12 +465,10 @@ void FCompilerResultsLog::InternalLogMessage(const TSharedRef<FTokenizedMessage>
 	}
 }
 
-void FCompilerResultsLog::AnnotateNode(class UEdGraphNode* Node, TSharedRef<FTokenizedMessage> LogLine)
+void FCompilerResultsLog::AnnotateNode(const TArray<UEdGraphNode*>& Nodes, TSharedRef<FTokenizedMessage> LogLine)
 {
-	if (Node != nullptr)
+	for(UEdGraphNode* Node : Nodes)
 	{
-		LogLine->SetMessageLink(FUObjectToken::Create(Node));
-
 		if (bAnnotateMentionedNodes)
 		{
 			// Determine if this message is the first or more important than the previous one (only showing one error/warning per node for now)
@@ -532,6 +619,43 @@ void FCompilerResultsLog::GetGlobalModuleCompilerDump(const FString& LogDump, EC
 
 	MessageLog.AddMessages(ParseCompilerLogDump(LogDump));
 }
+		
+void FCompilerResultsLog::GetNodesFromTokens(const TArray<TSharedRef<IMessageToken> >& MessageTokens, TArray<UEdGraphNode*>& OutOwnerNodes)
+{
+	for (TSharedRef<IMessageToken> const& Token : MessageTokens)
+	{
+		if (Token->GetType() == EMessageToken::Object)
+		{
+			FWeakObjectPtr ObjectPtr = ((FUObjectToken&)Token.Get()).GetObject();
+			if (!ObjectPtr.IsValid())
+			{
+				continue;
+			}
+			UObject* ObjectArgument = ObjectPtr.Get();
+			if(UEdGraphNode* Node = Cast<UEdGraphNode>(ObjectArgument))
+			{
+				OutOwnerNodes.Add(Node);
+			}
+		}
+		else if (Token->GetType() == EMessageToken::EdGraph)
+		{
+			FEdGraphToken& AsEdGraphToken = ((FEdGraphToken&)Token.Get());
+			const UEdGraphPin* PinBeingReferenced = AsEdGraphToken.GetPin();
+			UEdGraphNode* OwnerNode = Cast<UEdGraphNode>((UObject*)AsEdGraphToken.GetGraphObject());
+			if (OwnerNode == nullptr)
+			{
+				if (PinBeingReferenced)
+				{
+					OutOwnerNodes.Add(Cast<UEdGraphNode>(PinBeingReferenced->GetOwningNodeUnchecked()));
+				}
+			}
+			else
+			{
+				OutOwnerNodes.Add(OwnerNode);
+			}
+		}
+	}
+}
 
 void FCompilerResultsLog::Append(FCompilerResultsLog const& Other)
 {
@@ -558,37 +682,9 @@ void FCompilerResultsLog::Append(FCompilerResultsLog const& Other)
 		}
 		Messages.Add(Message);
 		
-		UEdGraphNode* OwnerNode = nullptr;
-		for (TSharedRef<IMessageToken> const& Token : Message->GetMessageTokens())
-		{
-			if (Token->GetType() == EMessageToken::Object)
-			{
-				FWeakObjectPtr ObjectPtr = ((FUObjectToken&)Token.Get()).GetObject();
-				if (!ObjectPtr.IsValid())
-				{
-					continue;
-				}
-				UObject* ObjectArgument = ObjectPtr.Get();
-
-				OwnerNode = Cast<UEdGraphNode>(ObjectArgument);
-				break;
-			}
-			else if (Token->GetType() == EMessageToken::EdGraph)
-			{
-				FEdGraphToken& AsEdGraphToken = ((FEdGraphToken&)Token.Get());
-				const UEdGraphPin* PinBeingReferenced = AsEdGraphToken.GetPin();
-				OwnerNode = Cast<UEdGraphNode>((UObject*)AsEdGraphToken.GetGraphObject());
-				if (OwnerNode == nullptr)
-				{
-					if (PinBeingReferenced)
-					{
-						OwnerNode = PinBeingReferenced->GetOwningNodeUnchecked();
-					}
-				}
-				break;
-			}
-		}
-		AnnotateNode(OwnerNode, Message);
+		TArray<UEdGraphNode*> OwnerNodes;
+		GetNodesFromTokens(Message->GetMessageTokens(), OwnerNodes);
+		AnnotateNode(OwnerNodes, Message);
 	}
 }
 
diff --git a/Engine/Source/Editor/UnrealEd/Private/SCSVImportOptions.cpp b/Engine/Source/Editor/UnrealEd/Private/SCSVImportOptions.cpp
--- a/Engine/Source/Editor/UnrealEd/Private/SCSVImportOptions.cpp
+++ b/Engine/Source/Editor/UnrealEd/Private/SCSVImportOptions.cpp
@@ -25,7 +25,7 @@ void SCSVImportOptions::Construct(const FArguments& InArgs)
 
 	// Find table row struct info
 	UScriptStruct* TableRowStruct = FindObjectChecked<UScriptStruct>(ANY_PACKAGE, TEXT("TableRowBase"));
-	if (TableRowStruct != NULL)
+	if (TableRowStruct != nullptr)
 	{
 		// Make combo of table rowstruct options
 		for (TObjectIterator<UScriptStruct> It; It; ++It)
@@ -151,6 +151,7 @@ void SCSVImportOptions::Construct(const FArguments& InArgs)
 					SNew(SButton)
 					.Text(LOCTEXT("OK", "OK"))
 					.OnClicked( this, &SCSVImportOptions::OnImport )
+					.IsEnabled( this, &SCSVImportOptions::CanImport )
 				]
 				+SHorizontalBox::Slot()
 				.AutoWidth()
@@ -179,7 +180,7 @@ void SCSVImportOptions::Construct(const FArguments& InArgs)
 	/** If we should import */
 bool SCSVImportOptions::ShouldImport()
 {
-	return ((SelectedStruct != NULL) || GetSelectedImportType() != ECSVImportType::ECSV_DataTable) && bImport;
+	return ((SelectedStruct != nullptr) || GetSelectedImportType() != ECSVImportType::ECSV_DataTable) && bImport;
 }
 
 /** Get the row struct we selected */
@@ -244,7 +245,7 @@ TSharedRef<SWidget> SCSVImportOptions::MakeImportTypeItemWidget(TSharedPtr<ECSVI
 /** Called to create a widget for each struct */
 TSharedRef<SWidget> SCSVImportOptions::MakeRowStructItemWidget(UScriptStruct* Struct)
 {
-	check( Struct != NULL );
+	check(Struct != nullptr);
 	return	SNew(STextBlock)
 			.Text(FText::FromString(Struct->GetName()));
 }
@@ -282,23 +283,44 @@ FReply SCSVImportOptions::OnImport()
 {
 	SelectedStruct = RowStructCombo->GetSelectedItem();
 	SelectedImportType = *ImportTypeCombo->GetSelectedItem();
-	if(CurveInterpCombo->GetSelectedItem().IsValid())
+	if (CurveInterpCombo->GetSelectedItem().IsValid())
 	{
 		SelectedCurveInterpMode = *CurveInterpCombo->GetSelectedItem();
 	}
 	bImport = true;
-	if ( WidgetWindow.IsValid() )
+	if (WidgetWindow.IsValid())
 	{
 		WidgetWindow.Pin()->RequestDestroyWindow();
 	}
 	return FReply::Handled();
 }
 
+bool SCSVImportOptions::CanImport() const
+{
+	ECSVImportType ImportType = *ImportTypeCombo->GetSelectedItem();
+
+	switch (ImportType)
+	{
+	case ECSVImportType::ECSV_DataTable:
+		return RowStructCombo->GetSelectedItem() != nullptr;
+		break;
+	case ECSVImportType::ECSV_CurveTable:
+		return CurveInterpCombo->GetSelectedItem().IsValid();
+		break;
+	case ECSVImportType::ECSV_CurveFloat:
+	case ECSVImportType::ECSV_CurveVector:
+	case ECSVImportType::ECSV_CurveLinearColor:
+		return true;
+	}
+	
+	return false;
+}
+
 /** Called when 'Cancel' button is pressed */
 FReply SCSVImportOptions::OnCancel()
 {
 	bImport = false;
-	if ( WidgetWindow.IsValid() )
+	if (WidgetWindow.IsValid())
 	{
 		WidgetWindow.Pin()->RequestDestroyWindow();
 	}
diff --git a/Engine/Source/Editor/UnrealEd/Private/UserDefinedStructEditorData.cpp b/Engine/Source/Editor/UnrealEd/Private/UserDefinedStructEditorData.cpp
--- a/Engine/Source/Editor/UnrealEd/Private/UserDefinedStructEditorData.cpp
+++ b/Engine/Source/Editor/UnrealEd/Private/UserDefinedStructEditorData.cpp
@@ -215,6 +215,9 @@ void UUserDefinedStructEditorData::RecreateDefaultInstance(FString* OutLog)
 			}
 		}
 	}
+
+	// Make sure StructFlags are in sync with the new default instance:
+	ScriptStruct->UpdateStructFlags();
 }
 
 void UUserDefinedStructEditorData::CleanDefaultInstance()
diff --git a/Engine/Source/Editor/UnrealEd/Public/EdGraphToken.h b/Engine/Source/Editor/UnrealEd/Public/EdGraphToken.h
--- a/Engine/Source/Editor/UnrealEd/Public/EdGraphToken.h
+++ b/Engine/Source/Editor/UnrealEd/Public/EdGraphToken.h
@@ -16,9 +16,9 @@ class FEdGraphToken : public IMessageToken
 {
 public:
 	/** Factory method, tokens can only be constructed as shared refs */
-	UNREALED_API static TSharedRef<IMessageToken> Create(const UObject* InObject, const FCompilerResultsLog* Log, UEdGraphNode*& OutSourceNode);
-	UNREALED_API static TSharedRef<IMessageToken> Create(const UEdGraphPin* InPin, const FCompilerResultsLog* Log, UEdGraphNode*& OutSourceNode);
-	UNREALED_API static TSharedRef<IMessageToken> Create(const TCHAR* String, const FCompilerResultsLog* Log, UEdGraphNode*& OutSourceNode);
+	UNREALED_API static void Create(const UObject* InObject, FCompilerResultsLog* Log, FTokenizedMessage &OutMessage, TArray<UEdGraphNode*>& OutSourceNodes);
+	UNREALED_API static void Create(const UEdGraphPin* InPin, FCompilerResultsLog* Log, FTokenizedMessage &OutMessage, TArray<UEdGraphNode*>& OutSourceNodes);
+	UNREALED_API static void Create(const TCHAR* String, FCompilerResultsLog* Log, FTokenizedMessage &OutMessage, TArray<UEdGraphNode*>& OutSourceNodes);
 
 	/** Begin IMessageToken interface */
 	virtual EMessageToken::Type GetType() const override
@@ -32,6 +32,8 @@ public:
 private:
 	/** Private constructor */
 	FEdGraphToken( const UObject* InObject, const UEdGraphPin* InPin );
+	/** Helper to facilitate code reuse between ::Create overloads */
+	static void CreateInternal(const UObject* InObject, FCompilerResultsLog* Log, FTokenizedMessage &OutMessage, TArray<UEdGraphNode*>& OutSourceNodes, const UEdGraphPin* Pin);
 
 	/** An object being referenced by this token, if any */
 	FWeakObjectPtr ObjectBeingReferenced;
diff --git a/Engine/Source/Editor/UnrealEd/Public/Kismet2/CompilerResultsLog.h b/Engine/Source/Editor/UnrealEd/Public/Kismet2/CompilerResultsLog.h
--- a/Engine/Source/Editor/UnrealEd/Public/Kismet2/CompilerResultsLog.h
+++ b/Engine/Source/Editor/UnrealEd/Public/Kismet2/CompilerResultsLog.h
@@ -117,6 +117,9 @@ protected:
 	// Map to track intermediate tunnel nodes back to the intermediate expansion tunnel instance.
 	TMap<TWeakObjectPtr<const UEdGraphNode>, TWeakObjectPtr<const UEdGraphNode>> IntermediateTunnelNodeToTunnelInstanceMap;
 
+	// Map to track intermediate nodes back to the source macro instance nodes
+	TMap< TWeakObjectPtr<const UEdGraphNode>, TWeakObjectPtr<UEdGraphNode> > FullMacroBacktrackMap;
+
 public:
 	FCompilerResultsLog(bool bIsCompatibleWithEvents = true);
 	virtual ~FCompilerResultsLog();
@@ -145,7 +148,7 @@ public:
 	{
 		++NumErrors;
 		TSharedRef<FTokenizedMessage> Line = FTokenizedMessage::Create(EMessageSeverity::Error);
-		InternalLogMessage(Format, Line, args...);
+		InternalLogMessage(NAME_None, Format, Line, args...);
 		return Line;
 	}
 
@@ -158,10 +161,28 @@ public:
 	{
 		++NumWarnings;
 		TSharedRef<FTokenizedMessage> Line = FTokenizedMessage::Create(EMessageSeverity::Warning);
-		InternalLogMessage(Format, Line, args...);
+		InternalLogMessage(NAME_None, Format, Line, args...);
 		return Line;
 	}
 
+	/**
+	 * Write a warning in to the compiler log.
+	 * Note: @@ will be replaced by node or pin links for nodes/pins passed via varargs 
+	 */
+	template<typename... Args>
+	void Warning(FName ID, const TCHAR* Format, Args... args)
+	{
+		if(!IsMessageEnabled(ID))
+		{
+			return;
+		}
+
+		++NumWarnings;
+		TSharedRef<FTokenizedMessage> Line = FTokenizedMessage::Create(EMessageSeverity::Warning);
+		InternalLogMessage(ID, Format, Line, args...);
+		return;
+	}
+
 	/**
 	 * Write a note in to the compiler log.
 	 * Note: @@ will be replaced by node or pin links for nodes/pins passed via varargs 
@@ -170,7 +191,7 @@ public:
 	TSharedRef<FTokenizedMessage> Note(const TCHAR* Format, Args... args)
 	{
 		TSharedRef<FTokenizedMessage> Line = FTokenizedMessage::Create(EMessageSeverity::Info);
-		InternalLogMessage(Format, Line, args...);
+		InternalLogMessage(NAME_None, Format, Line, args...);
 		return Line;
 	}
 
@@ -182,8 +203,8 @@ public:
 	TSharedRef<FTokenizedMessage> StorePotentialError(const UEdGraphNode* Source, const TCHAR* Format, Args... args)
 	{
 		TSharedRef<FTokenizedMessage> Line = FTokenizedMessage::Create(EMessageSeverity::Error);
-		UEdGraphNode* SourceNode = nullptr;
-		Tokenize(Format, *Line, SourceNode, args...);
+		TArray<UEdGraphNode*> SourceNodes;
+		Tokenize(Format, *Line, SourceNodes, args...);
 		PotentialMessages.FindOrAdd(Source).Add(Line);
 		return Line;
 	}
@@ -196,8 +217,8 @@ public:
 	TSharedRef<FTokenizedMessage> StorePotentialWarning(const UEdGraphNode* Source, const TCHAR* Format, Args... args)
 	{
 		TSharedRef<FTokenizedMessage> Line = FTokenizedMessage::Create(EMessageSeverity::Warning);
-		UEdGraphNode* SourceNode = nullptr;
-		Tokenize(Format, *Line, SourceNode, args...);
+		TArray<UEdGraphNode*> SourceNodes;
+		Tokenize(Format, *Line, SourceNodes, args...);
 		PotentialMessages.FindOrAdd(Source).Add(Line);
 		return Line;
 	}
@@ -210,8 +231,8 @@ public:
 	TSharedRef<FTokenizedMessage> StorePotentialNote(const UEdGraphNode* Source, const TCHAR* Format, Args... args)
 	{
 		TSharedRef<FTokenizedMessage> Line = FTokenizedMessage::Create(EMessageSeverity::Info);
-		UEdGraphNode* SourceNode = nullptr;
-		Tokenize(Format, *Line, SourceNode, args...);
+		TArray<UEdGraphNode*> SourceNodes;
+		Tokenize(Format, *Line, SourceNodes, args...);
 		PotentialMessages.FindOrAdd(Source).Add(Line);
 		return Line;
 	}
@@ -219,33 +240,7 @@ public:
 	/**
 	 * Commit all stored potential messages for a given node. Returns true if any messages were written.
 	 */
-	bool CommitPotentialMessages(UEdGraphNode* Source)
-	{
-		TArray<TSharedRef<FTokenizedMessage>> FoundMessages;
-		if (PotentialMessages.RemoveAndCopyValue(Source, FoundMessages))
-		{
-			for (const TSharedRef<FTokenizedMessage>& Message : FoundMessages)
-			{
-				switch (Message->GetSeverity())
-				{
-				case EMessageSeverity::Error:
-					++NumErrors;
-					break;
-
-				case EMessageSeverity::Warning:
-					++NumWarnings;
-					break;
-
-				default:
-					break;
-				}
-
-				InternalLogMessage(Message, Source);
-			}
-			return true;
-		}
-		return false;
-	}
+	bool CommitPotentialMessages(UEdGraphNode* Source);
 
 	/** Update the source backtrack map to note that NewObject was most closely generated/caused by the SourceObject */
 	void NotifyIntermediateObjectCreation(UObject* NewObject, UObject* SourceObject);
@@ -254,9 +249,13 @@ public:
 	/** Update the expansion map to note that Node was expanded from OuterTunnelInstance, both the node and tunnel instance should be intermediate nodes */
 	void NotifyIntermediateTunnelNode(const UEdGraphNode* Node, const UEdGraphNode* OuterTunnelInstance);
 
+	/** Update the map that tracks nodes created by macro instance nodes */
+	void NotifyIntermediateMacroNode(UEdGraphNode* SourceNode, const UEdGraphNode* IntermediateNode);
+
 	/** Returns the true source object for the passed in object */
 	UObject* FindSourceObject(UObject* PossiblyDuplicatedObject);
 	UObject const* FindSourceObject(UObject const* PossiblyDuplicatedObject) const;
+	UObject* FindSourceMacroInstance(const UEdGraphNode* IntermediateNode) const;
 
 	/** Returns the intermediate tunnel instance that generated the node */
 	const UEdGraphNode* GetIntermediateTunnelInstance(const UEdGraphNode* IntermediateNode) const;
@@ -301,21 +300,22 @@ protected:
 	/** Helper method to add a child event to the given parent event scope */
 	void AddChildEvent(TSharedPtr<FCompilerEvent>& ParentEventScope, TSharedRef<FCompilerEvent>& ChildEventScope);
 
-	void InternalLogMessage(const TSharedRef<FTokenizedMessage>& Message, UEdGraphNode* SourceNode );
+	void InternalLogMessage(FName MessageID, const TSharedRef<FTokenizedMessage>& Message, const TArray<UEdGraphNode*>& SourceNodes );
 
-	void Tokenize(const TCHAR* Text, FTokenizedMessage &OutMessage, UEdGraphNode*& OutSourceNode)
+	void Tokenize(const TCHAR* Text, FTokenizedMessage &OutMessage, TArray<UEdGraphNode*>& OutSourceNode)
 	{
 		OutMessage.AddToken(FTextToken::Create(FText::FromString(Text)));
 	}
 
 	template<typename T, typename... Args>
-	void Tokenize(const TCHAR* Format, FTokenizedMessage &OutMessage, UEdGraphNode*& OutSourceNode, T First, Args... Rest)
+	void Tokenize(const TCHAR* Format, FTokenizedMessage &OutMessage, TArray<UEdGraphNode*>& OutSourceNode, T First, Args... Rest)
 	{
 		// read to next "@@":
 		if (const TCHAR* DelimiterStr = FCString::Strstr(Format, TEXT("@@")))
 		{
 			OutMessage.AddToken(FTextToken::Create(FText::FromString(FString(DelimiterStr - Format, Format))));
-			OutMessage.AddToken(FEdGraphToken::Create(First, this, OutSourceNode));
+			FEdGraphToken::Create(First, this, OutMessage, OutSourceNode);
+
 			const TCHAR* NextChunk = DelimiterStr + FCString::Strlen(TEXT("@@"));
 			if (*NextChunk)
 			{
@@ -329,18 +329,18 @@ protected:
 	}
 
 	template<typename... Args>
-	void InternalLogMessage(const TCHAR* Format, const TSharedRef<FTokenizedMessage>& Message, Args... args)
+	void InternalLogMessage(FName MessageID, const TCHAR* Format, const TSharedRef<FTokenizedMessage>& Message, Args... args)
 	{
 		// Convention for SourceNode established by the original version of the compiler results log
 		// was to annotate the error on the first node we can find. I am preserving that behavior
 		// for this type safe, variadic version:
-		UEdGraphNode* SourceNode = nullptr;
-		Tokenize(Format, *Message, SourceNode, args...);
-		InternalLogMessage(Message, SourceNode);
+		TArray<UEdGraphNode*> SourceNodes;
+		Tokenize(Format, *Message, SourceNodes, args...);
+		InternalLogMessage(MessageID, Message, SourceNodes);
 	}
 
-	/** */
-	void AnnotateNode(class UEdGraphNode* Node, TSharedRef<FTokenizedMessage> LogLine);
+	/** Links the UEdGraphNode with the LogLine: */
+	void AnnotateNode(const TArray<UEdGraphNode*>& Nodes, TSharedRef<FTokenizedMessage> LogLine);
 
 	/** Internal method to append the final compiler results summary to the MessageLog */
 	void InternalLogSummary();
@@ -348,6 +348,8 @@ protected:
 	/** Internal helper method to recursively append event details into the MessageLog */
 	void InternalLogEvent(const FCompilerEvent& InEvent, int32 InDepth = 0);
 
+	/** Returns true if the user has requested this compiler message be suppressed */
+	bool IsMessageEnabled(FName ID);
 private:
 
 	/** Map of stored potential messages indexed by a node. Can be committed to the results log by calling CommitPotentialMessages for that node. */
@@ -361,6 +363,9 @@ private:
 
 	/** Callback function for binding the global compiler dump to open the static compiler log */
 	static void GetGlobalModuleCompilerDump(const FString& LogDump, ECompilationResult::Type CompilationResult, bool bShowLog);
+	
+	/** Searches a token list for referenced UEdGraphNodes, used to update the nodes when a log is committed */
+	static void GetNodesFromTokens(const TArray<TSharedRef<IMessageToken> >& MessageTokens, TArray<UEdGraphNode*>& OutOwnerNodes);
 
 	/** The log's name, for easy re-use */
 	static const FName Name;
diff --git a/Engine/Source/Editor/UnrealEd/Public/SCSVImportOptions.h b/Engine/Source/Editor/UnrealEd/Public/SCSVImportOptions.h
--- a/Engine/Source/Editor/UnrealEd/Public/SCSVImportOptions.h
+++ b/Engine/Source/Editor/UnrealEd/Public/SCSVImportOptions.h
@@ -34,7 +34,7 @@ public:
 	SCSVImportOptions()
 		: bImport(false)
 		, SelectedImportType(ECSVImportType::ECSV_DataTable)
-		, SelectedStruct(NULL)
+		, SelectedStruct(nullptr)
 		{}
 
 	void Construct(const FArguments& InArgs);
@@ -73,6 +73,9 @@ public:
 	/** Called when 'OK' button is pressed */
 	FReply OnImport();
 
+	/** Do we have all of the data we need to import this asset? */
+	bool CanImport() const;
+
 	/** Called when 'Cancel' button is pressed */
 	FReply OnCancel();
 
@@ -84,7 +87,7 @@ public:
 
 private:
 	/** Whether we should go ahead with import */
-	bool										bImport;
+	uint8										bImport : 1;
 
 	/** Window that owns us */
 	TWeakPtr< SWindow >							WidgetWindow;
diff --git a/Engine/Source/Editor/UnrealEd/Public/Settings/EditorProjectSettings.h b/Engine/Source/Editor/UnrealEd/Public/Settings/EditorProjectSettings.h
--- a/Engine/Source/Editor/UnrealEd/Public/Settings/EditorProjectSettings.h
+++ b/Engine/Source/Editor/UnrealEd/Public/Settings/EditorProjectSettings.h
@@ -183,5 +183,21 @@ public:
 	/** If enabled, the editor will load packages to look for soft references to actors when deleting/renaming them. This can be slow in large projects so disable this to improve performance but increase the chance of breaking blueprints/sequences that use soft actor references */
 	UPROPERTY(EditAnywhere, config, Category=Actors)
 	uint32 bValidateUnloadedSoftActorReferences : 1;
+
+	/** 
+	 * List of compiler messages that have been suppressed during cook (or other commandlets) for 
+	 * the current project - useful for silencing warnings that were added to the engine after 
+	 * project inception and are going to be addressed as they are found by content authors
+	 */
+	UPROPERTY(EditAnywhere, config, Category= Blueprints, DisplayName = "Compiler Messages Disabled during Cook")
+	TArray<FName> DisabledCompilerMessagesHeadless;
+	
+	/** 
+	 * List of compiler messages that have been suppressed completely - message suppression is only 
+	 * advisable when using blueprints that you cannot update and are raising innocuous warnings. 
+	 * If useless messages are being raised prefer to contact support rather than disabling messages
+	 */
+	UPROPERTY(EditAnywhere, config, Category= Blueprints, DisplayName = "Compiler Messages Disabled Entirely")
+	TArray<FName> DisabledCompilerMessages;
 };
 
diff --git a/Engine/Source/Programs/UnrealHeaderTool/Private/HeaderParser.cpp b/Engine/Source/Programs/UnrealHeaderTool/Private/HeaderParser.cpp
--- a/Engine/Source/Programs/UnrealHeaderTool/Private/HeaderParser.cpp
+++ b/Engine/Source/Programs/UnrealHeaderTool/Private/HeaderParser.cpp
@@ -951,6 +951,61 @@ namespace
 			}
 			break;
 
+			case ECheckedMetadataSpecifier::CommutativeAssociativeBinaryOperator:
+			{
+				if (UFunction* Function = Cast<UFunction>(Field))
+				{
+					bool bGoodParams = (Function->NumParms == 3);
+					if (bGoodParams)
+					{
+						UProperty* FirstParam = nullptr;
+						UProperty* SecondParam = nullptr;
+						UProperty* ReturnValue = nullptr;
+
+						TFieldIterator<UProperty> It(Function);
+
+						auto GetNextParam = [&]()
+						{
+							if (It)
+							{
+								if (It->HasAnyPropertyFlags(CPF_ReturnParm))
+								{
+									ReturnValue = *It;
+								}
+								else
+								{
+									if (FirstParam == nullptr)
+									{
+										FirstParam = *It;
+									}
+									else if (SecondParam == nullptr)
+									{
+										SecondParam = *It;
+									}
+								}
+								++It;
+							}
+						};
+
+						GetNextParam();
+						GetNextParam();
+						GetNextParam();
+						ensure(!It);
+
+						if (ReturnValue == nullptr || SecondParam == nullptr || !SecondParam->SameType(FirstParam))
+						{
+							bGoodParams = false;
+						}
+					}
+
+					if (!bGoodParams)
+					{
+						UE_LOG_ERROR_UHT(TEXT("Commutative asssociative binary operators must have exactly 2 parameters of the same type and a return value."));
+					}
+				}
+			}
+			break;
+
 			case ECheckedMetadataSpecifier::ExpandEnumAsExecs:
 			{
 				if (UFunction* Function = Cast<UFunction>(Field))
@@ -2698,6 +2753,14 @@ void FHeaderParser::FixupDelegateProperties( FClasses& AllClasses, UStruct* Stru
 							for (TFieldIterator<UProperty> PropIt(SourceDelegateFunction); PropIt && (PropIt->PropertyFlags & CPF_Parm); ++PropIt)
 							{
 								UProperty* FuncParam = *PropIt;
+
+								if (!IsPropertySupportedByBlueprint(FuncParam, false))
+								{
+									FString ExtendedCPPType;
+									FString CPPType = FuncParam->GetCPPType(&ExtendedCPPType);
+									UE_LOG_ERROR_UHT(TEXT("Type '%s%s' is not supported by blueprint. %s.%s"), *CPPType, *ExtendedCPPType, *SourceDelegateFunction->GetName(), *FuncParam->GetName());
+								}
+
 								if(FuncParam->HasAllPropertyFlags(CPF_OutParm) && !FuncParam->HasAllPropertyFlags(CPF_ConstParm)  )
 								{
 									const bool bClassGeneratedFromBP = FClass::IsDynamic(Struct);
diff --git a/Engine/Source/Programs/UnrealHeaderTool/Private/ParserHelper.h b/Engine/Source/Programs/UnrealHeaderTool/Private/ParserHelper.h
--- a/Engine/Source/Programs/UnrealHeaderTool/Private/ParserHelper.h
+++ b/Engine/Source/Programs/UnrealHeaderTool/Private/ParserHelper.h
@@ -1069,7 +1069,7 @@ struct FFuncInfo
 		FString FunctionName = FunctionReference->GetName();
 		if( FunctionReference->HasAnyFunctionFlags( FUNC_Delegate ) )
 		{
-			FunctionName = FunctionName.LeftChop( FString( TEXT( "__DelegateSignature" ) ).Len() );
+			FunctionName = FunctionName.LeftChop( FString(HEADER_GENERATED_DELEGATE_SIGNATURE_SUFFIX).Len() );
 		}
 		UnMarshallAndCallName = FString(TEXT("exec")) + FunctionName;
 		
diff --git a/Engine/Source/Programs/UnrealHeaderTool/Private/Specifiers/CheckedMetadataSpecifiers.def b/Engine/Source/Programs/UnrealHeaderTool/Private/Specifiers/CheckedMetadataSpecifiers.def
--- a/Engine/Source/Programs/UnrealHeaderTool/Private/Specifiers/CheckedMetadataSpecifiers.def
+++ b/Engine/Source/Programs/UnrealHeaderTool/Private/Specifiers/CheckedMetadataSpecifiers.def
@@ -5,6 +5,7 @@
 CHECKED_METADATA_SPECIFIER(BlueprintProtected)
 CHECKED_METADATA_SPECIFIER(ClampMax)
 CHECKED_METADATA_SPECIFIER(ClampMin)
+CHECKED_METADATA_SPECIFIER(CommutativeAssociativeBinaryOperator)
 CHECKED_METADATA_SPECIFIER(DevelopmentStatus)
 CHECKED_METADATA_SPECIFIER(ExpandEnumAsExecs)
 CHECKED_METADATA_SPECIFIER(UIMax)
diff --git a/Engine/Source/Runtime/AIModule/Classes/AITypes.h b/Engine/Source/Runtime/AIModule/Classes/AITypes.h
--- a/Engine/Source/Runtime/AIModule/Classes/AITypes.h
+++ b/Engine/Source/Runtime/AIModule/Classes/AITypes.h
@@ -426,7 +426,7 @@ namespace FAIResources
 }
 
 
-USTRUCT()
+USTRUCT(BlueprintType)
 struct AIMODULE_API FAIRequestID
 {
 	GENERATED_USTRUCT_BODY()
diff --git a/Engine/Source/Runtime/AIModule/Classes/EnvironmentQuery/EnvQueryManager.h b/Engine/Source/Runtime/AIModule/Classes/EnvironmentQuery/EnvQueryManager.h
--- a/Engine/Source/Runtime/AIModule/Classes/EnvironmentQuery/EnvQueryManager.h
+++ b/Engine/Source/Runtime/AIModule/Classes/EnvironmentQuery/EnvQueryManager.h
@@ -161,6 +161,8 @@ class AIMODULE_API UEnvQueryManager : public UObject, public FTickableGameObject
 	// Any loaded instance will get marked as PendingKill
 	virtual void PostLoad() override;
 
+	virtual void PostInitProperties() override;
+
 	// We need to implement GetWorld() so that any EQS-related blueprints (such as blueprint contexts) can implement
 	// GetWorld() and so provide access to blueprint nodes using hidden WorldContextObject parameters.
 	virtual UWorld* GetWorld() const override;
@@ -319,4 +321,11 @@ private:
 	void CheckQueryCount() const;
 	void LogQueryCountWarning() const;
 #endif
+
+#if WITH_EDITOR
+	FDelegateHandle OnBlueprintCompiledHandle;
+
+	// used to reset LocalContextMap in case we had BP contexts stored
+	void OnBlueprintCompiled();
+#endif // WITH_EDITOR
 };
diff --git a/Engine/Source/Runtime/AIModule/Classes/EnvironmentQuery/Generators/EnvQueryGenerator_SimpleGrid.h b/Engine/Source/Runtime/AIModule/Classes/EnvironmentQuery/Generators/EnvQueryGenerator_SimpleGrid.h
--- a/Engine/Source/Runtime/AIModule/Classes/EnvironmentQuery/Generators/EnvQueryGenerator_SimpleGrid.h
+++ b/Engine/Source/Runtime/AIModule/Classes/EnvironmentQuery/Generators/EnvQueryGenerator_SimpleGrid.h
@@ -15,7 +15,7 @@
  */
 
 UCLASS(meta = (DisplayName = "Points: Grid"))
-class UEnvQueryGenerator_SimpleGrid : public UEnvQueryGenerator_ProjectedPoints
+class AIMODULE_API UEnvQueryGenerator_SimpleGrid : public UEnvQueryGenerator_ProjectedPoints
 {
 	GENERATED_UCLASS_BODY()
 
diff --git a/Engine/Source/Runtime/AIModule/Private/EnvironmentQuery/EnvQueryManager.cpp b/Engine/Source/Runtime/AIModule/Private/EnvironmentQuery/EnvQueryManager.cpp
--- a/Engine/Source/Runtime/AIModule/Private/EnvironmentQuery/EnvQueryManager.cpp
+++ b/Engine/Source/Runtime/AIModule/Private/EnvironmentQuery/EnvQueryManager.cpp
@@ -132,6 +132,17 @@ void UEnvQueryManager::PostLoad()
 	MarkPendingKill();
 }
 
+void UEnvQueryManager::PostInitProperties()
+{
+	Super::PostInitProperties();
+#if WITH_EDITOR
+	if (GEditor)
+	{
+		OnBlueprintCompiledHandle = GEditor->OnBlueprintCompiled().AddUObject(this, &UEnvQueryManager::OnBlueprintCompiled);
+	}
+#endif // WITH_EDITOR
+}
+
 UWorld* UEnvQueryManager::GetWorld() const
 {
 	return Cast<UWorld>(GetOuter());
@@ -140,6 +151,14 @@ UWorld* UEnvQueryManager::GetWorld() const
 void UEnvQueryManager::FinishDestroy()
 {
 	FCoreUObjectDelegates::PreLoadMap.RemoveAll(this);
+
+#if WITH_EDITOR
+	if (GEditor)
+	{
+		GEditor->OnBlueprintCompiled().Remove(OnBlueprintCompiledHandle);
+	}
+#endif // WITH_EDITOR
+
 	Super::FinishDestroy();
 }
 
@@ -948,6 +967,13 @@ TSharedPtr<FEnvQueryInstance> UEnvQueryManager::FindQueryInstance(const int32 Qu
 	return nullptr;
 }
 
+#if WITH_EDITOR
+void UEnvQueryManager::OnBlueprintCompiled()
+{
+	LocalContextMap.Reset();
+}
+#endif // WITH_EDITOR
+
 //----------------------------------------------------------------------//
 // Exec functions (i.e. console commands)
 //----------------------------------------------------------------------//
diff --git a/Engine/Source/Runtime/Core/Private/Logging/TokenizedMessage.cpp b/Engine/Source/Runtime/Core/Private/Logging/TokenizedMessage.cpp
--- a/Engine/Source/Runtime/Core/Private/Logging/TokenizedMessage.cpp
+++ b/Engine/Source/Runtime/Core/Private/Logging/TokenizedMessage.cpp
@@ -102,6 +102,16 @@ FName FTokenizedMessage::GetSeverityIconName(EMessageSeverity::Type InSeverity)
 	return SeverityIconName;
 }
 
+FName FTokenizedMessage::GetIdentifier() const
+{
+	return Identifier;
+}
+
+void FTokenizedMessage::SetIdentifier(FName InIdentifier)
+{
+	Identifier = InIdentifier;
+}
+
 TSharedRef<FTokenizedMessage> FTokenizedMessage::AddToken( const TSharedRef<IMessageToken>& InToken )
 {
 	MessageTokens.Add( InToken );
@@ -124,17 +134,6 @@ EMessageSeverity::Type FTokenizedMessage::GetSeverity() const
 	return Severity;
 }
 
-TSharedRef<FTokenizedMessage> FTokenizedMessage::SetMessageData( FTokenizedMiscData* InMessageData )
-{
-	MessageData = InMessageData;
-	return AsShared();
-}
-
-FTokenizedMiscData* FTokenizedMessage::GetMessageData() const
-{
-	return MessageData;
-}
-
 const TArray<TSharedRef<IMessageToken> >& FTokenizedMessage::GetMessageTokens() const
 {
 	return MessageTokens;
diff --git a/Engine/Source/Runtime/Core/Public/Logging/TokenizedMessage.h b/Engine/Source/Runtime/Core/Public/Logging/TokenizedMessage.h
--- a/Engine/Source/Runtime/Core/Public/Logging/TokenizedMessage.h
+++ b/Engine/Source/Runtime/Core/Public/Logging/TokenizedMessage.h
@@ -102,10 +102,6 @@ protected:
 	FText CachedText;
 };
 
-
-/** Dummy struct that other structs need to inherit from if they wish to assign to MessageData */
-struct FTokenizedMiscData {};
-
 /** This class represents a rich tokenized message, such as would be used for compiler output with 'hyperlinks' to source file locations */
 class FTokenizedMessage : public TSharedFromThis<FTokenizedMessage>
 {
@@ -150,21 +146,6 @@ public:
 	 */
 	CORE_API EMessageSeverity::Type GetSeverity() const;
 
-	/** 
-	 * Sets the custom message data of this message
-	 * 
-	 * @param	InMessageData	The custom message data to set
-	 * @returns this message, so calls can be chained together
-	 */
-	CORE_API TSharedRef<FTokenizedMessage> SetMessageData( FTokenizedMiscData* InMessageData );
-
-	/** 
-	 * Gets the custom message data of this message
-	 * 
-	 * @returns the custom message data contained in this message
-	 */
-	CORE_API FTokenizedMiscData* GetMessageData() const;
-
 	/**
 	 * Get the tokens in this message
 	 * 
@@ -203,12 +184,18 @@ public:
 	 * @returns the name of the icon for the severity
 	 */
 	CORE_API static FName GetSeverityIconName(EMessageSeverity::Type InSeverity);
+	
+	/** @returns Identifier for the message, if set, else NAME_None */
+	CORE_API FName GetIdentifier() const;
+	
+	/** Assigns Identifier for the message to the provided name */
+	CORE_API void SetIdentifier(FName InIdentifier);
 
 private:
 	/** Private constructor - we want to only create these structures as shared references via Create() */
 	FTokenizedMessage()
 		: Severity( EMessageSeverity::Info )
-		, MessageData( nullptr )
+		, Identifier(NAME_None)
 	{ }
 
 protected:
@@ -224,8 +211,8 @@ private:
 	/** The severity of this message */
 	EMessageSeverity::Type Severity;
 
-	/** A payload for this message */
-	FTokenizedMiscData* MessageData;	// @todo Don't use TSharedPtr as it'll get free'd mid-sort! (refcount == 0)
+	/** Identifier for the message */
+	FName Identifier;
 };
 
 /** Basic message token with a localized text payload */
diff --git a/Engine/Source/Runtime/CoreUObject/Private/UObject/Class.cpp b/Engine/Source/Runtime/CoreUObject/Private/UObject/Class.cpp
--- a/Engine/Source/Runtime/CoreUObject/Private/UObject/Class.cpp
+++ b/Engine/Source/Runtime/CoreUObject/Private/UObject/Class.cpp
@@ -442,7 +442,7 @@ void UField::SetMetaData(const TCHAR* Key, const TCHAR* InValue)
 	UPackage* Package = GetOutermost();
 	check(Package);
 
-	return Package->GetMetaData()->SetValue(this, Key, InValue);
+	Package->GetMetaData()->SetValue(this, Key, InValue);
 }
 
 void UField::SetMetaData(const FName& Key, const TCHAR* InValue)
@@ -450,7 +450,7 @@ void UField::SetMetaData(const FName& Key, const TCHAR* InValue)
 	UPackage* Package = GetOutermost();
 	check(Package);
 
-	return Package->GetMetaData()->SetValue(this, Key, InValue);
+	Package->GetMetaData()->SetValue(this, Key, InValue);
 }
 
 UClass* UField::GetClassMetaData(const TCHAR* Key) const
diff --git a/Engine/Source/Runtime/CoreUObject/Private/UObject/PropertyStruct.cpp b/Engine/Source/Runtime/CoreUObject/Private/UObject/PropertyStruct.cpp
--- a/Engine/Source/Runtime/CoreUObject/Private/UObject/PropertyStruct.cpp
+++ b/Engine/Source/Runtime/CoreUObject/Private/UObject/PropertyStruct.cpp
@@ -84,6 +84,19 @@ void UStructProperty::LinkInternal(FArchive& Ar)
 		// User Defined structs won't have UScriptStruct::ICppStructOps. Setting their flags here.
 		PropertyFlags |= CPF_HasGetValueTypeHash;
 	}
+
+	if (Struct->StructFlags & STRUCT_ZeroConstructor)
+	{
+		PropertyFlags |= CPF_ZeroConstructor;
+	}
+	if (Struct->StructFlags & STRUCT_IsPlainOldData)
+	{
+		PropertyFlags |= CPF_IsPlainOldData;
+	}
+	if (Struct->StructFlags & STRUCT_NoDestructor)
+	{
+		PropertyFlags |= CPF_NoDestructor;
+	}
 }
 
 bool UStructProperty::Identical( const void* A, const void* B, uint32 PortFlags ) const
diff --git a/Engine/Source/Runtime/Engine/Classes/EdGraph/EdGraphNode.h b/Engine/Source/Runtime/Engine/Classes/EdGraph/EdGraphNode.h
--- a/Engine/Source/Runtime/Engine/Classes/EdGraph/EdGraphNode.h
+++ b/Engine/Source/Runtime/Engine/Classes/EdGraph/EdGraphNode.h
@@ -788,6 +788,28 @@ public:
 	/** Sets a flag if the comment bubble needs to be made visible immediately */
 	void SetMakeCommentBubbleVisible(bool MakeVisible);
 
+	/** Execute a provided function once for each node that is directly connected to this node, will not include the node itself */
+	void ForEachNodeDirectlyConnected(TFunctionRef<void(UEdGraphNode*)> Func);
+	
+	/** 
+	 * Often we are only interested in a subset of our connections (e.g. only output pins, or only output pins except our exec pin)
+	 * This function provides the ability to execute a provided function once for each node that is directly connected to this node, 
+	 * but first filters out which of this node's pins to consider:
+	 */
+	void ForEachNodeDirectlyConnectedIf(TFunctionRef<bool(const UEdGraphPin* Pin)> Filter, TFunctionRef<void(UEdGraphNode*)> Func);
+
+	/** 
+	 * Execute a provided function once for each node that is directly connected to this node's input pins, will not include the node itself
+	 * Implementation provides an example usage of ForEachNodeDirectlyConnectedIf.
+	 */
+	void ForEachNodeDirectlyConnectedToInputs(TFunctionRef<void(UEdGraphNode*)> Func);
+
+	/** 
+	 * Execute a provided function once for each node that is directly connected to this node's output pins, will not include the node itself
+	 * Implementation provides an example usage of ForEachNodeDirectlyConnectedIf.
+	 */
+	void ForEachNodeDirectlyConnectedToOutputs(TFunctionRef<void(UEdGraphNode*)> Func);
+	
 protected:
 	/**
 	 * Finds the difference in properties of node instance
diff --git a/Engine/Source/Runtime/Engine/Classes/Engine/UserDefinedStruct.h b/Engine/Source/Runtime/Engine/Classes/Engine/UserDefinedStruct.h
--- a/Engine/Source/Runtime/Engine/Classes/Engine/UserDefinedStruct.h
+++ b/Engine/Source/Runtime/Engine/Classes/Engine/UserDefinedStruct.h
@@ -124,4 +124,6 @@ public:
 	/** returns references from default instance */
 	static void AddReferencedObjects(UObject* InThis, FReferenceCollector& Collector);
 
+	/** Inspects properties and default values, setting appropriate StructFlags */
+	void UpdateStructFlags();
 };
diff --git a/Engine/Source/Runtime/Engine/Classes/Kismet/GameplayStatics.h b/Engine/Source/Runtime/Engine/Classes/Kismet/GameplayStatics.h
--- a/Engine/Source/Runtime/Engine/Classes/Kismet/GameplayStatics.h
+++ b/Engine/Source/Runtime/Engine/Classes/Kismet/GameplayStatics.h
@@ -141,7 +141,7 @@ class ENGINE_API UGameplayStatics : public UBlueprintFunctionLibrary
 	* @param Player		The player controller of the player to get the ID of
 	* @return			The ID of the passed in player. -1 if there is no controller for the passed in player
 	*/
-	UFUNCTION(BlueprintCallable, Category="Game", meta=(UnsafeDuringActorConstruction="true"))
+	UFUNCTION(BlueprintPure, Category="Game", meta=(UnsafeDuringActorConstruction="true"))
 	static int32 GetPlayerControllerID(APlayerController* Player);
 
 	/** 
diff --git a/Engine/Source/Runtime/Engine/Private/EdGraph/EdGraphNode.cpp b/Engine/Source/Runtime/Engine/Private/EdGraph/EdGraphNode.cpp
--- a/Engine/Source/Runtime/Engine/Private/EdGraph/EdGraphNode.cpp
+++ b/Engine/Source/Runtime/Engine/Private/EdGraph/EdGraphNode.cpp
@@ -765,6 +765,84 @@ void UEdGraphNode::SetMakeCommentBubbleVisible(bool MakeVisible)
 	bCommentBubbleMakeVisible = MakeVisible;
 }
 
+void UEdGraphNode::ForEachNodeDirectlyConnected(TFunctionRef<void(UEdGraphNode*)> Func)
+{
+	TSet<UEdGraphNode*> DirectNeighbors;
+	for( UEdGraphPin* Pin : Pins )
+	{
+		if(Pin->LinkedTo.Num() > 0)
+		{
+			for(UEdGraphPin* Connection : Pin->LinkedTo)
+			{
+				// avoid including the current node in the case of a self connection:
+				if(Connection->GetOwningNode() != this)
+				{
+					DirectNeighbors.Add(Connection->GetOwningNode());
+				}
+			}
+		}
+	}
+
+	for(UEdGraphNode* Neighbor : DirectNeighbors)
+	{
+		Func(Neighbor);
+	}
+}
+
+void UEdGraphNode::ForEachNodeDirectlyConnectedIf(TFunctionRef<bool(const UEdGraphPin* Pin)> Filter, TFunctionRef<void(UEdGraphNode*)> Func)
+{
+	TSet<UEdGraphNode*> NeighborsAcceptedForConsideration;
+	for( UEdGraphPin* Pin : Pins )
+	{
+		if(Pin->LinkedTo.Num() > 0 && Filter(Pin))
+		{
+			for(UEdGraphPin* Connection : Pin->LinkedTo)
+			{
+				// avoid including the current node in the case of a self connection:
+				if(Connection->GetOwningNode() != this)
+				{
+					NeighborsAcceptedForConsideration.Add(Connection->GetOwningNode());
+				}
+			}
+		}
+	}
+
+	for(UEdGraphNode* Neighbor : NeighborsAcceptedForConsideration)
+	{
+		Func(Neighbor);
+	}
+}
+
+void UEdGraphNode::ForEachNodeDirectlyConnectedToInputs(TFunctionRef<void(UEdGraphNode*)> Func)
+{
+	ForEachNodeDirectlyConnectedIf(
+		[](const UEdGraphPin* Pin)
+		{ 
+			if(Pin->Direction == EGPD_Input)
+			{
+				return true;
+			}
+			return false;
+		},
+		Func
+	);
+}
+
+void UEdGraphNode::ForEachNodeDirectlyConnectedToOutputs(TFunctionRef<void(UEdGraphNode*)> Func)
+{
+	ForEachNodeDirectlyConnectedIf(
+		[](const UEdGraphPin* Pin)
+		{ 
+			if(Pin->Direction == EGPD_Output)
+			{
+				return true;
+			}
+			return false;
+		},
+		Func
+	);
+}
+
 #endif	//#if WITH_EDITOR
 
 bool UEdGraphNode::IsInDevelopmentMode() const
diff --git a/Engine/Source/Runtime/Engine/Private/LevelStreaming.cpp b/Engine/Source/Runtime/Engine/Private/LevelStreaming.cpp
--- a/Engine/Source/Runtime/Engine/Private/LevelStreaming.cpp
+++ b/Engine/Source/Runtime/Engine/Private/LevelStreaming.cpp
@@ -425,7 +425,11 @@ bool ULevelStreaming::DetermineTargetState()
 		{
 			TargetState = ETargetState::UnloadedAndRemoved;
 		}
-		else if (!GUseBackgroundLevelStreaming && !World->GetShouldForceUnloadStreamingLevels())
+		else if (World->GetShouldForceUnloadStreamingLevels())
+		{
+			bContinueToConsider = false;
+		}
+		else if (!GUseBackgroundLevelStreaming)
 		{
 			TargetState = ETargetState::LoadedNotVisible;
 		}
@@ -444,7 +448,7 @@ bool ULevelStreaming::DetermineTargetState()
 		break;
 
 	case ECurrentState::LoadedNotVisible:
-		if (bIsRequestingUnloadAndRemoval)
+		if (bIsRequestingUnloadAndRemoval || World->GetShouldForceUnloadStreamingLevels())
 		{
 			TargetState = ETargetState::Unloaded;
 		}
@@ -467,7 +471,7 @@ bool ULevelStreaming::DetermineTargetState()
 		break;
 
 	case ECurrentState::LoadedVisible:
-		if (bIsRequestingUnloadAndRemoval)
+		if (bIsRequestingUnloadAndRemoval || World->GetShouldForceUnloadStreamingLevels())
 		{
 			TargetState = ETargetState::LoadedNotVisible;
 		}
diff --git a/Engine/Source/Runtime/Engine/Private/LevelTick.cpp b/Engine/Source/Runtime/Engine/Private/LevelTick.cpp
--- a/Engine/Source/Runtime/Engine/Private/LevelTick.cpp
+++ b/Engine/Source/Runtime/Engine/Private/LevelTick.cpp
@@ -448,7 +448,7 @@ void UWorld::TickNetClient( float DeltaSeconds )
 bool UWorld::IsPaused() const
 {
 	// pause if specifically set or if we're waiting for the end of the tick to perform streaming level loads (so actors don't fall through the world in the meantime, etc)
-	const AWorldSettings* Info = GetWorldSettings();
+	const AWorldSettings* Info = GetWorldSettings(/*bCheckStreamingPersistent=*/false, /*bChecked=*/false);
 	return ( (Info && Info->Pauser != nullptr && TimeSeconds >= PauseDelay) ||
 				(bRequestedBlockOnAsyncLoading && GetNetMode() == NM_Client) ||
 				(GEngine->ShouldCommitPendingMapChange(this)) ||
diff --git a/Engine/Source/Runtime/Engine/Private/PackageMapClient.cpp b/Engine/Source/Runtime/Engine/Private/PackageMapClient.cpp
--- a/Engine/Source/Runtime/Engine/Private/PackageMapClient.cpp
+++ b/Engine/Source/Runtime/Engine/Private/PackageMapClient.cpp
@@ -24,6 +24,7 @@
 #include "Engine/DemoNetDriver.h"
 #include "Serialization/MemoryReader.h"
 #include "Serialization/MemoryWriter.h"
+#include "Components/ChildActorComponent.h"
 
 #if WITH_EDITOR
 #include "UObject/ObjectRedirector.h"
@@ -335,7 +336,16 @@ bool UPackageMapClient::SerializeNewActor(FArchive& Ar, class UActorChannel *Cha
 
 		if (Ar.IsSaving())
 		{
-			Archetype = Actor->GetArchetype();
+			// ChildActor's need to be spawned from the ChildActorTemplate otherwise any non-replicated 
+			// customized properties will be incorrect on the Client.
+			if (UChildActorComponent* CAC = Actor->GetParentComponent())
+			{
+				Archetype = CAC->GetChildActorTemplate();
+			}
+			if (Archetype == nullptr)
+			{
+				Archetype = Actor->GetArchetype();
+			}
 			ActorLevel = Actor->GetLevel();
 
 			check( Archetype != nullptr );
diff --git a/Engine/Source/Runtime/Engine/Private/UserDefinedStruct.cpp b/Engine/Source/Runtime/Engine/Private/UserDefinedStruct.cpp
--- a/Engine/Source/Runtime/Engine/Private/UserDefinedStruct.cpp
+++ b/Engine/Source/Runtime/Engine/Private/UserDefinedStruct.cpp
@@ -126,6 +126,13 @@ void UUserDefinedStruct::Serialize(FStructuredArchive::FRecord Record)
 
 			FScopedPlaceholderRawContainerTracker TrackStruct(StructData);
 			SerializeItem(Record.EnterField(FIELD_NAME_TEXT("Data")), StructData, nullptr);
+
+			// Now that defaults have been loaded we can inspect our properties
+			// and default values and set the StructFlags accordingly:
+			if(UnderlyingArchive.IsLoading())
+			{
+				UpdateStructFlags();
+			}
 		}
 	}
 
@@ -434,4 +441,87 @@ void UUserDefinedStruct::AddReferencedObjects(UObject* InThis, FReferenceCollect
 	}
 
 	Super::AddReferencedObjects(This, Collector);
-}
\ No newline at end of file
+}
+
+void UUserDefinedStruct::UpdateStructFlags()
+{
+	// Adapted from PrepareCppStructOps, where we 'discover' zero constructability
+	// for native types:
+	bool bIsZeroConstruct = true;
+	{
+		uint8* StructData = DefaultStructInstance.GetStructMemory();
+		if (StructData)
+		{
+			int32 Size = GetStructureSize();
+			for (int32 Index = 0; Index < Size; Index++)
+			{
+				if (StructData[Index])
+				{
+					bIsZeroConstruct = false;
+					break;
+				}
+			}
+		}
+
+		if(bIsZeroConstruct)
+		{	
+			for (TFieldIterator<UProperty> It(this); It; ++It)
+			{
+				UProperty* Property = *It;
+				if (Property && !Property->HasAnyPropertyFlags(CPF_ZeroConstructor))
+				{
+					bIsZeroConstruct = false;
+					break;
+				}
+			}
+		}
+	}
+
+	// IsPOD/NoDtor could be derived earlier than bIsZeroConstruct because they do not depend on 
+	// the structs default values, but it is convenient to calculate them all in one place:
+	bool bIsPOD = true;
+	{
+		for (TFieldIterator<UProperty> It(this); It; ++It)
+		{
+			UProperty* Property = *It;
+			if (Property && !Property->HasAnyPropertyFlags(CPF_IsPlainOldData))
+			{
+				bIsPOD = false;
+				break;
+			}
+		}
+	}
+
+	bool bHasNoDtor = bIsPOD;
+	{
+		if(!bHasNoDtor)
+		{
+			// we're not POD, but we still may have no destructor, check properties:
+			bHasNoDtor = true;
+			for (TFieldIterator<UProperty> It(this); It; ++It)
+			{
+				UProperty* Property = *It;
+				if (Property && !Property->HasAnyPropertyFlags(CPF_NoDestructor))
+				{
+					bHasNoDtor = false;
+					break;
+				}
+			}
+		}
+	}
+	
+	StructFlags = EStructFlags(StructFlags | STRUCT_ZeroConstructor | STRUCT_IsPlainOldData | STRUCT_NoDestructor);
+	if(!bIsZeroConstruct)
+	{
+		StructFlags = EStructFlags(StructFlags & ~STRUCT_ZeroConstructor);
+	}
+	if(!bIsPOD)
+	{
+		StructFlags = EStructFlags(StructFlags & ~STRUCT_IsPlainOldData);
+	}
+	if(!bHasNoDtor)
+	{
+		StructFlags = EStructFlags(StructFlags & ~STRUCT_NoDestructor);
+	}
+
+}
diff --git a/Engine/Source/Runtime/GameplayTags/Classes/GameplayTagsManager.h b/Engine/Source/Runtime/GameplayTags/Classes/GameplayTagsManager.h
--- a/Engine/Source/Runtime/GameplayTags/Classes/GameplayTagsManager.h
+++ b/Engine/Source/Runtime/GameplayTags/Classes/GameplayTagsManager.h
@@ -72,6 +72,15 @@ enum class EGameplayTagSourceType : uint8
 	Invalid,			// Not a real source
 };
 
+UENUM()
+enum class EGameplayTagSelectionType : uint8
+{
+	None,
+	NonRestrictedOnly,
+	RestrictedOnly,
+	All
+};
+
 /** Struct defining where gameplay tags are loaded/saved from. Mostly for the editor */
 USTRUCT()
 struct FGameplayTagSource
@@ -511,8 +520,12 @@ class GAMEPLAYTAGS_API UGameplayTagsManager : public UObject
 	/** Refresh the gameplaytag tree due to an editor change */
 	void EditorRefreshGameplayTagTree();
 
-	/** Gets a Tag Container containing the all tags in the hierarchy that are children of this tag, and were explicitly added to the dictionary */
+	/** Gets a Tag Container containing all of the tags in the hierarchy that are children of this tag, and were explicitly added to the dictionary */
 	FGameplayTagContainer RequestGameplayTagChildrenInDictionary(const FGameplayTag& GameplayTag) const;
+#if WITH_EDITORONLY_DATA
+	/** Gets a Tag Container containing all of the tags in the hierarchy that are children of this tag, were explicitly added to the dictionary, and do not have any explicitly added tags between them and the specified tag */
+	FGameplayTagContainer RequestGameplayTagDirectDescendantsInDictionary(const FGameplayTag& GameplayTag, EGameplayTagSelectionType SelectionType) const;
+#endif // WITH_EDITORONLY_DATA
 
 	/** This is called when EditorRefreshGameplayTagTree. Useful if you need to do anything editor related when tags are added or removed */
 	static FSimpleMulticastDelegate OnEditorRefreshGameplayTagTree;
diff --git a/Engine/Source/Runtime/GameplayTags/Private/GameplayTagsManager.cpp b/Engine/Source/Runtime/GameplayTags/Private/GameplayTagsManager.cpp
--- a/Engine/Source/Runtime/GameplayTags/Private/GameplayTagsManager.cpp
+++ b/Engine/Source/Runtime/GameplayTags/Private/GameplayTagsManager.cpp
@@ -1239,6 +1239,48 @@ FGameplayTagContainer UGameplayTagsManager::RequestGameplayTagChildrenInDictiona
 	return TagContainer;
 }
 
+#if WITH_EDITORONLY_DATA
+FGameplayTagContainer UGameplayTagsManager::RequestGameplayTagDirectDescendantsInDictionary(const FGameplayTag& GameplayTag, EGameplayTagSelectionType SelectionType) const
+{
+	bool bIncludeRestrictedTags = (SelectionType == EGameplayTagSelectionType::RestrictedOnly || SelectionType == EGameplayTagSelectionType::All);
+	bool bIncludeNonRestrictedTags = (SelectionType == EGameplayTagSelectionType::NonRestrictedOnly || SelectionType == EGameplayTagSelectionType::All);
+
+	// Note this purposefully does not include the passed in GameplayTag in the container.
+	FGameplayTagContainer TagContainer;
+
+	TSharedPtr<FGameplayTagNode> GameplayTagNode = FindTagNode(GameplayTag);
+	if (GameplayTagNode.IsValid())
+	{
+		TArray< TSharedPtr<FGameplayTagNode> >& ChildrenNodes = GameplayTagNode->GetChildTagNodes();
+		int32 CurrArraySize = ChildrenNodes.Num();
+		for (int32 Idx = 0; Idx < CurrArraySize; ++Idx)
+		{
+			TSharedPtr<FGameplayTagNode> ChildNode = ChildrenNodes[Idx];
+			if (ChildNode.IsValid())
+			{
+				// if the tag isn't in the dictionary, add its children to the list
+				if (ChildNode->SourceName == NAME_None)
+				{
+					TArray< TSharedPtr<FGameplayTagNode> >& GrandChildrenNodes = ChildNode->GetChildTagNodes();
+					ChildrenNodes.Append(GrandChildrenNodes);
+					CurrArraySize = ChildrenNodes.Num();
+				}
+				else
+				{
+					// this tag is in the dictionary so add it to the list
+					if ((ChildNode->bIsRestrictedTag && bIncludeRestrictedTags) ||
+						(!ChildNode->bIsRestrictedTag && bIncludeNonRestrictedTags))
+					{
+						TagContainer.AddTag(ChildNode->GetCompleteTag());
+					}
+				}
+			}
+		}
+	}
+	return TagContainer;
+}
+#endif // WITH_EDITORONLY_DATA
+
 void UGameplayTagsManager::NotifyGameplayTagDoubleClickedEditor(FString TagName)
 {
 	FGameplayTag Tag = RequestGameplayTag(FName(*TagName), false);
