Change DisabledCompilerMessagesHeadless to DisabledCompilerMessagesExceptEditor to reduce warning spam when testing.

diff --git a/Engine/Source/Editor/UnrealEd/Private/Kismet2/CompilerResultsLog.cpp b/Engine/Source/Editor/UnrealEd/Private/Kismet2/CompilerResultsLog.cpp
--- a/Engine/Source/Editor/UnrealEd/Private/Kismet2/CompilerResultsLog.cpp
+++ b/Engine/Source/Editor/UnrealEd/Private/Kismet2/CompilerResultsLog.cpp
@@ -201,20 +201,17 @@ bool FCompilerResultsLog::IsMessageEnabled(FName ID)
 	}
 
 	const UBlueprintEditorProjectSettings* EditorProjectSettings = GetDefault<UBlueprintEditorProjectSettings>();
-	if(IsRunningCommandlet())
+	if(IsRunningCommandlet() || !GIsEditor)
 	{
-		if(	EditorProjectSettings->DisabledCompilerMessagesHeadless.Contains(ID) ||
-			EditorProjectSettings->DisabledCompilerMessages.Contains(ID))
+		if(	EditorProjectSettings->DisabledCompilerMessagesExceptEditor.Contains(ID))
 		{
 			return false;
 		}
 	}
-	else
+
+	if( EditorProjectSettings->DisabledCompilerMessages.Contains(ID) )
 	{
-		if( EditorProjectSettings->DisabledCompilerMessages.Contains(ID) )
-		{
-			return false;
-		}
+		return false;
 	}
 
 	return true;
diff --git a/Engine/Source/Editor/UnrealEd/Public/Settings/EditorProjectSettings.h b/Engine/Source/Editor/UnrealEd/Public/Settings/EditorProjectSettings.h
--- a/Engine/Source/Editor/UnrealEd/Public/Settings/EditorProjectSettings.h
+++ b/Engine/Source/Editor/UnrealEd/Public/Settings/EditorProjectSettings.h
@@ -185,12 +185,12 @@ public:
 	uint32 bValidateUnloadedSoftActorReferences : 1;
 
 	/** 
-	 * List of compiler messages that have been suppressed during cook (or other commandlets) for 
+	 * List of compiler messages that have been suppressed outside of full, interactive editor sessions for 
 	 * the current project - useful for silencing warnings that were added to the engine after 
 	 * project inception and are going to be addressed as they are found by content authors
 	 */
-	UPROPERTY(EditAnywhere, config, Category= Blueprints, DisplayName = "Compiler Messages Disabled during Cook")
-	TArray<FName> DisabledCompilerMessagesHeadless;
+	UPROPERTY(EditAnywhere, config, Category= Blueprints, DisplayName = "Compiler Messages Disabled Except in Editor")
+	TArray<FName> DisabledCompilerMessagesExceptEditor;
 	
 	/** 
 	 * List of compiler messages that have been suppressed completely - message suppression is only 
