[NO NON-PII COMMIT MESSAGE]

diff --git a/Engine/Source/Editor/UnrealEd/Private/Kismet2/DebuggerCommands.cpp b/Engine/Source/Editor/UnrealEd/Private/Kismet2/DebuggerCommands.cpp
--- a/Engine/Source/Editor/UnrealEd/Private/Kismet2/DebuggerCommands.cpp
+++ b/Engine/Source/Editor/UnrealEd/Private/Kismet2/DebuggerCommands.cpp
@@ -797,7 +797,7 @@ TSharedRef< SWidget > FPlayWorldCommands::GeneratePlayMenuContent( TSharedRef<FU
 			MenuBuilder.AddSubMenu(
 				LOCTEXT("TargetedMobilePreviewSubMenu", "Mobile Preview (PIE)"),
 				LOCTEXT("TargetedMobilePreviewSubMenu_ToolTip", "Play this level using a specified mobile device preview (runs in its own process)"),
-				FNewMenuDelegate::CreateStatic(&MakePreviewDeviceMenu), true,
+				FNewMenuDelegate::CreateStatic(&MakePreviewDeviceMenu), false,
 				FSlateIcon(FEditorStyle::GetStyleSetName(), "PlayWorld.PlayInMobilePreview")
 				);
 		}
