UE-61838 Android 'Launch On All' can't run during an existing launch on session on Windows

#4.21

Cancel previous sessions before launching the selected configuration.

diff --git a/Engine/Source/Editor/UnrealEd/Private/Kismet2/DebuggerCommands.cpp b/Engine/Source/Editor/UnrealEd/Private/Kismet2/DebuggerCommands.cpp
--- a/Engine/Source/Editor/UnrealEd/Private/Kismet2/DebuggerCommands.cpp
+++ b/Engine/Source/Editor/UnrealEd/Private/Kismet2/DebuggerCommands.cpp
@@ -2460,6 +2460,7 @@ void FInternalPlayWorldCommandCallbacks::LaunchOnDevice( const FString& DeviceId
 
 		if (FModuleManager::LoadModuleChecked<IProjectTargetPlatformEditorModule>("ProjectTargetPlatformEditor").ShowUnsupportedTargetWarning(*TargetDeviceId.GetPlatformName()))
 		{
+			GUnrealEd->CancelPlayingViaLauncher();
 			GUnrealEd->RequestPlaySession(DeviceId, DeviceName);
 		}
 	}
