Editor's feature level is now saved and restored between sessions

diff --git a/Engine/Source/Editor/LevelEditor/Private/LevelEditorActions.cpp b/Engine/Source/Editor/LevelEditor/Private/LevelEditorActions.cpp
--- a/Engine/Source/Editor/LevelEditor/Private/LevelEditorActions.cpp
+++ b/Engine/Source/Editor/LevelEditor/Private/LevelEditorActions.cpp
@@ -563,27 +563,7 @@ bool FLevelEditorActionCallbacks::IsMaterialQualityLevelChecked( EMaterialQualit
 
 void FLevelEditorActionCallbacks::SetPreviewPlatform(FName MaterialQualityPlatform, ERHIFeatureLevel::Type PreviewFeatureLevel)
 {
-	// If we have specified a MaterialQualityPlatform ensure its feature level matches the requested feature level.
-	check(MaterialQualityPlatform.IsNone() || GetMaxSupportedFeatureLevel(ShaderFormatToLegacyShaderPlatform(MaterialQualityPlatform)) == PreviewFeatureLevel);
-
-	UMaterialShaderQualitySettings* MaterialShaderQualitySettings = UMaterialShaderQualitySettings::Get();
-	const FName InitialPreviewPlatform = MaterialShaderQualitySettings->GetPreviewPlatform();
-
-	const ERHIFeatureLevel::Type InitialFeatureLevel = GetWorld()->FeatureLevel;
-	MaterialShaderQualitySettings->SetPreviewPlatform(MaterialQualityPlatform);
-	SetFeatureLevelPreview(PreviewFeatureLevel);
-
-	if (
-		// Rebuild materials if the preview platform has changed. 
-		InitialPreviewPlatform != MaterialQualityPlatform
-		// If the feature level changed then materials have been rebuilt already.
-		&& InitialFeatureLevel == PreviewFeatureLevel)
-	{
-		FGlobalComponentRecreateRenderStateContext Recreate;
-		FlushRenderingCommands();
-		UMaterial::AllMaterialsCacheResourceShadersForRendering();
-		UMaterialInstance::AllMaterialsCacheResourceShadersForRendering();
-	}
+	GEditor->SetPreviewPlatform(MaterialQualityPlatform, PreviewFeatureLevel);
 }
 
 bool FLevelEditorActionCallbacks::IsPreviewPlatformChecked(FName MaterialQualityPlatform, ERHIFeatureLevel::Type PreviewFeatureLevel)
@@ -594,23 +574,7 @@ bool FLevelEditorActionCallbacks::IsPreviewPlatformChecked(FName MaterialQuality
 
 void FLevelEditorActionCallbacks::SetFeatureLevelPreview(ERHIFeatureLevel::Type InPreviewFeatureLevel)
 {
-	// Record this feature level as we want to use it for all subsequent level creation and loading
-	GEditor->DefaultWorldFeatureLevel = InPreviewFeatureLevel;
-
-	GetWorld()->ChangeFeatureLevel(InPreviewFeatureLevel);
-
-	// Update any currently running PIE sessions.
-	for (TObjectIterator<UWorld> It; It; ++It)
-	{
-		UWorld* ItWorld = *It;
-		if (ItWorld->WorldType == EWorldType::PIE)
-		{
-			ItWorld->ChangeFeatureLevel(InPreviewFeatureLevel);
-		}
-	}
-
-	GUnrealEd->OnSceneMaterialsModified();
-	GUnrealEd->RedrawAllViewports();
+	GEditor->SetPreviewPlatform(NAME_None, InPreviewFeatureLevel);
 }
 
 bool FLevelEditorActionCallbacks::IsFeatureLevelPreviewChecked(ERHIFeatureLevel::Type InPreviewFeatureLevel)
diff --git a/Engine/Source/Editor/UnrealEd/Classes/Editor/EditorEngine.h b/Engine/Source/Editor/UnrealEd/Classes/Editor/EditorEngine.h
--- a/Engine/Source/Editor/UnrealEd/Classes/Editor/EditorEngine.h
+++ b/Engine/Source/Editor/UnrealEd/Classes/Editor/EditorEngine.h
@@ -2969,7 +2969,23 @@ public:
 	/** This function should be called to notify the editor that new materials were added to our scene or some materials were modified */
 	void OnSceneMaterialsModified();
 
+	/** Call this function to change the feature level and to override the material quality platform of the editor and PIE worlds */
+	void SetPreviewPlatform(const FName MaterialQualityPlatform, const ERHIFeatureLevel::Type PreviewFeatureLevel, const bool bSaveSettings = true);
+
 protected:
+	/** Call this function to change the feature level of the editor and PIE worlds */
+	void SetFeatureLevelPreview(const ERHIFeatureLevel::Type InPreviewFeatureLevel);
+
+	/** call this function to change the feature level for all materials */
+	void SetMaterialsFeatureLevel(const ERHIFeatureLevel::Type InFeatureLevel);
+
+	/** call this to recompile the materials */
+	void AllMaterialsCacheResourceShadersForRendering();
+
+	/** Function pair used to save and restore the global feature level */
+	void LoadEditorFeatureLevel();
+	void SaveEditorFeatureLevel();
+
 	/** For some platforms (e.g. mobiles), when running in editor mode, we emulate the shaders functionality on available running GPU (e.g. DirectX),
 	 *  but when displaying the shader complexity we need to be able compile and extract statistics (instruction count) from the real shaders that
 	 *  will be compiled when the game will run on the specific platform. Thus (if compiler available) we perform an 'offline' shader compilation step,
diff --git a/Engine/Source/Editor/UnrealEd/Classes/Editor/EditorPerProjectUserSettings.h b/Engine/Source/Editor/UnrealEd/Classes/Editor/EditorPerProjectUserSettings.h
--- a/Engine/Source/Editor/UnrealEd/Classes/Editor/EditorPerProjectUserSettings.h
+++ b/Engine/Source/Editor/UnrealEd/Classes/Editor/EditorPerProjectUserSettings.h
@@ -149,6 +149,14 @@ public:
 	UPROPERTY(config)
 	int32 MaterialQualityLevel;
 
+	/** The name of the shader platform that will be used in the editor */
+	UPROPERTY(config)
+	FName PreviewShaderPlatformName;
+
+	/** whether or not the above PreviewShaderPlatformName is a platform that needs to override the material settings in UMaterialShaderQualitySettings */
+	UPROPERTY(config)
+	bool bIsMaterialQualityOverridePlatform;
+
 public:
 
 	/** Delegate for when a user setting has changed */
diff --git a/Engine/Source/Editor/UnrealEd/Private/EditorEngine.cpp b/Engine/Source/Editor/UnrealEd/Private/EditorEngine.cpp
--- a/Engine/Source/Editor/UnrealEd/Private/EditorEngine.cpp
+++ b/Engine/Source/Editor/UnrealEd/Private/EditorEngine.cpp
@@ -202,10 +202,15 @@
 #include "Developer/HotReload/Public/IHotReload.h"
 #include "EditorBuildUtils.h"
 #include "MaterialStatsCommon.h"
+#include "MaterialShaderQualitySettings.h"
 
 #include "Bookmarks/IBookmarkTypeTools.h"
 #include "Bookmarks/BookMarkTypeActions.h"
 #include "Bookmarks/BookMark2DTypeActions.h"
+#include "ComponentReregisterContext.h"
+#include "Materials/Material.h"
+#include "Materials/MaterialInstance.h"
+#include "ComponentRecreateRenderStateContext.h"
 
 DEFINE_LOG_CATEGORY_STATIC(LogEditor, Log, All);
 
@@ -937,6 +942,8 @@ void UEditorEngine::Init(IEngineLoop* InEngineLoop)
 	GEditor = this;
 	InitEditor(InEngineLoop);
 
+	LoadEditorFeatureLevel();
+
 	Layers = FLayers::Create( MakeWeakObjectPtr( this ) );
 
 	// Init transactioning.
@@ -7367,4 +7374,143 @@ void UEditorEngine::OnSceneMaterialsModified()
 	UpdateShaderComplexityMaterials();
 }
 
+void UEditorEngine::SetMaterialsFeatureLevel(const ERHIFeatureLevel::Type InFeatureLevel)
+{
+	FScopedSlowTask SlowTask(100.f, NSLOCTEXT("Engine", "UpdatingMaterialsMessage", "Updating Materials"), true);
+	SlowTask.MakeDialog();
+
+	//invalidate global bound shader states so they will be created with the new shaders the next time they are set (in SetGlobalBoundShaderState)
+	for (TLinkedList<FGlobalBoundShaderStateResource*>::TIterator It(FGlobalBoundShaderStateResource::GetGlobalBoundShaderStateList()); It; It.Next())
+	{
+		BeginUpdateResourceRHI(*It);
+	}
+
+	FGlobalComponentReregisterContext RecreateComponents;
+	FlushRenderingCommands();
+
+	SlowTask.EnterProgressFrame(5.0f);
+
+	// Decrement refcount on old feature level
+	UMaterialInterface::SetGlobalRequiredFeatureLevel(InFeatureLevel, true);
+
+	SlowTask.EnterProgressFrame(50.0f);
+	UMaterial::AllMaterialsCacheResourceShadersForRendering();
+	
+	SlowTask.EnterProgressFrame(40.0f);
+	UMaterialInstance::AllMaterialsCacheResourceShadersForRendering();
+
+	CompileGlobalShaderMap(InFeatureLevel);
+	GShaderCompilingManager->ProcessAsyncResults(false, true);
+	SlowTask.EnterProgressFrame(5.0f);
+}
+
+void UEditorEngine::SetFeatureLevelPreview(const ERHIFeatureLevel::Type InPreviewFeatureLevel)
+{
+	if (DefaultWorldFeatureLevel != InPreviewFeatureLevel)
+	{
+		// Record this feature level as we want to use it for all subsequent level creation and loading
+		DefaultWorldFeatureLevel = InPreviewFeatureLevel;
+
+		FScopedSlowTask SlowTask(100.f, NSLOCTEXT("Engine", "ChangingPreviewRenderingLevelMessage", "Changing Preview Rendering Level"), true);
+		SlowTask.MakeDialog();
+
+		// first change the feature level for global/shared resources
+		SetMaterialsFeatureLevel(InPreviewFeatureLevel);
+
+		SlowTask.EnterProgressFrame(50.0f);
+
+		UWorld* MainWorld = GetEditorWorldContext().World();
+		if (MainWorld != nullptr)
+		{
+			MainWorld->ChangeFeatureLevel(InPreviewFeatureLevel, false);
+		}
+
+		SlowTask.EnterProgressFrame(25.0f);
+
+		// Update any currently running PIE sessions.
+		for (TObjectIterator<UWorld> It; It; ++It)
+		{
+			UWorld* ItWorld = *It;
+			if (ItWorld->WorldType == EWorldType::PIE)
+			{
+				ItWorld->ChangeFeatureLevel(InPreviewFeatureLevel, false);
+			}
+		}
+
+		SlowTask.EnterProgressFrame(25.0f);
+
+		GUnrealEd->OnSceneMaterialsModified();
+		GUnrealEd->RedrawAllViewports();
+	}
+}
+
+void UEditorEngine::AllMaterialsCacheResourceShadersForRendering()
+{
+	FGlobalComponentRecreateRenderStateContext Recreate;
+	FlushRenderingCommands();
+	UMaterial::AllMaterialsCacheResourceShadersForRendering();
+	UMaterialInstance::AllMaterialsCacheResourceShadersForRendering();
+}
+
+void UEditorEngine::SetPreviewPlatform(const FName MaterialQualityPlatform, const ERHIFeatureLevel::Type PreviewFeatureLevel, const bool bSaveSettings/* = true*/)
+{
+	// If we have specified a MaterialQualityPlatform ensure its feature level matches the requested feature level.
+	check(MaterialQualityPlatform.IsNone() || GetMaxSupportedFeatureLevel(ShaderFormatToLegacyShaderPlatform(MaterialQualityPlatform)) == PreviewFeatureLevel);
+
+	UMaterialShaderQualitySettings* MaterialShaderQualitySettings = UMaterialShaderQualitySettings::Get();
+	const FName InitialPreviewPlatform = MaterialShaderQualitySettings->GetPreviewPlatform();
+	MaterialShaderQualitySettings->SetPreviewPlatform(MaterialQualityPlatform);
+
+	if (DefaultWorldFeatureLevel != PreviewFeatureLevel)
+	{
+		// a new feature level will recompile the materials and apply the effect of any 'material quality platform'
+		SetFeatureLevelPreview(PreviewFeatureLevel);
+	}
+	else if (InitialPreviewPlatform != MaterialQualityPlatform)
+	{
+		// Rebuild materials if we have the same feature level but a different 'material quality platform'
+		AllMaterialsCacheResourceShadersForRendering();
+	}
+
+	if (bSaveSettings)
+	{
+		SaveEditorFeatureLevel();
+	}
+}
+
+void UEditorEngine::LoadEditorFeatureLevel()
+{
+ 	auto* Settings = GetMutableDefault<UEditorPerProjectUserSettings>();
+	auto* QualitySettings = UMaterialShaderQualitySettings::Get();
+
+	EShaderPlatform ShaderPlatform = ShaderFormatToLegacyShaderPlatform(Settings->PreviewShaderPlatformName);
+	if (ShaderPlatform != SP_NumPlatforms)
+	{
+		const FName MaterialQualityPlatform = Settings->bIsMaterialQualityOverridePlatform ? Settings->PreviewShaderPlatformName : NAME_None;
+		const ERHIFeatureLevel::Type PreviewFeatureLevel = GetMaxSupportedFeatureLevel(ShaderPlatform);
+		SetPreviewPlatform(MaterialQualityPlatform, PreviewFeatureLevel, false);
+	}
+}
+
+void UEditorEngine::SaveEditorFeatureLevel()
+{
+	auto* Settings = GetMutableDefault<UEditorPerProjectUserSettings>();
+
+	const FName& PreviewPlatormName = UMaterialShaderQualitySettings::Get()->GetPreviewPlatform();
+	if (PreviewPlatormName == NAME_None)
+	{
+		Settings->bIsMaterialQualityOverridePlatform = false;
+
+		const EShaderPlatform ShaderPlatform = GetFeatureLevelShaderPlatform(DefaultWorldFeatureLevel);
+		Settings->PreviewShaderPlatformName = LegacyShaderPlatformToShaderFormat(ShaderPlatform);
+	}
+	else
+	{
+		Settings->bIsMaterialQualityOverridePlatform = true;
+		Settings->PreviewShaderPlatformName = PreviewPlatormName;
+	}
+
+	Settings->PostEditChange();
+}
+
 #undef LOCTEXT_NAMESPACE 
diff --git a/Engine/Source/Editor/UnrealEd/Private/PlayLevel.cpp b/Engine/Source/Editor/UnrealEd/Private/PlayLevel.cpp
--- a/Engine/Source/Editor/UnrealEd/Private/PlayLevel.cpp
+++ b/Engine/Source/Editor/UnrealEd/Private/PlayLevel.cpp
@@ -2922,7 +2922,7 @@ UGameInstance* UEditorEngine::CreatePIEGameInstance(int32 InPIEInstance, bool bI
 	GameInstanceParams.bSimulateInEditor = bInSimulateInEditor;
 	GameInstanceParams.bStartInSpectatorMode = bStartInSpectatorMode;
 	GameInstanceParams.bRunAsDedicated = bRunAsDedicated;
-
+	GameInstanceParams.WorldFeatureLevel = DefaultWorldFeatureLevel;
 	
 	const FGameInstancePIEResult InitializeResult = GameInstance->InitializeForPlayInEditor(InPIEInstance, GameInstanceParams);
 	if (!InitializeResult.IsSuccess())
diff --git a/Engine/Source/Runtime/Engine/Classes/Engine/Engine.h b/Engine/Source/Runtime/Engine/Classes/Engine/Engine.h
--- a/Engine/Source/Runtime/Engine/Classes/Engine/Engine.h
+++ b/Engine/Source/Runtime/Engine/Classes/Engine/Engine.h
@@ -361,6 +361,9 @@ struct FWorldContext
 	/** The Prefix in front of PIE level names, empty is default */
 	FString	PIEPrefix;
 
+	/** The feature level that PIE world should use */
+	ERHIFeatureLevel::Type PIEWorldFeatureLevel;
+
 	/** Is this running as a dedicated server */
 	bool	RunAsDedicated;
 
@@ -414,6 +417,7 @@ struct FWorldContext
 		, bWaitingOnOnlineSubsystem(false)
 		, AudioDeviceHandle(INDEX_NONE)
 		, ThisCurrentWorld(nullptr)
+		, PIEWorldFeatureLevel(ERHIFeatureLevel::Num)
 	{ }
 
 private:
diff --git a/Engine/Source/Runtime/Engine/Classes/Engine/GameInstance.h b/Engine/Source/Runtime/Engine/Classes/Engine/GameInstance.h
--- a/Engine/Source/Runtime/Engine/Classes/Engine/GameInstance.h
+++ b/Engine/Source/Runtime/Engine/Classes/Engine/GameInstance.h
@@ -9,6 +9,7 @@
 #include "Templates/SubclassOf.h"
 #include "Engine/EngineBaseTypes.h"
 #include "Engine/NetworkDelegates.h"
+#include "RHIDefinitions.h"
 #include "GameInstance.generated.h"
 
 class AGameModeBase;
@@ -96,6 +97,9 @@ struct FGameInstancePIEParameters
 
 	// Is this a dedicated server instance for PIE?
 	bool bRunAsDedicated;
+
+	// The feature level that PIE world should use
+	ERHIFeatureLevel::Type WorldFeatureLevel = ERHIFeatureLevel::Num;
 };
 
 #endif
diff --git a/Engine/Source/Runtime/Engine/Private/GameInstance.cpp b/Engine/Source/Runtime/Engine/Private/GameInstance.cpp
--- a/Engine/Source/Runtime/Engine/Private/GameInstance.cpp
+++ b/Engine/Source/Runtime/Engine/Private/GameInstance.cpp
@@ -213,6 +213,8 @@ FGameInstancePIEResult UGameInstance::InitializeForPlayInEditor(int32 PIEInstanc
 		WorldContext->PIEInstance = PIEInstanceIndex;
 	}
 
+	WorldContext->PIEWorldFeatureLevel = Params.WorldFeatureLevel;
+
 	WorldContext->RunAsDedicated = Params.bRunAsDedicated;
 
 	WorldContext->OwningGameInstance = this;
diff --git a/Engine/Source/Runtime/Engine/Private/UnrealEngine.cpp b/Engine/Source/Runtime/Engine/Private/UnrealEngine.cpp
--- a/Engine/Source/Runtime/Engine/Private/UnrealEngine.cpp
+++ b/Engine/Source/Runtime/Engine/Private/UnrealEngine.cpp
@@ -12047,6 +12047,14 @@ bool UEngine::LoadMap( FWorldContext& WorldContext, FURL URL, class UPendingNetG
 				NewWorld->RenameToPIEWorld(WorldContext.PIEInstance);
 			}
 			ResetPIEAudioSetting(NewWorld);
+
+#if WITH_EDITOR
+			// PIE worlds should use the same feature level as the editor
+			if (WorldContext.PIEWorldFeatureLevel != ERHIFeatureLevel::Num && NewWorld->FeatureLevel != WorldContext.PIEWorldFeatureLevel)
+			{
+				NewWorld->ChangeFeatureLevel(WorldContext.PIEWorldFeatureLevel);
+			}
+#endif
 		}
 		else if (WorldContext.WorldType == EWorldType::Game)
 		{
diff --git a/Engine/Source/Runtime/Engine/Private/World.cpp b/Engine/Source/Runtime/Engine/Private/World.cpp
--- a/Engine/Source/Runtime/Engine/Private/World.cpp
+++ b/Engine/Source/Runtime/Engine/Private/World.cpp
@@ -6889,25 +6889,6 @@ void UWorld::ChangeFeatureLevel(ERHIFeatureLevel::Type InFeatureLevel, bool bSho
 			FGlobalComponentReregisterContext RecreateComponents;
 			FlushRenderingCommands();
 
-			// Decrement refcount on old feature level
-			UMaterialInterface::SetGlobalRequiredFeatureLevel(InFeatureLevel, true);
-
-            SlowTask.EnterProgressFrame(10.0f);
-            UMaterial::AllMaterialsCacheResourceShadersForRendering();
-            SlowTask.EnterProgressFrame(10.0f);
-            UMaterialInstance::AllMaterialsCacheResourceShadersForRendering();
-            SlowTask.EnterProgressFrame(10.0f);
-            CompileGlobalShaderMap(InFeatureLevel);
-            SlowTask.EnterProgressFrame(10.0f);
-            GShaderCompilingManager->ProcessAsyncResults(false, true);
-
-			SlowTask.EnterProgressFrame(10.0f);
-			//invalidate global bound shader states so they will be created with the new shaders the next time they are set (in SetGlobalBoundShaderState)
-			for (TLinkedList<FGlobalBoundShaderStateResource*>::TIterator It(FGlobalBoundShaderStateResource::GetGlobalBoundShaderStateList()); It; It.Next())
-			{
-				BeginUpdateResourceRHI(*It);
-			}
-
 			FeatureLevel = InFeatureLevel;
 
 			SlowTask.EnterProgressFrame(10.0f);
diff --git a/Engine/Source/Runtime/Renderer/Private/SystemTextures.h b/Engine/Source/Runtime/Renderer/Private/SystemTextures.h
--- a/Engine/Source/Runtime/Renderer/Private/SystemTextures.h
+++ b/Engine/Source/Runtime/Renderer/Private/SystemTextures.h
@@ -33,6 +33,7 @@ public:
 			return;
 		}
 
+		ReleaseDynamicRHI();
 		InternalInitializeTextures(RHICmdList, InFeatureLevel);
 	}
 
