Fix project browser reopening when creating a new project, if other command line arguments are present. Missing a space between the project file and previous argument list.

diff --git a/Engine/Source/Editor/UnrealEd/Private/UnrealEdMisc.cpp b/Engine/Source/Editor/UnrealEd/Private/UnrealEdMisc.cpp
--- a/Engine/Source/Editor/UnrealEd/Private/UnrealEdMisc.cpp
+++ b/Engine/Source/Editor/UnrealEd/Private/UnrealEdMisc.cpp
@@ -940,7 +940,7 @@ void FUnrealEdMisc::OnExit()
 	if( PendingProjName.Len() > 0 )
 	{
 		// If there is a pending project switch, spawn that process now and use the same command line parameters that were used for this editor instance.
-		FString Cmd = PendingProjName + FCommandLine::Get();
+		FString Cmd = FString::Printf(TEXT("%s %s"), *PendingProjName, FCommandLine::Get());
 
 		FString ExeFilename = CreateProjectPath();
 		FProcHandle Handle = FPlatformProcess::CreateProc( *ExeFilename, *Cmd, true, false, false, NULL, 0, NULL, NULL );
