Fix failure to launch from a DebugGame editor in an installed build, due to "Debug" in the executable name causing it to attempt to launch a debug client. Now queries for the current configuration instead.

#rnx

diff --git a/Engine/Source/Editor/UnrealEd/Private/PlayLevel.cpp b/Engine/Source/Editor/UnrealEd/Private/PlayLevel.cpp
--- a/Engine/Source/Editor/UnrealEd/Private/PlayLevel.cpp
+++ b/Engine/Source/Editor/UnrealEd/Private/PlayLevel.cpp
@@ -1955,15 +1955,7 @@ void UEditorEngine::PlayUsingLauncher()
 			break;
 		default:
 			// same as the running editor
-			FString ExeName = FUnrealEdMisc::Get().GetExecutableForCommandlets();
-			if (ExeName.Contains(TEXT("Debug")))
-			{
-				LauncherProfile->SetBuildConfiguration(EBuildConfigurations::Debug);
-			}
-			else
-			{
-				LauncherProfile->SetBuildConfiguration(EBuildConfigurations::Development);
-			}
+			LauncherProfile->SetBuildConfiguration(FApp::GetBuildConfiguration());
 			break;
 		}
 
