Add an -OnlyUnversioned option to the ResavePackages commandlet which will resave any unversioned packages.

diff --git a/Engine/Source/Editor/UnrealEd/Classes/Commandlets/ResavePackagesCommandlet.h b/Engine/Source/Editor/UnrealEd/Classes/Commandlets/ResavePackagesCommandlet.h
--- a/Engine/Source/Editor/UnrealEd/Classes/Commandlets/ResavePackagesCommandlet.h
+++ b/Engine/Source/Editor/UnrealEd/Classes/Commandlets/ResavePackagesCommandlet.h
@@ -92,6 +92,9 @@ protected:
 	/** only process packages containing materials **/
 	bool bOnlyMaterials;
 
+	/** Only save packages with a changelist of zero **/
+	bool bOnlyUnversioned;
+
 	/** Ignore package version changelist **/
 	bool bIgnoreChangelist;
 
diff --git a/Engine/Source/Editor/UnrealEd/Private/Commandlets/ContentCommandlets.cpp b/Engine/Source/Editor/UnrealEd/Private/Commandlets/ContentCommandlets.cpp
--- a/Engine/Source/Editor/UnrealEd/Private/Commandlets/ContentCommandlets.cpp
+++ b/Engine/Source/Editor/UnrealEd/Private/Commandlets/ContentCommandlets.cpp
@@ -71,6 +71,8 @@ DEFINE_LOG_CATEGORY(LogContentCommandlet);
 #include "ICollectionManager.h"
 #include "CollectionManagerModule.h"
 
+#pragma optimize("", off)
+
 /**-----------------------------------------------------------------------------
  *	UResavePackages commandlet.
  *
@@ -897,6 +899,8 @@ int32 UResavePackagesCommandlet::Main( const FString& Params )
 	bShouldBuildTextureStreaming = bShouldBuildTextureStreamingForAll || Switches.Contains(TEXT("buildtexturestreaming"));
 	/** determine if we can skip the version changelist check */
 	bIgnoreChangelist = Switches.Contains(TEXT("IgnoreChangelist"));
+	/** whether we should only save packages with changelist zero */
+	bOnlyUnversioned = Switches.Contains(TEXT("OnlyUnversioned"));
 	/** only process packages containing materials */
 	bOnlyMaterials = Switches.Contains(TEXT("onlymaterials"));
 
@@ -1206,6 +1210,12 @@ void UResavePackagesCommandlet::PerformPreloadOperations( FLinkerLoad* PackageLi
 		bSavePackage = false;
 	}
 
+	// Check if the changelist number is zero
+	if (bOnlyUnversioned && PackageLinker->Summary.SavedByEngineVersion.GetChangelist() != 0)
+	{
+		bSavePackage = false;
+	}
+
 	// If not, don't resave it.
 	if ( !bAllowResave )
 	{
