Fixes for resave commandlet.

#rnx

diff --git a/Engine/Source/Editor/UnrealEd/Private/Commandlets/ContentCommandlets.cpp b/Engine/Source/Editor/UnrealEd/Private/Commandlets/ContentCommandlets.cpp
--- a/Engine/Source/Editor/UnrealEd/Private/Commandlets/ContentCommandlets.cpp
+++ b/Engine/Source/Editor/UnrealEd/Private/Commandlets/ContentCommandlets.cpp
@@ -71,8 +71,6 @@ DEFINE_LOG_CATEGORY(LogContentCommandlet);
 #include "ICollectionManager.h"
 #include "CollectionManagerModule.h"
 
-#pragma optimize("", off)
-
 /**-----------------------------------------------------------------------------
  *	UResavePackages commandlet.
  *
