Merging //UE4/Dev-Main to Dev-Build (//UE4/Dev-Build)

#rnx

diff --git a/Engine/Plugins/Runtime/AnimationBudgetAllocator/Source/AnimationBudgetAllocator/Private/AnimationBudgetAllocator.cpp b/Engine/Plugins/Runtime/AnimationBudgetAllocator/Source/AnimationBudgetAllocator/Private/AnimationBudgetAllocator.cpp
--- a/Engine/Plugins/Runtime/AnimationBudgetAllocator/Source/AnimationBudgetAllocator/Private/AnimationBudgetAllocator.cpp
+++ b/Engine/Plugins/Runtime/AnimationBudgetAllocator/Source/AnimationBudgetAllocator/Private/AnimationBudgetAllocator.cpp
@@ -289,6 +289,19 @@ static FAutoConsoleVariableRef CVarSkelBatch_ReducedWorkThrottleMaxPerFrame(
 	}),
 	ECVF_Scalability);
 
+static float GBudgetPressureBeforeEmergencyReducedWork = 2.5f;
+
+static FAutoConsoleVariableRef CVarSkelBatch_BudgetPressureBeforeEmergencyReducedWork(
+	TEXT("a.Budget.GBudgetPressureBeforeEmergencyReducedWork"),
+	GBudgetPressureBeforeEmergencyReducedWork,
+	TEXT("Range > 0.0, Default = 2.5\n")
+	TEXT("Controls the budget pressure where emergency reduced work (applied to all components except those that are bAlwaysTick).\n"),
+	FConsoleVariableDelegate::CreateLambda([](IConsoleVariable* InVariable)
+	{
+		GBudgetPressureBeforeEmergencyReducedWork = FMath::Max(GBudgetPressureBeforeEmergencyReducedWork, 0.0f);
+	}),
+	ECVF_Scalability);
+
 FComponentData::FComponentData(USkeletalMeshComponentBudgeted* InComponent)
 	: Component(InComponent)
 	, RootPrerequisite(nullptr)
@@ -805,6 +818,8 @@ int32 FAnimationBudgetAllocator::CalculateWorkDistributionAndQueue(float InDelta
 
 		if(--ReducedComponentWorkCounter <= 0)
 		{
+			const bool bEmergencyReducedWork = SmoothedBudgetPressure >= GBudgetPressureBeforeEmergencyReducedWork;
+
 			// Scale num components to switch based on budget pressure
 			const int32 NumComponentsToSwitch = (int32)FMath::Lerp(1.0f, (float)GReducedWorkThrottleMaxPerFrame, BudgetPressureInterpAlpha);
 			int32 ComponentsSwitched = 0;
@@ -837,7 +852,10 @@ int32 FAnimationBudgetAllocator::CalculateWorkDistributionAndQueue(float InDelta
 				for (SortedComponentIndex = TotalIdealWorkUnits - 1; SortedComponentIndex >= FullIndexEnd; --SortedComponentIndex)
 				{
 					FComponentData& ComponentData = AllComponentData[AllSortedComponentData[SortedComponentIndex]];
-					if(ComponentData.bAllowReducedWork && !ComponentData.bReducedWork && ComponentData.Component->OnReduceWork().IsBound())
+
+					const bool bAllowReducedWork = (ComponentData.bAllowReducedWork || bEmergencyReducedWork) && !ComponentData.bAlwaysTick;
+
+					if(bAllowReducedWork && !ComponentData.bReducedWork && ComponentData.Component->OnReduceWork().IsBound())
 					{
 #if WITH_TICK_DEBUG
 						UE_LOG(LogTemp, Warning, TEXT("Reducing component work (mesh %s) (actor %llx)"), ComponentData.Component->SkeletalMesh ? *ComponentData.Component->SkeletalMesh->GetName() : TEXT("null"), (uint64)ComponentData.Component->GetOwner());
diff --git a/Engine/Source/Editor/UnrealEd/Private/Commandlets/AssetRegistryGenerator.cpp b/Engine/Source/Editor/UnrealEd/Private/Commandlets/AssetRegistryGenerator.cpp
--- a/Engine/Source/Editor/UnrealEd/Private/Commandlets/AssetRegistryGenerator.cpp
+++ b/Engine/Source/Editor/UnrealEd/Private/Commandlets/AssetRegistryGenerator.cpp
@@ -219,6 +219,7 @@ bool FAssetRegistryGenerator::GenerateStreamingInstallManifest(int64 InExtraFlav
 	// Add manifests for any staging-time only groups
 	if (bUseAssetManager && !TargetPlatform->HasSecurePackageFormat())
 	{
+		UE_LOG(LogAssetRegistryGenerator, Log, TEXT("Updating stage-only encryption manifests for platform %s"), *TargetPlatform->IniPlatformName());
 		FContentEncryptionConfig ContentEncryptionConfig;
 		UAssetManager::Get().GetContentEncryptionConfig(ContentEncryptionConfig);
 		const FContentEncryptionConfig::TGroupMap& EncryptedNonUFSFileGroups = ContentEncryptionConfig.GetPackageGroupMap();
@@ -227,19 +228,24 @@ bool FAssetRegistryGenerator::GenerateStreamingInstallManifest(int64 InExtraFlav
 		{
 			if (Element.Value.bStageTimeOnly)
 			{
+				UE_LOG(LogAssetRegistryGenerator, Log, TEXT("Adding stage-time only manifest for group '%s'"), *Element.Key.ToString());
+
 				FName GroupName = Element.Key;
 				const TSet<FName>& PackageNames = Element.Value.PackageNames;
 				
 				FChunkPackageSet* NewManifest = new FChunkPackageSet();
-				int32 ChunkID = UAssetManager::Get().GetContentEncryptionGroupChunkID(GroupName);
-				
 				for (FName PackageName : PackageNames)
 				{
 					NewManifest->Add(PackageName, FPackageName::LongPackageNameToFilename(PackageName.ToString()));
 				}
 
-				check(ChunkID < FinalChunkManifests.Num());
-				check(FinalChunkManifests[ChunkID] == nullptr);
+				int32 ChunkID = UAssetManager::Get().GetContentEncryptionGroupChunkID(GroupName);
+				if (ChunkID >= FinalChunkManifests.Num())
+				{
+					FinalChunkManifests.AddZeroed(ChunkID - FinalChunkManifests.Num() + 1);
+				}
+				checkf(ChunkID < FinalChunkManifests.Num(), TEXT("Chunk %i out of range. %i manifests available"), ChunkID, FinalChunkManifests.Num() - 1);
+				checkf(FinalChunkManifests[ChunkID] == nullptr, TEXT("Manifest already exists for chunk %i"), ChunkID);
 				FinalChunkManifests[ChunkID] = NewManifest;
 			}
 		}
@@ -1397,6 +1403,13 @@ void FAssetRegistryGenerator::FixupPackageDependenciesForChunks(FSandboxPlatform
 	UE_LOG(LogAssetRegistryGenerator, Log, TEXT("Starting FixupPackageDependenciesForChunks..."));
 	SCOPE_LOG_TIME_IN_SECONDS(TEXT("... FixupPackageDependenciesForChunks complete."), nullptr);
 
+	// Clear any existing manifests from the final array
+	for (FChunkPackageSet* Manifest : FinalChunkManifests)
+	{
+		delete Manifest;
+	}
+	FinalChunkManifests.Empty();
+
 	for (int32 ChunkID = 0, MaxChunk = ChunkManifests.Num(); ChunkID < MaxChunk; ++ChunkID)
 	{
 		FinalChunkManifests.Add(nullptr);
diff --git a/Engine/Source/Runtime/Core/Public/GenericPlatform/StandardPlatformString.h b/Engine/Source/Runtime/Core/Public/GenericPlatform/StandardPlatformString.h
--- a/Engine/Source/Runtime/Core/Public/GenericPlatform/StandardPlatformString.h
+++ b/Engine/Source/Runtime/Core/Public/GenericPlatform/StandardPlatformString.h
@@ -7,6 +7,7 @@
 #include "GenericPlatform/GenericPlatformStricmp.h"
 #include "GenericPlatform/GenericPlatformString.h"
 #include <wchar.h>
+#include <string.h>
 
 #if (PLATFORM_APPLE || PLATFORM_UNIX || PLATFORM_HTML5 || PLATFORM_PS4 || PLATFORM_SWITCH) && !PLATFORM_TCHAR_IS_CHAR16
 
