Fix ensure/warnings with PIE replays by not remapping actor class net exports

diff --git a/Engine/Source/Editor/UnrealEd/Private/EditorEngine.cpp b/Engine/Source/Editor/UnrealEd/Private/EditorEngine.cpp
--- a/Engine/Source/Editor/UnrealEd/Private/EditorEngine.cpp
+++ b/Engine/Source/Editor/UnrealEd/Private/EditorEngine.cpp
@@ -6802,7 +6802,22 @@ FORCEINLINE bool NetworkRemapPath_local(FWorldContext& Context, FString& Str, bo
 		// First strip any source prefix, then add the appropriate prefix for this context
 		FSoftObjectPath Path = UWorld::RemovePIEPrefix(Str);
 		
-		Path.FixupForPIE(Context.PIEInstance);
+		if (bIsReplay)
+		{
+			FString AssetName = Path.GetAssetName();
+			FString ShortName = FPackageName::GetShortName(Path.GetLongPackageName());
+
+			const bool bActorClass = FPackageName::IsValidObjectPath(Path.ToString()) && !AssetName.IsEmpty() && !ShortName.IsEmpty() && (AssetName == (ShortName + TEXT("_C")));
+			if (!bActorClass)
+			{
+				Path.FixupForPIE(Context.PIEInstance);
+			}
+		}
+		else
+		{
+			Path.FixupForPIE(Context.PIEInstance);
+		}
+
 		FString Remapped = Path.ToString();
 		if (!Remapped.Equals(Str, ESearchCase::CaseSensitive))
 		{
