Hook up an AdditionalLaunchParametersForMobile to be used during a PIE standalone session

diff --git a/Engine/Source/Editor/UnrealEd/Classes/Settings/LevelEditorPlaySettings.h b/Engine/Source/Editor/UnrealEd/Classes/Settings/LevelEditorPlaySettings.h
--- a/Engine/Source/Editor/UnrealEd/Classes/Settings/LevelEditorPlaySettings.h
+++ b/Engine/Source/Editor/UnrealEd/Classes/Settings/LevelEditorPlaySettings.h
@@ -256,6 +256,10 @@ public:
 	UPROPERTY(config, EditAnywhere, Category=PlayInStandaloneGame)
 	FString AdditionalLaunchParameters;
 
+	/** Extra parameters to be included as part of the command line for a mobile-on-PC standalone game. */
+	UPROPERTY(config, EditAnywhere, Category=PlayInStandaloneGame)
+	FString AdditionalLaunchParametersForMobile;
+
 public:
 
 	/** Whether to build the game before launching on device. */
diff --git a/Engine/Source/Editor/UnrealEd/Private/PlayLevel.cpp b/Engine/Source/Editor/UnrealEd/Private/PlayLevel.cpp
--- a/Engine/Source/Editor/UnrealEd/Private/PlayLevel.cpp
+++ b/Engine/Source/Editor/UnrealEd/Private/PlayLevel.cpp
@@ -1395,6 +1395,15 @@ void UEditorEngine::PlayStandaloneLocalPc(FString MapNameOverride, FIntPoint* Wi
 		AdditionalParameters += PlayInSettings->AdditionalLaunchParameters;
 	}
 
+	if (bPlayUsingMobilePreview)
+	{
+		if (PlayInSettings->AdditionalLaunchParametersForMobile.Len() > 0)
+		{
+			AdditionalParameters += TEXT(" ");
+			AdditionalParameters += PlayInSettings->AdditionalLaunchParametersForMobile;
+		}
+	}
+
 	uint16 ServerPort = 0;
 	if (bIsServer && PlayInSettings->GetServerPort(ServerPort))
 	{
