Placing actors in sublevels that aren't checked out when under source control, or outside the bounds of small sublevels, will pop up a confirmation dialog
#rnx

diff --git a/Engine/Source/Editor/UnrealEd/Classes/Settings/LevelEditorMiscSettings.h b/Engine/Source/Editor/UnrealEd/Classes/Settings/LevelEditorMiscSettings.h
--- a/Engine/Source/Editor/UnrealEd/Classes/Settings/LevelEditorMiscSettings.h
+++ b/Engine/Source/Editor/UnrealEd/Classes/Settings/LevelEditorMiscSettings.h
@@ -66,6 +66,15 @@ public:
 	UPROPERTY(EditAnywhere, config, Category=Levels)
 	TSubclassOf<ULevelStreaming> DefaultLevelStreamingClass;
 
+	UPROPERTY(EditAnywhere, config, Category = Levels)
+	bool bPromptWhenAddingToLevelBeforeCheckout;
+
+	UPROPERTY(EditAnywhere, config, Category = Levels)
+	bool bPromptWhenAddingToLevelOutsideBounds;
+
+	UPROPERTY(EditAnywhere, config, Category = Levels, Meta = (ClampMin = "0.0", ClampMax = "100.0"))
+	float PercentageThresholdForPrompt;
+
 public:
 
 	/** The save directory for newly created screenshots */
diff --git a/Engine/Source/Editor/UnrealEd/Private/AssetSelection.cpp b/Engine/Source/Editor/UnrealEd/Private/AssetSelection.cpp
--- a/Engine/Source/Editor/UnrealEd/Private/AssetSelection.cpp
+++ b/Engine/Source/Editor/UnrealEd/Private/AssetSelection.cpp
@@ -43,6 +43,10 @@
 #include "ObjectEditorUtils.h"
 #include "Framework/Notifications/NotificationManager.h"
 #include "Widgets/Notifications/SNotificationList.h"
+#include "Settings/LevelEditorMiscSettings.h"
+#include "Engine/LevelBounds.h"
+#include "SourceControlHelpers.h"
+#include "Dialogs/Dialogs.h"
 
 
 namespace AssetSelectionUtils
@@ -410,6 +414,53 @@ namespace AssetSelectionUtils
 	}
 }
 
+namespace ActorPlacementUtils
+{
+	bool IsLevelValidForActorPlacement(ULevel* InLevel, TArray<FTransform>& InActorTransforms)
+	{
+		if (FLevelUtils::IsLevelLocked(InLevel))
+		{
+			FNotificationInfo Info(NSLOCTEXT("UnrealEd", "Error_OperationDisallowedOnLockedLevel", "The requested operation could not be completed because the level is locked."));
+			Info.ExpireDuration = 3.0f;
+			FSlateNotificationManager::Get().AddNotification(Info);
+			return false;
+		}
+		if (InLevel && GetDefault<ULevelEditorMiscSettings>()->bPromptWhenAddingToLevelBeforeCheckout && SourceControlHelpers::IsAvailable())
+		{
+			FString FileName = SourceControlHelpers::PackageFilename(InLevel->GetPathName());
+			// Query file state also checks the source control status
+			FSourceControlState SCState = SourceControlHelpers::QueryFileState(FileName, true);
+			if (!SCState.bIsCheckedOut)
+			{
+				if (EAppReturnType::Ok != OpenMsgDlgInt(EAppMsgType::OkCancel, NSLOCTEXT("UnrealEd","LevelNotCheckedOutMsg", "This actor will be placed in a level that is not currently checked out. Continue?"), NSLOCTEXT("UnrealEd", "ActorPlacement_Title", "Actor Placement Warning")))
+				{
+					return false;
+				}
+			}
+		}
+		if (InLevel && GetDefault<ULevelEditorMiscSettings>()->bPromptWhenAddingToLevelOutsideBounds)
+		{
+			FBox CurrentLevelBounds = ALevelBounds::CalculateLevelBounds(InLevel);
+			FVector ExpandedScale = FVector(1.0f + (GetDefault<ULevelEditorMiscSettings>()->PercentageThresholdForPrompt / 100.0f));
+			FTransform ExpandedScaleTransform = FTransform::Identity;
+			ExpandedScaleTransform.SetScale3D(ExpandedScale);
+			CurrentLevelBounds.TransformBy(ExpandedScaleTransform);
+			for (int32 ActorTransformIndex = 0; ActorTransformIndex < InActorTransforms.Num(); ++ActorTransformIndex)
+			{
+				FTransform ActorTransform = InActorTransforms[ActorTransformIndex];
+				if (!CurrentLevelBounds.IsInsideOrOn(ActorTransform.GetLocation()))
+				{
+					if (EAppReturnType::Ok != OpenMsgDlgInt(EAppMsgType::OkCancel, NSLOCTEXT("UnrealEd", "LevelBoundsMsg", "The actor will be placed outside the bounds of the current level. Continue?"), NSLOCTEXT("UnrealEd", "ActorPlacement_Title", "Actor Placement Warning")))
+					{
+						return false;
+					}
+				}
+			}
+		}
+		return true;
+	}
+}
+
 /**
 * Creates an actor using the specified factory.  
 *
@@ -465,15 +516,16 @@ static AActor* PrivateAddActor( UObject* Asset, UActorFactory* Factory, bool Sel
 	FSnappingUtils::ClearSnappingHelpers( bClearImmediately );
 
 	ULevel* DesiredLevel = GWorld->GetCurrentLevel();
+	bool bSpawnActor = true;
 
-	// Don't spawn the actor if the current level is locked.
-	if (FLevelUtils::IsLevelLocked(DesiredLevel))
+	if ((ObjectFlags & RF_Transactional) != 0)
 	{
-		FNotificationInfo Info(NSLOCTEXT("UnrealEd", "Error_OperationDisallowedOnLockedLevel", "The requested operation could not be completed because the level is locked."));
-		Info.ExpireDuration = 3.0f;
-		FSlateNotificationManager::Get().AddNotification(Info);
+		TArray<FTransform> SpawningActorTransforms;
+		SpawningActorTransforms.Add(ActorTransform);
+		bSpawnActor = ActorPlacementUtils::IsLevelValidForActorPlacement(DesiredLevel, SpawningActorTransforms);
 	}
-	else
+
+	if(bSpawnActor)
 	{
 		FScopedTransaction Transaction( NSLOCTEXT("UnrealEd", "CreateActor", "Create Actor"), (ObjectFlags & RF_Transactional) != 0 );
 
diff --git a/Engine/Source/Editor/UnrealEd/Private/EditorActor.cpp b/Engine/Source/Editor/UnrealEd/Private/EditorActor.cpp
--- a/Engine/Source/Editor/UnrealEd/Private/EditorActor.cpp
+++ b/Engine/Source/Editor/UnrealEd/Private/EditorActor.cpp
@@ -63,6 +63,7 @@
 #include "HAL/PlatformApplicationMisc.h"
 #include "IAssetTools.h"
 #include "AssetToolsModule.h"
+#include "AssetSelection.h"
 
 #define LOCTEXT_NAMESPACE "UnrealEd.EditorActor"
 
@@ -419,12 +420,23 @@ public:
 	void DuplicateActorsToLevel(TArray<AActor*>& OutNewActors, ULevel* DestLevel, bool bOffsetLocations)
 	{
 		// Check neither level is locked
-		if ( FLevelUtils::IsLevelLocked(SrcLevel) || FLevelUtils::IsLevelLocked(DestLevel) )
+		if ( FLevelUtils::IsLevelLocked(SrcLevel) )
 		{
 			UE_LOG(LogEditorActor, Warning, TEXT("DuplicateActorsToLevel: The requested operation could not be completed because the level is locked."));
 			return;
 		}
 
+		TArray<FTransform> ActorTransforms;
+		for (AActor* Actor : Actors)
+		{
+			ActorTransforms.Add(Actor->GetActorTransform());
+		}
+
+		if (!ActorPlacementUtils::IsLevelValidForActorPlacement(DestLevel, ActorTransforms))
+		{
+			return;
+		}
+
 		// Cache the current source level
 		ULevel* OldLevel = SrcLevel->OwningWorld->GetCurrentLevel();
 		// Set the selection set to be precisely the actors belonging to this job.
diff --git a/Engine/Source/Editor/UnrealEd/Private/Settings/SettingsClasses.cpp b/Engine/Source/Editor/UnrealEd/Private/Settings/SettingsClasses.cpp
--- a/Engine/Source/Editor/UnrealEd/Private/Settings/SettingsClasses.cpp
+++ b/Engine/Source/Editor/UnrealEd/Private/Settings/SettingsClasses.cpp
@@ -344,6 +344,9 @@ ULevelEditorMiscSettings::ULevelEditorMiscSettings( const FObjectInitializer& Ob
 	SectionName = TEXT("Misc");
 	CategoryName = TEXT("LevelEditor");
 	EditorScreenshotSaveDirectory.Path = FPaths::ScreenShotDir();
+	bPromptWhenAddingToLevelBeforeCheckout = true;
+	bPromptWhenAddingToLevelOutsideBounds = true;
+	PercentageThresholdForPrompt = 20.0f;
 }
 
 
diff --git a/Engine/Source/Editor/UnrealEd/Public/AssetSelection.h b/Engine/Source/Editor/UnrealEd/Public/AssetSelection.h
--- a/Engine/Source/Editor/UnrealEd/Public/AssetSelection.h
+++ b/Engine/Source/Editor/UnrealEd/Public/AssetSelection.h
@@ -210,6 +210,11 @@ namespace AssetSelectionUtils
 	UNREALED_API bool IsBuilderBrushSelected();
 };
 
+namespace ActorPlacementUtils
+{
+	UNREALED_API bool IsLevelValidForActorPlacement(ULevel* Level, TArray<FTransform>& InActorTransforms);
+}
+
 class UNREALED_API FActorFactoryAssetProxy
 {
 public:
