Fix bug when reimporting custom LOD greater then 1 and the LOD 1 is generated

diff --git a/Engine/Source/Editor/UnrealEd/Private/Factories/SkeletalMeshImport.cpp b/Engine/Source/Editor/UnrealEd/Private/Factories/SkeletalMeshImport.cpp
--- a/Engine/Source/Editor/UnrealEd/Private/Factories/SkeletalMeshImport.cpp
+++ b/Engine/Source/Editor/UnrealEd/Private/Factories/SkeletalMeshImport.cpp
@@ -427,8 +427,7 @@ ExistingSkelMeshData* SaveExistingSkelMeshData(USkeletalMesh* ExistingSkelMesh,
 			ExistingSkelMesh->GetLODNum() == ImportedResource->LODModels.Num() )
 		{
 			int32 OffsetReductionLODIndex = 0;
-			
-			FSkeletalMeshLODInfo* LODInfo = ExistingSkelMesh->GetLODInfo(ReimportSpecificLOD);
+			FSkeletalMeshLODInfo* LODInfo = ExistingSkelMesh->GetLODInfo( ReimportLODIndex < 0 ? 0 : ReimportLODIndex);
 			ExistingMeshDataPtr->bIsReimportLODReduced = (LODInfo && LODInfo->bHasBeenSimplified);
 			if (ExistingMeshDataPtr->bIsReimportLODReduced)
 			{
