Changed "Enable COTF" to be a checkbox.

diff --git a/Engine/Source/Editor/UnrealEd/Private/Kismet2/DebuggerCommands.cpp b/Engine/Source/Editor/UnrealEd/Private/Kismet2/DebuggerCommands.cpp
--- a/Engine/Source/Editor/UnrealEd/Private/Kismet2/DebuggerCommands.cpp
+++ b/Engine/Source/Editor/UnrealEd/Private/Kismet2/DebuggerCommands.cpp
@@ -1226,7 +1226,7 @@ TSharedRef< SWidget > FPlayWorldCommands::GenerateLaunchMenuContent( TSharedRef<
 			FSlateIcon(),
 			UIAction,
 			NAME_None,
-			EUserInterfaceActionType::Check
+			EUserInterfaceActionType::ToggleButton
 		);
 	}
 
