Fixed logic issue in SaveExistingStaticMeshData

diff --git a/Engine/Source/Editor/UnrealEd/Private/StaticMeshEdit.cpp b/Engine/Source/Editor/UnrealEd/Private/StaticMeshEdit.cpp
--- a/Engine/Source/Editor/UnrealEd/Private/StaticMeshEdit.cpp
+++ b/Engine/Source/Editor/UnrealEd/Private/StaticMeshEdit.cpp
@@ -1093,7 +1093,7 @@ ExistingStaticMeshData* SaveExistingStaticMeshData(UStaticMesh* ExistingMesh, Un
 				ExistingMesh->SectionInfoMap.Remove(LodIndex, SectionIndex);
 			}
 		}
-		int32 TotalMaterialIndex=0;
+		int32 TotalMaterialIndex = ExistingMeshDataPtr->ExistingMaterials.Num();
 		for(int32 i=0; i<ExistingMesh->SourceModels.Num(); i++)
 		{
 			//If the last import was exceeding the maximum number of LOD the source model will contain more LOD so just break the loop
