Fix for SCC dialog popping up when messages should be suppressed when adding actors to checked out levels
#rnx

diff --git a/Engine/Plugins/Developer/PerforceSourceControl/Source/PerforceSourceControl/Private/PerforceSourceControlProvider.cpp b/Engine/Plugins/Developer/PerforceSourceControl/Source/PerforceSourceControl/Private/PerforceSourceControlProvider.cpp
--- a/Engine/Plugins/Developer/PerforceSourceControl/Source/PerforceSourceControl/Private/PerforceSourceControlProvider.cpp
+++ b/Engine/Plugins/Developer/PerforceSourceControl/Source/PerforceSourceControl/Private/PerforceSourceControlProvider.cpp
@@ -494,8 +494,13 @@ ECommandResult::Type FPerforceSourceControlProvider::ExecuteSynchronousCommand(F
 		}
 	};
 
+	FText TaskText = Task;
 	// Display the progress dialog
-	FScopedSourceControlProgress Progress(Task, FSimpleDelegate::CreateStatic(&Local::CancelCommand, &InCommand));
+	if (bSuppressResponseMsg)
+	{
+		TaskText = FText::GetEmpty();
+	}
+	FScopedSourceControlProgress Progress(TaskText, FSimpleDelegate::CreateStatic(&Local::CancelCommand, &InCommand));
 
 	// Perform the command asynchronously
 	IssueCommand( InCommand, false );
diff --git a/Engine/Source/Editor/UnrealEd/Private/AssetSelection.cpp b/Engine/Source/Editor/UnrealEd/Private/AssetSelection.cpp
--- a/Engine/Source/Editor/UnrealEd/Private/AssetSelection.cpp
+++ b/Engine/Source/Editor/UnrealEd/Private/AssetSelection.cpp
@@ -430,9 +430,9 @@ namespace ActorPlacementUtils
 			FString FileName = SourceControlHelpers::PackageFilename(InLevel->GetPathName());
 			// Query file state also checks the source control status
 			FSourceControlState SCState = SourceControlHelpers::QueryFileState(FileName, true);
-			if (!SCState.bIsCheckedOut)
+			if (!(SCState.bIsCheckedOut || SCState.bIsAdded || SCState.bIsUnknown))
 			{
-				if (EAppReturnType::Ok != OpenMsgDlgInt(EAppMsgType::OkCancel, NSLOCTEXT("UnrealEd","LevelNotCheckedOutMsg", "This actor will be placed in a level that is not currently checked out. Continue?"), NSLOCTEXT("UnrealEd", "ActorPlacement_Title", "Actor Placement Warning")))
+				if (EAppReturnType::Ok != OpenMsgDlgInt(EAppMsgType::OkCancel, NSLOCTEXT("UnrealEd","LevelNotCheckedOutMsg", "This actor will be placed in a level that is is source control but not currently checked out. Continue?"), NSLOCTEXT("UnrealEd", "ActorPlacement_Title", "Actor Placement Warning")))
 				{
 					return false;
 				}
