Fixed crash after loading a User Defined Structure with circular references to Blueprint classes by avoiding unnecessary default instance recreation

diff --git a/Engine/Source/Editor/UnrealEd/Classes/UserDefinedStructure/UserDefinedStructEditorData.h b/Engine/Source/Editor/UnrealEd/Classes/UserDefinedStructure/UserDefinedStructEditorData.h
--- a/Engine/Source/Editor/UnrealEd/Classes/UserDefinedStructure/UserDefinedStructEditorData.h
+++ b/Engine/Source/Editor/UnrealEd/Classes/UserDefinedStructure/UserDefinedStructEditorData.h
@@ -143,6 +143,7 @@ public:
 
 	const uint8* GetDefaultInstance() const;
 	void RecreateDefaultInstance(FString* OutLog = nullptr);
+	void ReinitializeDefaultInstance(FString* OutLog = nullptr);
 	void CleanDefaultInstance();
 
 private:
diff --git a/Engine/Source/Editor/UnrealEd/Private/UserDefinedStructEditorData.cpp b/Engine/Source/Editor/UnrealEd/Private/UserDefinedStructEditorData.cpp
--- a/Engine/Source/Editor/UnrealEd/Private/UserDefinedStructEditorData.cpp
+++ b/Engine/Source/Editor/UnrealEd/Private/UserDefinedStructEditorData.cpp
@@ -179,6 +179,12 @@ void UUserDefinedStructEditorData::RecreateDefaultInstance(FString* OutLog)
 {
 	UUserDefinedStruct* ScriptStruct = GetOwnerStruct();
 	ScriptStruct->DefaultStructInstance.Recreate(ScriptStruct);
+	ReinitializeDefaultInstance(OutLog);
+}
+
+void UUserDefinedStructEditorData::ReinitializeDefaultInstance(FString* OutLog)
+{
+	UUserDefinedStruct* ScriptStruct = GetOwnerStruct();
 	uint8* StructData = ScriptStruct->DefaultStructInstance.GetStructMemory();
 	ensure(ScriptStruct->DefaultStructInstance.IsValid() && ScriptStruct->DefaultStructInstance.GetStruct() == ScriptStruct);
 	if (ScriptStruct->DefaultStructInstance.IsValid() && StructData)
diff --git a/Engine/Source/Runtime/Engine/Private/UserDefinedStruct.cpp b/Engine/Source/Runtime/Engine/Private/UserDefinedStruct.cpp
--- a/Engine/Source/Runtime/Engine/Private/UserDefinedStruct.cpp
+++ b/Engine/Source/Runtime/Engine/Private/UserDefinedStruct.cpp
@@ -156,7 +156,18 @@ void UUserDefinedStruct::Serialize(FStructuredArchive::FRecord Record)
 				UnderlyingArchive.Preload(EditorData);
 				if (!(UnderlyingArchive.GetPortFlags() & PPF_Duplicate))
 				{
-					FStructureEditorUtils::RecreateDefaultInstanceInEditorData(this);
+					if(!DefaultStructInstance.IsValid())
+					{
+						FStructureEditorUtils::RecreateDefaultInstanceInEditorData(this);
+					}
+					else
+					{
+						UUserDefinedStructEditorData* UDSEditorData = Cast<UUserDefinedStructEditorData>(EditorData);
+						if (UDSEditorData)
+						{
+							UDSEditorData->ReinitializeDefaultInstance();
+						}
+					}
 				}
 			}
 
