Level checkout warning will only pop up once per level per editor session
#rnx

diff --git a/Engine/Source/Editor/UnrealEd/Private/AssetSelection.cpp b/Engine/Source/Editor/UnrealEd/Private/AssetSelection.cpp
--- a/Engine/Source/Editor/UnrealEd/Private/AssetSelection.cpp
+++ b/Engine/Source/Editor/UnrealEd/Private/AssetSelection.cpp
@@ -430,12 +430,16 @@ namespace ActorPlacementUtils
 			FString FileName = SourceControlHelpers::PackageFilename(InLevel->GetPathName());
 			// Query file state also checks the source control status
 			FSourceControlState SCState = SourceControlHelpers::QueryFileState(FileName, true);
-			if (!(SCState.bIsCheckedOut || SCState.bIsAdded || SCState.bIsUnknown))
+			if (!InLevel->bLevelOkayForPlacementWhileCheckedIn && !(SCState.bIsCheckedOut || SCState.bIsAdded || SCState.bIsUnknown))
 			{
 				if (EAppReturnType::Ok != OpenMsgDlgInt(EAppMsgType::OkCancel, NSLOCTEXT("UnrealEd","LevelNotCheckedOutMsg", "This actor will be placed in a level that is is source control but not currently checked out. Continue?"), NSLOCTEXT("UnrealEd", "ActorPlacement_Title", "Actor Placement Warning")))
 				{
 					return false;
 				}
+				else
+				{
+					InLevel->bLevelOkayForPlacementWhileCheckedIn = true;
+				}
 			}
 		}
 		if (InLevel && GetDefault<ULevelEditorMiscSettings>()->bPromptWhenAddingToLevelOutsideBounds)
diff --git a/Engine/Source/Runtime/Engine/Classes/Engine/Level.h b/Engine/Source/Runtime/Engine/Classes/Engine/Level.h
--- a/Engine/Source/Runtime/Engine/Classes/Engine/Level.h
+++ b/Engine/Source/Runtime/Engine/Classes/Engine/Level.h
@@ -591,6 +591,9 @@ public:
 
 	float FixupOverrideVertexColorsTime;
 	int32 FixupOverrideVertexColorsCount;
+
+	UPROPERTY(transient)
+	bool bLevelOkayForPlacementWhileCheckedIn;
 #endif //WITH_EDITORONLY_DATA
 
 	/** Actor which defines level logical bounding box				*/
