Don't prompt about source control or level bounds if there is only one level open
#rnx

diff --git a/Engine/Source/Editor/UnrealEd/Private/AssetSelection.cpp b/Engine/Source/Editor/UnrealEd/Private/AssetSelection.cpp
--- a/Engine/Source/Editor/UnrealEd/Private/AssetSelection.cpp
+++ b/Engine/Source/Editor/UnrealEd/Private/AssetSelection.cpp
@@ -425,12 +425,17 @@ namespace ActorPlacementUtils
 			FSlateNotificationManager::Get().AddNotification(Info);
 			return false;
 		}
+		if (InLevel && InLevel->OwningWorld && InLevel->OwningWorld->GetStreamingLevels().Num() == 0)
+		{
+			// if this is the only level
+			return true;
+		}
 		if (InLevel && GetDefault<ULevelEditorMiscSettings>()->bPromptWhenAddingToLevelBeforeCheckout && SourceControlHelpers::IsAvailable())
 		{
 			FString FileName = SourceControlHelpers::PackageFilename(InLevel->GetPathName());
 			// Query file state also checks the source control status
 			FSourceControlState SCState = SourceControlHelpers::QueryFileState(FileName, true);
-			if (!InLevel->bLevelOkayForPlacementWhileCheckedIn && !(SCState.bIsCheckedOut || SCState.bIsAdded || SCState.bIsUnknown))
+			if (!InLevel->bLevelOkayForPlacementWhileCheckedIn && !(SCState.bIsCheckedOut || SCState.bIsAdded || SCState.bCanAdd))
 			{
 				if (EAppReturnType::Ok != OpenMsgDlgInt(EAppMsgType::OkCancel, NSLOCTEXT("UnrealEd","LevelNotCheckedOutMsg", "This actor will be placed in a level that is is source control but not currently checked out. Continue?"), NSLOCTEXT("UnrealEd", "ActorPlacement_Title", "Actor Placement Warning")))
 				{
