Fixes for building lighting using Swarm when done from an editor built with a unique build environment.

* List of dependencies for Lightmass needs to be derived from the engine directory at runtime, rather than assuming that the base directory is Engine/Binaries/Win64.
* AgentInterface.dll is now copied to the same output folder as the editor SwarmInterface DLL at build time, rather than being checked into P4 in two places.

diff --git a/Engine/Source/Editor/SwarmInterface/SwarmInterface.Build.cs b/Engine/Source/Editor/SwarmInterface/SwarmInterface.Build.cs
--- a/Engine/Source/Editor/SwarmInterface/SwarmInterface.Build.cs
+++ b/Engine/Source/Editor/SwarmInterface/SwarmInterface.Build.cs
@@ -26,5 +26,9 @@ public class SwarmInterface : ModuleRules
 			PublicDependencyModuleNames.Add("Sockets");
 			PublicDependencyModuleNames.Add("Networking");
 		}
+
+		// Copy the AgentInterface DLLs to the same output directory as the editor DLL.
+		RuntimeDependencies.Add("$(BinaryOutputDir)/AgentInterface.dll", "$(EngineDir)/Binaries/DotNET/AgentInterface.dll", StagedFileType.NonUFS);
+		RuntimeDependencies.Add("$(BinaryOutputDir)/AgentInterface.pdb", "$(EngineDir)/Binaries/DotNET/AgentInterface.pdb", StagedFileType.DebugNonUFS);
 	}
 }
diff --git a/Engine/Source/Editor/UnrealEd/Private/Lightmass/Lightmass.cpp b/Engine/Source/Editor/UnrealEd/Private/Lightmass/Lightmass.cpp
--- a/Engine/Source/Editor/UnrealEd/Private/Lightmass/Lightmass.cpp
+++ b/Engine/Source/Editor/UnrealEd/Private/Lightmass/Lightmass.cpp
@@ -2860,6 +2860,40 @@ void FLightmassProcessor::IssueStaticShadowDepthMapTask(const ULightComponent* L
 	}
 }
 
+// Helper class for creating a list of paths relative to the engine directory, which can be queried as a list of raw strings.
+class FEngineDependencyPaths
+{
+public:
+	FEngineDependencyPaths(std::initializer_list<const TCHAR*> Paths)
+	{
+		Strings.Reserve(Paths.size());
+		for(const TCHAR* Path : Paths)
+		{
+			Strings.Add(FPaths::EngineDir() / Path);
+		}
+
+		RawStrings.Reserve(Strings.Num());
+		for(const FString& String : Strings)
+		{
+			RawStrings.Add(*String);
+		}
+	}
+
+	const TCHAR** GetArray() const
+	{
+		return const_cast<const TCHAR**>(RawStrings.GetData());
+	}
+
+	int32 Num() const
+	{
+		return RawStrings.Num();
+	}
+
+private:
+	TArray<FString> Strings;
+	TArray<const TCHAR*> RawStrings;
+};
+
 bool FLightmassProcessor::BeginRun()
 {
 	{
@@ -2887,88 +2921,84 @@ bool FLightmassProcessor::BeginRun()
 	}
 
 	// Setup dependencies for 32bit.
-	const TCHAR* LightmassExecutable32 = TEXT("../Win32/UnrealLightmass.exe");
-	const TCHAR* RequiredDependencyPaths32[] =
-	{
-		TEXT("../DotNET/SwarmInterface.dll"),
-		TEXT("../Win32/AgentInterface.dll"),
-		TEXT("../Win32/UnrealLightmass-SwarmInterface.dll"),
-		TEXT("../Win32/UnrealLightmass-ApplicationCore.dll"),
-		TEXT("../Win32/UnrealLightmass-Core.dll"),
-		TEXT("../Win32/UnrealLightmass-CoreUObject.dll"),
-		TEXT("../Win32/UnrealLightmass-Projects.dll"),
-		TEXT("../Win32/UnrealLightmass-Json.dll"),
-		TEXT("../Win32/UnrealLightmass-BuildSettings.dll")
+	static const FString LightmassExecutable32 = FPaths::EngineDir() / TEXT("Binaries/Win32/UnrealLightmass.exe");
+	static const FEngineDependencyPaths RequiredDependencyPaths32 =
+	{
+		TEXT("Binaries/DotNET/SwarmInterface.dll"),
+		TEXT("Binaries/Win32/AgentInterface.dll"),
+		TEXT("Binaries/Win32/UnrealLightmass-SwarmInterface.dll"),
+		TEXT("Binaries/Win32/UnrealLightmass-ApplicationCore.dll"),
+		TEXT("Binaries/Win32/UnrealLightmass-Core.dll"),
+		TEXT("Binaries/Win32/UnrealLightmass-CoreUObject.dll"),
+		TEXT("Binaries/Win32/UnrealLightmass-Projects.dll"),
+		TEXT("Binaries/Win32/UnrealLightmass-Json.dll"),
+		TEXT("Binaries/Win32/UnrealLightmass-BuildSettings.dll")
 	};
-	const int32 RequiredDependencyPaths32Count = ARRAY_COUNT(RequiredDependencyPaths32);
 
 	// Setup dependencies for 64bit.
 #if PLATFORM_WINDOWS
-	const TCHAR* LightmassExecutable64 = TEXT("../Win64/UnrealLightmass.exe");
-	const TCHAR* RequiredDependencyPaths64[] =
-	{
-		TEXT("../DotNET/SwarmInterface.dll"),
-		TEXT("../Win64/AgentInterface.dll"),
-		TEXT("../Win64/UnrealLightmass-SwarmInterface.dll"),
-		TEXT("../Win64/UnrealLightmass-ApplicationCore.dll"),
-		TEXT("../Win64/UnrealLightmass-Core.dll"),
-		TEXT("../Win64/UnrealLightmass-CoreUObject.dll"),
-		TEXT("../Win64/UnrealLightmass-Projects.dll"),
-		TEXT("../Win64/UnrealLightmass-Json.dll"),
-		TEXT("../Win64/UnrealLightmass-BuildSettings.dll"),
-		TEXT("../Win64/embree.dll"),
-		TEXT("../Win64/tbb.dll"),
-		TEXT("../Win64/tbbmalloc.dll")
+	static const FString LightmassExecutable64 = FPaths::EngineDir() / TEXT("Binaries/Win64/UnrealLightmass.exe");
+	static const FEngineDependencyPaths RequiredDependencyPaths64 =
+	{
+		TEXT("Binaries/DotNET/SwarmInterface.dll"),
+		TEXT("Binaries/Win64/AgentInterface.dll"),
+		TEXT("Binaries/Win64/UnrealLightmass-SwarmInterface.dll"),
+		TEXT("Binaries/Win64/UnrealLightmass-ApplicationCore.dll"),
+		TEXT("Binaries/Win64/UnrealLightmass-Core.dll"),
+		TEXT("Binaries/Win64/UnrealLightmass-CoreUObject.dll"),
+		TEXT("Binaries/Win64/UnrealLightmass-Projects.dll"),
+		TEXT("Binaries/Win64/UnrealLightmass-Json.dll"),
+		TEXT("Binaries/Win64/UnrealLightmass-BuildSettings.dll"),
+		TEXT("Binaries/Win64/embree.dll"),
+		TEXT("Binaries/Win64/tbb.dll"),
+		TEXT("Binaries/Win64/tbbmalloc.dll")
 	};
 #elif PLATFORM_MAC
-	const TCHAR* LightmassExecutable64 = TEXT("../Mac/UnrealLightmass");
-	const TCHAR* RequiredDependencyPaths64[] =
-	{
-		TEXT("../DotNET/Mac/AgentInterface.dll"),
-		TEXT("../Mac/UnrealLightmass-ApplicationCore.dylib"),
-		TEXT("../Mac/UnrealLightmass-Core.dylib"),
-		TEXT("../Mac/UnrealLightmass-CoreUObject.dylib"),
-		TEXT("../Mac/UnrealLightmass-Json.dylib"),
-		TEXT("../Mac/UnrealLightmass-Projects.dylib"),
-		TEXT("../Mac/UnrealLightmass-SwarmInterface.dylib"),
-		TEXT("../Mac/UnrealLightmass-BuildSettings.dylib"),
-		TEXT("../Mac/libembree.2.dylib"),
-		TEXT("../Mac/libtbb.dylib"),
-		TEXT("../Mac/libtbbmalloc.dylib")
+	static const FString LightmassExecutable64 = FPaths::EngineDir() / TEXT("Binaries/Mac/UnrealLightmass");
+	static const FLightmassDependencyPaths RequiredDependencyPaths64 =
+	{
+		TEXT("Binaries/DotNET/Mac/AgentInterface.dll"),
+		TEXT("Binaries/Mac/UnrealLightmass-ApplicationCore.dylib"),
+		TEXT("Binaries/Mac/UnrealLightmass-Core.dylib"),
+		TEXT("Binaries/Mac/UnrealLightmass-CoreUObject.dylib"),
+		TEXT("Binaries/Mac/UnrealLightmass-Json.dylib"),
+		TEXT("Binaries/Mac/UnrealLightmass-Projects.dylib"),
+		TEXT("Binaries/Mac/UnrealLightmass-SwarmInterface.dylib"),
+		TEXT("Binaries/Mac/UnrealLightmass-BuildSettings.dylib"),
+		TEXT("Binaries/Mac/libembree.2.dylib"),
+		TEXT("Binaries/Mac/libtbb.dylib"),
+		TEXT("Binaries/Mac/libtbbmalloc.dylib")
 	};
 #elif PLATFORM_LINUX
-	const TCHAR* LightmassExecutable64 = TEXT("../Linux/UnrealLightmass");
-	const TCHAR* RequiredDependencyPaths64[] =
-	{
-		TEXT("../DotNET/Linux/AgentInterface.dll"),
-		TEXT("../Linux/libUnrealLightmass-ApplicationCore.so"),
-		TEXT("../Linux/libUnrealLightmass-Core.so"),
-		TEXT("../Linux/libUnrealLightmass-CoreUObject.so"),
-		TEXT("../Linux/libUnrealLightmass-Json.so"),
-		TEXT("../Linux/libUnrealLightmass-Projects.so"),
-		TEXT("../Linux/libUnrealLightmass-SwarmInterface.so"),
-		TEXT("../Linux/libUnrealLightmass-Networking.so"),
-		TEXT("../Linux/libUnrealLightmass-Messaging.so"),
-		TEXT("../Linux/libUnrealLightmass-BuildSettings.so"),
-		TEXT("../../Plugins/Messaging/UdpMessaging/Binaries/Linux/libUnrealLightmass-UdpMessaging.so")
+	static const FString LightmassExecutable64 = FPaths::EngineDir() / TEXT("Binaries/Linux/UnrealLightmass");
+	static const FString RequiredDependencyPaths64[] =
+	{
+		TEXT("Binaries/DotNET/Linux/AgentInterface.dll"),
+		TEXT("Binaries/Linux/libUnrealLightmass-ApplicationCore.so"),
+		TEXT("Binaries/Linux/libUnrealLightmass-Core.so"),
+		TEXT("Binaries/Linux/libUnrealLightmass-CoreUObject.so"),
+		TEXT("Binaries/Linux/libUnrealLightmass-Json.so"),
+		TEXT("Binaries/Linux/libUnrealLightmass-Projects.so"),
+		TEXT("Binaries/Linux/libUnrealLightmass-SwarmInterface.so"),
+		TEXT("Binaries/Linux/libUnrealLightmass-Networking.so"),
+		TEXT("Binaries/Linux/libUnrealLightmass-Messaging.so"),
+		TEXT("Binaries/Linux/libUnrealLightmass-BuildSettings.so"),
+		TEXT("Plugins/Messaging/UdpMessaging/Binaries/Linux/libUnrealLightmass-UdpMessaging.so")
 	};
 #else // PLATFORM_LINUX
 #error "Unknown Lightmass platform"
 #endif
-	const int32 RequiredDependencyPaths64Count = ARRAY_COUNT(RequiredDependencyPaths64);
 
 	// Set up optional dependencies.  These might not exist in Launcher distributions, for example.
-	const TCHAR* OptionalDependencyPaths32[] =
+	const FEngineDependencyPaths OptionalDependencyPaths32 =
 	{
-		TEXT("../Win32/UnrealLightmass.pdb")
+		TEXT("Binaries/Win32/UnrealLightmass.pdb")
 	};
-	const int32 OptionalDependencyPaths32Count = ARRAY_COUNT(OptionalDependencyPaths32);
 
-	const TCHAR* OptionalDependencyPaths64[] =
+	const FEngineDependencyPaths OptionalDependencyPaths64 =
 	{
-		TEXT("../Win64/UnrealLightmass.pdb")
+		TEXT("Binaries/Win64/UnrealLightmass.pdb")
 	};
-	const int32 OptionalDependencyPaths64Count = ARRAY_COUNT(OptionalDependencyPaths64);
 
 	// Set up the description for the Job
 	const TCHAR* DescriptionKeys[] =
@@ -3003,7 +3033,7 @@ bool FLightmassProcessor::BeginRun()
 	
 	Statistics.SwarmJobOpenTime += FPlatformTime::Seconds() - SwarmJobStartTime;
 	
-	UE_LOG(LogLightmassSolver, Log,  TEXT("Swarm launching: %s %s"), bUse64bitProcess ? LightmassExecutable64 : LightmassExecutable32, *Exporter->SceneGuid.ToString() );
+	UE_LOG(LogLightmassSolver, Log,  TEXT("Swarm launching: %s %s"), bUse64bitProcess ? *LightmassExecutable64 : *LightmassExecutable32, *Exporter->SceneGuid.ToString() );
 
 	SwarmJobStartTime = FPlatformTime::Seconds();
 
@@ -3052,14 +3082,14 @@ bool FLightmassProcessor::BeginRun()
 	NSwarm::FJobSpecification JobSpecification32, JobSpecification64;
 	if ( !bUse64bitProcess )
 	{
-		JobSpecification32 = NSwarm::FJobSpecification( LightmassExecutable32, *CommandLineParameters, ( NSwarm::TJobTaskFlags )JobFlags );
-		JobSpecification32.AddDependencies( RequiredDependencyPaths32, RequiredDependencyPaths32Count, OptionalDependencyPaths32, OptionalDependencyPaths32Count );
+		JobSpecification32 = NSwarm::FJobSpecification( *LightmassExecutable32, *CommandLineParameters, ( NSwarm::TJobTaskFlags )JobFlags );
+		JobSpecification32.AddDependencies( RequiredDependencyPaths32.GetArray(), RequiredDependencyPaths32.Num(), OptionalDependencyPaths32.GetArray(), OptionalDependencyPaths32.Num() );
 		JobSpecification32.AddDescription( DescriptionKeys, DescriptionValues, ARRAY_COUNT(DescriptionKeys) );
 	}
 	else
 	{
-		JobSpecification64 = NSwarm::FJobSpecification( LightmassExecutable64, *CommandLineParameters, ( NSwarm::TJobTaskFlags )JobFlags );
-		JobSpecification64.AddDependencies( RequiredDependencyPaths64, RequiredDependencyPaths64Count, OptionalDependencyPaths64, OptionalDependencyPaths64Count );
+		JobSpecification64 = NSwarm::FJobSpecification( *LightmassExecutable64, *CommandLineParameters, ( NSwarm::TJobTaskFlags )JobFlags );
+		JobSpecification64.AddDependencies( RequiredDependencyPaths64.GetArray(), RequiredDependencyPaths64.Num(), OptionalDependencyPaths64.GetArray(), OptionalDependencyPaths64.Num() );
 		JobSpecification64.AddDescription( DescriptionKeys, DescriptionValues, ARRAY_COUNT(DescriptionKeys) );
 	}
 	int32 ErrorCode = Swarm.BeginJobSpecification( JobSpecification32, JobSpecification64 );
diff --git a/Engine/Source/Editor/UnrealEd/Private/StaticLightingSystem/StaticLightingSystem.cpp b/Engine/Source/Editor/UnrealEd/Private/StaticLightingSystem/StaticLightingSystem.cpp
--- a/Engine/Source/Editor/UnrealEd/Private/StaticLightingSystem/StaticLightingSystem.cpp
+++ b/Engine/Source/Editor/UnrealEd/Private/StaticLightingSystem/StaticLightingSystem.cpp
@@ -1920,7 +1920,7 @@ bool FStaticLightingSystem::CreateLightmassProcessor()
 		return false;
 	}
 
-	if (NSwarm::FSwarmInterface::Initialize(*(FString(FPlatformProcess::BaseDir()) + TEXT("..\\DotNET\\SwarmInterface.dll"))) == false)
+	if (NSwarm::FSwarmInterface::Initialize(*(FPaths::EngineDir() / TEXT("Binaries/DotNET/SwarmInterface.dll"))) == false)
 	{
 		UE_LOG(LogStaticLightingSystem, Warning, TEXT("Failed to initialize Swarm."));
 		FMessageDialog::Open(EAppMsgType::Ok,
diff --git a/Engine/Source/Programs/UnrealSwarm/AgentInterface/AgentInterface.csproj b/Engine/Source/Programs/UnrealSwarm/AgentInterface/AgentInterface.csproj
--- a/Engine/Source/Programs/UnrealSwarm/AgentInterface/AgentInterface.csproj
+++ b/Engine/Source/Programs/UnrealSwarm/AgentInterface/AgentInterface.csproj
@@ -46,7 +46,7 @@
     <DebugSymbols>True</DebugSymbols>
     <DebugType>pdbonly</DebugType>
     <Optimize>False</Optimize>
-    <OutputPath>..\..\..\..\Binaries\Win64\</OutputPath>
+    <OutputPath>..\..\..\..\Binaries\DotNET\</OutputPath>
     <DefineConstants>DEBUG;TRACE</DefineConstants>
     <ErrorReport>prompt</ErrorReport>
     <WarningLevel>4</WarningLevel>
@@ -55,7 +55,7 @@
   <PropertyGroup Condition=" '$(Configuration)|$(Platform)' == 'Development|AnyCPU' ">
     <DebugType>pdbonly</DebugType>
     <Optimize>True</Optimize>
-    <OutputPath>..\..\..\..\Binaries\Win64\</OutputPath>
+    <OutputPath>..\..\..\..\Binaries\DotNET\</OutputPath>
     <DefineConstants>TRACE</DefineConstants>
     <ErrorReport>prompt</ErrorReport>
     <WarningLevel>4</WarningLevel>
