Fix Linux CIS error.

diff --git a/Engine/Source/Editor/UnrealEd/Private/Lightmass/Lightmass.cpp b/Engine/Source/Editor/UnrealEd/Private/Lightmass/Lightmass.cpp
--- a/Engine/Source/Editor/UnrealEd/Private/Lightmass/Lightmass.cpp
+++ b/Engine/Source/Editor/UnrealEd/Private/Lightmass/Lightmass.cpp
@@ -2955,7 +2955,7 @@ bool FLightmassProcessor::BeginRun()
 	};
 #elif PLATFORM_MAC
 	static const FString LightmassExecutable64 = FPaths::EngineDir() / TEXT("Binaries/Mac/UnrealLightmass");
-	static const FLightmassDependencyPaths RequiredDependencyPaths64 =
+	static const FEngineDependencyPaths RequiredDependencyPaths64 =
 	{
 		TEXT("Binaries/DotNET/Mac/AgentInterface.dll"),
 		TEXT("Binaries/Mac/UnrealLightmass-ApplicationCore.dylib"),
@@ -2971,7 +2971,7 @@ bool FLightmassProcessor::BeginRun()
 	};
 #elif PLATFORM_LINUX
 	static const FString LightmassExecutable64 = FPaths::EngineDir() / TEXT("Binaries/Linux/UnrealLightmass");
-	static const FString RequiredDependencyPaths64[] =
+	static const FEngineDependencyPaths RequiredDependencyPaths64 =
 	{
 		TEXT("Binaries/DotNET/Linux/AgentInterface.dll"),
 		TEXT("Binaries/Linux/libUnrealLightmass-ApplicationCore.so"),
