IESLoader
- Added support for type A files
- Don't crash if reading a \n without a \r

diff --git a/Engine/Source/Editor/UnrealEd/Private/Factories/IESLoader.cpp b/Engine/Source/Editor/UnrealEd/Private/Factories/IESLoader.cpp
--- a/Engine/Source/Editor/UnrealEd/Private/Factories/IESLoader.cpp
+++ b/Engine/Source/Editor/UnrealEd/Private/Factories/IESLoader.cpp
@@ -12,6 +12,13 @@ enum EIESVersion
 	EIESV_2002, // IES LM-63-2002
 };
 
+enum class EIESPhotometricType
+{
+	TypeC = 1,
+	TypeB = 2,
+	TypeA = 3
+};
+
 DEFINE_LOG_CATEGORY_STATIC(LogIESLoader, Log, All);
 
 
@@ -55,15 +62,15 @@ static void GetLineContent(const uint8*& BufferPos, char Line[256], bool bStopOn
 		{
 			break;
 		}
-		else if(*BufferPos == 13 && *(BufferPos + 1) == 10)
+		else if(*BufferPos == '\r' && *(BufferPos + 1) == '\n')
 		{
 			BufferPos += 2;
 			break;
 		}
-		else if(*BufferPos == 10)
+		else if(*BufferPos == '\n')
 		{
-			// no valid MSDOS return file
-			check(0);
+			BufferPos += 1;
+			break;
 		}
 		else if(bStopOnWhitespace && (*BufferPos <= ' '))
 		{
@@ -106,6 +113,7 @@ static bool GetInt(const uint8*& BufferPos, int32& ret)
 FIESLoadHelper::FIESLoadHelper(const uint8* Buffer, uint32 BufferLength)
 	: Brightness(0)
 	, CachedIntegral(MAX_FLT)
+	, PhotometricType(EIESPhotometricType::TypeC)
 	, Error("No data loaded")
 {
 	check(!IsValid());
@@ -220,7 +228,12 @@ void FIESLoadHelper::Load(const uint8* Buffer)
 		return;
 	}
 
-	PARSE_INT(PhotometricType);
+	PARSE_INT(PhotometricTypeTemp);
+
+	if ( PhotometricTypeTemp >= (int32)EIESPhotometricType::TypeC && PhotometricTypeTemp <= (int32)EIESPhotometricType::TypeA )
+	{
+		PhotometricType = EIESPhotometricType(PhotometricTypeTemp);
+	}
 
 	// 1:feet, 2:meter
 	PARSE_INT(UnitType);
@@ -549,29 +562,49 @@ float FIESLoadHelper::InterpolateBilinear(float fX, float fY) const
 
 float FIESLoadHelper::Interpolate2D(float HAngle, float VAngle) const
 {
-	// Support symmetry, per the IES format doc:
-	// 0     There is only one horizontal angle, implying that the luminaire is laterally symmetric in all photometric planes.
-	// 90    The luminaire is assumed to be symmetric in each quadrant.
-	// 180   The luminaire is assumed to be bilaterally symmetric about the 0-180 degree photometric plane.
-	// 360   The luminaire is assumed to exhibit no lateral symmetry.
-	//
-	// A luminaire that is bilaterally symmetric about the 90-270 degree
-	// photometric plane will have a first value of 90 degrees and a last value of 270 degrees.
-
-	if ( HAngles.Num() > 0 )
+	float u = 0.0f;
+	float v = 0.0f;
+
+	if (PhotometricType == EIESPhotometricType::TypeA)
 	{
-		if ( HAngles.Last() > 0.f && HAngle > HAngles.Last() )
-		{
-			HAngle = HAngles.Last() - FMath::Fmod( HAngle, HAngles.Last() );
-		}
-		else if ( HAngles[0] > 0.f && HAngle < HAngles[0] )
+		float X = FMath::Cos(FMath::DegreesToRadians(VAngle));
+		float Y = FMath::Sin(FMath::DegreesToRadians(VAngle)) * FMath::Sin(FMath::DegreesToRadians(HAngle));
+		float Z = FMath::Sin(FMath::DegreesToRadians(VAngle)) * FMath::Cos(FMath::DegreesToRadians(HAngle));
+
+
+		float NewVAngle = FMath::RadiansToDegrees(FMath::Asin(Z));
+		float NewHAngle = -FMath::RadiansToDegrees(FMath::Atan2(Y, X));
+
+		u = ComputeFilterPos(NewHAngle, HAngles);
+		v = ComputeFilterPos(NewVAngle, VAngles);
+	}
+	else
+	{
+		// Support symmetry, per the IES format doc:
+		// 0     There is only one horizontal angle, implying that the luminaire is laterally symmetric in all photometric planes.
+		// 90    The luminaire is assumed to be symmetric in each quadrant.
+		// 180   The luminaire is assumed to be bilaterally symmetric about the 0-180 degree photometric plane.
+		// 360   The luminaire is assumed to exhibit no lateral symmetry.
+		//
+		// A luminaire that is bilaterally symmetric about the 90-270 degree
+		// photometric plane will have a first value of 90 degrees and a last value of 270 degrees.
+
+		if (HAngles.Num() > 0)
 		{
-			HAngle = HAngles.Last() - HAngles[0] - HAngle;
+			if (HAngles.Last() > 0.f && HAngle > HAngles.Last())
+			{
+				HAngle = HAngles.Last() - FMath::Fmod(HAngle, HAngles.Last());
+			}
+			else if (HAngles[0] > 0.f && HAngle < HAngles[0])
+			{
+				HAngle = HAngles.Last() - HAngles[0] - HAngle;
+			}
 		}
-	}
 
-	float u = ComputeFilterPos(HAngle, HAngles);
-	float v = ComputeFilterPos(VAngle, VAngles);
+		u = ComputeFilterPos(HAngle, HAngles);
+		v = ComputeFilterPos(VAngle, VAngles);
+	}
+	
 
 	return InterpolateBilinear(u, v);
 }
diff --git a/Engine/Source/Editor/UnrealEd/Private/Factories/IESLoader.h b/Engine/Source/Editor/UnrealEd/Private/Factories/IESLoader.h
--- a/Engine/Source/Editor/UnrealEd/Private/Factories/IESLoader.h
+++ b/Engine/Source/Editor/UnrealEd/Private/Factories/IESLoader.h
@@ -7,6 +7,8 @@
 
 class Error;
 
+enum class EIESPhotometricType;
+
 /**
  * To load the IES file image format. IES files exist for many real world lights. The file stores how much light is emitted in a specific direction.
  * The data is usually measured but tools to paint IES files exist.
@@ -40,6 +42,9 @@ private:
 	/** used by ComputeFullIntegral(), integrated over the unit sphere, to avoid computing it twice, -1 if not computed yet */
 	float CachedIntegral;
 
+	/** The photometric type associated with the IES file being loaded */
+	EIESPhotometricType PhotometricType;
+
 	// will set Error to 0 on success
 	void Load(const uint8* Buffer);
 
