Allow LOD in the geo/skinning workflow

diff --git a/Engine/Source/Editor/UnrealEd/Private/Factories/SkeletalMeshImport.cpp b/Engine/Source/Editor/UnrealEd/Private/Factories/SkeletalMeshImport.cpp
--- a/Engine/Source/Editor/UnrealEd/Private/Factories/SkeletalMeshImport.cpp
+++ b/Engine/Source/Editor/UnrealEd/Private/Factories/SkeletalMeshImport.cpp
@@ -848,18 +848,6 @@ void RestoreExistingSkelMeshData(ExistingSkelMeshData* MeshData, USkeletalMesh*
 		// We copy back and fix-up the LODs that still work with this skeleton.
 		if (MeshData->ExistingLODModels.Num() > 0)
 		{
-			// see if we have reduction avail
-			IMeshReductionManagerModule& Module = FModuleManager::Get().LoadModuleChecked<IMeshReductionManagerModule>("MeshReductionInterface");
-			static bool bAutoMeshReductionAvailable = Module.GetSkeletalMeshReductionInterface() != NULL;
-
-			bool bRegenLODs = !bImportSkinningOnly;
-			TArray<bool> ApplySkinnings;
-			ApplySkinnings.AddZeroed(MeshData->ExistingLODModels.Num());
-			for (bool& bApplySkinning : ApplySkinnings)
-			{
-				bApplySkinning = true;
-			}
-			
 			auto RestoreReductionSourceData = [&SkeletalMeshImportedModel, &MeshData](int32 ExistingIndex, int32 NewIndex)
 			{
 				if (!MeshData->ExistingOriginalReductionSourceMeshData.IsValidIndex(ExistingIndex) || MeshData->ExistingOriginalReductionSourceMeshData[ExistingIndex]->IsEmpty())
@@ -881,46 +869,9 @@ void RestoreExistingSkelMeshData(ExistingSkelMeshData* MeshData, USkeletalMesh*
 				SkeletalMeshImportedModel->OriginalReductionSourceMeshData.Add(ReductionLODData);
 			};
 
-			auto ApplySkinning = [&SkeletalMesh, &SkeletalMeshImportedModel, &MeshData, &ApplySkinnings, &RestoreReductionSourceData]()
-			{
-				FSkeletalMeshLODModel& BaseLodModel = SkeletalMeshImportedModel->LODModels[0];
-				int32 OffsetLOD = SkeletalMesh->GetLODNum();
-				//Apply the new skinning on all existing LOD
-				for (int32 Index = 0; Index < MeshData->ExistingLODModels.Num(); ++Index)
-				{
-					int32 RealIndex = OffsetLOD + Index;
-					if (!ApplySkinnings[Index])
-					{
-						continue;
-					}
-					FSkeletalMeshLODModel& LODModel = MeshData->ExistingLODModels[Index];
-					FSkeletalMeshLODInfo& LODInfo = MeshData->ExistingLODInfo[Index];
-
-					FSkeletalMeshLODModel* NewLODModel = nullptr;
-					if (RealIndex >= SkeletalMesh->GetLODNum())
-					{
-						NewLODModel = new FSkeletalMeshLODModel(LODModel);
-						SkeletalMeshImportedModel->LODModels.Add(NewLODModel);
-						// add LOD info back
-						SkeletalMesh->AddLODInfo(LODInfo);
-					}
-					else
-					{
-						NewLODModel = &(SkeletalMeshImportedModel->LODModels[RealIndex]);
-					}
-
-					RestoreReductionSourceData(Index, RealIndex);
-
-					//Apply the new skinning to the existing LOD geometry
-					SkeletalMeshHelper::ApplySkinning(SkeletalMesh, BaseLodModel, *NewLODModel);
-					ApplySkinnings[Index] = false;
-				}
-			};
 
-			bool bSkinningIsApply = false;
 			if (SkeletonsAreCompatible(SkeletalMesh->RefSkeleton, MeshData->ExistingRefSkeleton, bImportSkinningOnly))
 			{
-				bRegenLODs = false;
 				// First create mapping table from old skeleton to new skeleton.
 				TArray<int32> OldToNewMap;
 				OldToNewMap.AddUninitialized(MeshData->ExistingRefSkeleton.GetRawBoneNum());
@@ -1019,52 +970,20 @@ void RestoreExistingSkelMeshData(ExistingSkelMeshData* MeshData, USkeletalMesh*
 					{
 						UnFbx::FFbxImporter* FFbxImporter = UnFbx::FFbxImporter::GetInstance();
 						FFbxImporter->AddTokenizedErrorMessage(FTokenizedMessage::Create(EMessageSeverity::Warning, FText::Format(LOCTEXT("NewMeshMissingBoneFromLOD", "New mesh is missing bone '{0}' required by an LOD."), FText::FromName(MissingBoneName))), FFbxErrors::SkeletalMesh_LOD_MissingBone);
-						bRegenLODs = !bImportSkinningOnly;
 						break;
 					}
 					else
 					{
-						// if it has been regenerated, it try to regen and if we have reduction available 
-						if (bAutoMeshReductionAvailable && LODInfo.bHasBeenSimplified && LODInfo.ReductionSettings.BaseLOD == 0)
-						{
-							bRegenLODs = !bImportSkinningOnly;
-							if (bRegenLODs)
-							{
-								//We need to add LODInfo
-								FSkeletalMeshLODModel* NewLODModel = new FSkeletalMeshLODModel(LODModel);
-								SkeletalMeshImportedModel->LODModels.Add(NewLODModel);
-								SkeletalMesh->AddLODInfo(LODInfo);
-								RestoreReductionSourceData(i, SkeletalMesh->GetLODNum() - 1);
-							}
-						}
-						else
-						{
-							FSkeletalMeshLODModel* NewLODModel = new FSkeletalMeshLODModel(LODModel);
-							SkeletalMeshImportedModel->LODModels.Add(NewLODModel);
-							SkeletalMesh->AddLODInfo(LODInfo);
-
-							RestoreReductionSourceData(i, SkeletalMesh->GetLODNum() - 1);
-
-							ApplySkinnings[i] = false;
-						}
+						//We need to add LODInfo
+						FSkeletalMeshLODModel* NewLODModel = new FSkeletalMeshLODModel(LODModel);
+						SkeletalMeshImportedModel->LODModels.Add(NewLODModel);
+						SkeletalMesh->AddLODInfo(LODInfo);
+						RestoreReductionSourceData(i, SkeletalMesh->GetLODNum() - 1);
 					}
 				}
 			}
-			else if (bImportSkinningOnly)
-			{
-				ApplySkinning();
-				bRegenLODs = false;
-				bSkinningIsApply = true;
-			}
-
-			if (bRegenLODs)
-			{
-				TryRegenerateLODs(MeshData, SkeletalMesh);
-			}
-			else if (!bSkinningIsApply)
-			{
-				ApplySkinning();
-			}
+			//Always regenerate LODs when re-importing LOD 0
+			TryRegenerateLODs(MeshData, SkeletalMesh);
 		}
 
 		for (int32 AssetIndex = 0; AssetIndex < MeshData->ExistingPhysicsAssets.Num(); ++AssetIndex)
diff --git a/Engine/Source/Editor/UnrealEd/Private/Fbx/FbxSkeletalMeshImport.cpp b/Engine/Source/Editor/UnrealEd/Private/Fbx/FbxSkeletalMeshImport.cpp
--- a/Engine/Source/Editor/UnrealEd/Private/Fbx/FbxSkeletalMeshImport.cpp
+++ b/Engine/Source/Editor/UnrealEd/Private/Fbx/FbxSkeletalMeshImport.cpp
@@ -3692,6 +3692,8 @@ bool UnFbx::FFbxImporter::ImportSkeletalMeshLOD(USkeletalMesh* InSkeletalMesh, U
 		//Apply the base LOD skinning to the new LOD geometry, DestImportedResource here is the base lod and is the source for the apply skinning.
 		//The Dest here mean we merge the just import LOD to the base LodModel.
 		SkeletalMeshHelper::ApplySkinning(InSkeletalMesh, DestImportedResource->LODModels[0], ImportedResource->LODModels[0]);
+		IMeshUtilities& MeshUtilities = FModuleManager::Get().LoadModuleChecked<IMeshUtilities>("MeshUtilities");
+		MeshUtilities.RemoveBonesFromMesh(BaseSkeletalMesh, DesiredLOD, NULL);
 	}
 
 	FSkeletalMeshLODModel& NewLODModel = ImportedResource->LODModels[0];
