Fix assertion/crash when switching Feature Levels with the Shader Complexity view active

diff --git a/Engine/Source/Editor/UnrealEd/Classes/Editor/EditorEngine.h b/Engine/Source/Editor/UnrealEd/Classes/Editor/EditorEngine.h
--- a/Engine/Source/Editor/UnrealEd/Classes/Editor/EditorEngine.h
+++ b/Engine/Source/Editor/UnrealEd/Classes/Editor/EditorEngine.h
@@ -3081,8 +3081,12 @@ protected:
 	 *  but when displaying the shader complexity we need to be able compile and extract statistics (instruction count) from the real shaders that
 	 *  will be compiled when the game will run on the specific platform. Thus (if compiler available) we perform an 'offline' shader compilation step,
 	 *  extract the needed statistics and transfer them to the emulated editor running shaders.
-	 *  This function will be called from OnSceneMaterialsModified() */
-	void UpdateShaderComplexityMaterials();
+	 *  This function will be called from OnSceneMaterialsModified()
+	 *
+	 * @param	bForceUpdate	When true, view mode shaders are always updated for worlds displaying shader complexity materials.
+	 *                          When false, view mode shaders are rebuilt only when emulating a shader platform.
+	 */
+	void UpdateShaderComplexityMaterials(bool bForceUpdate);
 
 	/** Utility function that can determine whether some input world is using materials who's shaders are emulated in the editor */
 	bool IsEditorShaderPlatformEmulated(UWorld* World);
diff --git a/Engine/Source/Editor/UnrealEd/Private/EditorEngine.cpp b/Engine/Source/Editor/UnrealEd/Private/EditorEngine.cpp
--- a/Engine/Source/Editor/UnrealEd/Private/EditorEngine.cpp
+++ b/Engine/Source/Editor/UnrealEd/Private/EditorEngine.cpp
@@ -7300,7 +7300,7 @@ bool UEditorEngine::IsOfflineShaderCompilerAvailable(UWorld* World)
 	return FMaterialStatsUtils::IsPlatformOfflineCompilerAvailable(RealPlatform);
 }
 
-void UEditorEngine::UpdateShaderComplexityMaterials()
+void UEditorEngine::UpdateShaderComplexityMaterials(bool bForceUpdate)
 {
 	TSet<UWorld *> WorldSet;
 
@@ -7316,10 +7316,10 @@ void UEditorEngine::UpdateShaderComplexityMaterials()
 	for (auto* SomeWorld : WorldSet)
 	{
 		bool bShadersEmulated = IsEditorShaderPlatformEmulated(SomeWorld);
-		if (bShadersEmulated)
+		if (bShadersEmulated || bForceUpdate)
 		{
 			bool bOfflineCompilerAvailable = IsOfflineShaderCompilerAvailable(SomeWorld);
-			if (bOfflineCompilerAvailable)
+			if (bOfflineCompilerAvailable || bForceUpdate)
 			{
 				FEditorBuildUtils::CompileViewModeShaders(SomeWorld, VMI_ShaderComplexity);
 			}
@@ -7329,7 +7329,7 @@ void UEditorEngine::UpdateShaderComplexityMaterials()
 
 void UEditorEngine::OnSceneMaterialsModified()
 {
-	UpdateShaderComplexityMaterials();
+	UpdateShaderComplexityMaterials(false);
 }
 
 void UEditorEngine::SetMaterialsFeatureLevel(const ERHIFeatureLevel::Type InFeatureLevel)
@@ -7368,6 +7368,14 @@ void UEditorEngine::SetMaterialsFeatureLevel(const ERHIFeatureLevel::Type InFeat
 	GShaderCompilingManager->ProcessAsyncResults(false, true);
 
 	PreviewFeatureLevelChanged.Broadcast(InFeatureLevel);
+
+	// The feature level changed, so existing debug view materials are invalid and need to be rebuilt.
+	// This process must follow the PreviewFeatureLevelChanged event, because any listeners need
+	// opportunity to switch to the new feature level first.
+	void ClearDebugViewMaterials(UMaterialInterface*);
+	ClearDebugViewMaterials(nullptr);
+
+	UpdateShaderComplexityMaterials(true);
 }
 
 void UEditorEngine::SetFeatureLevelPreview(const ERHIFeatureLevel::Type InPreviewFeatureLevel)
@@ -7433,7 +7441,14 @@ void UEditorEngine::ToggleFeatureLevelPreview()
 
 	PreviewFeatureLevelChanged.Broadcast(NewPreviewFeatureLevel);
 	
-	GEditor->OnSceneMaterialsModified();
+	// The feature level changed, so existing debug view materials are invalid and need to be rebuilt.
+	// This process must follow the PreviewFeatureLevelChanged event, because any listeners need
+	// opportunity to switch to the new feature level first.
+	void ClearDebugViewMaterials(UMaterialInterface*);
+	ClearDebugViewMaterials(nullptr);
+
+	UpdateShaderComplexityMaterials(true);
+
 	GEditor->RedrawAllViewports();
 }
 
