[NO NON-PII COMMIT MESSAGE]

diff --git a/Engine/Source/Editor/UnrealEd/Private/EditorEngine.cpp b/Engine/Source/Editor/UnrealEd/Private/EditorEngine.cpp
--- a/Engine/Source/Editor/UnrealEd/Private/EditorEngine.cpp
+++ b/Engine/Source/Editor/UnrealEd/Private/EditorEngine.cpp
@@ -4467,7 +4467,7 @@ void UEditorEngine::OnPreSaveWorld(uint32 SaveFlags, UWorld* World)
 
 	// If we can get the streaming level, we should remove the editor transform before saving
 	ULevelStreaming* LevelStreaming = FLevelUtils::FindStreamingLevel( World->PersistentLevel );
-	if ( LevelStreaming )
+	if ( LevelStreaming && World->PersistentLevel->bAlreadyMovedActors )
 	{
 		FLevelUtils::RemoveEditorTransform(LevelStreaming);
 	}
@@ -4550,7 +4550,7 @@ void UEditorEngine::OnPostSaveWorld(uint32 SaveFlags, UWorld* World, uint32 Orig
 
 	// If got the streaming level, we should re-apply the editor transform after saving
 	ULevelStreaming* LevelStreaming = FLevelUtils::FindStreamingLevel( World->PersistentLevel );
-	if ( LevelStreaming )
+	if ( LevelStreaming && World->PersistentLevel->bAlreadyMovedActors )
 	{
 		FLevelUtils::ApplyEditorTransform(LevelStreaming);
 	}
