FEditorModeRegistry::UnregisterMode() no longer triggers events if the mode was not registered.
Play in editor stopping and editor shut down were two of the major places this was happening.

diff --git a/Engine/Source/Editor/UnrealEd/Private/EditorModeRegistry.cpp b/Engine/Source/Editor/UnrealEd/Private/EditorModeRegistry.cpp
--- a/Engine/Source/Editor/UnrealEd/Private/EditorModeRegistry.cpp
+++ b/Engine/Source/Editor/UnrealEd/Private/EditorModeRegistry.cpp
@@ -145,8 +145,9 @@ void FEditorModeRegistry::RegisterMode(FEditorModeID ModeID, TSharedRef<IEditorM
 void FEditorModeRegistry::UnregisterMode(FEditorModeID ModeID)
 {
 	// First off delete the factory
-	ModeFactories.Remove(ModeID);
-
-	OnModeUnregisteredEvent.Broadcast(ModeID);
-	RegisteredModesChanged.Broadcast();
+	if (ModeFactories.Remove(ModeID) > 0)
+	{
+		OnModeUnregisteredEvent.Broadcast(ModeID);
+		RegisteredModesChanged.Broadcast();
+	}
 }
