Added support for unattended mode to Alembic and FbxScene factories.

diff --git a/Engine/Plugins/Experimental/AlembicImporter/Source/AlembicImporter/Private/AlembicImportFactory.cpp b/Engine/Plugins/Experimental/AlembicImporter/Source/AlembicImporter/Private/AlembicImportFactory.cpp
--- a/Engine/Plugins/Experimental/AlembicImporter/Source/AlembicImporter/Private/AlembicImportFactory.cpp
+++ b/Engine/Plugins/Experimental/AlembicImporter/Source/AlembicImporter/Private/AlembicImportFactory.cpp
@@ -82,7 +82,7 @@ UObject* UAlembicImportFactory::FactoryCreateFile(UClass* InClass, UObject* InPa
 
 	bOutOperationCanceled = false;
 
-	if (bShowOption)
+	if (!GIsRunningUnattendedScript && bShowOption)
 	{
 		TSharedPtr<SAlembicImportOptions> Options;
 		ShowImportOptionsWindow(Options, UFactory::CurrentFilename, Importer);
diff --git a/Engine/Source/Editor/UnrealEd/Private/Fbx/FbxSceneImportFactory.cpp b/Engine/Source/Editor/UnrealEd/Private/Fbx/FbxSceneImportFactory.cpp
--- a/Engine/Source/Editor/UnrealEd/Private/Fbx/FbxSceneImportFactory.cpp
+++ b/Engine/Source/Editor/UnrealEd/Private/Fbx/FbxSceneImportFactory.cpp
@@ -134,18 +134,21 @@ bool GetFbxSceneImportOptions(UnFbx::FFbxImporter* FbxImporter
 
 	GlobalImportSettings->OverrideMaterials.Reset();
 
-	TSharedPtr<SWindow> ParentWindow;
-	if (FModuleManager::Get().IsModuleLoaded("MainFrame"))
+	// Don't show the import options in unattended mode
+	if (!GIsRunningUnattendedScript)
 	{
-		IMainFrameModule& MainFrame = FModuleManager::LoadModuleChecked<IMainFrameModule>("MainFrame");
-		ParentWindow = MainFrame.GetParentWindow();
-	}
-	TSharedRef<SWindow> Window = SNew(SWindow)
-		.ClientSize(FVector2D(820.f, 650.f))
-		.Title(NSLOCTEXT("UnrealEd", "FBXSceneImportOpionsTitle", "FBX Scene Import Options"));
-	TSharedPtr<SFbxSceneOptionWindow> FbxSceneOptionWindow;
+		TSharedPtr<SWindow> ParentWindow;
+		if (FModuleManager::Get().IsModuleLoaded("MainFrame"))
+		{
+			IMainFrameModule& MainFrame = FModuleManager::LoadModuleChecked<IMainFrameModule>("MainFrame");
+			ParentWindow = MainFrame.GetParentWindow();
+		}
+		TSharedRef<SWindow> Window = SNew(SWindow)
+			.ClientSize(FVector2D(820.f, 650.f))
+			.Title(NSLOCTEXT("UnrealEd", "FBXSceneImportOpionsTitle", "FBX Scene Import Options"));
+		TSharedPtr<SFbxSceneOptionWindow> FbxSceneOptionWindow;
 
-	Window->SetContent
+		Window->SetContent
 		(
 			SAssignNew(FbxSceneOptionWindow, SFbxSceneOptionWindow)
 			.SceneInfo(SceneInfoPtr)
@@ -156,13 +159,14 @@ bool GetFbxSceneImportOptions(UnFbx::FFbxImporter* FbxImporter
 			.SceneImportOptionsSkeletalMeshDisplay(SkeletalMeshImportData)
 			.OwnerWindow(Window)
 			.FullPath(Path)
-			);
+		);
 
-	FSlateApplication::Get().AddModalWindow(Window, ParentWindow, false);
+		FSlateApplication::Get().AddModalWindow(Window, ParentWindow, false);
 
-	if (!FbxSceneOptionWindow->ShouldImport())
-	{
-		return false;
+		if (!FbxSceneOptionWindow->ShouldImport())
+		{
+			return false;
+		}
 	}
 
 	//setup all options
