Disabled camera speed scale when Shift-translating objects in LevelEditor.

diff --git a/Engine/Source/Editor/UnrealEd/Private/LevelEditorViewport.cpp b/Engine/Source/Editor/UnrealEd/Private/LevelEditorViewport.cpp
--- a/Engine/Source/Editor/UnrealEd/Private/LevelEditorViewport.cpp
+++ b/Engine/Source/Editor/UnrealEd/Private/LevelEditorViewport.cpp
@@ -1636,6 +1636,7 @@ FLevelEditorViewportClient::FLevelEditorViewportClient(const TSharedPtr<SLevelVi
 	, SoundShowFlags(ESoundShowFlags::Disabled)
 	, bEditorCameraCut(false)
 	, bWasEditorCameraCut(false)
+	, bApplyCameraSpeedScaleByDistance(true)
 {
 	// By default a level editor viewport is pointed to the editor world
 	SetReferenceToWorldContext(GEditor->GetEditorWorldContext());
@@ -2507,8 +2508,10 @@ bool FLevelEditorViewportClient::InputWidgetDelta(FViewport* InViewport, EAxisLi
 
 				if( IsShiftPressed() )
 				{
+					bApplyCameraSpeedScaleByDistance = false;
 					FVector CameraDelta( Drag );
 					MoveViewportCamera( CameraDelta, FRotator::ZeroRotator );
+					bApplyCameraSpeedScaleByDistance = true;
 				}
 
 				TArray<FEdMode*> ActiveModes; 
@@ -2528,6 +2531,11 @@ bool FLevelEditorViewportClient::InputWidgetDelta(FViewport* InViewport, EAxisLi
 	return bHandled;
 }
 
+bool FLevelEditorViewportClient::ShouldScaleCameraSpeedByDistance() const
+{
+	return bApplyCameraSpeedScaleByDistance && FEditorViewportClient::ShouldScaleCameraSpeedByDistance();
+}
+
 TSharedPtr<FDragTool> FLevelEditorViewportClient::MakeDragTool( EDragTool::Type DragToolType )
 {
 	// Let the drag tool handle the transaction
diff --git a/Engine/Source/Editor/UnrealEd/Public/LevelEditorViewport.h b/Engine/Source/Editor/UnrealEd/Public/LevelEditorViewport.h
--- a/Engine/Source/Editor/UnrealEd/Public/LevelEditorViewport.h
+++ b/Engine/Source/Editor/UnrealEd/Public/LevelEditorViewport.h
@@ -181,6 +181,7 @@ public:
 	virtual void EndCameraMovement() override;
 	virtual void SetVREditView(bool bGameViewEnable) override;
 	virtual bool GetPivotForOrbit(FVector& Pivot) const override;
+	virtual bool ShouldScaleCameraSpeedByDistance() const override;
 
 	virtual bool OverrideHighResScreenshotCaptureRegion(FIntRect& OutCaptureRegion) override;
 
@@ -800,4 +801,6 @@ private:
 
 	/** Stores the previous frame's value of bEditorCameraCut in order to reset it back to false on the next frame */
 	bool					bWasEditorCameraCut;
+
+	bool bApplyCameraSpeedScaleByDistance;
 };
