Fix crash cause by a slate DetailLayoutBuilder bad usage

diff --git a/Engine/Source/Editor/Persona/Private/PersonaMeshDetails.cpp b/Engine/Source/Editor/Persona/Private/PersonaMeshDetails.cpp
--- a/Engine/Source/Editor/Persona/Private/PersonaMeshDetails.cpp
+++ b/Engine/Source/Editor/Persona/Private/PersonaMeshDetails.cpp
@@ -1055,8 +1055,8 @@ void FPersonaMeshDetails::AddLODLevelCategories(IDetailLayoutBuilder& DetailLayo
 							.ButtonFlags(ButtonFlag)
 							.OnApplyLODChangeClicked(this, &FPersonaMeshDetails::RegenerateLOD, LODIndex)
 							.OnRemoveLODClicked(this, &FPersonaMeshDetails::RemoveOneLOD, LODIndex)
-							.OnReimportClicked(this, &FPersonaMeshDetails::OnReimportLodClicked, &DetailLayout, EReimportButtonType::Reimport, LODIndex)
-							.OnReimportNewFileClicked(this, &FPersonaMeshDetails::OnReimportLodClicked, &DetailLayout, EReimportButtonType::ReimportWithNewFile, LODIndex)
+							.OnReimportClicked(this, &FPersonaMeshDetails::OnReimportLodClicked, EReimportButtonType::Reimport, LODIndex)
+							.OnReimportNewFileClicked(this, &FPersonaMeshDetails::OnReimportLodClicked, EReimportButtonType::ReimportWithNewFile, LODIndex)
 						];
 				}
 			}
@@ -1897,7 +1897,7 @@ void FPersonaMeshDetails::OnSetPostProcessBlueprint(const FAssetData& AssetData,
 	}
 }
 
-FReply FPersonaMeshDetails::OnReimportLodClicked(IDetailLayoutBuilder* DetailLayout, EReimportButtonType InReimportType, int32 InLODIndex)
+FReply FPersonaMeshDetails::OnReimportLodClicked(EReimportButtonType InReimportType, int32 InLODIndex)
 {
 	if(USkeletalMesh* SkelMesh = GetPersonaToolkit()->GetMesh())
 	{
@@ -1922,14 +1922,6 @@ FReply FPersonaMeshDetails::OnReimportLodClicked(IDetailLayoutBuilder* DetailLay
 			SkelMesh->GetLODInfo(InLODIndex)->SourceImportFilename = SourceFilenameBackup;
 		}
 
-		//Regenerate dependent LODs
-		RegenerateDependentLODs(InLODIndex);
-
-		if(DetailLayout)
-		{
-			DetailLayout->ForceRefreshDetails();
-		}
-
 		return FReply::Handled();
 	}
 
diff --git a/Engine/Source/Editor/Persona/Private/PersonaMeshDetails.h b/Engine/Source/Editor/Persona/Private/PersonaMeshDetails.h
--- a/Engine/Source/Editor/Persona/Private/PersonaMeshDetails.h
+++ b/Engine/Source/Editor/Persona/Private/PersonaMeshDetails.h
@@ -404,7 +404,7 @@ private:
 	};
 
 	// Handler for reimport buttons in LOD details
-	FReply OnReimportLodClicked(IDetailLayoutBuilder* DetailLayout, EReimportButtonType InReimportType, int32 InLODIndex);
+	FReply OnReimportLodClicked(EReimportButtonType InReimportType, int32 InLODIndex);
 
 	void OnCopySectionList(int32 LODIndex);
 	bool OnCanCopySectionList(int32 LODIndex) const;
diff --git a/Engine/Source/Editor/UnrealEd/Private/Factories/SkeletalMeshImport.cpp b/Engine/Source/Editor/UnrealEd/Private/Factories/SkeletalMeshImport.cpp
--- a/Engine/Source/Editor/UnrealEd/Private/Factories/SkeletalMeshImport.cpp
+++ b/Engine/Source/Editor/UnrealEd/Private/Factories/SkeletalMeshImport.cpp
@@ -542,6 +542,33 @@ ExistingSkelMeshData* SaveExistingSkelMeshData(USkeletalMesh* ExistingSkelMesh,
 	return ExistingMeshDataPtr;
 }
 
+void RegenerateDependentLODs(USkeletalMesh* SkelMesh, int32 LODIndex)
+{
+	check(SkelMesh);
+
+	IMeshReductionModule& ReductionModule = FModuleManager::Get().LoadModuleChecked<IMeshReductionModule>("MeshReductionInterface");
+	IMeshReduction* MeshReduction = ReductionModule.GetSkeletalMeshReductionInterface();
+	if (MeshReduction && MeshReduction->IsSupported())
+	{
+		FSkeletalMeshUpdateContext UpdateContext;
+		UpdateContext.SkeletalMesh = SkelMesh;
+		TArray<bool> DependentLODs;
+		DependentLODs.AddZeroed(SkelMesh->GetLODNum());
+		DependentLODs[LODIndex] = true;
+		for (int32 CurrentLODIndex = LODIndex + 1; CurrentLODIndex < DependentLODs.Num(); ++CurrentLODIndex)
+		{
+			FSkeletalMeshLODInfo& CurrentLODInfo = *(SkelMesh->GetLODInfo(CurrentLODIndex));
+			FSkeletalMeshOptimizationSettings& Settings = CurrentLODInfo.ReductionSettings;
+			if (CurrentLODInfo.bHasBeenSimplified && DependentLODs[Settings.BaseLOD])
+			{
+				DependentLODs[CurrentLODIndex] = true;
+				//Regenerate this LOD
+				FLODUtilities::SimplifySkeletalMeshLOD(UpdateContext, CurrentLODIndex, false);
+			}
+		}
+	}
+}
+
 void TryRegenerateLODs(ExistingSkelMeshData* MeshData, USkeletalMesh* SkeletalMesh)
 {
 	check(SkeletalMesh != nullptr);
@@ -816,12 +843,16 @@ void RestoreExistingSkelMeshData(ExistingSkelMeshData* MeshData, USkeletalMesh*
 		{
 			SkeletalMeshImportedModel->OriginalReductionSourceMeshData[ReimportLODIndex]->EmptyBulkData();
 		}
-		//Regenerate the reimport LOD
-		GWarn->BeginSlowTask(LOCTEXT("RegenReimportedLOD", "Generating reimported LOD"), true);
-		FSkeletalMeshUpdateContext UpdateContext;
-		UpdateContext.SkeletalMesh = SkeletalMesh;
-		FLODUtilities::SimplifySkeletalMeshLOD(UpdateContext, ReimportLODIndex, false);
-		GWarn->EndSlowTask();
+
+		if (SkeletalMesh->IsReductionActive(ReimportLODIndex))
+		{
+			//Regenerate the reimport LOD
+			GWarn->BeginSlowTask(LOCTEXT("RegenReimportedLOD", "Generating reimported LOD"), true);
+			FSkeletalMeshUpdateContext UpdateContext;
+			UpdateContext.SkeletalMesh = SkeletalMesh;
+			FLODUtilities::SimplifySkeletalMeshLOD(UpdateContext, ReimportLODIndex, false);
+			GWarn->EndSlowTask();
+		}
 	}
 
 	//Do everything we need for base LOD re-import
@@ -1028,6 +1059,10 @@ void RestoreExistingSkelMeshData(ExistingSkelMeshData* MeshData, USkeletalMesh*
 
 		SkeletalMesh->SetSamplingInfo(MeshData->ExistingSamplingInfo);
 	}
+	else
+	{
+		RegenerateDependentLODs(SkeletalMesh, ReimportLODIndex);
+	}
 
 	//Restore the section change only for the reimport LOD, other LOD are not affected since the material array can only grow.
 	if (MeshData->UseMaterialNameSlotWorkflow)
