Fixed potential GC leak in Python when switching worlds that contained sub-levels

diff --git a/Engine/Source/Editor/UnrealEd/Private/EditorServer.cpp b/Engine/Source/Editor/UnrealEd/Private/EditorServer.cpp
--- a/Engine/Source/Editor/UnrealEd/Private/EditorServer.cpp
+++ b/Engine/Source/Editor/UnrealEd/Private/EditorServer.cpp
@@ -2075,6 +2075,14 @@ void UEditorEngine::EditorDestroyWorld( FWorldContext & Context, const FText& Cl
 	}
 
 	FEditorSupportDelegates::PrepareToCleanseEditorObject.Broadcast(ContextWorld);
+	for (ULevel* Level : ContextWorld->GetLevels())
+	{
+		UWorld* LevelWorld = Level->GetTypedOuter<UWorld>();
+		if (ensureAlways(LevelWorld) && LevelWorld != ContextWorld && LevelWorld != NewWorld)
+		{
+			FEditorSupportDelegates::PrepareToCleanseEditorObject.Broadcast(LevelWorld);
+		}
+	}
 
 	ContextWorld->DestroyWorld( true, NewWorld );
 	Context.SetCurrentWorld(NULL);
