Fix the specific content re-import to show the open file dialog before the import option dialog

diff --git a/Engine/Source/Developer/AssetTools/Private/AssetTypeActions/AssetTypeActions_SkeletalMesh.cpp b/Engine/Source/Developer/AssetTools/Private/AssetTypeActions/AssetTypeActions_SkeletalMesh.cpp
--- a/Engine/Source/Developer/AssetTools/Private/AssetTypeActions/AssetTypeActions_SkeletalMesh.cpp
+++ b/Engine/Source/Developer/AssetTools/Private/AssetTypeActions/AssetTypeActions_SkeletalMesh.cpp
@@ -648,7 +648,7 @@ void FAssetTypeActions_SkeletalMesh::GetSourceFileLabels(const TArray<UObject*>&
 		SkeletalMesh->AssetImportData->ExtractFilenames(SourceFilePaths);
 		for (int32 SourceIndex = 0; SourceIndex < SourceFilePaths.Num(); ++SourceIndex)
 		{
-			FText SourceIndexLabel = SourceIndex == 0 ? NSSkeletalMeshSourceFileLabels::GeoAndSkinningText() : SourceIndex == 1 ? NSSkeletalMeshSourceFileLabels::GeometryText() : NSSkeletalMeshSourceFileLabels::SkinningText();
+			FText SourceIndexLabel = USkeletalMesh::GetSourceFileLabelFromIndex(SourceIndex);
 			OutSourceFileLabels.Add(SourceIndexLabel.ToString());
 		}
 	}
diff --git a/Engine/Source/Editor/UnrealEd/Classes/Factories/FbxSkeletalMeshImportData.h b/Engine/Source/Editor/UnrealEd/Classes/Factories/FbxSkeletalMeshImportData.h
--- a/Engine/Source/Editor/UnrealEd/Classes/Factories/FbxSkeletalMeshImportData.h
+++ b/Engine/Source/Editor/UnrealEd/Classes/Factories/FbxSkeletalMeshImportData.h
@@ -21,26 +21,6 @@ enum EFBXImportContentType
 	FBXICT_MAX,
 };
 
-namespace NSSkeletalMeshSourceFileLabels
-{
-	static FText GeoAndSkinningText()
-	{
-		static FText GeoAndSkinningText = (NSLOCTEXT("FBXReimport", "ImportContentTypeAll", "Geometry and Skinning Weights"));
-		return GeoAndSkinningText;
-	}
-
-	static FText GeometryText()
-	{
-		static FText GeometryText = (NSLOCTEXT("FBXReimport", "ImportContentTypeGeometry", "Geometry"));
-		return GeometryText;
-	}
-	static FText SkinningText()
-	{
-		static FText SkinningText = (NSLOCTEXT("FBXReimport", "ImportContentTypeSkinning", "Skinning Weights"));
-		return SkinningText;
-	}
-}
-
 /**
  * Import data and options used when importing a static mesh from fbx
  */
diff --git a/Engine/Source/Editor/UnrealEd/Private/Editor.cpp b/Engine/Source/Editor/UnrealEd/Private/Editor.cpp
--- a/Engine/Source/Editor/UnrealEd/Private/Editor.cpp
+++ b/Engine/Source/Editor/UnrealEd/Private/Editor.cpp
@@ -276,9 +276,10 @@ bool FReimportManager::Reimport( UObject* Obj, bool bAskForNewFileIfMissing, boo
 		if(CanReimportHandler != nullptr)
 		{
 			TArray<int32> MissingFileIndex;
+
 			// Check all filenames for missing files
 			bool bMissingFiles = false;
-			if (SourceFilenames.Num() > 0)
+			if (!bForceNewFile && SourceFilenames.Num() > 0)
 			{
 				for (int32 FileIndex = 0; FileIndex < SourceFilenames.Num(); ++FileIndex)
 				{
@@ -294,7 +295,22 @@ bool FReimportManager::Reimport( UObject* Obj, bool bAskForNewFileIfMissing, boo
 			}
 			else
 			{
-				MissingFileIndex.AddUnique(SourceFileIndex == INDEX_NONE ? 0 : SourceFileIndex);
+				int32 RealSourceFileIndex = SourceFileIndex == INDEX_NONE ? 0 : SourceFileIndex;
+				if (bForceNewFile)
+				{
+					if (SourceFilenames.IsValidIndex(RealSourceFileIndex))
+					{
+						SourceFilenames[RealSourceFileIndex].Empty();
+					}
+					else
+					{
+						//Add the missing entries
+						SourceFilenames.AddDefaulted(RealSourceFileIndex - (SourceFilenames.Num() - 1));
+					}
+					bAskForNewFileIfMissing = true;
+				}
+
+				MissingFileIndex.AddUnique(RealSourceFileIndex);
 				bMissingFiles = true;
 			}
 
diff --git a/Engine/Source/Editor/UnrealEd/Private/Factories/EditorFactories.cpp b/Engine/Source/Editor/UnrealEd/Private/Factories/EditorFactories.cpp
--- a/Engine/Source/Editor/UnrealEd/Private/Factories/EditorFactories.cpp
+++ b/Engine/Source/Editor/UnrealEd/Private/Factories/EditorFactories.cpp
@@ -5632,7 +5632,17 @@ void UReimportFbxSkeletalMeshFactory::SetReimportPaths( UObject* Obj, const FStr
 	{
 		SkeletalMesh->Modify();
 		UFbxSkeletalMeshImportData* ImportData = UFbxSkeletalMeshImportData::GetImportDataForSkeletalMesh(SkeletalMesh, ImportUI->SkeletalMeshImportData);
-		ImportData->UpdateFilenameOnly(NewReimportPath, SourceFileIndex);
+		int32 RealSourceFileIndex = SourceFileIndex == INDEX_NONE ? 0 : SourceFileIndex;
+		if (RealSourceFileIndex < ImportData->GetSourceFileCount())
+		{
+			ImportData->UpdateFilenameOnly(NewReimportPath, SourceFileIndex);
+		}
+		else
+		{
+			//Create a source file entry, this case happen when user import a specific content for the first time
+			FString SourceIndexLabel = USkeletalMesh::GetSourceFileLabelFromIndex(RealSourceFileIndex).ToString();
+			ImportData->AddFileName(NewReimportPath, RealSourceFileIndex, SourceIndexLabel);
+		}
 	}
 }
 
@@ -5727,9 +5737,9 @@ EReimportResult::Type UReimportFbxSkeletalMeshFactory::Reimport( UObject* Obj, i
 				AbsoluteFilenames.Reset();
 				ImportDataPtr->ExtractFilenames(AbsoluteFilenames);
 				//Set both geo and skinning filepath. Reuse existing file path if possible. Use the first filename(geo and skin) if it has to be create.
-				FString FilenameToAdd = SourceIndex == 1 ? OutFilename : AbsoluteFilenames.Num() > SourceIndex ? AbsoluteFilenames[1] : AbsoluteFilenames[0];
+				FString FilenameToAdd = SourceIndex == 1 ? OutFilename : AbsoluteFilenames.IsValidIndex(1) ? AbsoluteFilenames[1] : AbsoluteFilenames[0];
 				ImportDataPtr->AddFileName(FilenameToAdd, 1, NSSkeletalMeshSourceFileLabels::GeometryText().ToString());
-				FilenameToAdd = SourceIndex == 2 ? OutFilename : AbsoluteFilenames.Num() > SourceIndex ? AbsoluteFilenames[2] : AbsoluteFilenames[0];
+				FilenameToAdd = SourceIndex == 2 ? OutFilename : AbsoluteFilenames.IsValidIndex(2) ? AbsoluteFilenames[2] : AbsoluteFilenames[0];
 				ImportDataPtr->AddFileName(FilenameToAdd, 2, NSSkeletalMeshSourceFileLabels::SkinningText().ToString());
 			}
 			return true;
diff --git a/Engine/Source/Runtime/Engine/Classes/Engine/SkeletalMesh.h b/Engine/Source/Runtime/Engine/Classes/Engine/SkeletalMesh.h
--- a/Engine/Source/Runtime/Engine/Classes/Engine/SkeletalMesh.h
+++ b/Engine/Source/Runtime/Engine/Classes/Engine/SkeletalMesh.h
@@ -417,6 +417,29 @@ struct FSkeletalMaterial
 DECLARE_MULTICAST_DELEGATE_OneParam(FOnPostMeshCache, class USkeletalMesh*);
 #endif
 
+#if WITH_EDITORONLY_DATA
+namespace NSSkeletalMeshSourceFileLabels
+{
+	static FText GeoAndSkinningText()
+	{
+		static FText GeoAndSkinningText = (NSLOCTEXT("FBXReimport", "ImportContentTypeAll", "Geometry and Skinning Weights"));
+		return GeoAndSkinningText;
+	}
+
+	static FText GeometryText()
+	{
+		static FText GeometryText = (NSLOCTEXT("FBXReimport", "ImportContentTypeGeometry", "Geometry"));
+		return GeometryText;
+	}
+	static FText SkinningText()
+	{
+		static FText SkinningText = (NSLOCTEXT("FBXReimport", "ImportContentTypeSkinning", "Skinning Weights"));
+		return SkinningText;
+	}
+}
+#endif
+
+
 /**
  * SkeletalMesh is geometry bound to a hierarchical skeleton of bones which can be animated for the purpose of deforming the mesh.
  * Skeletal Meshes are built up of two parts; a set of polygons composed to make up the surface of the mesh, and a hierarchical skeleton which can be used to animate the polygons.
@@ -644,6 +667,8 @@ public:
 	UPROPERTY(EditAnywhere, Instanced, Category=ImportSettings)
 	class UAssetImportData* AssetImportData;
 
+	static FText GetSourceFileLabelFromIndex(int32 SourceFileIndex);
+
 	/** Path to the resource used to construct this skeletal mesh */
 	UPROPERTY()
 	FString SourceFilePath_DEPRECATED;
diff --git a/Engine/Source/Runtime/Engine/Private/SkeletalMesh.cpp b/Engine/Source/Runtime/Engine/Private/SkeletalMesh.cpp
--- a/Engine/Source/Runtime/Engine/Private/SkeletalMesh.cpp
+++ b/Engine/Source/Runtime/Engine/Private/SkeletalMesh.cpp
@@ -3141,6 +3141,15 @@ void USkeletalMesh::GetMappableNodeData(TArray<FName>& OutNames, TArray<FNodeIte
 		}
 	}
 }
+
+#if WITH_EDITORONLY_DATA
+FText USkeletalMesh::GetSourceFileLabelFromIndex(int32 SourceFileIndex)
+{
+	int32 RealSourceFileIndex = SourceFileIndex == INDEX_NONE ? 0 : SourceFileIndex;
+	return RealSourceFileIndex == 0 ? NSSkeletalMeshSourceFileLabels::GeoAndSkinningText() : RealSourceFileIndex == 1 ? NSSkeletalMeshSourceFileLabels::GeometryText() : NSSkeletalMeshSourceFileLabels::SkinningText();
+}
+#endif //WITH_EDITOR
+
 /*-----------------------------------------------------------------------------
 USkeletalMeshSocket
 -----------------------------------------------------------------------------*/
