Fixed the order issue with FKismetEditorUtilities::AddComponentsToBlueprint.

The RemoveAtSwap was inversing the order of the components after the first component.

diff --git a/Engine/Source/Editor/UnrealEd/Private/Kismet2/Kismet2.cpp b/Engine/Source/Editor/UnrealEd/Private/Kismet2/Kismet2.cpp
--- a/Engine/Source/Editor/UnrealEd/Private/Kismet2/Kismet2.cpp
+++ b/Engine/Source/Editor/UnrealEd/Private/Kismet2/Kismet2.cpp
@@ -1399,6 +1399,10 @@ void FKismetEditorUtilities::AddComponentsToBlueprint(UBlueprint* Blueprint, TAr
 	TArray<FSceneComponentAndFirstParent> SceneComponentNodes;
 	SceneComponentNodes.Reserve(Components.Num());
 
+	// Array of the valid ActorComponents
+	TArray<UActorComponent*> ActorComponents;
+	ActorComponents.Reserve(Components.Num());
+
 	// The boolean indicate if the scene component was added to the SceneComponentNodes array
 	TMap<USceneComponent*, bool> SceneComponentsMap;
 	SceneComponentsMap.Reserve(Components.Num());
@@ -1407,23 +1411,17 @@ void FKismetEditorUtilities::AddComponentsToBlueprint(UBlueprint* Blueprint, TAr
 	{
 		UActorComponent* ActorComponent = Components[CompIndex];
 
-		// Clear out nulls and the components we won't be able to create.
-		if (ActorComponent == nullptr || !ActorComponent->GetClass()->HasMetaData(FBlueprintMetadata::MD_BlueprintSpawnableComponent))
-		{
-			Components.RemoveAtSwap(CompIndex, 1, false);
-			// We don't want to skip the component swapped in
-			CompIndex--;
-		}
-		else
+		// Filter out nulls and the components we won't be able to create.
+		if (ActorComponent && ActorComponent->GetClass()->HasMetaData(FBlueprintMetadata::MD_BlueprintSpawnableComponent))
 		{
 			USceneComponent* SceneComponent = Cast<USceneComponent>(Components[CompIndex]);
-			// Remove the scene components from the array as they will be added in the array SceneComponentNodes
 			if (SceneComponent)
 			{
 				SceneComponentsMap.Add(SceneComponent, false);
-				Components.RemoveAtSwap(CompIndex, 1, false);
-				// We don't want to skip the component swapped in
-				CompIndex--;
+			}
+			else
+			{
+				ActorComponents.Add(ActorComponent);
 			}
 		}
 	}
@@ -1465,9 +1463,9 @@ void FKismetEditorUtilities::AddComponentsToBlueprint(UBlueprint* Blueprint, TAr
 	}
 
 	// The easy part to add the non-scene components.
-	for (int32 CompIndex = 0; CompIndex < Components.Num(); ++CompIndex)
+	for (int32 CompIndex = 0; CompIndex < ActorComponents.Num(); ++CompIndex)
 	{
-		UActorComponent* ActorComponent = Components[CompIndex];
+		UActorComponent* ActorComponent = ActorComponents[CompIndex];
 		AActor* Actor = ActorComponent->GetOwner();
 		check(Actor);
 
