Copying task stream Dev-cook to Main
Integrate all cook -iterate optimizations from //UE4/Dev-core/...

Test Scenario:
1) BuildCookStageAndRun with these arguments: -platform=Win64 -configuration=Development
2) CookIterate with these arguments: -run=Cook -CookCultures=en -TargetPlatform=WindowsClient -unversioned -stdout -unattended -iterate

Wall Time Results (as an average of running step 2) two times):
Before: ~09:40 (580 seconds) cook commandlet time
After: ~02:30 (150 seconds) cook commandlet time
=> 07:10 (430 seconds ) faster, i.e. a ~ 3.9x speedup

Win32 FileSystem Results:
Before: 1.5 million GetFileAttribute calls and 1.2 million FindNextFile calls
After: 35 0000 GetFileAttribute calls and 1.6 million FindNextFile calls
=> ~400 000 calls to FindNextFile replaces ~1.5 million calls to GetFileAttribute

diff --git a/Engine/Plugins/Editor/AssetManagerEditor/Source/AssetManagerEditor/Private/AssetManagerEditorModule.cpp b/Engine/Plugins/Editor/AssetManagerEditor/Source/AssetManagerEditor/Private/AssetManagerEditorModule.cpp
--- a/Engine/Plugins/Editor/AssetManagerEditor/Source/AssetManagerEditor/Private/AssetManagerEditorModule.cpp
+++ b/Engine/Plugins/Editor/AssetManagerEditor/Source/AssetManagerEditor/Private/AssetManagerEditorModule.cpp
@@ -296,6 +296,7 @@ public:
 	void PerformAuditConsoleCommand(const TArray<FString>& Args);
 	void PerformDependencyChainConsoleCommand(const TArray<FString>& Args);
 	void PerformDependencyClassConsoleCommand(const TArray<FString>& Args);
+	void DumpAssetRegistry(const TArray<FString>& Args);
 	void DumpAssetDependencies(const TArray<FString>& Args);
 
 	virtual void OpenAssetAuditUI(TArray<FAssetData> SelectedAssets) override;
@@ -435,6 +436,15 @@ void FAssetManagerEditorModule::StartupModule()
 			ECVF_Default
 			));
 
+	#if ASSET_REGISTRY_STATE_DUMPING_ENABLED
+		AuditCmds.Add(IConsoleManager::Get().RegisterConsoleCommand(
+			TEXT("AssetManager.DumpAssetRegistry"),
+			TEXT("Prints entries in the asset registry. Arguments are required: ObjectPath, PackageName, Path, Class, Tag, Dependencies, PackageData."),
+			FConsoleCommandWithArgsDelegate::CreateRaw(this, &FAssetManagerEditorModule::DumpAssetRegistry),
+			ECVF_Default
+		));
+	#endif
+
 		AuditCmds.Add(IConsoleManager::Get().RegisterConsoleCommand(
 			TEXT("AssetManager.DumpAssetDependencies"),
 			TEXT("Shows a list of all primary assets and the secondary assets that they depend on. Also writes out a .graphviz file"),
@@ -1997,6 +2007,21 @@ void FAssetManagerEditorModule::LogAssetsWithMultipleLabels()
 	}
 }
 
+
+void FAssetManagerEditorModule::DumpAssetRegistry(const TArray<FString>& Args)
+{
+#if ASSET_REGISTRY_STATE_DUMPING_ENABLED
+	UAssetManager& Manager = UAssetManager::Get();
+	TArray<FString> ReportLines;
+
+	const FAssetRegistryState* State = Manager.GetAssetRegistry().GetAssetRegistryState();
+	State->Dump(Args, ReportLines);
+
+	Manager.WriteCustomReport(FString::Printf(TEXT("AssetRegistryState-%s.txt"), *FDateTime::Now().ToString()), ReportLines);
+#endif
+}
+
+
 void FAssetManagerEditorModule::DumpAssetDependencies(const TArray<FString>& Args)
 {
 	if (!UAssetManager::IsValid())
diff --git a/Engine/Source/Editor/UnrealEd/Classes/CookOnTheSide/CookOnTheFlyServer.h b/Engine/Source/Editor/UnrealEd/Classes/CookOnTheSide/CookOnTheFlyServer.h
--- a/Engine/Source/Editor/UnrealEd/Classes/CookOnTheSide/CookOnTheFlyServer.h
+++ b/Engine/Source/Editor/UnrealEd/Classes/CookOnTheSide/CookOnTheFlyServer.h
@@ -197,7 +197,7 @@ private:
 	int32 LastCookedPackagesCount = 0;
 	double LastProgressDisplayTime = 0;
 
-	FString ConvertCookedPathToUncookedPath(const FString& CookedPackageName) const;
+	FName ConvertCookedPathToUncookedPath(const FString& SandboxPath, const FString& CookedPath) const;
 
 	/** Get dependencies for package */
 	const TArray<FName>& GetFullPackageDependencies(const FName& PackageName) const;
diff --git a/Engine/Source/Editor/UnrealEd/Private/Commandlets/AssetRegistryGenerator.cpp b/Engine/Source/Editor/UnrealEd/Private/Commandlets/AssetRegistryGenerator.cpp
--- a/Engine/Source/Editor/UnrealEd/Private/Commandlets/AssetRegistryGenerator.cpp
+++ b/Engine/Source/Editor/UnrealEd/Private/Commandlets/AssetRegistryGenerator.cpp
@@ -177,6 +177,51 @@ int64 FAssetRegistryGenerator::GetMaxChunkSizePerPlatform(const ITargetPlatform*
 	return -1;
 }
 
+class FPackageFileSizeVisitor : public IPlatformFile::FDirectoryStatVisitor
+{
+	TMap<FString, int64>& PackageFileSizes;
+public:
+	FPackageFileSizeVisitor(TMap<FString, int64>& InFileSizes)
+		: PackageFileSizes(InFileSizes)
+	{}
+	virtual bool Visit(const TCHAR* FilenameOrDirectory, const FFileStatData& StatData)
+	{
+		const TCHAR* Extensions[] = { TEXT(".uexp"), TEXT(".uasset"), TEXT(".ubulk"), TEXT(".ufont"), TEXT(".umap"), TEXT(".uptnl") };
+
+		if (StatData.bIsDirectory)
+			return true;
+
+		const TCHAR* Extension = FCString::Strrchr(FilenameOrDirectory, '.');
+		if (!Extension)
+			return true;
+
+		int32 ExtIndex = 0;
+		for (; ExtIndex < ARRAY_COUNT(Extensions); ++ExtIndex)
+		{
+			if (0 == FCString::Stricmp(Extension, Extensions[ExtIndex]))
+				break;
+		}
+
+		if (ExtIndex >= ARRAY_COUNT(Extensions))
+			return true;
+
+		int32 LengthWithoutExtension = Extension - FilenameOrDirectory;
+		FString FilenameWithoutExtension(LengthWithoutExtension, FilenameOrDirectory);
+
+		if (int64* CurrentPackageSize = PackageFileSizes.Find(FilenameWithoutExtension))
+		{
+			int64& TotalPackageSize = *CurrentPackageSize;
+			TotalPackageSize += StatData.FileSize;
+		}
+		else
+		{
+			PackageFileSizes.Add(FilenameWithoutExtension, StatData.FileSize);
+		}
+
+		return true;
+	}
+};
+
 bool FAssetRegistryGenerator::GenerateStreamingInstallManifest(int64 InExtraFlavorChunkSize, FSandboxPlatformFile* InSandboxFile)
 {
 	const FString Platform = TargetPlatform->PlatformName();
@@ -259,6 +304,15 @@ bool FAssetRegistryGenerator::GenerateStreamingInstallManifest(int64 InExtraFlav
 		}
 	}
 
+	TMap<FString, int64> PackageFileSizes;
+	if (MaxChunkSize > 0)
+	{
+		FString SandboxPath = InSandboxFile->GetSandboxDirectory();
+		SandboxPath.ReplaceInline(TEXT("[Platform]"), *Platform);
+		FPackageFileSizeVisitor PackageSearch(PackageFileSizes);
+		IFileManager::Get().IterateDirectoryStatRecursively(*SandboxPath, PackageSearch);
+	}
+
 	// generate per-chunk pak list files
 	for (int32 PakchunkIndex = 0; PakchunkIndex < FinalChunkManifests.Num(); ++PakchunkIndex)
 	{
@@ -271,19 +325,26 @@ bool FAssetRegistryGenerator::GenerateStreamingInstallManifest(int64 InExtraFlav
 		int32 FilenameIndex = 0;
 		TArray<FString> ChunkFilenames;
 		FinalChunkManifests[PakchunkIndex]->GenerateValueArray(ChunkFilenames);
-		int32 SubChunkIndex = 0;
-		while ( true )
+		bool bFinishedAllFiles = false;
+		for (int32 SubChunkIndex = 0; !bFinishedAllFiles; ++SubChunkIndex)
 		{
-			FString PakChunkFilename = FString::Printf(TEXT("pakchunk%d.txt"), PakchunkIndex);
-			if ( SubChunkIndex > 0 )
+			const FString PakChunkFilename = (SubChunkIndex > 0)
+				? FString::Printf(TEXT("pakchunk%d_s%d.txt"), PakchunkIndex, SubChunkIndex)
+				: FString::Printf(TEXT("pakchunk%d.txt"), PakchunkIndex);
+
+			const FString PakListFilename = FString::Printf(TEXT("%s/%s"), *TmpPackagingDir, *PakChunkFilename);
+			TUniquePtr<FArchive> PakListFile(IFileManager::Get().CreateFileWriter(*PakListFilename));
+
+			if (!PakListFile)
 			{
-				PakChunkFilename = FString::Printf(TEXT("pakchunk%d_s%d.txt"), PakchunkIndex, SubChunkIndex);
+				UE_LOG(LogAssetRegistryGenerator, Error, TEXT("Failed to open output paklist file %s"), *PakListFilename);
+				return false;
 			}
 
 			FString PakChunkOptions;
-			for ( FString CompressedChunkWildcard : CompressedChunkWildcards)
+			for (const FString& CompressedChunkWildcard : CompressedChunkWildcards)
 			{
-				if ( PakChunkFilename.MatchesWildcard(CompressedChunkWildcard) )
+				if (PakChunkFilename.MatchesWildcard(CompressedChunkWildcard))
 				{
 					PakChunkOptions += " compressed";
 					break;
@@ -312,37 +373,22 @@ bool FAssetRegistryGenerator::GenerateStreamingInstallManifest(int64 InExtraFlav
 				}
 			}
 
-			++SubChunkIndex;
-			FString PakListFilename = FString::Printf(TEXT("%s/%s"), *TmpPackagingDir, *PakChunkFilename);
-			TUniquePtr<FArchive> PakListFile(IFileManager::Get().CreateFileWriter(*PakListFilename));
-
-			if (!PakListFile)
-			{
-				UE_LOG(LogAssetRegistryGenerator, Error, TEXT("Failed to open output paklist file %s"), *PakListFilename);
-				return false;
-			}
-
-			int64 CurrentPakSize = 0;
-			bool bFinishedAllFiles = true;
-
-			if (bUseAssetManager)
+			if (bUseAssetManager && SubChunkIndex == 0)
 			{
 				// Sort so the order is consistent. If load order is important then it should be specified as a load order file to UnrealPak
 				ChunkFilenames.Sort();
 			}
 
+			int64 CurrentPakSize = 0;
+			bFinishedAllFiles = true;
 			for (; FilenameIndex < ChunkFilenames.Num(); ++FilenameIndex)
 			{
 				FString Filename = ChunkFilenames[FilenameIndex];
 				FString PakListLine = FPaths::ConvertRelativePathToFull(Filename.Replace(TEXT("[Platform]"), *Platform));
 				if (MaxChunkSize > 0)
 				{
-					const TCHAR* Extensions[] = { TEXT(".uexp"), TEXT(".uasset"), TEXT(".ubulk"), TEXT(".ufont"), TEXT(".umap"), TEXT(".uptnl") };
-					for (int32 ExtIndex = 0; ExtIndex < ARRAY_COUNT(Extensions); ++ExtIndex)
-					{
-						int64 FileSize = IFileManager::Get().FileSize(*(PakListLine + Extensions[ExtIndex]));
-						CurrentPakSize += FileSize > 0 ? FileSize : 0;
-					}
+					const int64* PackageFileSize = PackageFileSizes.Find(PakListLine);
+					CurrentPakSize += PackageFileSize ? *PackageFileSize : 0;
 					if (MaxChunkSize < CurrentPakSize)
 					{
 						// early out if we are over memory limit
@@ -368,11 +414,6 @@ bool FAssetRegistryGenerator::GenerateStreamingInstallManifest(int64 InExtraFlav
 			FString LayerString = FString::Printf(TEXT("%d\r\n"), TargetLayer);
 
 			ChunkLayerFile->Serialize(TCHAR_TO_ANSI(*LayerString), LayerString.Len());
-			
-			if (bFinishedAllFiles)
-			{
-				break;
-			}
 		}
 
 	}
@@ -977,8 +1018,8 @@ bool FAssetRegistryGenerator::SaveAssetRegistry(const FString& SandboxPath, bool
 		{
 			// Prune out the development only packages, and any assets that belong in a different chunk asset registry
 			FAssetRegistryState NewState;
-			NewState.InitializeFromExisting(State, SaveOptions);
-			NewState.PruneAssetData(CookedPackages, TSet<FName>(), ChunkBucketElement.Value, SaveOptions);
+			NewState.InitializeFromExistingAndPrune(State, CookedPackages, TSet<FName>(), ChunkBucketElement.Value, SaveOptions);
+
 			InjectEncryptionData(NewState);
 
 			// Create runtime registry data
@@ -1181,12 +1222,8 @@ bool FAssetRegistryGenerator::GatherAllPackageDependencies(FName PackageName, TA
 
 bool FAssetRegistryGenerator::GenerateAssetChunkInformationCSV(const FString& OutputPath, bool bWriteIndividualFiles)
 {
-	FString TmpString;
-	FString CSVString;
-	FString HeaderText(TEXT("ChunkID, Package Name, Class Type, Hard or Soft Chunk, File Size, Other Chunks\n"));
-	FString EndLine(TEXT("\n"));
-	FString NoneText(TEXT("None\n"));
-	CSVString = HeaderText;
+	FString TmpString, TmpStringChunks;
+	ANSICHAR HeaderText[] = "ChunkID, Package Name, Class Type, Hard or Soft Chunk, File Size, Other Chunks\n";
 
 	const TMap<FName, const FAssetData*>& ObjectToDataMap = State.GetObjectPathToAssetDataMap();
 	TArray<const FAssetData*> AssetDataList;
@@ -1201,69 +1238,93 @@ bool FAssetRegistryGenerator::GenerateAssetChunkInformationCSV(const FString& Ou
 		return A.ObjectPath < B.ObjectPath;
 	});
 
-	for (int32 PakchunkIndex = 0; PakchunkIndex < FinalChunkManifests.Num(); ++PakchunkIndex)
+	// Create file for all chunks
+	TUniquePtr<FArchive> AllChunksFile(IFileManager::Get().CreateFileWriter(*FPaths::Combine(*OutputPath, TEXT("AllChunksInfo.csv"))));
+	if (!AllChunksFile.IsValid())
+	{
+		return false;
+	}
+
+	AllChunksFile->Serialize(HeaderText, sizeof(HeaderText)-1);
+
+	// Create file for each chunk if needed
+	TArray<TUniquePtr<FArchive>> ChunkFiles;
+	if (bWriteIndividualFiles)
 	{
-		FString PerChunkManifestCSV = HeaderText;
-		for (const FAssetData* AssetDataPtr : AssetDataList)
+		for (int32 PakchunkIndex = 0, ChunkNum = FinalChunkManifests.Num(); PakchunkIndex < ChunkNum; ++PakchunkIndex)
 		{
-			const FAssetData& AssetData = *AssetDataPtr;
-			// Add only assets that have actually been cooked and belong to any chunk
-			if (AssetData.ChunkIDs.Num() > 0)
+			FArchive* ChunkFile = IFileManager::Get().CreateFileWriter(*FPaths::Combine(*OutputPath, *FString::Printf(TEXT("Chunks%dInfo.csv"), PakchunkIndex)));
+			if (ChunkFile != nullptr)
 			{
-				const FAssetPackageData* PackageData = State.GetAssetPackageData(AssetData.PackageName);
-				if (AssetData.ChunkIDs.Contains(PakchunkIndex) && PackageData && PackageData->DiskSize >= 0)
+				return false;
+			}
+			ChunkFile->Serialize(HeaderText, sizeof(HeaderText)-1);
+			ChunkFiles.Add(TUniquePtr<FArchive>(ChunkFile));
+		}
+	}
+
+	for (const FAssetData* AssetDataPtr : AssetDataList)
+	{
+		const FAssetData& AssetData = *AssetDataPtr;
+		const FAssetPackageData* PackageData = State.GetAssetPackageData(AssetData.PackageName);
+
+		// Add only assets that have actually been cooked and belong to any chunk and that have a file size
+		if (AssetData.ChunkIDs.Num() > 0 && PackageData->DiskSize > 0)
+		{
+			for (int32 PakchunkIndex : AssetData.ChunkIDs)
+			{
+				const int64 FileSize = PackageData->DiskSize;
+				FString SoftChain;
+				bool bHardChunk = false;
+				if (PakchunkIndex < ChunkManifests.Num())
 				{
-					int64 FileSize = PackageData->DiskSize;
-					FString SoftChain;
-					bool bHardChunk = false;
-					if (PakchunkIndex < ChunkManifests.Num())
-					{
-						bHardChunk = ChunkManifests[PakchunkIndex] && ChunkManifests[PakchunkIndex]->Contains(AssetData.PackageName);
-					
-						if (!bHardChunk)
-						{
-							//
-							SoftChain = GetShortestReferenceChain(AssetData.PackageName, PakchunkIndex);
-						}
-					}
-					if (SoftChain.IsEmpty())
+					bHardChunk = ChunkManifests[PakchunkIndex] && ChunkManifests[PakchunkIndex]->Contains(AssetData.PackageName);
+
+					if (!bHardChunk)
 					{
-						SoftChain = TEXT("Soft: Possibly Unassigned Asset");
+						SoftChain = GetShortestReferenceChain(AssetData.PackageName, PakchunkIndex);
 					}
+				}
+				if (SoftChain.IsEmpty())
+				{
+					SoftChain = TEXT("Soft: Possibly Unassigned Asset");
+				}
 
-					TmpString = FString::Printf(TEXT("%d,%s,%s,%s,%lld,"), PakchunkIndex, *AssetData.PackageName.ToString(), *AssetData.AssetClass.ToString(), bHardChunk ? TEXT("Hard") : *SoftChain, FileSize);
-					CSVString += TmpString;
-					PerChunkManifestCSV += TmpString;
-					if (AssetData.ChunkIDs.Num() == 1)
+				// Build "other chunks" string or None if not part of
+				TmpStringChunks.Empty(64);
+				for (const auto& OtherChunk : AssetData.ChunkIDs)
+				{
+					if (OtherChunk != PakchunkIndex)
 					{
-						CSVString += NoneText;
-						PerChunkManifestCSV += NoneText;
+						TmpString = FString::Printf(TEXT("%d "), OtherChunk);
 					}
-					else
+				}
+
+				// Build csv line
+				TmpString = FString::Printf(TEXT("%d,%s,%s,%s,%lld,%s\n"),
+					PakchunkIndex,
+					*AssetData.PackageName.ToString(),
+					*AssetData.AssetClass.ToString(),
+					bHardChunk ? TEXT("Hard") : *SoftChain,
+					FileSize,
+					AssetData.ChunkIDs.Num() == 1 ? TEXT("None") : *TmpStringChunks
+				);
+
+				// Write line to all chunks file and individual chunks files if requested
+				{
+					auto Src = StringCast<ANSICHAR>(*TmpString, TmpString.Len());
+					AllChunksFile->Serialize((ANSICHAR*)Src.Get(), Src.Length() * sizeof(ANSICHAR));
+					
+					if (bWriteIndividualFiles)
 					{
-						for (const auto& OtherChunk : AssetData.ChunkIDs)
-						{
-							if (OtherChunk != PakchunkIndex)
-							{
-								TmpString = FString::Printf(TEXT("%d "), OtherChunk);
-								CSVString += TmpString;
-								PerChunkManifestCSV += TmpString;
-							}
-						}
-						CSVString += EndLine;
-						PerChunkManifestCSV += EndLine;
+						ChunkFiles[PakchunkIndex]->Serialize((ANSICHAR*)Src.Get(), Src.Length() * sizeof(ANSICHAR));
 					}
 				}
 			}
 		}
-
-		if (bWriteIndividualFiles)
-		{
-			FFileHelper::SaveStringToFile(PerChunkManifestCSV, *FPaths::Combine(*OutputPath, *FString::Printf(TEXT("Chunks%dInfo.csv"), PakchunkIndex)));
-		}
 	}
 
-	return FFileHelper::SaveStringToFile(CSVString, *FPaths::Combine(*OutputPath, TEXT("AllChunksInfo.csv")));
+	return true;
 }
 
 void FAssetRegistryGenerator::AddPackageToManifest(const FString& PackageSandboxPath, FName PackageName, int32 ChunkId)
@@ -1295,21 +1356,21 @@ void FAssetRegistryGenerator::RemovePackageFromManifest(FName PackageName, int32
 	}
 }
 
-void FAssetRegistryGenerator::ResolveChunkDependencyGraph(const FChunkDependencyTreeNode& Node, const FChunkPackageSet& BaseAssetSet, TArray<TArray<FName>>& OutPackagesMovedBetweenChunks)
+void FAssetRegistryGenerator::ResolveChunkDependencyGraph(const FChunkDependencyTreeNode& Node, const TSet<FName>& BaseAssetSet, TArray<TArray<FName>>& OutPackagesMovedBetweenChunks)
 {
 	if (FinalChunkManifests.Num() > Node.ChunkID && FinalChunkManifests[Node.ChunkID])
 	{
 		for (auto It = BaseAssetSet.CreateConstIterator(); It; ++It)
 		{
 			// Remove any assets belonging to our parents.			
-			if (FinalChunkManifests[Node.ChunkID]->Remove(It.Key()) > 0)
+			if (FinalChunkManifests[Node.ChunkID]->Remove(*It) > 0)
 			{
-				OutPackagesMovedBetweenChunks[Node.ChunkID].Add(It.Key());
-				UE_LOG(LogAssetRegistryGenerator, Verbose, TEXT("Removed %s from chunk %i because it is duplicated in another chunk."), *It.Key().ToString(), Node.ChunkID);
+				OutPackagesMovedBetweenChunks[Node.ChunkID].Add(*It);
+				UE_LOG(LogAssetRegistryGenerator, Verbose, TEXT("Removed %s from chunk %i because it is duplicated in another chunk."), *It->ToString(), Node.ChunkID);
 			}
 		}
 		
-		FChunkPackageSet ModifiedAssetSet;
+		TSet<FName> ModifiedAssetSet;
 		
 		// Add the current Chunk's assets
 		for (auto It = FinalChunkManifests[Node.ChunkID]->CreateConstIterator(); It; ++It)//for (const auto It : *(FinalChunkManifests[Node.ChunkID]))
@@ -1319,7 +1380,7 @@ void FAssetRegistryGenerator::ResolveChunkDependencyGraph(const FChunkDependency
 				ModifiedAssetSet = BaseAssetSet;
 			}
 			
-			ModifiedAssetSet.Add(It.Key(), It.Value());
+			ModifiedAssetSet.Add(It.Key());
 		}
 		
 		const auto& AssetSet = ModifiedAssetSet.Num() ? ModifiedAssetSet : BaseAssetSet;
@@ -1483,7 +1544,7 @@ void FAssetRegistryGenerator::FixupPackageDependenciesForChunks(FSandboxPlatform
 	PackagesRemovedFromChunks.AddDefaulted(ChunkManifests.Num());
 
 	//Finally, if the previous step may added any extra packages to the 0 chunk. Pull them out of other chunks and save space
-	ResolveChunkDependencyGraph(*ChunkDepGraph, FChunkPackageSet(), PackagesRemovedFromChunks);
+	ResolveChunkDependencyGraph(*ChunkDepGraph, TSet<FName>(), PackagesRemovedFromChunks);
 
 	for (int32 PakchunkIndex = 0; PakchunkIndex < ChunkManifests.Num(); ++PakchunkIndex)
 	{
diff --git a/Engine/Source/Editor/UnrealEd/Private/CookOnTheFlyServer.cpp b/Engine/Source/Editor/UnrealEd/Private/CookOnTheFlyServer.cpp
--- a/Engine/Source/Editor/UnrealEd/Private/CookOnTheFlyServer.cpp
+++ b/Engine/Source/Editor/UnrealEd/Private/CookOnTheFlyServer.cpp
@@ -1427,8 +1427,6 @@ UCookOnTheFlyServer::UCookOnTheFlyServer(const FObjectInitializer& ObjectInitial
 	bIsSavingPackage(false),
 	AssetRegistry(nullptr)
 {
-	PackageNameCache = new FPackageNameCache;
-	PackageTracker	 = new FPackageTracker(PackageNameCache);
 }
 
 UCookOnTheFlyServer::~UCookOnTheFlyServer()
@@ -1436,7 +1434,7 @@ UCookOnTheFlyServer::~UCookOnTheFlyServer()
 	FCoreDelegates::OnFConfigCreated.RemoveAll(this);
 	FCoreDelegates::OnFConfigDeleted.RemoveAll(this);
 
-		delete CookByTheBookOptions;
+	delete CookByTheBookOptions;
 	CookByTheBookOptions = nullptr;
 
 	delete PackageTracker;
@@ -1444,6 +1442,8 @@ UCookOnTheFlyServer::~UCookOnTheFlyServer()
 
 	delete PackageNameCache;
 	PackageNameCache = nullptr;
+
+	ClearHierarchyTimers();
 }
 
 // this tick only happens in the editor cook commandlet directly calls tick on the side
@@ -2070,12 +2070,8 @@ uint32 UCookOnTheFlyServer::TickCookOnTheSide( const float TimeSlice, uint32 &Co
 	{
 		uint32 Result = FullLoadAndSave(CookedPackageCount);
 
-		UE_LOG(LogCook, Display, TEXT("Finishing up..."));
-
 		CookByTheBookFinished();
 
-		UE_LOG(LogCook, Display, TEXT("Done!"));
-
 		return Result;
 	}
 
@@ -2135,15 +2131,26 @@ uint32 UCookOnTheFlyServer::TickCookOnTheSide( const float TimeSlice, uint32 &Co
 			}
 		}
 
-		FFilePlatformRequest ToBuild;
-		
-		if (HasCookRequests())
+		auto GetNextUncookedRequestLambda = [this](FFilePlatformRequest& OutToBuild) -> bool
 		{
-			PackageTracker->CookRequests.Dequeue(/* out */ ToBuild);
-		}
-		else
+			while (this->PackageTracker->CookRequests.Dequeue(OutToBuild))
+			{
+				if (this->PackageTracker->CookedPackages.Exists(OutToBuild))
+				{
+#if DEBUG_COOKONTHEFLY
+					UE_LOG(LogCook, Display, TEXT("Package for platform already cooked %s, discarding request"), *ToBuild.GetFilename().ToString());
+#endif
+					continue;
+				}
+				return true;
+			}
+			return false;
+		};
+
+		FFilePlatformRequest ToBuild;
+		if (!GetNextUncookedRequestLambda(ToBuild))
 		{
-			// no more to do this tick break out and do some other stuff
+			// no more real work to do this tick, break out and do some other stuff
 			break;
 		}
 
@@ -2399,6 +2406,9 @@ uint32 UCookOnTheFlyServer::TickCookOnTheSide( const float TimeSlice, uint32 &Co
 		check(IsCookByTheBookMode());
 
 		// if we are out of stuff and we are in cook by the book from the editor mode then we finish up
+		UE_CLOG(!(TickFlags & ECookTickFlags::HideProgressDisplay) && (GCookProgressDisplay & (int32)ECookProgressDisplayMode::RemainingPackages),
+			LogCook, Display, TEXT("Cooked packages %d Packages Remain %d Total %d"),
+			PackageTracker->CookedPackages.Num(), 0, PackageTracker->CookedPackages.Num());
 		CookByTheBookFinished();
 	}
 
@@ -4045,9 +4055,13 @@ void UCookOnTheFlyServer::Initialize( ECookMode::Type DesiredCookMode, ECookInit
 
 	FCoreUObjectDelegates::GetPreGarbageCollectDelegate().AddUObject(this, &UCookOnTheFlyServer::PreGarbageCollect);
 
-	if (IsCookByTheBookMode() && !IsCookingInEditor())
+	if (CurrentCookMode != ECookMode::CookByTheBook)
 	{
-		FCoreUObjectDelegates::PackageCreatedForLoad.AddUObject(this, &UCookOnTheFlyServer::MaybeMarkPackageAsAlreadyLoaded);
+		// For standalone CookByTheBook the PackageNameCache and PackageTracker are initialized 
+		// later to benefit from a fully scanned and initialized AssetRegistry.
+		// For IsCookingInEditor() these objects are required for the FCoreUObjectDelegates below
+		PackageNameCache = new FPackageNameCache();
+		PackageTracker = new FPackageTracker(PackageNameCache);
 	}
 
 	if (IsCookingInEditor())
@@ -4956,16 +4970,16 @@ bool UCookOnTheFlyServer::SaveCurrentIniSettings(const ITargetPlatform* TargetPl
 
 }
 
-FString UCookOnTheFlyServer::ConvertCookedPathToUncookedPath(const FString& CookedRelativeFilename) const 
+FName UCookOnTheFlyServer::ConvertCookedPathToUncookedPath(const FString& SandboxPath, const FString& CookedPath) const
 {
 	// Check for remapped plugins' cooked content
-	if (PluginsToRemap.Num() > 0 && CookedRelativeFilename.Contains(REMAPPED_PLUGGINS))
+	if (PluginsToRemap.Num() > 0 && CookedPath.Contains(REMAPPED_PLUGGINS))
 	{
-		int32 RemappedIndex = CookedRelativeFilename.Find(REMAPPED_PLUGGINS);
+		int32 RemappedIndex = CookedPath.Find(REMAPPED_PLUGGINS);
 		check(RemappedIndex >= 0);
 		static uint32 RemappedPluginStrLen = FCString::Strlen(REMAPPED_PLUGGINS);
 		// Snip everything up through the RemappedPlugins/ off so we can find the plugin it corresponds to
-		FString PluginPath = CookedRelativeFilename.RightChop(RemappedIndex + RemappedPluginStrLen + 1);
+		FString PluginPath = CookedPath.RightChop(RemappedIndex + RemappedPluginStrLen + 1);
 		FString FullUncookedPath;
 		// Find the plugin that owns this content
 		for (TSharedRef<IPlugin> Plugin : PluginsToRemap)
@@ -4982,40 +4996,24 @@ FString UCookOnTheFlyServer::ConvertCookedPathToUncookedPath(const FString& Cook
 
 		if (FullUncookedPath.Len() > 0)
 		{
-			return FullUncookedPath;
+			return FName(*FullUncookedPath);
 		}
 		// Otherwise fall through to sandbox handling
 	}
 
-	const FString CookedFilename = FPaths::ConvertRelativePathToFull(CookedRelativeFilename);
-
+	FString UncookedPath;
 
-	FString SandboxDirectory = SandboxFile->GetSandboxDirectory();
-
-	SandboxDirectory.ReplaceInline(TEXT("[PLATFORM]"), TEXT(""));
-	SandboxDirectory.ReplaceInline(TEXT("//"), TEXT("/"));
-
-	FString CookedFilenameNoSandbox = CookedFilename;
-	CookedFilenameNoSandbox.RemoveFromStart(SandboxDirectory);
-	int32 EndOfPlatformIndex = 0;
-
-	// assume at this point the cook platform is the next thing on the path
-	FString CookedFilenameNoPlatform = CookedFilename;
-
-	if (CookedFilenameNoSandbox.FindChar(TEXT('/'), EndOfPlatformIndex))
+	if (CookedPath.StartsWith(SandboxPath))
 	{
-		CookedFilenameNoPlatform = SandboxFile->GetSandboxDirectory() / CookedFilenameNoSandbox.Mid(EndOfPlatformIndex);
-		CookedFilenameNoPlatform.ReplaceInline(TEXT("//"), TEXT("/"));
+		UncookedPath = CookedPath.Replace(*SandboxPath, *FPaths::GetRelativePathToRoot());
+	}
+	else
+	{
+		FString FullCookedFilename = FPaths::ConvertRelativePathToFull(CookedPath);
+		UncookedPath = FullCookedFilename.Replace(*SandboxPath, *FPaths::GetRelativePathToRoot());
 	}
-	
-	// after removing the cooked platform we can use the sandbox file to convert back to a uncooked path
-	FString FullUncookedPath = SandboxFile->ConvertFromSandboxPath(*CookedFilenameNoPlatform);
-
-	// make the result a standard filename (relative)
-	FPaths::MakeStandardFilename(FullUncookedPath);
-	return FullUncookedPath;
-
 
+	return FName(*UncookedPath);
 }
 
 void UCookOnTheFlyServer::GetAllCookedFiles(TMap<FName, FName>& UncookedPathToCookedPath, const FString& SandboxPath)
@@ -5027,11 +5025,7 @@ void UCookOnTheFlyServer::GetAllCookedFiles(TMap<FName, FName>& UncookedPathToCo
 	for (const FString& CookedFile : CookedFiles)
 	{
 		const FName CookedFName(*CookedFile);
-
-		const FString CookedFullPath = FPaths::ConvertRelativePathToFull(CookedFile);
-		const FString UncookedFilename = ConvertCookedPathToUncookedPath(CookedFullPath);
-
-		const FName UncookedFName(*UncookedFilename);
+		const FName UncookedFName = ConvertCookedPathToUncookedPath(SandboxPath, CookedFile);
 		UncookedPathToCookedPath.Add(UncookedFName, CookedFName);
 	}
 }
@@ -5141,6 +5135,8 @@ void UCookOnTheFlyServer::PopulateCookedPackagesFromDisk(const TArray<ITargetPla
 		const static FName NAME_DummyCookedFilename(TEXT("DummyCookedFilename")); // pls never name a package dummycookedfilename otherwise shit might go wonky
 		if (bIsIterateSharedBuild)
 		{
+			check(IFileManager::Get().FileExists(*NAME_DummyCookedFilename.ToString()) == false);
+
 			TSet<FName> ExistingPackages = ModifiedPackages;
 			ExistingPackages.Append(RemovedPackages);
 			ExistingPackages.Append(IdenticalCookedPackages);
@@ -5196,11 +5192,6 @@ void UCookOnTheFlyServer::PopulateCookedPackagesFromDisk(const TArray<ITargetPla
 					bShouldKeep = false;
 				}
 			}
-				
-			if (CookedFile == NAME_DummyCookedFilename)
-			{
-				check(IFileManager::Get().FileExists(*CookedFile.ToString()) == false);
-			}
 
 			TArray<FName> PlatformNames;
 			PlatformNames.Add(PlatformFName);
@@ -5208,10 +5199,6 @@ void UCookOnTheFlyServer::PopulateCookedPackagesFromDisk(const TArray<ITargetPla
 			if (bShouldKeep)
 			{
 				// Add this package to the CookedPackages list so that we don't try cook it again
-				if (CookedFile != NAME_DummyCookedFilename)
-				{
-					check(IFileManager::Get().FileExists(*CookedFile.ToString()));
-				}
 				TArray<bool> Succeeded;
 				Succeeded.Add(true);
 
@@ -5452,16 +5439,23 @@ void UCookOnTheFlyServer::GenerateAssetRegistry()
 
 void UCookOnTheFlyServer::GenerateLongPackageNames(TArray<FName>& FilesInPath)
 {
+	TSet<FName> FilesInPathSet;
 	TArray<FName> FilesInPathReverse;
+	FilesInPathSet.Reserve(FilesInPath.Num());
 	FilesInPathReverse.Reserve(FilesInPath.Num());
 
-	for( int32 FileIndex = 0; FileIndex < FilesInPath.Num(); FileIndex++ )
+	for (int32 FileIndex = 0; FileIndex < FilesInPath.Num(); FileIndex++)
 	{
-		const FString& FileInPath = FilesInPath[FilesInPath.Num() - FileIndex - 1].ToString();
+		const FName& FileInPathFName = FilesInPath[FilesInPath.Num() - FileIndex - 1];
+		const FString& FileInPath = FileInPathFName.ToString();
 		if (FPackageName::IsValidLongPackageName(FileInPath))
 		{
-			const FName FileInPathFName(*FileInPath);
-			FilesInPathReverse.AddUnique(FileInPathFName);
+			bool bIsAlreadyAdded;
+			FilesInPathSet.Add(FileInPathFName, &bIsAlreadyAdded);
+			if (!bIsAlreadyAdded)
+			{
+				FilesInPathReverse.Add(FileInPathFName);
+			}
 		}
 		else
 		{
@@ -5470,7 +5464,12 @@ void UCookOnTheFlyServer::GenerateLongPackageNames(TArray<FName>& FilesInPath)
 			if (FPackageName::TryConvertFilenameToLongPackageName(FileInPath, LongPackageName, &FailureReason))
 			{
 				const FName LongPackageFName(*LongPackageName);
-				FilesInPathReverse.AddUnique(LongPackageFName);
+				bool bIsAlreadyAdded;
+				FilesInPathSet.Add(LongPackageFName, &bIsAlreadyAdded);
+				if (!bIsAlreadyAdded)
+				{
+					FilesInPathReverse.Add(LongPackageFName);
+				}
 			}
 			else
 			{
@@ -5479,12 +5478,8 @@ void UCookOnTheFlyServer::GenerateLongPackageNames(TArray<FName>& FilesInPath)
 			}
 		}
 	}
-	// Exchange(FilesInPathReverse, FilesInPath);
 	FilesInPath.Empty(FilesInPathReverse.Num());
-	for ( const FName& Files : FilesInPathReverse )
-	{
-		FilesInPath.Add(Files);
-	}
+	FilesInPath.Append(FilesInPathReverse);
 }
 
 void UCookOnTheFlyServer::AddFileToCook( TArray<FName>& InOutFilesToCook, const FString &InFilename ) const
@@ -6231,6 +6226,7 @@ void UCookOnTheFlyServer::CookByTheBookFinished()
 	check( IsCookByTheBookMode() );
 	check( CookByTheBookOptions->bRunning == true );
 
+	UE_LOG(LogCook, Display, TEXT("Finishing up..."));
 
 	UPackage::WaitForAsyncFileWrites();
 
@@ -6420,6 +6416,8 @@ void UCookOnTheFlyServer::CookByTheBookFinished()
 
 	OutputHierarchyTimers();
 	ClearHierarchyTimers();
+
+	UE_LOG(LogCook, Display, TEXT("Done!"));
 }
 
 void UCookOnTheFlyServer::BuildMapDependencyGraph(const FName& PlatformName)
@@ -6659,6 +6657,14 @@ void UCookOnTheFlyServer::StartCookByTheBook( const FCookByTheBookStartupOptions
 	CookByTheBookOptions->bErrorOnEngineContentUse = CookByTheBookStartupOptions.bErrorOnEngineContentUse;
 
 	GenerateAssetRegistry();
+	if (CurrentCookMode == ECookMode::CookByTheBook)
+	{
+		check(!PackageNameCache);
+		check(!PackageTracker);
+		PackageNameCache = new FPackageNameCache(AssetRegistry);
+		PackageTracker = new FPackageTracker(PackageNameCache);
+		FCoreUObjectDelegates::PackageCreatedForLoad.AddUObject(this, &UCookOnTheFlyServer::MaybeMarkPackageAsAlreadyLoaded);
+	}
 
 	// Find all the localized packages and map them back to their source package
 	{
@@ -6677,22 +6683,30 @@ void UCookOnTheFlyServer::StartCookByTheBook( const FCookByTheBookStartupOptions
 
 		TArray<FString> RootPaths;
 		FPackageName::QueryRootContentPaths(RootPaths);
+
+		FARFilter Filter;
+		Filter.bRecursivePaths = true;
+		Filter.bIncludeOnlyOnDiskAssets = false;
+		Filter.PackagePaths.Reserve(AllCulturesToCook.Num() * RootPaths.Num());
 		for (const FString& RootPath : RootPaths)
 		{
 			for (const FString& CultureName : AllCulturesToCook)
 			{
-				TArray<FAssetData> AssetDataForCulture;
-				AssetRegistry->GetAssetsByPath(*(RootPath / TEXT("L10N") / CultureName), AssetDataForCulture, true);
+				FString LocalizedPackagePath = RootPath / TEXT("L10N") / CultureName;
+				Filter.PackagePaths.Add(*LocalizedPackagePath);
+			}
+		}
 
-				for (const FAssetData& AssetData : AssetDataForCulture)
-				{
-					const FName LocalizedPackageName = AssetData.PackageName;
-					const FName SourcePackageName = *FPackageName::GetSourcePackagePath(LocalizedPackageName.ToString());
+		TArray<FAssetData> AssetDataForCultures;
+		AssetRegistry->GetAssets(Filter, AssetDataForCultures);
 
-					TArray<FName>& LocalizedPackageNames = CookByTheBookOptions->SourceToLocalizedPackageVariants.FindOrAdd(SourcePackageName);
-					LocalizedPackageNames.AddUnique(LocalizedPackageName);
-				}
-			}
+		for (const FAssetData& AssetData : AssetDataForCultures)
+		{
+			const FName LocalizedPackageName = AssetData.PackageName;
+			const FName SourcePackageName = *FPackageName::GetSourcePackagePath(LocalizedPackageName.ToString());
+
+			TArray<FName>& LocalizedPackageNames = CookByTheBookOptions->SourceToLocalizedPackageVariants.FindOrAdd(SourcePackageName);
+			LocalizedPackageNames.AddUnique(LocalizedPackageName);
 		}
 	}
 
diff --git a/Engine/Source/Editor/UnrealEd/Private/Cooker/PackageNameCache.h b/Engine/Source/Editor/UnrealEd/Private/Cooker/PackageNameCache.h
--- a/Engine/Source/Editor/UnrealEd/Private/Cooker/PackageNameCache.h
+++ b/Engine/Source/Editor/UnrealEd/Private/Cooker/PackageNameCache.h
@@ -3,6 +3,8 @@
 #pragma once
 
 #include "CoreMinimal.h"
+#include "IAssetRegistry.h"
+#include "Algo/Find.h"
 #include "Misc/PackageName.h"
 
 struct FCachedPackageFilename
@@ -30,7 +32,7 @@ struct FCachedPackageFilename
 
 struct FPackageNameCache
 {
-	// Package name cache
+	FPackageNameCache(IAssetRegistry* InAssetRegistry = nullptr) : AssetRegistry(InAssetRegistry) {}
 
 	FString			GetCachedPackageFilename(const UPackage* Package) const;
 
@@ -45,8 +47,11 @@ struct FPackageNameCache
 	bool			ClearPackageFilenameCacheForPackage(const UPackage* Package) const;
 
 private:
+	bool DoesPackageExist(const FName& PackageName, FString* OutFilename) const;
 	const FCachedPackageFilename& Cache(const FName& PackageName) const;
 
+	IAssetRegistry* AssetRegistry;
+
 	mutable TMap<FName, FCachedPackageFilename> PackageFilenameCache; // filename cache (only process the string operations once)
 	mutable TMap<FName, FName>					PackageFilenameToPackageFNameCache;
 
@@ -84,6 +89,31 @@ bool FPackageNameCache::ClearPackageFilenameCacheForPackage(const UPackage* Pack
 	return PackageFilenameCache.Remove(Package->GetFName()) >= 1;
 }
 
+bool FPackageNameCache::DoesPackageExist(const FName& PackageName, FString* OutFilename) const
+{
+	if (!AssetRegistry)
+	{
+		return FPackageName::DoesPackageExist(PackageName.ToString(), NULL, OutFilename, false);
+	}
+
+	TArray<FAssetData> Assets;
+	AssetRegistry->GetAssetsByPackageName(PackageName, Assets, /*bIncludeOnlyDiskAssets*/ true);
+
+	if (Assets.Num() <= 0)
+	{
+		return false;
+	}
+
+	if (OutFilename)
+	{
+		const bool ContainsMap = Algo::FindByPredicate(Assets, [](const FAssetData& Asset) { return Asset.PackageFlags & PKG_ContainsMap; }) != nullptr;
+		const FString& PackageExtension = ContainsMap ? FPackageName::GetMapPackageExtension() : FPackageName::GetAssetPackageExtension();
+		*OutFilename = FPackageName::LongPackageNameToFilename(PackageName.ToString(), PackageExtension);
+	}
+
+	return true;
+}
+
 const FCachedPackageFilename& FPackageNameCache::Cache(const FName& PackageName) const
 {
 	check(IsInGameThread());
@@ -102,7 +132,7 @@ const FCachedPackageFilename& FPackageNameCache::Cache(const FName& PackageName)
 	FString StandardFilename;
 	FName StandardFileFName = NAME_None;
 
-	if (FPackageName::DoesPackageExist(PackageName.ToString(), NULL, &Filename, false))
+	if (DoesPackageExist(PackageName, &Filename))
 	{
 		StandardFilename = PackageFilename = FPaths::ConvertRelativePathToFull(Filename);
 
diff --git a/Engine/Source/Editor/UnrealEd/Public/Commandlets/AssetRegistryGenerator.h b/Engine/Source/Editor/UnrealEd/Public/Commandlets/AssetRegistryGenerator.h
--- a/Engine/Source/Editor/UnrealEd/Public/Commandlets/AssetRegistryGenerator.h
+++ b/Engine/Source/Editor/UnrealEd/Public/Commandlets/AssetRegistryGenerator.h
@@ -335,9 +335,9 @@ private:
 	FString	GetShortestReferenceChain(FName PackageName, int32 ChunkID);
 
 	/** Deprecated method to remove redundant chunks */
-	void ResolveChunkDependencyGraph(const FChunkDependencyTreeNode& Node, const FChunkPackageSet& BaseAssetSet, TArray<TArray<FName>>& OutPackagesMovedBetweenChunks);
+	void ResolveChunkDependencyGraph(const FChunkDependencyTreeNode& Node, const TSet<FName>& BaseAssetSet, TArray<TArray<FName>>& OutPackagesMovedBetweenChunks);
 
-	/** Helper function to verify Chunk asset assigment is valid */
+	/** Helper function to verify Chunk asset assignment is valid */
 	bool CheckChunkAssetsAreNotInChild(const FChunkDependencyTreeNode& Node);
 
 	/** Helper function to create a given collection. */
diff --git a/Engine/Source/Runtime/AssetRegistry/Private/AssetRegistryState.cpp b/Engine/Source/Runtime/AssetRegistry/Private/AssetRegistryState.cpp
--- a/Engine/Source/Runtime/AssetRegistry/Private/AssetRegistryState.cpp
+++ b/Engine/Source/Runtime/AssetRegistry/Private/AssetRegistryState.cpp
@@ -215,6 +215,125 @@ void FAssetRegistryState::Reset()
 	CachedPackageData.Empty();
 }
 
+void FAssetRegistryState::FilterTags(const FAssetDataTagMapSharedView& InTagsAndValues, FAssetDataTagMap& OutTagsAndValues, const TSet<FName>* ClassSpecificFilterlist, const FAssetRegistrySerializationOptions& Options)
+{
+	static FName WildcardName(TEXT("*"));
+	const TSet<FName>* AllClassesFilterlist = Options.CookFilterlistTagsByClass.Find(WildcardName);
+
+	// Exclude blacklisted tags or include only white listed tags, based on how we were configured in ini
+	for (const auto& TagPair : InTagsAndValues)
+	{
+		const bool bInAllClasseslist = AllClassesFilterlist && (AllClassesFilterlist->Contains(TagPair.Key) || AllClassesFilterlist->Contains(WildcardName));
+		const bool bInClassSpecificlist = ClassSpecificFilterlist && (ClassSpecificFilterlist->Contains(TagPair.Key) || ClassSpecificFilterlist->Contains(WildcardName));
+		if (Options.bUseAssetRegistryTagsWhitelistInsteadOfBlacklist)
+		{
+			// It's a white list, only include it if it is in the all classes list or in the class specific list
+			if (bInAllClasseslist || bInClassSpecificlist)
+			{
+				// It is in the white list. Keep it.
+				OutTagsAndValues.Add(TagPair.Key, TagPair.Value);
+			}
+		}
+		else
+		{
+			// It's a blacklist, include it unless it is in the all classes list or in the class specific list
+			if (!bInAllClasseslist && !bInClassSpecificlist)
+			{
+				// It isn't in the blacklist. Keep it.
+				OutTagsAndValues.Add(TagPair.Key, TagPair.Value);
+			}
+		}
+	}
+}
+
+void FAssetRegistryState::InitializeFromExistingAndPrune(const FAssetRegistryState & ExistingState, const TSet<FName>& RequiredPackages, const TSet<FName>& RemovePackages, const TSet<int32> ChunksToKeep, const FAssetRegistrySerializationOptions & Options)
+{
+	const bool bIsFilteredByChunkId = ChunksToKeep.Num() != 0;
+	const bool bIsFilteredByRequiredPackage = RequiredPackages.Num() != 0;
+	const bool bIsFilteredByRemovedPackages = RemovePackages.Num() != 0;
+	const bool bIsFilteredByNoTags = Options.bFilterAssetDataWithNoTags;
+
+	// Duplicate asset data entries
+	for (const TPair<FName, FAssetData*>& AssetPair : ExistingState.CachedAssetsByObjectPath)
+	{
+		const FAssetData* AssetData = AssetPair.Value;
+		if (AssetData)
+		{
+			bool bDiscard = false;
+			bDiscard |= bIsFilteredByChunkId && !AssetData->ChunkIDs.ContainsByPredicate([&](int32 ChunkId) { return ChunksToKeep.Contains(ChunkId); });
+			bDiscard |= bIsFilteredByRequiredPackage && !RequiredPackages.Contains(AssetData->PackageName);
+			bDiscard |= bIsFilteredByRemovedPackages && RemovePackages.Contains(AssetData->PackageName);
+			bDiscard |= bIsFilteredByNoTags && AssetData->TagsAndValues.Num() == 0 && !FPackageName::IsLocalizedPackage(AssetData->PackageName.ToString());
+
+			if (bDiscard)
+			{
+				continue;
+			}
+
+			FAssetDataTagMap NewTagsAndValues;
+			FAssetRegistryState::FilterTags(AssetData->TagsAndValues, NewTagsAndValues, Options.CookFilterlistTagsByClass.Find(AssetData->AssetClass), Options);
+
+			FAssetData* NewAssetData = new FAssetData(AssetData->PackageName, AssetData->PackagePath, AssetData->AssetName,
+				AssetData->AssetClass, NewTagsAndValues, AssetData->ChunkIDs, AssetData->PackageFlags);
+			// Add asset to new state
+			AddAssetData(NewAssetData);
+			// Also add it's package if it doesn't exist
+			CreateOrGetAssetPackageData(NewAssetData->PackageName);
+		}
+	}
+
+	TSet<FAssetIdentifier> ScriptPackages;
+
+	// Build list of script packages for depend nodes copying
+	for (const FName& Name : RequiredPackages)
+	{
+		bool bIsScriptPackage = FPackageName::IsScriptPackage(Name.ToString());
+		const FAssetPackageData* OldData = ExistingState.GetAssetPackageData(Name);
+		if (bIsScriptPackage) 
+		{
+			if (OldData)
+			{
+				ScriptPackages.Add(Name);
+
+				FAssetPackageData* NewData = CreateOrGetAssetPackageData(Name);
+				*NewData = *OldData;
+			}
+		}
+		else if (CachedAssetsByPackageName.Find(Name))
+		{
+			FAssetPackageData* NewData = CreateOrGetAssetPackageData(Name);
+			*NewData = *OldData;
+		}
+	}
+
+
+	// Duplicate dependency nodes
+	for (const TPair<FAssetIdentifier, FDependsNode*>& Pair : ExistingState.CachedDependsNodes)
+	{
+		FDependsNode* OldNode = Pair.Value;
+
+		// Only add if this node is already part of the filtered asset set
+		if (!CachedAssetsByPackageName.Contains(OldNode->GetIdentifier().PackageName))
+		{
+			continue;
+		}
+
+		FDependsNode* NewNode = CreateOrFindDependsNode(Pair.Key);
+		NewNode->Reserve(OldNode);
+	
+		OldNode->IterateOverDependencies([&, OldNode, NewNode](FDependsNode* InDependency, EAssetRegistryDependencyType::Type InDependencyType) {
+			const FAssetIdentifier& Identifier = InDependency->GetIdentifier();
+			if (CachedAssetsByPackageName.Find(Identifier.PackageName) || ScriptPackages.Contains(Identifier))
+			{
+				// Only add link if it's part of the filtered asset set
+				FDependsNode* NewDependency = CreateOrFindDependsNode(Identifier);
+				NewNode->AddDependency(NewDependency, InDependencyType, true);
+				NewDependency->AddReferencer(NewNode);
+			}
+		});
+	}
+}
+
 void FAssetRegistryState::InitializeFromExisting(const TMap<FName, FAssetData*>& AssetDataMap, const TMap<FAssetIdentifier, FDependsNode*>& DependsNodeMap, const TMap<FName, FAssetPackageData*>& AssetPackageDataMap, const FAssetRegistrySerializationOptions& Options, EInitializationMode InInitializationMode)
 {
 	if (InInitializationMode == EInitializationMode::Rebuild)
@@ -240,35 +359,8 @@ void FAssetRegistryState::InitializeFromExisting(const TMap<FName, FAssetData*>&
 			// Filter asset registry tags now
 			const FAssetData& AssetData = *Pair.Value;
 
-			static FName WildcardName(TEXT("*"));
-			const TSet<FName>* AllClassesFilterlist = Options.CookFilterlistTagsByClass.Find(WildcardName);
-			const TSet<FName>* ClassSpecificFilterlist = Options.CookFilterlistTagsByClass.Find(AssetData.AssetClass);
-
-			// Exclude blacklisted tags or include only whitelisted tags, based on how we were configured in ini
 			FAssetDataTagMap LocalTagsAndValues;
-			for (const auto& TagPair : AssetData.TagsAndValues)
-			{
-				const bool bInAllClasseslist = AllClassesFilterlist && (AllClassesFilterlist->Contains(TagPair.Key) || AllClassesFilterlist->Contains(WildcardName));
-				const bool bInClassSpecificlist = ClassSpecificFilterlist && (ClassSpecificFilterlist->Contains(TagPair.Key) || ClassSpecificFilterlist->Contains(WildcardName));
-				if (Options.bUseAssetRegistryTagsWhitelistInsteadOfBlacklist)
-				{
-					// It's a whitelist, only include it if it is in the all classes list or in the class specific list
-					if (bInAllClasseslist || bInClassSpecificlist)
-					{
-						// It is in the whitelist. Keep it.
-						LocalTagsAndValues.Add(TagPair.Key, TagPair.Value);
-					}
-				}
-				else
-				{
-					// It's a blacklist, include it unless it is in the all classes list or in the class specific list
-					if (!bInAllClasseslist && !bInClassSpecificlist)
-					{
-						// It isn't in the blacklist. Keep it.
-						LocalTagsAndValues.Add(TagPair.Key, TagPair.Value);
-					}
-				}
-			}
+			FAssetRegistryState::FilterTags(AssetData.TagsAndValues, LocalTagsAndValues, Options.CookFilterlistTagsByClass.Find(AssetData.AssetClass), Options);
 
 			if (InInitializationMode == EInitializationMode::OnlyUpdateExisting)
 			{
@@ -635,6 +727,7 @@ bool FAssetRegistryState::GetReferencers(const FAssetIdentifier& AssetIdentifier
 		TArray<FDependsNode*> DependencyNodes;
 		Node->GetReferencers(DependencyNodes, InReferenceType);
 
+		OutReferencers.Reserve(DependencyNodes.Num());
 		for (FDependsNode* DependencyNode : DependencyNodes)
 		{
 			OutReferencers.Add(DependencyNode->GetIdentifier());
@@ -1509,3 +1602,130 @@ bool FAssetRegistryState::IsFilterValid(const FARFilter& Filter, bool bAllowRecu
 
 	return true;
 }
+
+#if ASSET_REGISTRY_STATE_DUMPING_ENABLED
+
+void FAssetRegistryState::Dump(const TArray<FString>& Arguments, TArray<FString>& OutLines) const
+{
+	struct FAssetRefSorter
+	{
+		FORCEINLINE bool operator()(const FAssetData& A, const FAssetData& B) const
+		{
+			return (A.AssetName < B.AssetName);
+		}
+	};
+
+	if (Arguments.Contains(TEXT("ObjectPath")))
+	{
+		OutLines.Add(FString::Printf(TEXT("CachedAssetsByObjectPath (%d): "), CachedAssetsByObjectPath.Num()));
+		for (const TPair<FName, FAssetData*>& Pair : CachedAssetsByObjectPath)
+		{
+			OutLines.Add(FString::Printf(TEXT("	%s: %s"), *Pair.Key.ToString(), *Pair.Value->PackageName.ToString()));
+		}
+	}
+
+	if (Arguments.Contains(TEXT("PackageName")))
+	{
+		OutLines.Add(FString::Printf(TEXT("CachedAssetsByPackageName (%d): "), CachedAssetsByPackageName.Num()));
+		for (const TPair<FName, TArray<FAssetData*>>& Pair : CachedAssetsByPackageName)
+		{
+			if (Pair.Value.Num() == 0)
+			{
+				continue;
+			}
+
+			TArray<FAssetData*> Datas(Pair.Value);
+			Datas.HeapSort(FAssetRefSorter());
+			FString Tmp;
+			for (const FAssetData* data : Datas)
+			{
+				Tmp.Append(FString::Printf(TEXT("%s, "), *data->PackageName.ToString()));
+			}
+			OutLines.Add(FString::Printf(TEXT("	%s: %s"), *Pair.Key.ToString(), *Tmp));
+		}
+	}
+
+	if (Arguments.Contains(TEXT("Path")))
+	{
+		OutLines.Add(FString::Printf(TEXT("CachedAssetsByPath (%d): "), CachedAssetsByPath.Num()));
+		for (const TPair<FName, TArray<FAssetData*>>& Pair : CachedAssetsByPath)
+		{
+			if (Pair.Value.Num() == 0)
+			{
+				continue;
+			}
+
+			TArray<FAssetData*> Datas(Pair.Value);
+			Datas.HeapSort(FAssetRefSorter());
+			FString Tmp;
+			for (const FAssetData* data : Datas)
+			{
+				Tmp.Append(FString::Printf(TEXT("%s, "), *data->PackageName.ToString()));
+			}
+			OutLines.Add(FString::Printf(TEXT("	%s: %s"), *Pair.Key.ToString(), *Tmp));
+		}
+	}
+
+	if (Arguments.Contains(TEXT("Class")))
+	{
+		OutLines.Add(FString::Printf(TEXT("CachedAssetsByClass (%d): "), CachedAssetsByClass.Num()));
+		for (const TPair<FName, TArray<FAssetData*>>& Pair : CachedAssetsByClass)
+		{
+			if (Pair.Value.Num() == 0)
+			{
+				continue;
+			}
+
+			TArray<FAssetData*> Datas(Pair.Value);
+			Datas.HeapSort(FAssetRefSorter());
+			FString Tmp;
+			for (const FAssetData* data : Datas)
+			{
+				Tmp.Append(FString::Printf(TEXT("%s, "), *data->PackageName.ToString()));
+			}
+			OutLines.Add(FString::Printf(TEXT("	%s: %s"), *Pair.Key.ToString(), *Tmp));
+		}
+	}
+
+	if (Arguments.Contains(TEXT("Tag")))
+	{
+		OutLines.Add(FString::Printf(TEXT("CachedAssetsByTag (%d): "), CachedAssetsByTag.Num()));
+		for (const TPair<FName, TArray<FAssetData*>>& Pair : CachedAssetsByTag)
+		{
+			if (Pair.Value.Num() == 0)
+			{
+				continue;
+			}
+			TArray<FAssetData*> Datas(Pair.Value);
+			Datas.HeapSort(FAssetRefSorter());
+			FString Tmp;
+			for (const FAssetData* data : Datas)
+			{
+				Tmp.Append(FString::Printf(TEXT("%s, "), *data->PackageName.ToString()));
+			}
+			OutLines.Add(FString::Printf(TEXT("	%s: %s"), *Pair.Key.ToString(), *Tmp));
+		}
+	}
+
+	if (Arguments.Contains(TEXT("Dependencies")))
+	{
+		OutLines.Add(FString::Printf(TEXT("CachedDependsNodes (%d): "), CachedDependsNodes.Num()));
+		for (const TPair<FAssetIdentifier, FDependsNode*>& Pair : CachedDependsNodes)
+		{
+			const FDependsNode* Node = Pair.Value;
+			OutLines.Add(FString::Printf(TEXT("	%s: %s (%d connection)"), *Pair.Key.PackageName.ToString(), *Node->GetIdentifier().PackageName.ToString(), Node->GetConnectionCount()));
+		}
+	}
+
+	if (Arguments.Contains(TEXT("PackageData")))
+	{
+		OutLines.Add(FString::Printf(TEXT("CachedPackageData (%d): "), CachedPackageData.Num()));
+		for (const TPair<FName, FAssetPackageData*>& Pair : CachedPackageData)
+		{
+			const FAssetPackageData* PackageData = Pair.Value;
+			OutLines.Add(FString::Printf(TEXT("	%s: %s %d bytes"), *Pair.Key.ToString(), *PackageData->PackageGuid.ToString(), PackageData->DiskSize));
+		}
+	}
+}
+
+#endif // ASSET_REGISTRY_STATE_DUMPING_ENABLED
diff --git a/Engine/Source/Runtime/AssetRegistry/Private/DependsNode.cpp b/Engine/Source/Runtime/AssetRegistry/Private/DependsNode.cpp
--- a/Engine/Source/Runtime/AssetRegistry/Private/DependsNode.cpp
+++ b/Engine/Source/Runtime/AssetRegistry/Private/DependsNode.cpp
@@ -109,7 +109,7 @@ void FDependsNode::RemoveDependency(FDependsNode* InDependency)
 {
 	IterateOverDependencyArrays([InDependency](TArray<FDependsNode*>& InArray, EAssetRegistryDependencyType::Type)
 	{
-		InArray.Remove(InDependency);
+		InArray.RemoveSingleSwap(InDependency, false);
 	});
 }
 
diff --git a/Engine/Source/Runtime/AssetRegistry/Private/DependsNode.h b/Engine/Source/Runtime/AssetRegistry/Private/DependsNode.h
--- a/Engine/Source/Runtime/AssetRegistry/Private/DependsNode.h
+++ b/Engine/Source/Runtime/AssetRegistry/Private/DependsNode.h
@@ -40,7 +40,7 @@ public:
 	/** Remove a dependency from this node */
 	void RemoveDependency(FDependsNode* InDependency);
 	/** Remove a referencer from this node */
-	void RemoveReferencer(FDependsNode* InReferencer) { Referencers.Remove(InReferencer); }
+	void RemoveReferencer(FDependsNode* InReferencer) { Referencers.RemoveSingleSwap(InReferencer, false); }
 	/** Clear all dependency records from this node */
 	void ClearDependencies();
 	/** Removes Manage dependencies on this node and clean up referencers array. Manage references are the only ones safe to remove at runtime */
@@ -86,6 +86,10 @@ public:
 		Referencers.Reserve(InNumReferencers);
 	}
 
+	void Reserve(const FDependsNode* Other) {
+		Reserve(Other->HardDependencies.Num(), Other->SoftDependencies.Num(), Other->NameDependencies.Num(), Other->SoftManageDependencies.Num(), Other->HardManageDependencies.Num(), Other->Referencers.Num());
+	}
+
 private:
 
 	/** Iterate over all the separate dependency arrays. Const cast to avoid duplication */
diff --git a/Engine/Source/Runtime/AssetRegistry/Public/AssetRegistryState.h b/Engine/Source/Runtime/AssetRegistry/Public/AssetRegistryState.h
--- a/Engine/Source/Runtime/AssetRegistry/Public/AssetRegistryState.h
+++ b/Engine/Source/Runtime/AssetRegistry/Public/AssetRegistryState.h
@@ -9,6 +9,10 @@
 class FDependsNode;
 struct FARFilter;
 
+#ifndef ASSET_REGISTRY_STATE_DUMPING_ENABLED
+	#define ASSET_REGISTRY_STATE_DUMPING_ENABLED !(UE_BUILD_SHIPPING || UE_BUILD_TEST)
+#endif
+
 /** Load/Save options used to modify how the cache is serialized. These are read out of the AssetRegistry section of Engine.ini and can be changed per platform. */
 struct FAssetRegistrySerializationOptions
 {
@@ -262,6 +266,18 @@ public:
 	void PruneAssetData(const TSet<FName>& RequiredPackages, const TSet<FName>& RemovePackages, const TSet<int32> ChunksToKeep, const FAssetRegistrySerializationOptions& Options);
 	void PruneAssetData(const TSet<FName>& RequiredPackages, const TSet<FName>& RemovePackages, const FAssetRegistrySerializationOptions& Options);
 
+	
+	/**
+	 * Initializes a cache from an existing using a set of filters. This is more efficient than calling InitalizeFromExisting and then PruneAssetData.
+	 * @param ExistingState State to use initialize from
+	 * @param RequiredPackages If set, only these packages will be maintained. If empty it will keep all unless filtered by other parameters
+	 * @param RemovePackages These packages will be removed from the current set
+	 * @param ChunksToKeep The list of chunks that are allowed to remain. Any assets in other chunks are pruned. If empty, all assets are kept regardless of chunk
+	 * @param Options Serialization options to read filter info from
+	 */
+	void InitializeFromExistingAndPrune(const FAssetRegistryState& ExistingState, const TSet<FName>& RequiredPackages, const TSet<FName>& RemovePackages, const TSet<int32> ChunksToKeep, const FAssetRegistrySerializationOptions& Options);
+
+
 	/** Serialize the registry to/from a file, skipping editor only data */
 	bool Serialize(FArchive& Ar, const FAssetRegistrySerializationOptions& Options);
 
@@ -274,6 +290,15 @@ public:
 	/** Returns the number of assets in this state */
 	int32 GetNumAssets() const { return NumAssets; }
 
+#if ASSET_REGISTRY_STATE_DUMPING_ENABLED
+	/**
+	 * Writes out the state in textual form. Use arguments to control which segments to emit.
+	 * @param Arguments List of segments to emit. Possible values: 'ObjectPath', 'PackageName', 'Path', 'Class', 'Tag', 'Dependencies' and 'PackageData'
+	 * @param OutLines Textual representation will be written to this array.
+	 */
+	void Dump(const TArray<FString>& Arguments, TArray<FString>& OutLines) const;
+#endif
+
 private:
 	/** Find the first non-redirector dependency node starting from InDependency. */
 	FDependsNode* ResolveRedirector(FDependsNode* InDependency, TMap<FName, FAssetData*>& InAllowedAssets, TMap<FDependsNode*, FDependsNode*>& InCache);
@@ -290,6 +315,9 @@ private:
 	/** Shrink all contained data structures. */
 	void Shrink();
 
+	/** Filter a set of tags and output a copy of the filtered set. */
+	static void FilterTags(const FAssetDataTagMapSharedView& InTagsAndValues, FAssetDataTagMap& OutTagsAndValues, const TSet<FName>* ClassSpecificFilterlist, const FAssetRegistrySerializationOptions & Options);
+
 	/** Set up the data structures for USE_COMPACT_ASSET_REGISTRY, doesn't really belong here */
 	static void IngestIniSettingsForCompact(TArray<FString>& AsFName, TArray<FString>& AsPathName, TArray<FString>& AsLocText);
 
diff --git a/Engine/Source/Runtime/CoreUObject/Private/Misc/RedirectCollector.cpp b/Engine/Source/Runtime/CoreUObject/Private/Misc/RedirectCollector.cpp
--- a/Engine/Source/Runtime/CoreUObject/Private/Misc/RedirectCollector.cpp
+++ b/Engine/Source/Runtime/CoreUObject/Private/Misc/RedirectCollector.cpp
@@ -25,7 +25,6 @@ void FRedirectCollector::OnSoftObjectPathLoaded(const FSoftObjectPath& InPath, F
 		return;
 	}
 
-	FPackagePropertyPair ContainingPackageAndProperty;
 	FSoftObjectPathThreadContext& ThreadContext = FSoftObjectPathThreadContext::Get();
 
 	FName PackageName, PropertyName;
@@ -40,20 +39,12 @@ void FRedirectCollector::OnSoftObjectPathLoaded(const FSoftObjectPath& InPath, F
 		return;
 	}
 
-	if (PackageName != NAME_None)
-	{
-		ContainingPackageAndProperty.SetPackage(PackageName);
-		if (PropertyName != NAME_None)
-		{
-			ContainingPackageAndProperty.SetProperty(PropertyName);
-		}	
-	}
-
-	ContainingPackageAndProperty.SetReferencedByEditorOnlyProperty(CollectType == ESoftObjectPathCollectType::EditorOnlyCollect);
+	const bool bReferencedByEditorOnlyProperty = (CollectType == ESoftObjectPathCollectType::EditorOnlyCollect);
+	FSoftObjectPathProperty SoftObjectPathProperty(InPath.GetAssetPathName(), PropertyName, bReferencedByEditorOnlyProperty);
 
 	FScopeLock ScopeLock(&CriticalSection);
 
-	SoftObjectPathMap.AddUnique(InPath.GetAssetPathName(), ContainingPackageAndProperty);
+	SoftObjectPathMap.FindOrAdd(PackageName).Add(SoftObjectPathProperty);
 }
 
 void FRedirectCollector::OnStringAssetReferenceLoaded(const FString& InString)
@@ -77,69 +68,79 @@ FString FRedirectCollector::OnStringAssetReferenceSaved(const FString& InString)
 
 void FRedirectCollector::ResolveAllSoftObjectPaths(FName FilterPackage)
 {	
+	auto LoadSoftObjectPathLambda = [this](const FSoftObjectPathProperty& SoftObjectPathProperty)
+	{
+		const FName& ToLoadFName = SoftObjectPathProperty.GetAssetPathName();
+		const FString ToLoad = ToLoadFName.ToString();
+
+		if (ToLoad.Len() > 0 )
+		{
+			UE_LOG(LogRedirectors, Verbose, TEXT("Resolving Soft Object Path '%s'"), *ToLoad);
+			UE_CLOG(SoftObjectPathProperty.GetPropertyName().ToString().Len(), LogRedirectors, Verbose, TEXT("    Referenced by '%s'"), *SoftObjectPathProperty.GetPropertyName().ToString());
+
+			int32 DotIndex = ToLoad.Find(TEXT("."));
+			FString PackageName = DotIndex != INDEX_NONE ? ToLoad.Left(DotIndex) : ToLoad;
+
+			// If is known missing don't try
+			if (FLinkerLoad::IsKnownMissingPackage(FName(*PackageName)))
+			{
+				return;
+			}
+
+			UObject *Loaded = LoadObject<UObject>(NULL, *ToLoad, NULL, SoftObjectPathProperty.GetReferencedByEditorOnlyProperty() ? LOAD_EditorOnly | LOAD_NoWarn : LOAD_NoWarn, NULL);
+
+			if (Loaded)
+			{
+				FString Dest = Loaded->GetPathName();
+				UE_LOG(LogRedirectors, Verbose, TEXT("    Resolved to '%s'"), *Dest);
+				if (Dest != ToLoad)
+				{
+					AssetPathRedirectionMap.Add(ToLoadFName, FName(*Dest));
+				}
+			}
+			else
+			{
+				const FString Referencer = SoftObjectPathProperty.GetPropertyName().ToString().Len() ? SoftObjectPathProperty.GetPropertyName().ToString() : TEXT("Unknown");
+				UE_LOG(LogRedirectors, Warning, TEXT("Soft Object Path '%s' was not found when resolving paths! (Referencer '%s')"), *ToLoad, *Referencer);
+			}
+		}
+	};
+
 	FScopeLock ScopeLock(&CriticalSection);
 
-	TMultiMap<FName, FPackagePropertyPair> SkippedReferences;
-	SkippedReferences.Empty(SoftObjectPathMap.Num());
+	FSoftObjectPathMap KeepSoftObjectPathMap;
+	KeepSoftObjectPathMap.Reserve(SoftObjectPathMap.Num());
 	while (SoftObjectPathMap.Num())
 	{
-		TMultiMap<FName, FPackagePropertyPair> CurrentReferences;
-		Swap(SoftObjectPathMap, CurrentReferences);
+		FSoftObjectPathMap LocalSoftObjectPathMap;
+		Swap(SoftObjectPathMap, LocalSoftObjectPathMap);
 
-		for (const TPair<FName, FPackagePropertyPair>& CurrentReference : CurrentReferences)
+		for (TPair<FName, FSoftObjectPathPropertySet>& CurrentPackage : LocalSoftObjectPathMap)
 		{
-			const FName& ToLoadFName = CurrentReference.Key;
-			const FPackagePropertyPair& RefFilenameAndProperty = CurrentReference.Value;
+			const FName& CurrentPackageName = CurrentPackage.Key;
+			FSoftObjectPathPropertySet& SoftObjectPathProperties = CurrentPackage.Value;
 
 			if ((FilterPackage != NAME_None) && // not using a filter
-				(FilterPackage != RefFilenameAndProperty.GetCachedPackageName()) && // this is the package we are looking for
-				(RefFilenameAndProperty.GetCachedPackageName() != NAME_None) // if we have an empty package name then process it straight away
+				(FilterPackage != CurrentPackageName) && // this is the package we are looking for
+				(CurrentPackageName != NAME_None) // if we have an empty package name then process it straight away
 				)
 			{
-				// If we have a valid filter and it doesn't match, skip this reference
-				SkippedReferences.Add(ToLoadFName, RefFilenameAndProperty);
+				// If we have a valid filter and it doesn't match, skip processing of this package and keep it
+				KeepSoftObjectPathMap.FindOrAdd(CurrentPackageName).Append(MoveTemp(SoftObjectPathProperties));
 				continue;
 			}
 
-			const FString ToLoad = ToLoadFName.ToString();
-
-			if (ToLoad.Len() > 0 )
+			// This will call LoadObject which may trigger OnSoftObjectPathLoaded and add new soft object paths to the SoftObjectPathMap
+			for (const FSoftObjectPathProperty& SoftObjecPathProperty : SoftObjectPathProperties)
 			{
-				UE_LOG(LogRedirectors, Verbose, TEXT("Resolving Soft Object Path '%s'"), *ToLoad);
-				UE_CLOG(RefFilenameAndProperty.GetProperty().ToString().Len(), LogRedirectors, Verbose, TEXT("    Referenced by '%s'"), *RefFilenameAndProperty.GetProperty().ToString());
-
-				int32 DotIndex = ToLoad.Find(TEXT("."));
-				FString PackageName = DotIndex != INDEX_NONE ? ToLoad.Left(DotIndex) : ToLoad;
-
-				// If is known missing don't try
-				if (FLinkerLoad::IsKnownMissingPackage(FName(*PackageName)))
-				{
-					continue;
-				}
-
-				UObject *Loaded = LoadObject<UObject>(NULL, *ToLoad, NULL, RefFilenameAndProperty.GetReferencedByEditorOnlyProperty() ? LOAD_EditorOnly | LOAD_NoWarn : LOAD_NoWarn, NULL);
-
-				if (Loaded)
-				{
-					FString Dest = Loaded->GetPathName();
-					UE_LOG(LogRedirectors, Verbose, TEXT("    Resolved to '%s'"), *Dest);
-					if (Dest != ToLoad)
-					{
-						AssetPathRedirectionMap.Add(ToLoadFName, FName(*Dest));
-					}
-				}
-				else
-				{
-					const FString Referencer = RefFilenameAndProperty.GetProperty().ToString().Len() ? RefFilenameAndProperty.GetProperty().ToString() : TEXT("Unknown");
-					UE_LOG(LogRedirectors, Warning, TEXT("Soft Object Path '%s' was not found when resolving paths! (Referencer '%s')"), *ToLoad, *Referencer);
-				}
+				LoadSoftObjectPathLambda(SoftObjecPathProperty);
 			}
 		}
 	}
 
 	check(SoftObjectPathMap.Num() == 0);
-	// Add any skipped references back into the map for the next time this is called
-	Swap(SoftObjectPathMap, SkippedReferences);
+	// Add any non processed packages back into the global map for the next time this is called
+	Swap(SoftObjectPathMap, KeepSoftObjectPathMap);
 	// we shouldn't have any references left if we decided to resolve them all
 	check((SoftObjectPathMap.Num() == 0) || (FilterPackage != NAME_None));
 }
@@ -148,24 +149,26 @@ void FRedirectCollector::ProcessSoftObjectPathPackageList(FName FilterPackage, b
 {
 	FScopeLock ScopeLock(&CriticalSection);
 
-	// Iterate map, remove all matching and potentially add to OutReferencedPackages
-	for (auto It = SoftObjectPathMap.CreateIterator(); It; ++It)
+	TSet<FSoftObjectPathProperty>* SoftObjectPathProperties = SoftObjectPathMap.Find(FilterPackage);
+	if (!SoftObjectPathProperties)
 	{
-		const FName& ToLoadFName = It->Key;
-		const FPackagePropertyPair& RefFilenameAndProperty = It->Value;
+		return;
+	}
 
-		// Package name may be null, if so return the set of things not associated with a package
-		if (FilterPackage == RefFilenameAndProperty.GetCachedPackageName())
+	// potentially add soft object path package names to OutReferencedPackages
+	OutReferencedPackages.Reserve(SoftObjectPathProperties->Num());
+	for (const FSoftObjectPathProperty& SoftObjecPathProperty : *SoftObjectPathProperties)
+	{
+		if (!SoftObjecPathProperty.GetReferencedByEditorOnlyProperty() || bGetEditorOnly)
 		{
-			if (!RefFilenameAndProperty.GetReferencedByEditorOnlyProperty() || bGetEditorOnly)
-			{
-				FString PackageNameString = FPackageName::ObjectPathToPackageName(ToLoadFName.ToString());
-				OutReferencedPackages.Add(FName(*PackageNameString));
-			}
-
-			It.RemoveCurrent();
+			const FName& ToLoadFName = SoftObjecPathProperty.GetAssetPathName();
+			FString PackageNameString = FPackageName::ObjectPathToPackageName(ToLoadFName.ToString());
+			OutReferencedPackages.Add(FName(*PackageNameString));
 		}
 	}
+
+	// always remove all data for the processed FilterPackage
+	SoftObjectPathMap.Remove(FilterPackage);
 }
 
 void FRedirectCollector::AddAssetPathRedirection(FName OriginalPath, FName RedirectedPath)
diff --git a/Engine/Source/Runtime/CoreUObject/Public/Misc/RedirectCollector.h b/Engine/Source/Runtime/CoreUObject/Public/Misc/RedirectCollector.h
--- a/Engine/Source/Runtime/CoreUObject/Public/Misc/RedirectCollector.h
+++ b/Engine/Source/Runtime/CoreUObject/Public/Misc/RedirectCollector.h
@@ -16,49 +16,38 @@ class COREUOBJECT_API FRedirectCollector
 private:
 	
 	/** Helper struct for soft object path tracking */
-	struct FPackagePropertyPair
+	struct FSoftObjectPathProperty
 	{
-		FPackagePropertyPair() : bReferencedByEditorOnlyProperty(false) {}
-
-		FPackagePropertyPair(const FName& InPackage, const FName& InProperty)
-		: Package(InPackage)
-		, Property(InProperty)
-		, bReferencedByEditorOnlyProperty(false)
+		FSoftObjectPathProperty(FName InAssetPathName, FName InProperty, bool bInReferencedByEditorOnlyProperty)
+			: AssetPathName(InAssetPathName)
+			, PropertyName(InProperty)
+			, bReferencedByEditorOnlyProperty(bInReferencedByEditorOnlyProperty)
 		{}
 
-		bool operator==(const FPackagePropertyPair& Other) const
-		{
-			return Package == Other.Package &&
-				Property == Other.Property &&
-				bReferencedByEditorOnlyProperty == Other.bReferencedByEditorOnlyProperty;
-		}
+		 bool operator==(const FSoftObjectPathProperty& Other) const
+		 {
+		 	return AssetPathName == Other.AssetPathName &&
+		 		PropertyName == Other.PropertyName &&
+		 		bReferencedByEditorOnlyProperty == Other.bReferencedByEditorOnlyProperty;
+		 }
 
-		const FName& GetCachedPackageName() const
+		friend inline uint32 GetTypeHash(const FSoftObjectPathProperty& Key)
 		{
-			return Package;
+			uint32 Hash = 0;
+			Hash = HashCombine(Hash, GetTypeHash(Key.AssetPathName));
+			Hash = HashCombine(Hash, GetTypeHash(Key.PropertyName));
+			Hash = HashCombine(Hash, (uint32)Key.bReferencedByEditorOnlyProperty);
+			return Hash;
 		}
 
-		void SetPackage(const FName& InPackage)
-		{
-			Package = InPackage;
-		}
-		void SetProperty(const FName& InProperty)
+		const FName& GetAssetPathName() const
 		{
-			Property = InProperty;
-		}
-		
-		void SetReferencedByEditorOnlyProperty(bool InReferencedByEditorOnlyProperty)
-		{
-			bReferencedByEditorOnlyProperty = InReferencedByEditorOnlyProperty;
+			return AssetPathName;
 		}
 
-		const FName& GetPackage() const
-		{
-			return Package;
-		}
-		const FName& GetProperty() const
+		const FName& GetPropertyName() const
 		{
-			return Property;
+			return PropertyName;
 		}
 
 		bool GetReferencedByEditorOnlyProperty() const
@@ -67,8 +56,8 @@ private:
 		}
 
 	private:
-		FName Package;
-		FName Property;
+		FName AssetPathName;
+		FName PropertyName;
 		bool bReferencedByEditorOnlyProperty;
 	};
 
@@ -127,8 +116,10 @@ public:
 
 private:
 
-	/** A map of assets referenced by soft object paths, with the key being the asset being referenced and the value equal to the package with the reference */
-	TMultiMap<FName, FPackagePropertyPair> SoftObjectPathMap;
+	/** A map of assets referenced by soft object paths, with the key being the package with the reference */
+	typedef TSet<FSoftObjectPathProperty> FSoftObjectPathPropertySet;
+	typedef TMap<FName, FSoftObjectPathPropertySet> FSoftObjectPathMap;
+	FSoftObjectPathMap SoftObjectPathMap;
 
 	/** When saving, apply this remapping to all soft object paths */
 	TMap<FName, FName> AssetPathRedirectionMap;
@@ -140,4 +131,4 @@ private:
 // global redirect collector callback structure
 COREUOBJECT_API extern FRedirectCollector GRedirectCollector;
 
-#endif // WITH_EDITOR
\ No newline at end of file
+#endif // WITH_EDITOR
diff --git a/Engine/Source/Runtime/Engine/Private/AssetManager.cpp b/Engine/Source/Runtime/Engine/Private/AssetManager.cpp
--- a/Engine/Source/Runtime/Engine/Private/AssetManager.cpp
+++ b/Engine/Source/Runtime/Engine/Private/AssetManager.cpp
@@ -2635,6 +2635,7 @@ bool UAssetManager::GetPackageManagers(FName PackageName, bool bRecurseToParents
 
 	bool bFoundAny = false;
 	TArray<FAssetIdentifier> ReferencingPrimaryAssets;
+	ReferencingPrimaryAssets.Reserve(128);
 
 	AssetRegistry.GetReferencers(PackageName, ReferencingPrimaryAssets, EAssetRegistryDependencyType::Manage);
 
@@ -3235,6 +3236,7 @@ bool UAssetManager::GetPackageChunkIds(FName PackageName, const ITargetPlatform*
 
 	// Add all chunk ids from the asset rules of managers. By default priority will not override other chunks
 	TSet<FPrimaryAssetId> Managers;
+	Managers.Reserve(128);
 
 	GetPackageManagers(PackageName, true, Managers);
 	return GetPrimaryAssetSetChunkIds(Managers, TargetPlatform, ExistingChunkList, OutChunkList);
diff --git a/Engine/Source/Runtime/RenderCore/Private/ShaderCodeLibrary.cpp b/Engine/Source/Runtime/RenderCore/Private/ShaderCodeLibrary.cpp
--- a/Engine/Source/Runtime/RenderCore/Private/ShaderCodeLibrary.cpp
+++ b/Engine/Source/Runtime/RenderCore/Private/ShaderCodeLibrary.cpp
@@ -16,6 +16,7 @@ ShaderCodeLibrary.cpp: Bound shader state cache implementation.
 #include "Async/AsyncFileHandle.h"
 #include "PipelineFileCache.h"
 #include "Interfaces/IPluginManager.h"
+#include "Hash/CityHash.h"
 
 #include "Interfaces/IShaderFormatArchive.h"
 #include "ShaderPipelineCache.h"
@@ -100,9 +101,24 @@ static void FShaderLibraryHelperCompressCode(EShaderPlatform Platform, const TAr
 		CompressedCode.Shrink();
 }
 
+
+FORCEINLINE FName ParseFNameCached(const FString& Src, TMap<uint32,FName>& NameCache)
+{
+	uint32 SrcHash = CityHash32((char*)Src.GetCharArray().GetData(), Src.GetCharArray().GetAllocatedSize());
+	if (FName* Name = NameCache.Find(SrcHash))
+	{
+		return *Name;
+	}
+	else
+	{
+		return NameCache.Emplace(SrcHash, FName(*Src));
+	}
+}
+
 FString FCompactFullName::ToString() const
 {
 	FString RetString;
+	RetString.Reserve(256);
 	if (!ObjectClassAndPath.Num())
 	{
 		RetString += TEXT("empty");
@@ -161,6 +177,28 @@ uint32 GetTypeHash(const FCompactFullName& A)
 	return Hash;
 }
 
+void FixupUnsanitizedNames(const FString& Src, TArray<FString>& OutFields) 
+{
+	FString NewSrc(Src);
+
+	int32 ParenOpen = -1;
+	int32 ParenClose = -1;
+
+	if (NewSrc.FindChar(TCHAR('('), ParenOpen) && NewSrc.FindChar(TCHAR(')'), ParenClose) && ParenOpen < ParenClose && ParenOpen >= 0 && ParenClose >= 0)
+	{
+		for (int32 Index = ParenOpen + 1; Index < ParenClose; Index++)
+		{
+			if (NewSrc[Index] == TCHAR(','))
+			{
+				NewSrc[Index] = ' ';
+			}
+		}
+		OutFields.Empty();
+		NewSrc.TrimStartAndEnd().ParseIntoArray(OutFields, TEXT(","), false);
+		check(OutFields.Num() == 11);
+	}
+}
+
 void FStableShaderKeyAndValue::ComputeKeyHash()
 {
 	KeyHash = GetTypeHash(ClassNameAndObjectPath);
@@ -185,28 +223,9 @@ void FStableShaderKeyAndValue::ParseFromString(const FString& Src)
 	if (Fields.Num() > 11)
 	{
 		// hack fix for unsanitized names, should not occur anymore.
-
-		FString NewSrc(Src);
-
-		int32 ParenOpen = -1;
-		int32 ParenClose = -1;
-
-		if (NewSrc.FindChar(TCHAR('('), ParenOpen) && NewSrc.FindChar(TCHAR(')'), ParenClose) && ParenOpen < ParenClose && ParenOpen >= 0 && ParenClose >= 0)
-		{
-			for (int32 Index = ParenOpen + 1; Index < ParenClose; Index++)
-			{
-				if (NewSrc[Index] == TCHAR(','))
-				{
-					NewSrc[Index] = ' ';
-				}
-			}
-			Fields.Empty();
-			NewSrc.TrimStartAndEnd().ParseIntoArray(Fields, TEXT(","), false);
-			check(Fields.Num() == 11);
-		}
+		FixupUnsanitizedNames(Src, Fields);
 	}
 
-
 	check(Fields.Num() == 11);
 
 	int32 Index = 0;
@@ -231,6 +250,44 @@ void FStableShaderKeyAndValue::ParseFromString(const FString& Src)
 	ComputeKeyHash();
 }
 
+
+void FStableShaderKeyAndValue::ParseFromStringCached(const FString& Src, TMap<uint32, FName>& NameCache)
+{
+	TArray<FString> Fields;
+	Src.TrimStartAndEnd().ParseIntoArray(Fields, TEXT(","), false);
+
+	if (Fields.Num() > 11)
+	{
+		// hack fix for unsanitized names, should not occur anymore.
+		FixupUnsanitizedNames(Src, Fields);
+	}
+	
+	check(Fields.Num() == 11);
+
+	int32 Index = 0;
+	ClassNameAndObjectPath.ParseFromString(Fields[Index++]);
+
+	// There is a high level of uniformity on the following FNames, use
+	// the local name cache to accelerate lookup
+	ShaderType = ParseFNameCached(Fields[Index++], NameCache);
+	ShaderClass = ParseFNameCached(Fields[Index++], NameCache);
+	MaterialDomain = ParseFNameCached(Fields[Index++], NameCache);
+	FeatureLevel = ParseFNameCached(Fields[Index++], NameCache);
+
+	QualityLevel = ParseFNameCached(Fields[Index++], NameCache);
+	TargetFrequency = ParseFNameCached(Fields[Index++], NameCache);
+	TargetPlatform = ParseFNameCached(Fields[Index++], NameCache);
+
+	VFType = ParseFNameCached(Fields[Index++], NameCache);
+	PermutationId = ParseFNameCached(Fields[Index++], NameCache);
+
+	OutputHash.FromString(Fields[Index++]);
+
+	check(Index == 11);
+
+	ComputeKeyHash();
+}
+
 FString FStableShaderKeyAndValue::ToString() const
 {
 	FString Result;
@@ -521,7 +578,6 @@ public:
 			FScopeLock ScopeLock(&ReadRequestLock);
 
 			Entry->NumRefs--;
-			check(Entry->NumRefs >= 0);
 			if (Entry->NumRefs == 0)
 			{
 				// should not attempt to release shader code while it's loading
@@ -1212,6 +1268,9 @@ struct FEditorShaderStableInfo
 	{
 		check(LibraryName.Len() > 0);
 
+		TMap<uint32, FName> NameCache;
+		NameCache.Reserve(2048);
+
 		const FString ShaderIntermediateLocation = FPaths::ProjectSavedDir() / TEXT("Shaders") / FormatName.ToString();
 
 		TArray<FString> ShaderFiles;
@@ -1227,7 +1286,7 @@ struct FEditorShaderStableInfo
 					for (int32 Index = 1; Index < SourceFileContents.Num(); Index++)
 					{
 						FStableShaderKeyAndValue Item;
-						Item.ParseFromString(SourceFileContents[Index]);
+						Item.ParseFromStringCached(SourceFileContents[Index], NameCache);
 						AddShader(Item);
 					}
 				}
@@ -1255,17 +1314,23 @@ struct FEditorShaderStableInfo
 			// Write to a intermediate file
 			FString IntermediateFormatPath = GetStableInfoArchiveFilename(FPaths::ProjectSavedDir() / TEXT("Shaders") / FormatName.ToString(), LibraryName, FormatName);
 
-			TArray<FString> FileContents;
-			FileContents.Reserve(StableMap.Num() + 1);
+			// Write directly to the file
+			TUniquePtr<FArchive> IntermediateFormatAr(IFileManager::Get().CreateFileWriter(*IntermediateFormatPath));
+
+			const FString HeaderText = FStableShaderKeyAndValue::HeaderLine();
+			auto HeaderSrc = StringCast<ANSICHAR>(*HeaderText, HeaderText.Len());
+
+			IntermediateFormatAr->Serialize((ANSICHAR*)HeaderSrc.Get(), HeaderSrc.Length() * sizeof(ANSICHAR));
 
-			FileContents.Add(FStableShaderKeyAndValue::HeaderLine());
+			FString LineBuffer;
+			LineBuffer.Reserve(512);
 
 			for (const FStableShaderKeyAndValue& Item : StableMap)
 			{
-				FString& LineBuffer = FileContents.Emplace_GetRef();
 				Item.ToString(LineBuffer);
+				auto LineConverted = StringCast<ANSICHAR>(*LineBuffer, LineBuffer.Len());
+				IntermediateFormatAr->Serialize((ANSICHAR*)LineConverted.Get(), LineConverted.Length() * sizeof(ANSICHAR));
 			}
-			FFileHelper::SaveStringArrayToFile(FileContents, *IntermediateFormatPath);
 
 			// Only the master cooker needs to write to the output directory, child cookers only write to the Saved directory
 			if (bMasterCooker)
diff --git a/Engine/Source/Runtime/RenderCore/Public/ShaderCodeLibrary.h b/Engine/Source/Runtime/RenderCore/Public/ShaderCodeLibrary.h
--- a/Engine/Source/Runtime/RenderCore/Public/ShaderCodeLibrary.h
+++ b/Engine/Source/Runtime/RenderCore/Public/ShaderCodeLibrary.h
@@ -87,6 +87,7 @@ struct RENDERCORE_API FStableShaderKeyAndValue
 
 	void ComputeKeyHash();
 	void ParseFromString(const FString& Src);
+	void ParseFromStringCached(const FString& Src, class TMap<uint32, FName>& NameCache);
 	FString ToString() const;
 	void ToString(FString& OutResult) const;
 	static FString HeaderLine();
