Fix nullptr filter to allow for reference replacment to function correctly

diff --git a/Engine/Source/Editor/UnrealEd/Private/Toolkits/AssetEditorToolkit.cpp b/Engine/Source/Editor/UnrealEd/Private/Toolkits/AssetEditorToolkit.cpp
--- a/Engine/Source/Editor/UnrealEd/Private/Toolkits/AssetEditorToolkit.cpp
+++ b/Engine/Source/Editor/UnrealEd/Private/Toolkits/AssetEditorToolkit.cpp
@@ -1117,18 +1117,10 @@ void FAssetEditorToolkit::RemoveAllToolbarWidgets()
 
 void FAssetEditorToolkit::FGCEditingObjects::AddReferencedObjects(FReferenceCollector& Collector)
 {
-	// Remove null objects as a safe guard against assets being forcibly GC'd
-	TArray<UObject*> EditingObjectsCopy = OwnerToolkit.EditingObjects;
 	Collector.AddReferencedObjects(OwnerToolkit.EditingObjects);
-
-	OwnerToolkit.EditingObjects.Reset();
-	for (UObject* Object : EditingObjectsCopy)
-	{
-		if (Object != nullptr)
-		{
-			OwnerToolkit.EditingObjects.Add(Object);
-		}
-	}
+	
+	// Remove null objects as a safe guard against assets being forcibly GC'd
+	OwnerToolkit.EditingObjects.RemoveAllSwap([](UObject* Obj) { return Obj == nullptr; } );
 }
 
 
