Fix for edge case of persistent level with unknown SCC status
#rnx

diff --git a/Engine/Source/Editor/UnrealEd/Private/AssetSelection.cpp b/Engine/Source/Editor/UnrealEd/Private/AssetSelection.cpp
--- a/Engine/Source/Editor/UnrealEd/Private/AssetSelection.cpp
+++ b/Engine/Source/Editor/UnrealEd/Private/AssetSelection.cpp
@@ -435,7 +435,7 @@ namespace ActorPlacementUtils
 			FString FileName = SourceControlHelpers::PackageFilename(InLevel->GetPathName());
 			// Query file state also checks the source control status
 			FSourceControlState SCState = SourceControlHelpers::QueryFileState(FileName, true);
-			if (!InLevel->bLevelOkayForPlacementWhileCheckedIn && !(SCState.bIsCheckedOut || SCState.bIsAdded || SCState.bCanAdd))
+			if (!InLevel->bLevelOkayForPlacementWhileCheckedIn && !(SCState.bIsCheckedOut || SCState.bIsAdded || SCState.bCanAdd || SCState.bIsUnknown))
 			{
 				if (EAppReturnType::Ok != OpenMsgDlgInt(EAppMsgType::OkCancel, NSLOCTEXT("UnrealEd","LevelNotCheckedOutMsg", "This actor will be placed in a level that is is source control but not currently checked out. Continue?"), NSLOCTEXT("UnrealEd", "ActorPlacement_Title", "Actor Placement Warning")))
 				{
