Make sure the skeletalmesh morph targets are not remove when importing "weights only".

diff --git a/Engine/Source/Editor/UnrealEd/Private/Fbx/FbxSkeletalMeshImport.cpp b/Engine/Source/Editor/UnrealEd/Private/Fbx/FbxSkeletalMeshImport.cpp
--- a/Engine/Source/Editor/UnrealEd/Private/Fbx/FbxSkeletalMeshImport.cpp
+++ b/Engine/Source/Editor/UnrealEd/Private/Fbx/FbxSkeletalMeshImport.cpp
@@ -2225,7 +2225,8 @@ USkeletalMesh* UnFbx::FFbxImporter::ReimportSkeletalMesh(USkeletalMesh* Mesh, UF
 			Mesh->GetLODInfo(LODCounter)->bImportWithBaseMesh = false;
 		}
 		//Empty the morph target before re-importing, it will prevent to have old data that can point on random vertex
-		if (Mesh->MorphTargets.Num() > 0)
+		//Do not empty the morph if we import weights only
+		if (Mesh->MorphTargets.Num() > 0 && !ImportOptions->bImportAsSkeletalSkinning)
 		{
 			Mesh->UnregisterAllMorphTarget();
 		}
