[NO NON-PII COMMIT MESSAGE]

diff --git a/Engine/Source/Editor/UnrealEd/Private/DragAndDrop/CollectionDragDropOp.cpp b/Engine/Source/Editor/UnrealEd/Private/DragAndDrop/CollectionDragDropOp.cpp
--- a/Engine/Source/Editor/UnrealEd/Private/DragAndDrop/CollectionDragDropOp.cpp
+++ b/Engine/Source/Editor/UnrealEd/Private/DragAndDrop/CollectionDragDropOp.cpp
@@ -1,3 +1,5 @@
+// Copyright 1998-2019 Epic Games, Inc. All Rights Reserved.
+
 #include "DragAndDrop/CollectionDragDropOp.h"
 #include "ICollectionManager.h"
 #include "CollectionManagerModule.h"
