[NO NON-PII COMMIT MESSAGE]

diff --git a/Engine/Source/Editor/UnrealEd/Private/Commandlets/AssetRegistryGenerator.cpp b/Engine/Source/Editor/UnrealEd/Private/Commandlets/AssetRegistryGenerator.cpp
--- a/Engine/Source/Editor/UnrealEd/Private/Commandlets/AssetRegistryGenerator.cpp
+++ b/Engine/Source/Editor/UnrealEd/Private/Commandlets/AssetRegistryGenerator.cpp
@@ -1257,7 +1257,7 @@ bool FAssetRegistryGenerator::GenerateAssetChunkInformationCSV(const FString& Ou
 		for (int32 PakchunkIndex = 0, ChunkNum = FinalChunkManifests.Num(); PakchunkIndex < ChunkNum; ++PakchunkIndex)
 		{
 			FArchive* ChunkFile = IFileManager::Get().CreateFileWriter(*FPaths::Combine(*OutputPath, *FString::Printf(TEXT("Chunks%dInfo.csv"), PakchunkIndex)));
-			if (ChunkFile != nullptr)
+			if (ChunkFile == nullptr)
 			{
 				return false;
 			}
