Fixing typo, clarifying title text
#rnx

diff --git a/Engine/Source/Editor/UnrealEd/Private/AssetSelection.cpp b/Engine/Source/Editor/UnrealEd/Private/AssetSelection.cpp
--- a/Engine/Source/Editor/UnrealEd/Private/AssetSelection.cpp
+++ b/Engine/Source/Editor/UnrealEd/Private/AssetSelection.cpp
@@ -438,7 +438,7 @@ namespace ActorPlacementUtils
 			FSourceControlState SCState = SourceControlHelpers::QueryFileState(FileName, true);
 			if (!InLevel->bLevelOkayForPlacementWhileCheckedIn && !(SCState.bIsCheckedOut || SCState.bIsAdded || SCState.bCanAdd || SCState.bIsUnknown))
 			{
-				if (EAppReturnType::Ok != OpenMsgDlgInt(EAppMsgType::OkCancel, NSLOCTEXT("UnrealEd","LevelNotCheckedOutMsg", "This actor will be placed in a level that is is source control but not currently checked out. Continue?"), NSLOCTEXT("UnrealEd", "ActorPlacement_Title", "Actor Placement Warning")))
+				if (EAppReturnType::Ok != OpenMsgDlgInt(EAppMsgType::OkCancel, NSLOCTEXT("UnrealEd","LevelNotCheckedOutMsg", "This actor will be placed in a level that is in source control but not currently checked out. Continue?"), NSLOCTEXT("UnrealEd", "LevelCheckout_Title", "Level Checkout Warning")))
 				{
 					return false;
 				}
