Make sure re-importing "Weights Only" do not mess the vertexcolor. Remove the vertex color option when re-importing "Weights Only".

diff --git a/Engine/Source/Editor/UnrealEd/Classes/Factories/FbxSkeletalMeshImportData.h b/Engine/Source/Editor/UnrealEd/Classes/Factories/FbxSkeletalMeshImportData.h
--- a/Engine/Source/Editor/UnrealEd/Classes/Factories/FbxSkeletalMeshImportData.h
+++ b/Engine/Source/Editor/UnrealEd/Classes/Factories/FbxSkeletalMeshImportData.h
@@ -38,11 +38,11 @@ public:
 	TEnumAsByte<enum EFBXImportContentType> LastImportContentType;
 
 	/** Specify how vertex colors should be imported */
-	UPROPERTY(EditAnywhere, BlueprintReadWrite, AdvancedDisplay, config, Category = Mesh, meta = (OBJRestrict = "true", ImportType = "SkeletalMesh"))
+	UPROPERTY(EditAnywhere, BlueprintReadWrite, AdvancedDisplay, config, Category = Mesh, meta = (OBJRestrict = "true", ImportType = "SkeletalMesh|GeoOnly"))
 	TEnumAsByte<EVertexColorImportOption::Type> VertexColorImportOption;
 
 	/** Specify override color in the case that VertexColorImportOption is set to Override */
-	UPROPERTY(EditAnywhere, BlueprintReadWrite, AdvancedDisplay, config, Category = Mesh, meta = (OBJRestrict = "true", ImportType = "SkeletalMesh"))
+	UPROPERTY(EditAnywhere, BlueprintReadWrite, AdvancedDisplay, config, Category = Mesh, meta = (OBJRestrict = "true", ImportType = "SkeletalMesh|GeoOnly"))
 	FColor VertexOverrideColor;
 
 	/** Enable this option to update Skeleton (of the mesh)'s reference pose. Mesh's reference pose is always updated.  */
diff --git a/Engine/Source/Editor/UnrealEd/Private/Factories/EditorFactories.cpp b/Engine/Source/Editor/UnrealEd/Private/Factories/EditorFactories.cpp
--- a/Engine/Source/Editor/UnrealEd/Private/Factories/EditorFactories.cpp
+++ b/Engine/Source/Editor/UnrealEd/Private/Factories/EditorFactories.cpp
@@ -5878,6 +5878,7 @@ EReimportResult::Type UReimportFbxSkeletalMeshFactory::Reimport( UObject* Obj, i
 			ImportOptions->bImportSkeletalMeshLODs = false;
 			ImportOptions->bImportAnimations = false;
 			ImportOptions->bImportMorph = false;
+			ImportOptions->VertexColorImportOption = EVertexColorImportOption::Ignore;
 		}
 		else if (ImportOptions->bImportAsSkeletalGeometry)
 		{
diff --git a/Engine/Source/Editor/UnrealEd/Private/Fbx/FbxSkeletalMeshImport.cpp b/Engine/Source/Editor/UnrealEd/Private/Fbx/FbxSkeletalMeshImport.cpp
--- a/Engine/Source/Editor/UnrealEd/Private/Fbx/FbxSkeletalMeshImport.cpp
+++ b/Engine/Source/Editor/UnrealEd/Private/Fbx/FbxSkeletalMeshImport.cpp
@@ -116,6 +116,158 @@ void MatrixAdd(FbxAMatrix& pDstMatrix, FbxAMatrix& pSrcMatrix)
 	}
 }
 
+/** Helper struct for the mesh component vert position octree */
+struct FSkeletalMeshVertPosOctreeSemantics
+{
+	enum { MaxElementsPerLeaf = 16 };
+	enum { MinInclusiveElementsPerNode = 7 };
+	enum { MaxNodeDepth = 12 };
+
+	typedef TInlineAllocator<MaxElementsPerLeaf> ElementAllocator;
+
+	/**
+	 * Get the bounding box of the provided octree element. In this case, the box
+	 * is merely the point specified by the element.
+	 *
+	 * @param	Element	Octree element to get the bounding box for
+	 *
+	 * @return	Bounding box of the provided octree element
+	 */
+	FORCEINLINE static FBoxCenterAndExtent GetBoundingBox(const FSoftSkinVertex& Element)
+	{
+		return FBoxCenterAndExtent(Element.Position, FVector::ZeroVector);
+	}
+
+	/**
+	 * Determine if two octree elements are equal
+	 *
+	 * @param	A	First octree element to check
+	 * @param	B	Second octree element to check
+	 *
+	 * @return	true if both octree elements are equal, false if they are not
+	 */
+	FORCEINLINE static bool AreElementsEqual(const FSoftSkinVertex& A, const FSoftSkinVertex& B)
+	{
+		return (A.Position == B.Position && A.UVs[0] == B.UVs[0]);
+	}
+
+	/** Ignored for this implementation */
+	FORCEINLINE static void SetElementId(const FSoftSkinVertex& Element, FOctreeElementId Id)
+	{
+	}
+};
+typedef TOctree<FSoftSkinVertex, FSkeletalMeshVertPosOctreeSemantics> TSKCVertPosOctree;
+
+void RemapSkeletalMeshVertexColorToImportData(const USkeletalMesh* SkeletalMesh, const int32 LODIndex, FSkeletalMeshImportData* SkelMeshImportData)
+{
+	//Make sure we have all the source data we need to do the remap
+	if (!SkeletalMesh->GetImportedModel() || !SkeletalMesh->GetImportedModel()->LODModels.IsValidIndex(LODIndex) || !SkeletalMesh->bHasVertexColors)
+	{
+		return;
+	}
+
+	// Find the extents formed by the cached vertex positions in order to optimize the octree used later
+	FBox Bounds(ForceInitToZero);
+	SkelMeshImportData->bHasVertexColors = true;
+
+	int32 WedgeNumber = SkelMeshImportData->Wedges.Num();
+	for (int32 WedgeIndex = 0; WedgeIndex < WedgeNumber; ++WedgeIndex)
+	{
+		SkeletalMeshImportData::FVertex& Wedge = SkelMeshImportData->Wedges[WedgeIndex];
+		const FVector& Position = SkelMeshImportData->Points[Wedge.VertexIndex];
+		Bounds += Position;
+	}
+
+	TArray<FSoftSkinVertex> Vertices;
+	SkeletalMesh->GetImportedModel()->LODModels[LODIndex].GetVertices(Vertices);
+	for (int32 SkinVertexIndex = 0; SkinVertexIndex < Vertices.Num(); ++SkinVertexIndex)
+	{
+		const FSoftSkinVertex& SkinVertex = Vertices[SkinVertexIndex];
+		Bounds += SkinVertex.Position;
+	}
+
+	TSKCVertPosOctree VertPosOctree(Bounds.GetCenter(), Bounds.GetExtent().GetMax());
+
+	// Add each old vertex to the octree
+	for (int32 SkinVertexIndex = 0; SkinVertexIndex < Vertices.Num(); ++SkinVertexIndex)
+	{
+		const FSoftSkinVertex& SkinVertex = Vertices[SkinVertexIndex];
+		VertPosOctree.AddElement(SkinVertex);
+	}
+
+	// Iterate over each new vertex position, attempting to find the old vertex it is closest to, applying
+	// the color of the old vertex to the new position if possible.
+	const float DistanceOverNormalThreshold =KINDA_SMALL_NUMBER;
+	for (int32 WedgeIndex = 0; WedgeIndex < WedgeNumber; ++WedgeIndex)
+	{
+		SkeletalMeshImportData::FVertex& Wedge = SkelMeshImportData->Wedges[WedgeIndex];
+		const FVector& Position = SkelMeshImportData->Points[Wedge.VertexIndex];
+		const FVector2D UV = Wedge.UVs[0];
+		TArray<FSoftSkinVertex> PointsToConsider;
+		TSKCVertPosOctree::TConstIterator<> OctreeIter(VertPosOctree);
+		// Iterate through the octree attempting to find the vertices closest to the current new point
+		while (OctreeIter.HasPendingNodes())
+		{
+			const TSKCVertPosOctree::FNode& CurNode = OctreeIter.GetCurrentNode();
+			const FOctreeNodeContext& CurContext = OctreeIter.GetCurrentContext();
+
+			// Find the child of the current node, if any, that contains the current new point
+			FOctreeChildNodeRef ChildRef = CurContext.GetContainingChild(FBoxCenterAndExtent(Position, FVector::ZeroVector));
+
+			if (!ChildRef.IsNULL())
+			{
+				const TSKCVertPosOctree::FNode* ChildNode = CurNode.GetChild(ChildRef);
+
+				// If the specified child node exists and contains any of the old vertices, push it to the iterator for future consideration
+				if (ChildNode && ChildNode->GetInclusiveElementCount() > 0)
+				{
+					OctreeIter.PushChild(ChildRef);
+				}
+				// If the child node doesn't have any of the old vertices in it, it's not worth pursuing any further. In an attempt to find
+				// anything to match vs. the new point, add all of the children of the current octree node that have old points in them to the
+				// iterator for future consideration.
+				else
+				{
+					FOREACH_OCTREE_CHILD_NODE(OctreeChildRef)
+					{
+						if (CurNode.HasChild(OctreeChildRef))
+						{
+							OctreeIter.PushChild(OctreeChildRef);
+						}
+					}
+				}
+			}
+
+			// Add all of the elements in the current node to the list of points to consider for closest point calculations
+			PointsToConsider.Append(CurNode.GetElements());
+			OctreeIter.Advance();
+		}
+
+		if (PointsToConsider.Num() > 0)
+		{
+			//Get the closest position
+			float MinDistance = MAX_FLT;
+			float MinUVDistance = MAX_FLT;
+			int32 MatchIndex = INDEX_NONE;
+			for (int32 ConsiderationIndex = 0; ConsiderationIndex < PointsToConsider.Num(); ++ConsiderationIndex)
+			{
+				const FSoftSkinVertex& SkinVertex = PointsToConsider[ConsiderationIndex];
+				const FVector2D& SkinVertexUV = SkinVertex.UVs[0];
+				const float UVDistanceSqr = FVector2D::DistSquared(UV, SkinVertexUV);
+				if (UVDistanceSqr < MinUVDistance)
+				{
+					MinUVDistance = FMath::Min(MinUVDistance, UVDistanceSqr);
+					MatchIndex = ConsiderationIndex;
+				}
+			}
+			if (PointsToConsider.IsValidIndex(MatchIndex))
+			{
+				Wedge.Color = PointsToConsider[MatchIndex].Color;
+			}
+		}
+	}
+}
+
 void FFbxImporter::SkinControlPointsToPose(FSkeletalMeshImportData& ImportData, FbxMesh* FbxMesh, FbxShape* FbxShape, bool bUseT0 )
 {
 	FbxTime poseTime = FBXSDK_TIME_INFINITE;
@@ -1512,6 +1664,17 @@ USkeletalMesh* UnFbx::FFbxImporter::ImportSkeletalMesh(FImportSkeletalMeshArgs &
 		// or texture that will be create with the same name
 		SkeletalMesh = NewObject<USkeletalMesh>(ImportSkeletalMeshArgs.InParent, ImportSkeletalMeshArgs.Name, ImportSkeletalMeshArgs.Flags);
 	}
+	else
+	{
+		//If we should keep the current vertex color but we do not have any vertex color, use the one from the fbx file (same behavior as the staticmesh)
+		if (!ImportOptions->bImportAsSkeletalSkinning && ImportOptions->VertexColorImportOption == EVertexColorImportOption::Ignore)
+		{
+			if (!ExistingSkelMesh->bHasVertexColors)
+			{
+				ImportOptions->VertexColorImportOption = EVertexColorImportOption::Replace;
+			}
+		}
+	}
 
 	FSkeletalMeshImportData TempData;
 	// Fill with data from buffer - contains the full .FBX file. 	
@@ -1540,20 +1703,32 @@ USkeletalMesh* UnFbx::FFbxImporter::ImportSkeletalMesh(FImportSkeletalMeshArgs &
 	}
 
 	//Adjust the import data from the import options
-	if (ExistingSkelMesh != nullptr && (ImportOptions->bImportAsSkeletalSkinning || ImportOptions->bImportAsSkeletalGeometry))
+	if (ExistingSkelMesh != nullptr)
 	{
-		if (ImportOptions->bImportAsSkeletalSkinning)
+		if ((ImportOptions->bImportAsSkeletalSkinning || ImportOptions->bImportAsSkeletalGeometry))
 		{
-			//Replace geometry import data by original existing skel mesh geometry data
-			ReplaceSkeletalMeshGeometryImportData(ExistingSkelMesh, SkelMeshImportDataPtr, ImportSkeletalMeshArgs.LodIndex);
+			if (ImportOptions->bImportAsSkeletalSkinning)
+			{
+				//Replace geometry import data by original existing skel mesh geometry data
+				ReplaceSkeletalMeshGeometryImportData(ExistingSkelMesh, SkelMeshImportDataPtr, ImportSkeletalMeshArgs.LodIndex);
+			}
+			else if (ImportOptions->bImportAsSkeletalGeometry)
+			{
+				//Replace skinning import data by original existing skel mesh skinning data
+				ReplaceSkeletalMeshSkinningImportData(ExistingSkelMesh, SkelMeshImportDataPtr, ImportSkeletalMeshArgs.LodIndex);
+			}
 		}
-		else if (ImportOptions->bImportAsSkeletalGeometry)
+
+		//Reuse the vertex color in case we are not re-importing weights only and there is some vertexcolor in the existing skeleMesh
+		if (!ImportOptions->bImportAsSkeletalSkinning && ExistingSkelMesh->bHasVertexColors && ImportOptions->VertexColorImportOption == EVertexColorImportOption::Ignore)
 		{
-			//Replace skinning import data by original existing skel mesh skinning data
-			ReplaceSkeletalMeshSkinningImportData(ExistingSkelMesh, SkelMeshImportDataPtr, ImportSkeletalMeshArgs.LodIndex);
+			int32 LODIndex = ImportSkeletalMeshArgs.LodIndex < 0 ? 0 : ImportSkeletalMeshArgs.LodIndex;
+			RemapSkeletalMeshVertexColorToImportData(ExistingSkelMesh, LODIndex, SkelMeshImportDataPtr);
 		}
 	}
 
+	
+
 	// Create initial bounding box based on expanded version of reference pose for meshes without physics assets. Can be overridden by artist.
 	FBox BoundingBox(SkelMeshImportDataPtr->Points.GetData(), SkelMeshImportDataPtr->Points.Num());
 	FBox Temp = BoundingBox;
