Added chunkid to each asset which changes in the diffassetregistries commandlet.

diff --git a/Engine/Source/Editor/UnrealEd/Private/Commandlets/DiffAssetRegistriesCommandlet.cpp b/Engine/Source/Editor/UnrealEd/Private/Commandlets/DiffAssetRegistriesCommandlet.cpp
--- a/Engine/Source/Editor/UnrealEd/Private/Commandlets/DiffAssetRegistriesCommandlet.cpp
+++ b/Engine/Source/Editor/UnrealEd/Private/Commandlets/DiffAssetRegistriesCommandlet.cpp
@@ -774,7 +774,7 @@ void UDiffAssetRegistriesCommandlet::LogChangedFiles(FArchive *CSVFile, FString
 		CSVFile->Logf(TEXT("x, no binary change (shouldn't ever happen)"));
 		CSVFile->Logf(TEXT(""));
 
-		CSVFile->Logf(TEXT("Modification,Name,Class,NewSize,OldSize,Changelist"));
+		CSVFile->Logf(TEXT("Modification,Name,Class,NewSize,OldSize,Changelist,Chunk"));
 
 		UE_LOG(LogDiffAssets, Display, TEXT("Saving CSV results to %s"), *CSVFilename);
 	}
@@ -795,12 +795,18 @@ void UDiffAssetRegistriesCommandlet::LogChangedFiles(FArchive *CSVFile, FString
 		{
 			ClassName = GetClassName(NewState, AssetPath);
 		}
+		auto GetChunkIDString = [this, AssetPath=AssetPath](FAssetRegistryState& State) {
+				TArray<int32> ChunkIds = GetAssetChunks(State, AssetPath);
+				FString ChunkIdString = FString::JoinBy(ChunkIds, TEXT(" & "), [](int32 Value) { return FString::Printf(TEXT("%d"), Value); });
+				return ChunkIdString;
+			};
+		
 
 		if (ChangeInfo.Adds)
 		{
 			if (CSVFile)
 			{
-				CSVFile->Logf(TEXT("a,%s,%s,%d,0,%d"), *AssetPath.ToString(), *ClassName.ToString(), ChangeInfo.AddedBytes, Changelist);
+				CSVFile->Logf(TEXT("a,%s,%s,%d,0,%d,%s"), *AssetPath.ToString(), *ClassName.ToString(), ChangeInfo.AddedBytes, Changelist, *GetChunkIDString(NewState));
 			}
 
 			if (bIsVerbose)
@@ -839,7 +845,7 @@ void UDiffAssetRegistriesCommandlet::LogChangedFiles(FArchive *CSVFile, FString
 
 			if (CSVFile)
 			{
-				CSVFile->Logf(TEXT("%c,%s,%s,%d,%d,%d"), classification, *AssetPath.ToString(), *ClassName.ToString(), ChangeInfo.ChangedBytes, PrevData->DiskSize, Changelist);
+				CSVFile->Logf(TEXT("%c,%s,%s,%d,%d,%d,%s"), classification, *AssetPath.ToString(), *ClassName.ToString(), ChangeInfo.ChangedBytes, PrevData->DiskSize, Changelist, *GetChunkIDString(NewState));
 			}
 
 			if (bIsVerbose)
@@ -874,7 +880,7 @@ void UDiffAssetRegistriesCommandlet::LogChangedFiles(FArchive *CSVFile, FString
 
 			if (CSVFile)
 			{
-				CSVFile->Logf(TEXT("r,%s,%s,0,%d,0"), *AssetPath.ToString(), *ClassName.ToString(), PrevData->DiskSize);
+				CSVFile->Logf(TEXT("r,%s,%s,0,%d,0,0"), *AssetPath.ToString(), *ClassName.ToString(), PrevData->DiskSize);
 			}
 
 			if (bIsVerbose)
