Fixing usage of RegexBranchCL in DiffAssetRegistryCommandlet to parse the Branch+CL out from the NewPath param when we pass one that includes them

diff --git a/Engine/Source/Editor/UnrealEd/Private/Commandlets/DiffAssetRegistriesCommandlet.cpp b/Engine/Source/Editor/UnrealEd/Private/Commandlets/DiffAssetRegistriesCommandlet.cpp
--- a/Engine/Source/Editor/UnrealEd/Private/Commandlets/DiffAssetRegistriesCommandlet.cpp
+++ b/Engine/Source/Editor/UnrealEd/Private/Commandlets/DiffAssetRegistriesCommandlet.cpp
@@ -196,6 +196,17 @@ int32 UDiffAssetRegistriesCommandlet::Main(const FString& FullCommandLine)
 	if (NewPathVal)
 	{
 		FindAssetRegistryPath(*NewPathVal, NewPath);
+
+		if (!RegexBranchCL.IsEmpty())
+		{
+			const FRegexPattern CLPattern(RegexBranchCL);
+			FRegexMatcher CLMatcher(CLPattern, NewPath);
+			if (CLMatcher.FindNext())
+			{
+				Branch = CLMatcher.GetCaptureGroup(1);
+				CL = CLMatcher.GetCaptureGroup(2);
+			}
+		}
 	}
 
 	bMatchChangelists = false;
