Changing comment on param bShouldActuallyTransact

#rnx

diff --git a/Engine/Source/Editor/UnrealEd/Public/ScopedTransaction.h b/Engine/Source/Editor/UnrealEd/Public/ScopedTransaction.h
--- a/Engine/Source/Editor/UnrealEd/Public/ScopedTransaction.h
+++ b/Engine/Source/Editor/UnrealEd/Public/ScopedTransaction.h
@@ -16,14 +16,14 @@ public:
 	 * @param TransactionContext		The context for this transaction, ie editor module or producing system
 	 * @param SessionName				The description of this transaction, ie what is occurring
 	 * @param PrimaryObject				The main object being edited (if known)
-	 * @param bShouldActuallyTransact	Debug feature, if false the transaction is ignored 
+	 * @param bShouldActuallyTransact	If false the transaction is ignored. Good for creating optional transactions based on condition
 	 */
 	UNREALED_API FScopedTransaction(const TCHAR* TransactionContext, const FText& SessionName, UObject* PrimaryObject, const bool bShouldActuallyTransact = true);
 
 	/**
 	 * Construct an FScopedTransaction with minimal context 
 	 * @param SessionName				The description of this transaction, ie what is occurring
-	 * @param bShouldActuallyTransact	Debug feature, if false the transaction is ignored 
+	 * @param bShouldActuallyTransact	If false the transaction is ignored. Good for creating optional transactions based on condition 
 	 */
 	UNREALED_API FScopedTransaction(const FText& SessionName, const bool bShouldActuallyTransact = true);
 	UNREALED_API ~FScopedTransaction();
