Make sure skeletalmesh re-import follow correctly EVertexColorImportOption::Ignore option.

diff --git a/Engine/Source/Editor/UnrealEd/Private/Fbx/FbxSkeletalMeshImport.cpp b/Engine/Source/Editor/UnrealEd/Private/Fbx/FbxSkeletalMeshImport.cpp
--- a/Engine/Source/Editor/UnrealEd/Private/Fbx/FbxSkeletalMeshImport.cpp
+++ b/Engine/Source/Editor/UnrealEd/Private/Fbx/FbxSkeletalMeshImport.cpp
@@ -1720,7 +1720,7 @@ USkeletalMesh* UnFbx::FFbxImporter::ImportSkeletalMesh(FImportSkeletalMeshArgs &
 		}
 
 		//Reuse the vertex color in case we are not re-importing weights only and there is some vertexcolor in the existing skeleMesh
-		if (!ImportOptions->bImportAsSkeletalSkinning && ExistingSkelMesh->bHasVertexColors && ImportOptions->VertexColorImportOption == EVertexColorImportOption::Ignore)
+		if (ExistingSkelMesh->bHasVertexColors && ImportOptions->VertexColorImportOption == EVertexColorImportOption::Ignore)
 		{
 			int32 LODIndex = ImportSkeletalMeshArgs.LodIndex < 0 ? 0 : ImportSkeletalMeshArgs.LodIndex;
 			RemapSkeletalMeshVertexColorToImportData(ExistingSkelMesh, LODIndex, SkelMeshImportDataPtr);
