Save active feature level preview state between sessions

diff --git a/Engine/Source/Editor/LevelEditor/Private/SLevelEditor.cpp b/Engine/Source/Editor/LevelEditor/Private/SLevelEditor.cpp
--- a/Engine/Source/Editor/LevelEditor/Private/SLevelEditor.cpp
+++ b/Engine/Source/Editor/LevelEditor/Private/SLevelEditor.cpp
@@ -162,7 +162,7 @@ void SLevelEditor::Construct( const SLevelEditor::FArguments& InArgs)
 
 	// Set the initial preview feature level.
 	UEditorEngine* Editor = (UEditorEngine*)GEngine;
-	World->ChangeFeatureLevel(Editor->PreviewFeatureLevel);
+	World->ChangeFeatureLevel(Editor->GetActiveFeatureLevelPreviewType());
 
 	// Patch into the OnPreviewFeatureLevelChanged() delegate to swap out the current feature level with a user selection.
 	PreviewFeatureLevelChangedHandle = Editor->OnPreviewFeatureLevelChanged().AddLambda([this](ERHIFeatureLevel::Type NewFeatureLevel)
diff --git a/Engine/Source/Editor/UnrealEd/Classes/Editor/EditorEngine.h b/Engine/Source/Editor/UnrealEd/Classes/Editor/EditorEngine.h
--- a/Engine/Source/Editor/UnrealEd/Classes/Editor/EditorEngine.h
+++ b/Engine/Source/Editor/UnrealEd/Classes/Editor/EditorEngine.h
@@ -3057,12 +3057,18 @@ public:
 	/** Toggle the feature level preview */
 	void ToggleFeatureLevelPreview();
 
-	/** Return whether the feature level preview is able to be enabled */
+	/** Return whether the feature level preview is enabled for use via the user accessing Settings->Preview Rendering Level. */
 	bool IsFeatureLevelPreviewEnabled() const;
 
-	/** Return whether the feature level preview is currently active */
+	/** Return whether the feature level preview enabled via Settings->Preview Rendering Level is currently active/displaying. */
 	bool IsFeatureLevelPreviewActive() const;
 
+	/**
+	 * Return the active feature level. This will be the chosen Settings->Preview Rendering Level if the Preview Mode button
+	 * is highlighted or SM5 if the Preview Mode button has been clicked and isn't highlighted.
+	 */
+	ERHIFeatureLevel::Type GetActiveFeatureLevelPreviewType() const;
+
 	/** Return the delegate that is called when the preview feature level changes */
 	FPreviewFeatureLevelChanged& OnPreviewFeatureLevelChanged() { return PreviewFeatureLevelChanged; }
 
diff --git a/Engine/Source/Editor/UnrealEd/Classes/Editor/EditorPerProjectUserSettings.h b/Engine/Source/Editor/UnrealEd/Classes/Editor/EditorPerProjectUserSettings.h
--- a/Engine/Source/Editor/UnrealEd/Classes/Editor/EditorPerProjectUserSettings.h
+++ b/Engine/Source/Editor/UnrealEd/Classes/Editor/EditorPerProjectUserSettings.h
@@ -157,6 +157,10 @@ public:
 	UPROPERTY(config)
 	bool bIsMaterialQualityOverridePlatform;
 
+	/** whether or not the Feature Level Preview is active or whether it is GMaxRHIFeatureLevel is being shown instead */
+	UPROPERTY(config)
+	bool bIsFeatureLevelPreviewActive;
+
 public:
 
 	/** Delegate for when a user setting has changed */
diff --git a/Engine/Source/Editor/UnrealEd/Private/EditorEngine.cpp b/Engine/Source/Editor/UnrealEd/Private/EditorEngine.cpp
--- a/Engine/Source/Editor/UnrealEd/Private/EditorEngine.cpp
+++ b/Engine/Source/Editor/UnrealEd/Private/EditorEngine.cpp
@@ -7421,6 +7421,10 @@ void UEditorEngine::SetPreviewPlatform(const FName MaterialQualityPlatform, ERHI
 	const FName InitialPreviewPlatform = MaterialShaderQualitySettings->GetPreviewPlatform();
 	MaterialShaderQualitySettings->SetPreviewPlatform(MaterialQualityPlatform);
 
+	// Force activation of the feature level preview, because it may have been inactive before.
+ 	auto* Settings = GetMutableDefault<UEditorPerProjectUserSettings>();
+	Settings->bIsFeatureLevelPreviewActive = true;
+
 	if (PreviewFeatureLevel != InPreviewFeatureLevel)
 	{
 		// a new feature level will recompile the materials and apply the effect of any 'material quality platform'
@@ -7440,7 +7444,10 @@ void UEditorEngine::SetPreviewPlatform(const FName MaterialQualityPlatform, ERHI
 
 void UEditorEngine::ToggleFeatureLevelPreview()
 {
-	ERHIFeatureLevel::Type NewPreviewFeatureLevel = GWorld->FeatureLevel == GMaxRHIFeatureLevel ? PreviewFeatureLevel : GMaxRHIFeatureLevel;
+ 	auto* Settings = GetMutableDefault<UEditorPerProjectUserSettings>();
+ 	Settings->bIsFeatureLevelPreviewActive ^= 1;
+
+	ERHIFeatureLevel::Type NewPreviewFeatureLevel = Settings->bIsFeatureLevelPreviewActive ? PreviewFeatureLevel : GMaxRHIFeatureLevel;
 
 	PreviewFeatureLevelChanged.Broadcast(NewPreviewFeatureLevel);
 	
@@ -7453,6 +7460,8 @@ void UEditorEngine::ToggleFeatureLevelPreview()
 	UpdateShaderComplexityMaterials(true);
 
 	GEditor->RedrawAllViewports();
+	
+	SaveEditorFeatureLevel();
 }
 
 bool UEditorEngine::IsFeatureLevelPreviewEnabled() const
@@ -7462,7 +7471,14 @@ bool UEditorEngine::IsFeatureLevelPreviewEnabled() const
 
 bool UEditorEngine::IsFeatureLevelPreviewActive() const
 {
-	return GWorld->FeatureLevel != GMaxRHIFeatureLevel;
+ 	auto* Settings = GetMutableDefault<UEditorPerProjectUserSettings>();
+ 	return Settings->bIsFeatureLevelPreviewActive;
+}
+
+ERHIFeatureLevel::Type UEditorEngine::GetActiveFeatureLevelPreviewType() const
+{
+ 	auto* Settings = GetMutableDefault<UEditorPerProjectUserSettings>();
+	return Settings->bIsFeatureLevelPreviewActive ? PreviewFeatureLevel : GMaxRHIFeatureLevel;
 }
 
 void UEditorEngine::LoadEditorFeatureLevel()
@@ -7482,7 +7498,11 @@ void UEditorEngine::LoadEditorFeatureLevel()
 			UMaterialShaderQualitySettings::Get()->GetShaderPlatformQualitySettings(Settings->PreviewShaderPlatformName);
 		}
 
+		// Save off bIsFeatureLevelPreviewActive, because SetPreviewPlatform() will force it on.
+		// We want it to be what the user chose.
+		bool bSaveIsFeatureLevelPreviewActive = Settings->bIsFeatureLevelPreviewActive;
 		SetPreviewPlatform(MaterialQualityPlatform, FeatureLevel, false);
+		Settings->bIsFeatureLevelPreviewActive = bSaveIsFeatureLevelPreviewActive;
 	}
 }
 
