Use the proper memory address when grabbing debug info for output parameters. This fixes a crash that can occur when mousing over output parameter pins during blueprint debugging.

diff --git a/Engine/Source/Editor/UnrealEd/Private/Kismet2/KismetDebugUtilities.cpp b/Engine/Source/Editor/UnrealEd/Private/Kismet2/KismetDebugUtilities.cpp
--- a/Engine/Source/Editor/UnrealEd/Private/Kismet2/KismetDebugUtilities.cpp
+++ b/Engine/Source/Editor/UnrealEd/Private/Kismet2/KismetDebugUtilities.cpp
@@ -1180,7 +1180,19 @@ FKismetDebugUtilities::EWatchTextResult FKismetDebugUtilities::FindDebuggingData
 			{
 				if (Property->IsIn(TestFrame->Node))
 				{
-					PropertyBase = TestFrame->Locals;
+					// output parameters need special handling
+					for (FOutParmRec* OutParmRec = TestFrame->OutParms; OutParmRec != nullptr; OutParmRec = OutParmRec->NextOutParm)
+					{
+						if (OutParmRec->Property == Property)
+						{
+							PropertyBase = OutParmRec->PropAddr;
+							break;
+						}
+					}
+					if (PropertyBase == nullptr)
+					{
+						PropertyBase = TestFrame->Locals;
+					}
 					break;
 				}
 			}
