Fix clothing assignment when generating LODs

diff --git a/Engine/Source/Editor/Persona/Persona.Build.cs b/Engine/Source/Editor/Persona/Persona.Build.cs
--- a/Engine/Source/Editor/Persona/Persona.Build.cs
+++ b/Engine/Source/Editor/Persona/Persona.Build.cs
@@ -62,6 +62,7 @@ public class Persona : ModuleRules
                 "UnrealEd",
                 "CommonMenuExtensions",
                 "PinnedCommandList",
+                "RenderCore",
             }
         );
 
diff --git a/Engine/Source/Editor/Persona/Private/PersonaMeshDetails.cpp b/Engine/Source/Editor/Persona/Private/PersonaMeshDetails.cpp
--- a/Engine/Source/Editor/Persona/Private/PersonaMeshDetails.cpp
+++ b/Engine/Source/Editor/Persona/Private/PersonaMeshDetails.cpp
@@ -1446,7 +1446,7 @@ FReply FPersonaMeshDetails::OnApplyChanges()
 	return FReply::Handled();
 }
 
-void FPersonaMeshDetails::RegenerateOneLOD(int32 LODIndex)
+void FPersonaMeshDetails::RegenerateOneLOD(int32 LODIndex, bool bReregisterComponent /*= true*/)
 {
 	USkeletalMesh* SkelMesh = GetPersonaToolkit()->GetMesh();
 	check(SkelMesh);
@@ -1462,7 +1462,7 @@ void FPersonaMeshDetails::RegenerateOneLOD(int32 LODIndex)
 		{
 			//Restore the base LOD data
 			CurrentLODInfo.bHasBeenSimplified = false;
-			FLODUtilities::RestoreSkeletalMeshLODImportedData(SkelMesh, LODIndex, true);
+			FLODUtilities::RestoreSkeletalMeshLODImportedData(SkelMesh, LODIndex, bReregisterComponent);
 			return;
 		}
 		else if (!CurrentLODInfo.bHasBeenSimplified
@@ -1489,14 +1489,14 @@ void FPersonaMeshDetails::RegenerateOneLOD(int32 LODIndex)
 		UpdateContext.SkeletalMesh = SkelMesh;
 		UpdateContext.AssociatedComponents.Push(GetPersonaToolkit()->GetPreviewMeshComponent());
 
-		FLODUtilities::SimplifySkeletalMeshLOD(UpdateContext, LODIndex);
+		FLODUtilities::SimplifySkeletalMeshLOD(UpdateContext, LODIndex, bReregisterComponent);
 	}
 	return;
 }
 
 //Regenerate dependent LODs if we re-import LOD X any LOD Z using X has source must be regenerated
 //Also just generate already simplified mesh
-void FPersonaMeshDetails::RegenerateDependentLODs(int32 LODIndex)
+void FPersonaMeshDetails::RegenerateDependentLODs(int32 LODIndex, bool bReregisterComponent /*= true*/)
 {
 	USkeletalMesh* SkelMesh = GetPersonaToolkit()->GetMesh();
 	check(SkelMesh);
@@ -1516,7 +1516,7 @@ void FPersonaMeshDetails::RegenerateDependentLODs(int32 LODIndex)
 			{
 				DependentLODs[CurrentLODIndex] = true;
 				//Regenerate this LOD
-				RegenerateOneLOD(CurrentLODIndex);
+				RegenerateOneLOD(CurrentLODIndex, bReregisterComponent);
 			}
 		}
 	}
@@ -1554,9 +1554,41 @@ FReply FPersonaMeshDetails::RegenerateLOD(int32 LODIndex)
 			}
 		}
 	}
+	
+	//Reregister scope
+	{
+		TComponentReregisterContext<USkinnedMeshComponent> ReregisterContext;
+		SkelMesh->ReleaseResources();
+		SkelMesh->ReleaseResourcesFence.Wait();
+		SkelMesh->PreEditChange(nullptr);
+		SkelMesh->Modify();
+
+		// Unbind any existing clothing assets before we regenerate all LODs
+		TArray<ClothingAssetUtils::FClothingAssetMeshBinding> ClothingBindings;
+		ClothingAssetUtils::GetMeshClothingAssetBindings(SkelMesh, ClothingBindings);
+
+		for (ClothingAssetUtils::FClothingAssetMeshBinding& Binding : ClothingBindings)
+		{
+			Binding.Asset->UnbindFromSkeletalMesh(SkelMesh, Binding.LODIndex);
+		}
+
+		RegenerateOneLOD(LODIndex, false);
+		RegenerateDependentLODs(LODIndex, false);
+
+		//Restore all clothing we can
+		for (ClothingAssetUtils::FClothingAssetMeshBinding& Binding : ClothingBindings)
+		{
+			if (SkelMesh->GetImportedModel()->LODModels.IsValidIndex(Binding.LODIndex) &&
+				SkelMesh->GetImportedModel()->LODModels[Binding.LODIndex].Sections.IsValidIndex(Binding.SectionIndex))
+			{
+				Binding.Asset->BindToSkeletalMesh(SkelMesh, Binding.LODIndex, Binding.SectionIndex, Binding.AssetInternalLodIndex);
+			}
+		}
+
+		SkelMesh->PostEditChange();
+		SkelMesh->InitResources();
+	}
 
-	RegenerateOneLOD(LODIndex);
-	RegenerateDependentLODs(LODIndex);
 	return FReply::Handled();
 }
 
diff --git a/Engine/Source/Editor/Persona/Private/PersonaMeshDetails.h b/Engine/Source/Editor/Persona/Private/PersonaMeshDetails.h
--- a/Engine/Source/Editor/Persona/Private/PersonaMeshDetails.h
+++ b/Engine/Source/Editor/Persona/Private/PersonaMeshDetails.h
@@ -364,9 +364,9 @@ private:
 	/** apply LOD changes if the user modified LOD reduction settings */
 	FReply OnApplyChanges();
 	/** regenerate one specific LOD Index no dependencies*/
-	void RegenerateOneLOD(int32 LODIndex);
+	void RegenerateOneLOD(int32 LODIndex, bool bReregisterComponent = true);
 	/** regenerate the specific all LODs dependent of InLODIndex. This is not regenerating the InLODIndex*/
-	void RegenerateDependentLODs(int32 LODIndex);
+	void RegenerateDependentLODs(int32 LODIndex, bool bReregisterComponent = true);
 	/** Apply specified LOD Index */
 	FReply RegenerateLOD(int32 LODIndex);
 	/** Removes the specified lod from the skeletal mesh */
diff --git a/Engine/Source/Editor/SkeletalMeshEditor/Private/SkeletalMeshEditor.cpp b/Engine/Source/Editor/SkeletalMeshEditor/Private/SkeletalMeshEditor.cpp
--- a/Engine/Source/Editor/SkeletalMeshEditor/Private/SkeletalMeshEditor.cpp
+++ b/Engine/Source/Editor/SkeletalMeshEditor/Private/SkeletalMeshEditor.cpp
@@ -546,7 +546,8 @@ void FSkeletalMeshEditor::OnRemoveSectionFromLodAndBelowMenuItemClicked(int32 Lo
 		const FSkeletalMeshLODInfo* CurrentSkeletalMeshLODInfo = SkeletalMesh->GetLODInfo(GenerateLodIndex);
 		if (CurrentSkeletalMeshLODInfo != nullptr && CurrentSkeletalMeshLODInfo->bHasBeenSimplified && BaseLodIndexes.Contains(CurrentSkeletalMeshLODInfo->ReductionSettings.BaseLOD))
 		{
-			FLODUtilities::SimplifySkeletalMeshLOD(UpdateContext, GenerateLodIndex);
+			bool bRestoreClothing = true;
+			FLODUtilities::SimplifySkeletalMeshLOD(UpdateContext, GenerateLodIndex, true, bRestoreClothing);
 			BaseLodIndexes.Add(GenerateLodIndex);
 		}
 	}
@@ -937,7 +938,8 @@ void ReimportAllCustomLODs(USkeletalMesh* SkeletalMesh, UDebugSkelMeshComponent*
 			FSkeletalMeshUpdateContext UpdateContext;
 			UpdateContext.SkeletalMesh = SkeletalMesh;
 			UpdateContext.AssociatedComponents.Push(PreviewMeshComponent);
-			FLODUtilities::SimplifySkeletalMeshLOD(UpdateContext, LodIndex);
+			bool bRestoreClothing = true;
+			FLODUtilities::SimplifySkeletalMeshLOD(UpdateContext, LodIndex, true, bRestoreClothing);
 			Dependencies[LodIndex] = true;
 		}
 	}
diff --git a/Engine/Source/Editor/UnrealEd/Private/LODUtilities.cpp b/Engine/Source/Editor/UnrealEd/Private/LODUtilities.cpp
--- a/Engine/Source/Editor/UnrealEd/Private/LODUtilities.cpp
+++ b/Engine/Source/Editor/UnrealEd/Private/LODUtilities.cpp
@@ -14,6 +14,7 @@
 #include "Engine/SkeletalMesh.h"
 #include "EditorFramework/AssetImportData.h"
 #include "MeshUtilities.h"
+#include "Assets/ClothingAsset.h"
 
 #if WITH_APEX_CLOTHING
 	#include "ApexClothingUtils.h"
@@ -22,12 +23,23 @@
 #include "ComponentReregisterContext.h"
 #include "IMeshReductionManagerModule.h"
 
+
+
 DEFINE_LOG_CATEGORY_STATIC(LogLODUtilities, Log, All);
 
 bool FLODUtilities::RegenerateLOD(USkeletalMesh* SkeletalMesh, int32 NewLODCount /*= 0*/, bool bRegenerateEvenIfImported /*= false*/, bool bGenerateBaseLOD /*= false*/)
 {
 	if (SkeletalMesh)
 	{
+		// Unbind any existing clothing assets before we regenerate all LODs
+		TArray<ClothingAssetUtils::FClothingAssetMeshBinding> ClothingBindings;
+		ClothingAssetUtils::GetMeshClothingAssetBindings(SkeletalMesh, ClothingBindings);
+
+		for (ClothingAssetUtils::FClothingAssetMeshBinding& Binding : ClothingBindings)
+		{
+			Binding.Asset->UnbindFromSkeletalMesh(SkeletalMesh, Binding.LODIndex);
+		}
+
 		int32 LODCount = SkeletalMesh->GetLODNum();
 
 		if (NewLODCount > 0)
@@ -71,6 +83,16 @@ bool FLODUtilities::RegenerateLOD(USkeletalMesh* SkeletalMesh, int32 NewLODCount
 			}
 		}
 
+		//Restore all clothing we can
+		for (ClothingAssetUtils::FClothingAssetMeshBinding& Binding : ClothingBindings)
+		{
+			if (SkeletalMesh->GetImportedModel()->LODModels.IsValidIndex(Binding.LODIndex) &&
+				SkeletalMesh->GetImportedModel()->LODModels[Binding.LODIndex].Sections.IsValidIndex(Binding.SectionIndex))
+			{
+				Binding.Asset->BindToSkeletalMesh(SkeletalMesh, Binding.LODIndex, Binding.SectionIndex, Binding.AssetInternalLodIndex);
+			}
+		}
+
 		SkeletalMesh->PostEditChange();
 
 		return true;
@@ -718,7 +740,7 @@ void FLODUtilities::ApplyMorphTargetsToLOD(USkeletalMesh* SkeletalMesh, int32 So
 	CreateLODMorphTarget(SkeletalMesh, ReductionBaseSkeletalMeshBulkData, SourceLOD, DestinationLOD, PerMorphTargetBaseIndexToMorphTargetDelta, BaseMorphIndexToTargetIndexList, TargetVertices, TargetMatchData);
 }
 
-void FLODUtilities::SimplifySkeletalMeshLOD( USkeletalMesh* SkeletalMesh, int32 DesiredLOD, bool bReregisterComponent /*= true*/ )
+void FLODUtilities::SimplifySkeletalMeshLOD( USkeletalMesh* SkeletalMesh, int32 DesiredLOD, bool bReregisterComponent /*= true*/, bool bRestoreClothing /*= false*/)
 {
 	IMeshReductionModule& ReductionModule = FModuleManager::Get().LoadModuleChecked<IMeshReductionModule>("MeshReductionInterface");
 	IMeshReduction* MeshReduction = ReductionModule.GetSkeletalMeshReductionInterface();
@@ -754,6 +776,20 @@ void FLODUtilities::SimplifySkeletalMeshLOD( USkeletalMesh* SkeletalMesh, int32
 		GWarn->BeginSlowTask(StatusUpdate, true);
 	}
 
+	// Unbind DesiredLOD existing clothing assets before we simplify this LOD
+	TArray<ClothingAssetUtils::FClothingAssetMeshBinding> ClothingBindings;
+	if (bRestoreClothing)
+	{
+		ClothingAssetUtils::GetMeshClothingAssetBindings(SkeletalMesh, ClothingBindings);
+		for (ClothingAssetUtils::FClothingAssetMeshBinding& Binding : ClothingBindings)
+		{
+			if (DesiredLOD == Binding.LODIndex)
+			{
+				Binding.Asset->UnbindFromSkeletalMesh(SkeletalMesh, Binding.LODIndex);
+			}
+		}
+	}
+
 	if (SkeletalMesh->GetLODInfo(DesiredLOD) != nullptr)
 	{
 		FSkeletalMeshModel* SkeletalMeshResource = SkeletalMesh->GetImportedModel();
@@ -842,10 +878,27 @@ void FLODUtilities::SimplifySkeletalMeshLOD( USkeletalMesh* SkeletalMesh, int32
 		const FText Message = FText::Format(NSLOCTEXT("UnrealEd", "MeshSimp_GenerateLODFailed_F", "An error occurred while simplifying the geometry for mesh '{SkeletalMeshName}'.  Consider adjusting simplification parameters and re-simplifying the mesh."), Args);
 		FMessageDialog::Open(EAppMsgType::Ok, Message);
 	}
+
+	//Put back the clothing for the DesiredLOD
+	if (bRestoreClothing)
+	{
+		for (ClothingAssetUtils::FClothingAssetMeshBinding& Binding : ClothingBindings)
+		{
+			if (SkeletalMesh->GetImportedModel()->LODModels.IsValidIndex(Binding.LODIndex) &&
+				SkeletalMesh->GetImportedModel()->LODModels[Binding.LODIndex].Sections.IsValidIndex(Binding.SectionIndex))
+			{
+				if (DesiredLOD == Binding.LODIndex)
+				{
+					Binding.Asset->BindToSkeletalMesh(SkeletalMesh, Binding.LODIndex, Binding.SectionIndex, Binding.AssetInternalLodIndex);
+				}
+			}
+		}
+	}
+
 	GWarn->EndSlowTask();
 }
 
-void FLODUtilities::SimplifySkeletalMeshLOD(FSkeletalMeshUpdateContext& UpdateContext, int32 DesiredLOD, bool bReregisterComponent /*= true*/)
+void FLODUtilities::SimplifySkeletalMeshLOD(FSkeletalMeshUpdateContext& UpdateContext, int32 DesiredLOD, bool bReregisterComponent /*= true*/, bool bRestoreClothing /*= false*/)
 {
 	USkeletalMesh* SkeletalMesh = UpdateContext.SkeletalMesh;
 	IMeshReductionModule& ReductionModule = FModuleManager::Get().LoadModuleChecked<IMeshReductionModule>("MeshReductionInterface");
@@ -853,7 +906,7 @@ void FLODUtilities::SimplifySkeletalMeshLOD(FSkeletalMeshUpdateContext& UpdateCo
 
 	if (MeshReduction && MeshReduction->IsSupported() && SkeletalMesh)
 	{
-		SimplifySkeletalMeshLOD(SkeletalMesh, DesiredLOD, bReregisterComponent);
+		SimplifySkeletalMeshLOD(SkeletalMesh, DesiredLOD, bReregisterComponent, bRestoreClothing);
 		
 		if (UpdateContext.OnLODChanged.IsBound())
 		{
@@ -871,6 +924,19 @@ void FLODUtilities::RestoreSkeletalMeshLODImportedData(USkeletalMesh* SkeletalMe
 		//There is nothing to restore
 		return;
 	}
+
+	// Unbind LodIndex existing clothing assets before restoring the LOD
+	TArray<ClothingAssetUtils::FClothingAssetMeshBinding> ClothingBindings;
+	ClothingAssetUtils::GetMeshClothingAssetBindings(SkeletalMesh, ClothingBindings);
+	for (ClothingAssetUtils::FClothingAssetMeshBinding& Binding : ClothingBindings)
+	{
+		//Unbind only the LOD we restore
+		if (Binding.LODIndex == LodIndex)
+		{
+			Binding.Asset->UnbindFromSkeletalMesh(SkeletalMesh, Binding.LODIndex);
+		}
+	}
+
 	FSkeletalMeshLODModel ImportedBaseLODModel;
 	TMap<FString, TArray<FMorphTargetDelta>> ImportedBaseLODMorphTargetData;
 	SkeletalMesh->GetImportedModel()->OriginalReductionSourceMeshData[LodIndex]->LoadReductionData(ImportedBaseLODModel, ImportedBaseLODMorphTargetData);
@@ -904,6 +970,16 @@ void FLODUtilities::RestoreSkeletalMeshLODImportedData(USkeletalMesh* SkeletalMe
 		//Empty the bulkdata since we restore it
 		SkeletalMesh->GetImportedModel()->OriginalReductionSourceMeshData[LodIndex]->EmptyBulkData();
 
+		//Put back the clothing for the restore LOD
+		for (ClothingAssetUtils::FClothingAssetMeshBinding& Binding : ClothingBindings)
+		{
+			if (LodIndex == Binding.LODIndex && SkeletalMesh->GetImportedModel()->LODModels.IsValidIndex(Binding.LODIndex) &&
+				SkeletalMesh->GetImportedModel()->LODModels[Binding.LODIndex].Sections.IsValidIndex(Binding.SectionIndex))
+			{
+				Binding.Asset->BindToSkeletalMesh(SkeletalMesh, Binding.LODIndex, Binding.SectionIndex, Binding.AssetInternalLodIndex);
+			}
+		}
+
 		if (bReregisterComponent)
 		{
 			SkeletalMesh->PostEditChange();
diff --git a/Engine/Source/Editor/UnrealEd/Public/LODUtilities.h b/Engine/Source/Editor/UnrealEd/Public/LODUtilities.h
--- a/Engine/Source/Editor/UnrealEd/Public/LODUtilities.h
+++ b/Engine/Source/Editor/UnrealEd/Public/LODUtilities.h
@@ -49,7 +49,7 @@ public:
 	* @param UpdateContext - The skeletal mesh and actor components to operate on.
 	* @param DesiredLOD - The LOD to simplify
 	*/
-	static void SimplifySkeletalMeshLOD(FSkeletalMeshUpdateContext& UpdateContext, int32 DesiredLOD, bool bReregisterComponent = true);
+	static void SimplifySkeletalMeshLOD(FSkeletalMeshUpdateContext& UpdateContext, int32 DesiredLOD, bool bReregisterComponent = true, bool bRestoreClothing = false);
 
 	/**
 	*	Restore the LOD imported model to the last imported data. Call this function if you want to remove the reduce on the base LOD
@@ -80,7 +80,7 @@ private:
 	 * @param SkeletalMesh - The skeletal mesh and actor components to operate on.
 	 * @param DesiredLOD - Desired LOD
 	 */
-	static void SimplifySkeletalMeshLOD(USkeletalMesh* SkeletalMesh, int32 DesiredLOD, bool bReregisterComponent = true);
+	static void SimplifySkeletalMeshLOD(USkeletalMesh* SkeletalMesh, int32 DesiredLOD, bool bReregisterComponent = true, bool bRestoreClothing = false);
 
 	/**
 	*  Remap the morph targets of the base LOD onto the desired LOD.
