Remove trailing ; that clang didn't care about but MSVC does.

diff --git a/Engine/Source/Editor/UnrealEd/Private/Commandlets/ShaderCodeLibraryToolsCommandlet.cpp b/Engine/Source/Editor/UnrealEd/Private/Commandlets/ShaderCodeLibraryToolsCommandlet.cpp
--- a/Engine/Source/Editor/UnrealEd/Private/Commandlets/ShaderCodeLibraryToolsCommandlet.cpp
+++ b/Engine/Source/Editor/UnrealEd/Private/Commandlets/ShaderCodeLibraryToolsCommandlet.cpp
@@ -35,7 +35,7 @@ int32 UShaderCodeLibraryToolsCommandlet::StaticMain(const FString& Params)
 		bool bNativeFormat = Switches.Contains(TEXT("PreferNativeArchives")) || Switches.Contains(TEXT("-PreferNativeArchives"));
 		
 		bool bArchive = false;
-		if (GConfig->GetBool(TEXT("/Script/UnrealEd.ProjectPackagingSettings"), TEXT("bSharedMaterialNativeLibraries"), bArchive, GGameIni));
+		if (GConfig->GetBool(TEXT("/Script/UnrealEd.ProjectPackagingSettings"), TEXT("bSharedMaterialNativeLibraries"), bArchive, GGameIni))
 		{
 			bNativeFormat |= bArchive;
 		}
