Fix skeletal mesh vertex paint remap to find a better match when many wedge shared the same position. Adding normal dot product comparison.

diff --git a/Engine/Source/Editor/UnrealEd/Private/Fbx/FbxSkeletalMeshImport.cpp b/Engine/Source/Editor/UnrealEd/Private/Fbx/FbxSkeletalMeshImport.cpp
--- a/Engine/Source/Editor/UnrealEd/Private/Fbx/FbxSkeletalMeshImport.cpp
+++ b/Engine/Source/Editor/UnrealEd/Private/Fbx/FbxSkeletalMeshImport.cpp
@@ -195,14 +195,26 @@ void RemapSkeletalMeshVertexColorToImportData(const USkeletalMesh* SkeletalMesh,
 		VertPosOctree.AddElement(SkinVertex);
 	}
 
+	TMap<int32, FVector> WedgeIndexToNormal;
+	WedgeIndexToNormal.Reserve(WedgeNumber);
+	for (int32 FaceIndex = 0; FaceIndex < SkelMeshImportData->Faces.Num(); ++FaceIndex)
+	{
+		const SkeletalMeshImportData::FTriangle& Triangle = SkelMeshImportData->Faces[FaceIndex];
+		for (int32 Corner = 0; Corner < 3; ++Corner)
+		{
+			WedgeIndexToNormal.Add(Triangle.WedgeIndex[Corner], Triangle.TangentZ[Corner]);
+		}
+	}
+
 	// Iterate over each new vertex position, attempting to find the old vertex it is closest to, applying
 	// the color of the old vertex to the new position if possible.
-	const float DistanceOverNormalThreshold =KINDA_SMALL_NUMBER;
 	for (int32 WedgeIndex = 0; WedgeIndex < WedgeNumber; ++WedgeIndex)
 	{
 		SkeletalMeshImportData::FVertex& Wedge = SkelMeshImportData->Wedges[WedgeIndex];
 		const FVector& Position = SkelMeshImportData->Points[Wedge.VertexIndex];
 		const FVector2D UV = Wedge.UVs[0];
+		const FVector& Normal = WedgeIndexToNormal.FindChecked(WedgeIndex);
+
 		TArray<FSoftSkinVertex> PointsToConsider;
 		TSKCVertPosOctree::TConstIterator<> OctreeIter(VertPosOctree);
 		// Iterate through the octree attempting to find the vertices closest to the current new point
@@ -246,7 +258,7 @@ void RemapSkeletalMeshVertexColorToImportData(const USkeletalMesh* SkeletalMesh,
 		if (PointsToConsider.Num() > 0)
 		{
 			//Get the closest position
-			float MinDistance = MAX_FLT;
+			float MaxNormalDot = -MAX_FLT;
 			float MinUVDistance = MAX_FLT;
 			int32 MatchIndex = INDEX_NONE;
 			for (int32 ConsiderationIndex = 0; ConsiderationIndex < PointsToConsider.Num(); ++ConsiderationIndex)
@@ -258,6 +270,18 @@ void RemapSkeletalMeshVertexColorToImportData(const USkeletalMesh* SkeletalMesh,
 				{
 					MinUVDistance = FMath::Min(MinUVDistance, UVDistanceSqr);
 					MatchIndex = ConsiderationIndex;
+					MaxNormalDot = Normal | SkinVertex.TangentZ;
+				}
+				else if (FMath::IsNearlyEqual(UVDistanceSqr, MinUVDistance, KINDA_SMALL_NUMBER))
+				{
+					//This case is useful when we have hard edge that shared vertice, somtime not all the shared wedge have the same paint color
+					//Think about a cube where each face have different vertex color.
+					float NormalDot = Normal | SkinVertex.TangentZ;
+					if (NormalDot > MaxNormalDot)
+					{
+						MaxNormalDot = NormalDot;
+						MatchIndex = ConsiderationIndex;
+					}
 				}
 			}
 			if (PointsToConsider.IsValidIndex(MatchIndex))
