Reverting previous restriction in Parent Class menu
#rnx

diff --git a/Engine/Source/Editor/UnrealEd/Private/Kismet2/BlueprintEditorUtils.cpp b/Engine/Source/Editor/UnrealEd/Private/Kismet2/BlueprintEditorUtils.cpp
--- a/Engine/Source/Editor/UnrealEd/Private/Kismet2/BlueprintEditorUtils.cpp
+++ b/Engine/Source/Editor/UnrealEd/Private/Kismet2/BlueprintEditorUtils.cpp
@@ -7657,7 +7657,6 @@ TSharedRef<SWidget> FBlueprintEditorUtils::ConstructBlueprintParentClassPicker(
 		bIsAnimBlueprint |= Blueprint->IsA(UAnimBlueprint::StaticClass());
 		bIsLevelScriptActor |= Blueprint->ParentClass->IsChildOf( ALevelScriptActor::StaticClass() );
 		bIsComponentBlueprint |= Blueprint->ParentClass->IsChildOf( UActorComponent::StaticClass() );
-		bIsEditorOnlyBlueprint |= IsEditorOnlyObject(Blueprint);
 		bIsWidgetBlueprint = Blueprint->IsA(UBaseWidgetBlueprint::StaticClass());
 		if(Blueprint->GeneratedClass)
 		{
