Merging using //Ocean/Dev-EngineMerge-To-//UE4/Main
#rnx

diff --git a/Engine/Plugins/Runtime/AR/Apple/AppleARKit/Source/AppleARKit/Private/AppleARKitSystem.cpp b/Engine/Plugins/Runtime/AR/Apple/AppleARKit/Source/AppleARKit/Private/AppleARKitSystem.cpp
--- a/Engine/Plugins/Runtime/AR/Apple/AppleARKit/Source/AppleARKit/Private/AppleARKitSystem.cpp
+++ b/Engine/Plugins/Runtime/AR/Apple/AppleARKit/Source/AppleARKit/Private/AppleARKitSystem.cpp
@@ -206,6 +206,7 @@ FAppleARKitSystem::FAppleARKitSystem()
 , CameraDepth(nullptr)
 , LastTrackedGeometry_DebugId(0)
 , FaceARSupport(nullptr)
+, TimecodeProvider(nullptr)
 {
 	// See Initialize(), as we need access to SharedThis()
 }
diff --git a/Engine/Source/Editor/UnrealEd/Classes/Settings/LevelEditorPlaySettings.h b/Engine/Source/Editor/UnrealEd/Classes/Settings/LevelEditorPlaySettings.h
--- a/Engine/Source/Editor/UnrealEd/Classes/Settings/LevelEditorPlaySettings.h
+++ b/Engine/Source/Editor/UnrealEd/Classes/Settings/LevelEditorPlaySettings.h
@@ -86,6 +86,7 @@ enum EPlayNetMode
 	PIE_Standalone UMETA(DisplayName="Play Offline"),
 	PIE_ListenServer UMETA(DisplayName="Play As Listen Server"),
 	PIE_Client UMETA(DisplayName="Play As Client"),
+	PIE_StandaloneWithServer UMETA(DisplayName="Play Standalone With Server", ToolTip="Behaves exactly like 'Play Offline', but also launches the dedicated server, allowing you to override the map name"),
 };
 
 
@@ -335,6 +336,10 @@ private:
 	UPROPERTY(config, EditAnywhere, Category=MultiplayerOptions, meta=(ClampMin = 0))
 	int32 ClientWindowHeight;
 
+	/** Override the map launched by the dedicated server (currently only used when in PIE_StandaloneWithServer net mode) */
+	UPROPERTY(config, EditAnywhere, Category = MultiplayerOptions)
+	FString ServerMapNameOverride;
+
 	/** Additional options that will be passed to the server as URL parameters, in the format ?bIsLanMatch=1?listen - any additional command line switches should be passed in the Command Line Arguments field below. */
 	UPROPERTY(config, EditAnywhere, Category=MultiplayerOptions)
 	FString AdditionalServerGameOptions;
@@ -356,7 +361,7 @@ public:
 	bool GetRunUnderOneProcess( bool &OutRunUnderOneProcess ) const { OutRunUnderOneProcess = RunUnderOneProcess; return IsRunUnderOneProcessActive(); }
 	
 	void SetPlayNetDedicated( const bool InPlayNetDedicated ) { PlayNetDedicated = InPlayNetDedicated; }
-	bool IsPlayNetDedicatedActive() const { return (RunUnderOneProcess ? true : PlayNetMode == PIE_Client); }
+	bool IsPlayNetDedicatedActive() const { return (RunUnderOneProcess ? true : PlayNetMode == PIE_Client || PlayNetMode == PIE_StandaloneWithServer); }
 	bool GetPlayNetDedicated( bool &OutPlayNetDedicated ) const { OutPlayNetDedicated = PlayNetDedicated; return IsPlayNetDedicatedActive(); }
 
 	void SetPlayNumberOfClients( const int32 InPlayNumberOfClients ) { PlayNumberOfClients = InPlayNumberOfClients; }
@@ -375,6 +380,10 @@ public:
 	bool GetRouteGamepadToSecondWindow( bool &OutRouteGamepadToSecondWindow ) const { OutRouteGamepadToSecondWindow = RouteGamepadToSecondWindow; return IsRouteGamepadToSecondWindowActive(); }
 	EVisibility GetRouteGamepadToSecondWindowVisibility() const { return (RunUnderOneProcess ? EVisibility::Visible : EVisibility::Hidden); }
 
+	bool IsServerMapNameOverrideActive() const { return (PlayNetMode == PIE_StandaloneWithServer); }
+	bool GetServerMapNameOverride( FString& OutStandaloneServerMapName ) const { OutStandaloneServerMapName = ServerMapNameOverride; return IsServerMapNameOverrideActive(); }
+	EVisibility GetServerMapNameOverrideVisibility() const { return (PlayNetMode == PIE_StandaloneWithServer ? EVisibility::Visible : EVisibility::Hidden); }
+
 	bool IsAdditionalServerGameOptionsActive() const { return (PlayNetMode != PIE_Standalone) || RunUnderOneProcess; }
 	bool GetAdditionalServerGameOptions( FString &OutAdditionalServerGameOptions ) const { OutAdditionalServerGameOptions = AdditionalServerGameOptions; return IsAdditionalServerGameOptionsActive(); }
 
diff --git a/Engine/Source/Editor/UnrealEd/Private/EditorSelectUtils.cpp b/Engine/Source/Editor/UnrealEd/Private/EditorSelectUtils.cpp
--- a/Engine/Source/Editor/UnrealEd/Private/EditorSelectUtils.cpp
+++ b/Engine/Source/Editor/UnrealEd/Private/EditorSelectUtils.cpp
@@ -561,15 +561,6 @@ void UUnrealEdEngine::SelectActor(AActor* Actor, bool bInSelected, bool bNotify,
 	// Select the actor and update its internals.
 	if( !bSelectionHandled )
 	{
-		if(bInSelected)
-		{
-			// If trying to select an Actor spawned by a ChildACtorComponent, instead select Actor that spawned us
-			if (UChildActorComponent* ParentComponent = Actor->GetParentComponent())
-			{
-				Actor = ParentComponent->GetOwner();
-			}
-		}
-
 		if (UActorGroupingUtils::IsGroupingActive())
 		{
 			// if this actor is a group, do a group select/deselect
diff --git a/Engine/Source/Editor/UnrealEd/Private/PlayLevel.cpp b/Engine/Source/Editor/UnrealEd/Private/PlayLevel.cpp
--- a/Engine/Source/Editor/UnrealEd/Private/PlayLevel.cpp
+++ b/Engine/Source/Editor/UnrealEd/Private/PlayLevel.cpp
@@ -1214,9 +1214,15 @@ void UEditorEngine::StartQueuedPlayMapRequest()
 		FIntPoint WinPosition((int32)PreferredWorkArea.Right, (int32)PreferredWorkArea.Top);
 
 		// We'll need to spawn a server if we're playing outside the editor or the editor wants to run as a client
-		if (bPlayOnLocalPcSession || PlayNetMode == PIE_Client)
+		if (bPlayOnLocalPcSession || PlayNetMode == PIE_Client || PlayNetMode == PIE_StandaloneWithServer)
 		{			
-			PlayStandaloneLocalPc(TEXT(""), &WinPosition, NumClients, true);
+			FString MapNameOverride;
+			if (PlayInSettings->IsServerMapNameOverrideActive())
+			{
+				PlayInSettings->GetServerMapNameOverride(MapNameOverride);
+			}
+
+			PlayStandaloneLocalPc(*MapNameOverride, &WinPosition, NumClients, true);
 			
 			const bool CanPlayNetDedicated = [&PlayInSettings]{ bool PlayNetDedicated(false); return (PlayInSettings->GetPlayNetDedicated(PlayNetDedicated) && PlayNetDedicated); }();
 			if (!CanPlayNetDedicated)
@@ -1236,15 +1242,16 @@ void UEditorEngine::StartQueuedPlayMapRequest()
 		}
 
 		// Build the connection String
-		FString ConnectionAddr(TEXT("127.0.0.1"));
+		FString ConnectionAddr;
 
 		// Ignore the user's settings if the autoconnect option is inaccessible due to settings conflicts.
-		const bool WillAutoConnectToServer = [&PlayInSettings] { bool AutoConnectToServer(false); 
+		const bool WillAutoConnectToServer = [&PlayInSettings, PlayNetMode] { bool AutoConnectToServer(false);
 			return (PlayInSettings->GetAutoConnectToServerVisibility() == EVisibility::Visible) ? 
-				(PlayInSettings->GetAutoConnectToServer(AutoConnectToServer) && AutoConnectToServer) : true; }();
+				(PlayInSettings->GetAutoConnectToServer(AutoConnectToServer) && AutoConnectToServer) : (PlayNetMode != PIE_StandaloneWithServer); }();
 
 		if (WillAutoConnectToServer)
 		{
+			ConnectionAddr += TEXT("127.0.0.1");
 			uint16 ServerPort = 0;
 			if (PlayInSettings->GetServerPort(ServerPort))
 			{
@@ -2659,7 +2666,7 @@ bool UEditorEngine::CreatePIEWorldFromLogin(FWorldContext& PieWorldContext, EPla
 	GetMultipleInstancePositions(DataStruct.SettingsIndex, DataStruct.NextX, DataStruct.NextY);
 	
 	const bool CanPlayNetDedicated = [&PlayInSettings]{ bool PlayNetDedicated(false); return (PlayInSettings->GetPlayNetDedicated(PlayNetDedicated) && PlayNetDedicated); }();
-	const bool ActAsClient = PlayNetMode == EPlayNetMode::PIE_Client || PlayNetMode == EPlayNetMode::PIE_Standalone;
+	const bool ActAsClient = PlayNetMode == EPlayNetMode::PIE_Client || PlayNetMode == EPlayNetMode::PIE_StandaloneWithServer || PlayNetMode == EPlayNetMode::PIE_Standalone;
 	UGameInstance* const GameInstance = CreatePIEGameInstance(PieWorldContext.PIEInstance, false, DataStruct.bAnyBlueprintErrors, DataStruct.bStartInSpectatorMode, ActAsClient ? false : CanPlayNetDedicated, false, DataStruct.PIEStartTime);
 	
 	// Restore window settings
diff --git a/Engine/Source/Editor/UnrealEd/Private/Settings/LevelEditorPlaySettingsCustomization.h b/Engine/Source/Editor/UnrealEd/Private/Settings/LevelEditorPlaySettingsCustomization.h
--- a/Engine/Source/Editor/UnrealEd/Private/Settings/LevelEditorPlaySettingsCustomization.h
+++ b/Engine/Source/Editor/UnrealEd/Private/Settings/LevelEditorPlaySettingsCustomization.h
@@ -655,6 +655,10 @@ public:
 				.DisplayName(LOCTEXT("RunDedicatedServerLabel", "Run Dedicated Server"))
 				.IsEnabled(TAttribute<bool>::Create(TAttribute<bool>::FGetter::CreateSP(this, &FLevelEditorPlaySettingsCustomization::HandlePlayNetDedicatedPropertyIsEnabled)));
 
+			NetworkCategory.AddProperty("ServerMapNameOverride")
+				.DisplayName(LOCTEXT("ServerMapNameOverrideLabel", "Server Map Name Override"))
+				.Visibility(TAttribute<EVisibility>::Create(TAttribute<EVisibility>::FGetter::CreateSP(this, &FLevelEditorPlaySettingsCustomization::HandleServerMapNameOverrideVisibility)));
+
 			// client window size
 			TSharedRef<IPropertyHandle> WindowHeightHandle = LayoutBuilder.GetProperty("ClientWindowHeight");
 			TSharedRef<IPropertyHandle> WindowWidthHandle = LayoutBuilder.GetProperty("ClientWindowWidth");
@@ -893,6 +897,12 @@ private:
 		return GetDefault<ULevelEditorPlaySettings>()->GetPlayNetModeVisibility();
 	}
 
+	// Callback for getting the visibility of the StandaloneServerMapName property.
+	EVisibility HandleServerMapNameOverrideVisibility() const
+	{
+		return GetDefault<ULevelEditorPlaySettings>()->GetServerMapNameOverrideVisibility();
+	}
+
 	// Callback for getting the visibility of the AdditionalLaunchOptions property.
 	EVisibility HandleCmdLineVisibility( ) const
 	{
diff --git a/Engine/Source/Runtime/Core/Private/IOS/IOSErrorOutputDevice.cpp b/Engine/Source/Runtime/Core/Private/IOS/IOSErrorOutputDevice.cpp
--- a/Engine/Source/Runtime/Core/Private/IOS/IOSErrorOutputDevice.cpp
+++ b/Engine/Source/Runtime/Core/Private/IOS/IOSErrorOutputDevice.cpp
@@ -7,6 +7,7 @@
 #include "IOS/IOSErrorOutputDevice.h"
 #include "HAL/PlatformMisc.h"
 #include "HAL/PlatformAtomics.h"
+#include "Misc/OutputDeviceHelper.h"
 #include "Misc/OutputDeviceRedirector.h"
 #include "CoreGlobals.h"
 
@@ -17,10 +18,8 @@ FIOSErrorOutputDevice::FIOSErrorOutputDevice()
 
 void FIOSErrorOutputDevice::Serialize( const TCHAR* Msg, ELogVerbosity::Type Verbosity, const class FName& Category )
 {
-	// @todo iosmerge: this was a big mess copied from Mac that we don't want at all. can we use default like other consoles?
-	if( GIsGuarded )
+	FPlatformMisc::LowLevelOutputDebugString(*FOutputDeviceHelper::FormatLogLine(Verbosity, Category, Msg, GPrintLogTimes));	if( GIsGuarded )
 	{
-//		FOutputDevice::Serialize(Msg, Verbosity, Category);
 		UE_DEBUG_BREAK();
 	}
 	else
diff --git a/Engine/Source/Runtime/CoreUObject/Private/Serialization/BulkData.cpp b/Engine/Source/Runtime/CoreUObject/Private/Serialization/BulkData.cpp
--- a/Engine/Source/Runtime/CoreUObject/Private/Serialization/BulkData.cpp
+++ b/Engine/Source/Runtime/CoreUObject/Private/Serialization/BulkData.cpp
@@ -1619,7 +1619,8 @@ void FUntypedBulkData::LoadDataIntoMemory( void* Dest )
 		// is it possible for streaming mips to be loaded in non streaming ways.
 		if (CVarTextureStreamingEnabled->GetValueOnAnyThread() != 0)
 		{
-			UE_CLOG(GEventDrivenLoaderEnabled && (IsInGameThread() || IsInAsyncLoadingThread()), LogSerialization, Error, TEXT("Attempt to sync load bulk data with EDL enabled (LoadDataIntoMemory). This is not desireable. File %s"), *Filename);
+			bool bIsBulkFile = Filename.EndsWith(TEXT(".ubulk"));
+			UE_CLOG(GEventDrivenLoaderEnabled && bIsBulkFile && (IsInGameThread() || IsInAsyncLoadingThread()), LogSerialization, Error, TEXT("Attempt to sync load bulk data with EDL enabled (LoadDataIntoMemory). This is not desireable. File %s"), *Filename);
 		}
 #endif
 
diff --git a/Engine/Source/Runtime/Engine/Classes/Engine/SkeletalMesh.h b/Engine/Source/Runtime/Engine/Classes/Engine/SkeletalMesh.h
--- a/Engine/Source/Runtime/Engine/Classes/Engine/SkeletalMesh.h
+++ b/Engine/Source/Runtime/Engine/Classes/Engine/SkeletalMesh.h
@@ -954,14 +954,9 @@ public:
 	UFUNCTION(BlueprintCallable, Category = "Animation")
 	USkeletalMeshSocket* GetSocketByIndex(int32 Index) const;
 
-#if !WITH_EDITOR
-private:
-	/** Called internally to rebuild an invalid socket map */
+	/** Called to rebuild an out-of-date or invalid socket map */
 	void RebuildSocketMap();
 
-public:
-#endif
-
 	// @todo document
 	FMatrix GetRefPoseMatrix( int32 BoneIndex ) const;
 
diff --git a/Engine/Source/Runtime/Engine/Private/Curves/CurveLinearColor.cpp b/Engine/Source/Runtime/Engine/Private/Curves/CurveLinearColor.cpp
--- a/Engine/Source/Runtime/Engine/Private/Curves/CurveLinearColor.cpp
+++ b/Engine/Source/Runtime/Engine/Private/Curves/CurveLinearColor.cpp
@@ -314,13 +314,20 @@ void UCurveLinearColor::PostLoad()
 {
 	Super::PostLoad();
 
-#if WITH_EDITOR
-	if (GetLinkerCustomVersion(FReleaseObjectVersion::GUID) < FReleaseObjectVersion::UnclampRGBColorCurves)
+	// BEGIN OCEAN HACK
+	// This warning will keep coming up each time a new curve asset is saved in the release branch and gets robo-merged back down to main because the engine version is lower in the release branch
+	// Remove this hack once the release branch is brought up to the latest engine version
+	if (false)
 	{
-		FMessageLog("LoadErrors").Warning(FText::Format(NSLOCTEXT("CurveEditor","CurveDataUpdate", "Linear color curves now accurately handle RGB values > 1. If you were relying on HSV clamping, please update {0}"),
-			FText::FromString(GetName())));
-	}
+#if WITH_EDITOR
+		if (GetLinkerCustomVersion(FReleaseObjectVersion::GUID) < FReleaseObjectVersion::UnclampRGBColorCurves)
+		{
+			FMessageLog("LoadErrors").Warning(FText::Format(NSLOCTEXT("CurveEditor", "CurveDataUpdate", "Linear color curves now accurately handle RGB values > 1. If you were relying on HSV clamping, please update {0}"),
+				FText::FromString(GetName())));
+		}
 #endif
+	}
+	// END OCEAN HACK
 }
 
 void UCurveLinearColor::Serialize(FArchive& Ar)
diff --git a/Engine/Source/Runtime/Engine/Private/SkeletalMesh.cpp b/Engine/Source/Runtime/Engine/Private/SkeletalMesh.cpp
--- a/Engine/Source/Runtime/Engine/Private/SkeletalMesh.cpp
+++ b/Engine/Source/Runtime/Engine/Private/SkeletalMesh.cpp
@@ -2178,10 +2178,11 @@ USkeletalMeshSocket* USkeletalMesh::GetSocketByIndex(int32 Index) const
 	return nullptr;
 }
 
-#if !WITH_EDITOR
 
 void USkeletalMesh::RebuildSocketMap()
 {
+#if !WITH_EDITOR
+
 	check(IsInGameThread());
 
 	SocketMap.Reset();
@@ -2206,9 +2207,10 @@ void USkeletalMesh::RebuildSocketMap()
 			}
 		}
 	}
+
+#endif // !WITH_EDITOR
 }
 
-#endif
 
 /**
  * This will return detail info about this specific object. (e.g. AudioComponent will return the name of the cue,
diff --git a/Engine/Source/Runtime/Engine/Private/SkeletalMeshMerge.cpp b/Engine/Source/Runtime/Engine/Private/SkeletalMeshMerge.cpp
--- a/Engine/Source/Runtime/Engine/Private/SkeletalMeshMerge.cpp
+++ b/Engine/Source/Runtime/Engine/Private/SkeletalMeshMerge.cpp
@@ -1019,6 +1019,8 @@ void FSkeletalMeshMerge::BuildSockets(const TArray<USkeletalMesh*>& SourceMeshLi
 			AddSockets(NewSkeletonSocketList, true);
 		}
 	}
+
+	MergeMesh->RebuildSocketMap();
 }
 
 void FSkeletalMeshMerge::OverrideSocket(const USkeletalMeshSocket* SourceSocket)
@@ -1095,4 +1097,6 @@ void FSkeletalMeshMerge::OverrideMergedSockets(const TArray<FRefPoseOverride>& P
 			}
 		}
 	}
+
+	MergeMesh->RebuildSocketMap();
 }
diff --git a/Engine/Source/Runtime/Slate/Private/Widgets/Input/SMenuAnchor.cpp b/Engine/Source/Runtime/Slate/Private/Widgets/Input/SMenuAnchor.cpp
--- a/Engine/Source/Runtime/Slate/Private/Widgets/Input/SMenuAnchor.cpp
+++ b/Engine/Source/Runtime/Slate/Private/Widgets/Input/SMenuAnchor.cpp
@@ -93,6 +93,7 @@ void SMenuAnchor::Construct( const FArguments& InArgs )
 	OnGetMenuContent                       = InArgs._OnGetMenuContent;
 	OnMenuOpenChanged                      = InArgs._OnMenuOpenChanged;
 	Placement                              = InArgs._Placement;
+	bFitInWindow						   = InArgs._FitInWindow;
 	Method                                 = InArgs._Method;
 	bShouldDeferPaintingAfterWindowContent = InArgs._ShouldDeferPaintingAfterWindowContent;
 	bUseApplicationMenuStack               = InArgs._UseApplicationMenuStack;
@@ -153,11 +154,21 @@ void SMenuAnchor::Tick( const FGeometry& AllottedGeometry, const double InCurren
 		// geometry to use. Tick() is always in DesktopSpace, so cache the solution here and just use
 		// it in OnArrangeChildren().
 		const FPopupPlacement LocalPlacement(AllottedGeometry, Children[1].GetWidget()->GetDesiredSize(), Placement.Get());
-		const FSlateRect WindowRectLocalSpace = TransformRect(Inverse(AllottedGeometry.GetAccumulatedLayoutTransform()), PopupWindow->GetClientRectInScreen());
-		const FVector2D FittedPlacement = ComputePopupFitInRect(
-			LocalPlacement.AnchorLocalSpace,
-			FSlateRect(LocalPlacement.LocalPopupOffset, LocalPlacement.LocalPopupOffset + LocalPlacement.LocalPopupSize),
-			LocalPlacement.Orientation, WindowRectLocalSpace);
+
+		FVector2D FittedPlacement;
+
+		if (bFitInWindow)
+		{
+			const FSlateRect WindowRectLocalSpace = TransformRect(Inverse(AllottedGeometry.GetAccumulatedLayoutTransform()), PopupWindow->GetClientRectInScreen());
+			FittedPlacement = ComputePopupFitInRect(
+				LocalPlacement.AnchorLocalSpace,
+				FSlateRect(LocalPlacement.LocalPopupOffset, LocalPlacement.LocalPopupOffset + LocalPlacement.LocalPopupSize),
+				LocalPlacement.Orientation, WindowRectLocalSpace);
+		}
+		else
+		{
+			FittedPlacement = LocalPlacement.LocalPopupOffset;
+		}
 
 		LocalPopupPosition = FittedPlacement;
 		ScreenPopupPosition = AllottedGeometry.GetAccumulatedLayoutTransform().TransformPoint(LocalPopupPosition);
@@ -568,6 +579,24 @@ FVector2D SMenuAnchor::GetMenuPosition() const
 	return Pos;
 }
 
+void SMenuAnchor::SetMenuPlacement(TAttribute<EMenuPlacement> InMenuPlacement)
+{
+	if (!Placement.IsSet() || !Placement.IdenticalTo(InMenuPlacement))
+	{
+		Placement = InMenuPlacement;
+		Invalidate(EInvalidateWidget::Layout);
+	}
+}
+
+void SMenuAnchor::SetFitInWindow(bool bFit)
+{
+	if (bFitInWindow != bFit)
+	{
+		bFitInWindow = bFit;
+		Invalidate(EInvalidateWidget::Layout);
+	}
+}
+
 bool SMenuAnchor::HasOpenSubMenus() const
 {
 	bool Result = false;
@@ -575,6 +604,7 @@ bool SMenuAnchor::HasOpenSubMenus() const
 	{
 		Result = FSlateApplication::Get().HasOpenSubMenus(PopupMenuPtr.Pin());
 	}
+
 	return Result;
 }
 
diff --git a/Engine/Source/Runtime/Slate/Public/Widgets/Input/SMenuAnchor.h b/Engine/Source/Runtime/Slate/Public/Widgets/Input/SMenuAnchor.h
--- a/Engine/Source/Runtime/Slate/Public/Widgets/Input/SMenuAnchor.h
+++ b/Engine/Source/Runtime/Slate/Public/Widgets/Input/SMenuAnchor.h
@@ -37,6 +37,7 @@ public:
 		, _MenuContent( SNew(STextBlock) .Text( NSLOCTEXT("SMenuAnchor", "NoMenuContent", "No Menu Content Assigned; use .MenuContent") ) )
 		, _OnGetMenuContent()
 		, _Placement( MenuPlacement_BelowAnchor )
+		, _FitInWindow( true )
 		, _Method()
 		, _ShouldDeferPaintingAfterWindowContent(true)
 		, _UseApplicationMenuStack(true)
@@ -56,6 +57,8 @@ public:
 
 		SLATE_ATTRIBUTE( EMenuPlacement, Placement )
 
+		SLATE_ARGUMENT(bool, FitInWindow)
+
 		SLATE_ARGUMENT(TOptional<EPopupMethod>, Method)
 
 		SLATE_ARGUMENT(bool, ShouldDeferPaintingAfterWindowContent)
@@ -105,6 +108,10 @@ public:
 	/** @return The current menu position */
 	FVector2D GetMenuPosition() const;
 
+	void SetMenuPlacement(TAttribute<EMenuPlacement> InMenuPlacement);
+
+	void SetFitInWindow(bool bFit);
+
 	/** @return Whether this menu has open submenus */
 	bool HasOpenSubMenus() const;
 
@@ -185,6 +192,9 @@ protected:
 	/** How should the popup be placed relative to the anchor. */
 	TAttribute<EMenuPlacement> Placement;
 
+	/** Should the menu anchor fit inside the window? */
+	bool bFitInWindow;
+
 	/** Was the menu just dismissed this tick? */
 	bool bDismissedThisTick;
 
diff --git a/Engine/Source/Runtime/UMG/Private/Components/MenuAnchor.cpp b/Engine/Source/Runtime/UMG/Private/Components/MenuAnchor.cpp
--- a/Engine/Source/Runtime/UMG/Private/Components/MenuAnchor.cpp
+++ b/Engine/Source/Runtime/UMG/Private/Components/MenuAnchor.cpp
@@ -17,6 +17,7 @@ UMenuAnchor::UMenuAnchor(const FObjectInitializer& ObjectInitializer)
 	, UseApplicationMenuStack(true)
 {
 	Placement = MenuPlacement_ComboBox;
+	bFitInWindow = true;
 }
 
 void UMenuAnchor::ReleaseSlateResources(bool bReleaseChildren)
@@ -30,6 +31,7 @@ TSharedRef<SWidget> UMenuAnchor::RebuildWidget()
 {
 	MyMenuAnchor = SNew(SMenuAnchor)
 		.Placement(Placement)
+		.FitInWindow(bFitInWindow)
 		.OnGetMenuContent(BIND_UOBJECT_DELEGATE(FOnGetContent, HandleGetMenuContent))
 		.OnMenuOpenChanged(BIND_UOBJECT_DELEGATE(FOnIsOpenChanged, HandleMenuOpenChanged))
 		.ShouldDeferPaintingAfterWindowContent(ShouldDeferPaintingAfterWindowContent)
@@ -129,6 +131,24 @@ bool UMenuAnchor::IsOpen() const
 	return false;
 }
 
+void UMenuAnchor::SetPlacement(TEnumAsByte<EMenuPlacement> InPlacement)
+{
+	Placement = InPlacement;
+	if (MyMenuAnchor.IsValid())
+	{
+		return MyMenuAnchor->SetMenuPlacement(Placement);
+	}
+}
+
+void UMenuAnchor::FitInWindow(bool bFit)
+{
+	bFitInWindow = bFit;
+	if (MyMenuAnchor.IsValid())
+	{
+		return MyMenuAnchor->SetFitInWindow(bFitInWindow);
+	}
+}
+
 bool UMenuAnchor::ShouldOpenDueToClick() const
 {
 	if ( MyMenuAnchor.IsValid() )
diff --git a/Engine/Source/Runtime/UMG/Public/Components/MenuAnchor.h b/Engine/Source/Runtime/UMG/Public/Components/MenuAnchor.h
--- a/Engine/Source/Runtime/UMG/Public/Components/MenuAnchor.h
+++ b/Engine/Source/Runtime/UMG/Public/Components/MenuAnchor.h
@@ -42,6 +42,10 @@ public:
 	UPROPERTY(EditAnywhere, BlueprintReadOnly, Category="Menu Anchor")
 	TEnumAsByte<EMenuPlacement> Placement;
 
+	/** Should the menu anchor attempt to fit the menu inside the window. */
+	UPROPERTY(EditAnywhere, BlueprintReadOnly, Category = "Menu Anchor")
+	bool bFitInWindow;
+
 	UPROPERTY(EditAnywhere, BlueprintReadOnly, AdvancedDisplay, Category = "Menu Anchor")
 	bool ShouldDeferPaintingAfterWindowContent;
 
@@ -56,7 +60,12 @@ public:
 
 public:
 	//TODO UMG Add Set MenuClass
-	//TODO UMG Add Set Placement
+
+	UFUNCTION(BlueprintCallable, Category = "Menu Anchor")
+	void SetPlacement(TEnumAsByte<EMenuPlacement> InPlacement);
+
+	UFUNCTION(BlueprintCallable, Category = "Menu Anchor")
+	void FitInWindow(bool bFit);
 
 public:
 
