Fix crash when deleting open levels by preventing users from deleting open levels

diff --git a/Engine/Source/Editor/UnrealEd/Private/ObjectTools.cpp b/Engine/Source/Editor/UnrealEd/Private/ObjectTools.cpp
--- a/Engine/Source/Editor/UnrealEd/Private/ObjectTools.cpp
+++ b/Engine/Source/Editor/UnrealEd/Private/ObjectTools.cpp
@@ -2017,6 +2017,39 @@ namespace ObjectTools
 		}
 	}
 
+	bool ContainsWorldInUse(const TArray< UObject* >& ObjectsToDelete)
+	{
+		TArray<const UObject*> ActiveWorlds;
+		for (const FWorldContext& WorldContext : GEditor->GetWorldContexts())
+		{
+			if (UWorld* World = WorldContext.World())
+			{
+				ActiveWorlds.AddUnique(World);
+
+				for (const ULevelStreaming* StreamingLevel : World->GetStreamingLevels())
+				{
+					if (StreamingLevel && StreamingLevel->GetLoadedLevel() && StreamingLevel->GetLoadedLevel()->GetOuter())
+					{
+						ActiveWorlds.AddUnique(StreamingLevel->GetLoadedLevel()->GetOuter());
+					}
+				}
+			}
+		}
+
+		for (const UObject* ObjectToDelete : ObjectsToDelete)
+		{
+			if (const UWorld* World = Cast<UWorld>(ObjectToDelete))
+			{
+				if (ActiveWorlds.Contains(World))
+				{
+					return true;
+				}
+			}
+		}
+
+		return false;
+	}
+
 	int32 DeleteObjects( const TArray< UObject* >& InObjectsToDelete, bool bShowConfirmation )
 	{
 		// Allows deleting of sounds after they have been previewed
@@ -2062,6 +2095,17 @@ namespace ObjectTools
 			return 0;
 		}
 
+		if (ContainsWorldInUse(ObjectsToDelete))
+		{
+			OpenMsgDlgInt(
+				EAppMsgType::Ok,
+				NSLOCTEXT("UnrealEd", "DeleteFailedWorldInUse", "Unable to delete level while it is open"),
+				NSLOCTEXT("UnrealEd", "DeleteFailedWorldInUseTitle", "Unable to delete level")
+			);
+
+			return 0;
+		}
+
 		// let systems clean up any unnecessary references that they may have 
 		// (so that they're not flagged in the dialog)
 		FEditorDelegates::OnAssetsPreDelete.Broadcast(ObjectsToDelete);
