Added csv compression for PSO processing commandlet
Output will be compressed if user specifies extension as stablepc.csv.compressed instead of stablepc.csv
Compressed csv works will all PSO commandlet commands

diff --git a/Engine/Source/Editor/UnrealEd/Private/Commandlets/ShaderPipelineCacheToolsCommandlet.cpp b/Engine/Source/Editor/UnrealEd/Private/Commandlets/ShaderPipelineCacheToolsCommandlet.cpp
--- a/Engine/Source/Editor/UnrealEd/Private/Commandlets/ShaderPipelineCacheToolsCommandlet.cpp
+++ b/Engine/Source/Editor/UnrealEd/Private/Commandlets/ShaderPipelineCacheToolsCommandlet.cpp
@@ -7,9 +7,16 @@
 #include "ShaderCodeLibrary.h"
 #include "Misc/FileHelper.h"
 #include "ShaderPipelineCache.h"
+#include "Serialization/MemoryReader.h"
 
 DEFINE_LOG_CATEGORY_STATIC(LogShaderPipelineCacheTools, Log, All);
 
+const TCHAR* STABLE_CSV_EXT = TEXT("stablepc.csv");
+const TCHAR* STABLE_CSV_COMPRESSED_EXT = TEXT("stablepc.csv.compressed");
+const TCHAR* STABLE_COMPRESSED_EXT = TEXT(".compressed");
+const int32  STABLE_COMPRESSED_EXT_LEN = 11; // len of ".compressed";
+const int32  STABLE_COMPRESSED_VER = 1;
+
 
 void ExpandWildcards(TArray<FString>& Parts)
 {
@@ -18,8 +25,19 @@ void ExpandWildcards(TArray<FString>& Parts)
 	{
 		if (OldPart.Contains(TEXT("*")) || OldPart.Contains(TEXT("?")))
 		{
+			FString CleanPath = FPaths::GetPath(OldPart);
+			FString CleanFilename = FPaths::GetCleanFilename(OldPart);
+			
 			TArray<FString> ExpandedFiles;
-			IFileManager::Get().FindFilesRecursive(ExpandedFiles, *FPaths::GetPath(OldPart), *FPaths::GetCleanFilename(OldPart), true, false);
+			IFileManager::Get().FindFilesRecursive(ExpandedFiles, *CleanPath, *CleanFilename, true, false);
+			
+			if (CleanFilename.EndsWith(STABLE_CSV_EXT))
+			{
+				// look for stablepc.csv.compressed as well
+				CleanFilename.Append(STABLE_COMPRESSED_EXT);
+				IFileManager::Get().FindFilesRecursive(ExpandedFiles, *CleanPath, *CleanFilename, true, false, false);
+			}
+			
 			UE_CLOG(!ExpandedFiles.Num(), LogShaderPipelineCacheTools, Warning, TEXT("Expanding %s....did not match anything."), *OldPart);
 			UE_CLOG(ExpandedFiles.Num(), LogShaderPipelineCacheTools, Log, TEXT("Expanding matched %4d files: %s"), ExpandedFiles.Num(), *OldPart);
 			for (const FString& Item : ExpandedFiles)
@@ -62,6 +80,130 @@ void LoadStableSCL(TMultiMap<FStableShaderKeyAndValue, FSHAHash>& StableMap, con
 	UE_LOG(LogShaderPipelineCacheTools, Display, TEXT("Loaded %d shader info lines"), SourceFileContents.Num() - 1);
 }
 
+static bool LoadAndDecompressStableCSV(const FString& Filename, TArray<uint8>& UncompressedData)
+{
+	bool bResult = false;
+	FArchive* Ar = IFileManager::Get().CreateFileReader(*Filename);
+	if (Ar)
+	{
+		if (Ar->TotalSize() > 8)
+		{
+			int32 CompressedVersion = 0;
+			int32 UncompressedSize = 0;
+			int32 CompressedSize = 0;
+			
+			Ar->Serialize(&CompressedVersion, sizeof(int32));
+			Ar->Serialize(&UncompressedSize, sizeof(int32));
+			Ar->Serialize(&CompressedSize, sizeof(int32));
+			
+			TArray<uint8> CompressedData;
+			CompressedData.SetNumUninitialized(CompressedSize);
+			Ar->Serialize(CompressedData.GetData(), CompressedSize);
+
+			UncompressedData.SetNumUninitialized(UncompressedSize);
+			bResult = FCompression::UncompressMemory(NAME_Zlib, UncompressedData.GetData(), UncompressedSize, CompressedData.GetData(), CompressedSize);
+			if (!bResult)
+			{
+				UE_LOG(LogShaderPipelineCacheTools, Display, TEXT("Failed to decompress file %s"), *Filename);
+			}
+		}
+		else
+		{
+			UE_LOG(LogShaderPipelineCacheTools, Display, TEXT("Corrupted file %s"), *Filename);
+		}
+	
+		delete Ar;
+	}
+	else
+	{
+		UE_LOG(LogShaderPipelineCacheTools, Display, TEXT("Failed to open file %s"), *Filename);
+	}
+
+	return bResult;
+}
+
+static bool LoadStableCSV(const FString& Filename, TArray<FString>& OutputLines)
+{
+	bool bResult = false;
+	if (Filename.EndsWith(STABLE_CSV_COMPRESSED_EXT))
+	{
+		TArray<uint8> DecompressedData;
+		if (LoadAndDecompressStableCSV(Filename, DecompressedData))
+		{
+			FMemoryReader MemArchive(DecompressedData);
+			FString LineCSV;
+			while (!MemArchive.AtEnd())
+			{
+				MemArchive << LineCSV;
+				OutputLines.Add(LineCSV);
+			}
+			bResult = true;
+		}
+	}
+	else
+	{
+		bResult = FFileHelper::LoadFileToStringArray(OutputLines, *Filename);
+	}
+	
+	return bResult;
+}
+
+static int64 SaveStableCSV(const FString& Filename, const TArray<uint8>& UncompressedData)
+{
+	if (Filename.EndsWith(STABLE_CSV_COMPRESSED_EXT))
+	{
+		int32 UncompressedSize = UncompressedData.Num();
+		UE_LOG(LogShaderPipelineCacheTools, Display, TEXT("Compressing output, size = %.1fKB"), UncompressedSize/1024.f);
+		int32 CompressedSize = FCompression::CompressMemoryBound(NAME_Zlib, UncompressedSize);
+		TArray<uint8> CompressedData;
+		CompressedData.SetNumZeroed(CompressedSize);
+
+		if (FCompression::CompressMemory(NAME_Zlib, CompressedData.GetData(), CompressedSize, UncompressedData.GetData(), UncompressedSize))
+		{
+			FArchive* Ar = IFileManager::Get().CreateFileWriter(*Filename);
+			if (!Ar)
+			{
+				UE_LOG(LogShaderPipelineCacheTools, Fatal, TEXT("Failed to open %s"), *Filename);
+				return -1;
+			}
+			
+			int32 CompressedVersion = STABLE_COMPRESSED_VER;
+			
+			Ar->Serialize(&CompressedVersion, sizeof(int32));
+			Ar->Serialize(&UncompressedSize, sizeof(int32));
+			Ar->Serialize(&CompressedSize, sizeof(int32));
+			Ar->Serialize(CompressedData.GetData(), CompressedSize);
+			delete Ar;
+		}
+		else
+		{
+			UE_LOG(LogShaderPipelineCacheTools, Fatal, TEXT("Failed to compress (%.1f KB)"), UncompressedSize/1024.f);
+			return -1;
+		}
+	}
+	else
+	{
+		FMemoryReader MemArchive(UncompressedData);
+		FString CombinedCSV;
+		FString LineCSV;
+		while (!MemArchive.AtEnd())
+		{
+			MemArchive << LineCSV;
+			CombinedCSV.Append(LineCSV);
+			CombinedCSV.Append(LINE_TERMINATOR);
+		}
+
+		FFileHelper::SaveStringToFile(CombinedCSV, *Filename);
+	}
+	
+	int64 Size = IFileManager::Get().FileSize(*Filename);
+	if (Size < 1)
+	{
+		UE_LOG(LogShaderPipelineCacheTools, Fatal, TEXT("Failed to write %s"), *Filename);
+	}
+
+	return Size;
+}
 
 static void PrintShaders(const TMap<FSHAHash, TArray<FString>>& InverseMap, const FSHAHash& Shader)
 {
@@ -318,7 +460,7 @@ void GeneratePermuations(TArray<FPermuation>& Permutations, FPermuation& Working
 int32 ExpandPSOSC(const TArray<FString>& Tokens)
 {
 	TMultiMap<FStableShaderKeyAndValue, FSHAHash> StableMap;
-	check(Tokens.Last().EndsWith(TEXT(".stablepc.csv")));
+	check(Tokens.Last().EndsWith(STABLE_CSV_EXT) || Tokens.Last().EndsWith(STABLE_CSV_COMPRESSED_EXT));
 	for (int32 Index = 0; Index < Tokens.Num() - 1; Index++)
 	{
 		if (Tokens[Index].EndsWith(TEXT(".scl.csv")))
@@ -623,7 +765,9 @@ int32 ExpandPSOSC(const TArray<FString>& Tokens)
 		return 1;
 	}
 
-	TArray<FString> OutputLines;
+	int32 NumLines = 0;
+	TArray<uint8> UncomressedOutputLines;
+	FMemoryWriter OutputLinesAr(UncomressedOutputLines);
 	TSet<FString> DeDup;
 
 	{
@@ -633,7 +777,9 @@ int32 ExpandPSOSC(const TArray<FString>& Tokens)
 		{
 			PSOLine += FString::Printf(TEXT(",\"shaderslot%d: %s\""), SlotIndex, *FStableShaderKeyAndValue::HeaderLine());
 		}
-		OutputLines.Add(PSOLine);
+
+		OutputLinesAr << PSOLine;
+		NumLines++;
 	}
 
 	for (const FPermsPerPSO& Item : StableResults)
@@ -705,28 +851,54 @@ int32 ExpandPSOSC(const TArray<FString>& Tokens)
 
 			if (!DeDup.Contains(PSOLine))
 			{
-				DeDup.Add(PSOLine); 
-				OutputLines.Add(PSOLine);
+				DeDup.Add(PSOLine);
+				OutputLinesAr << PSOLine;
+				NumLines++;
 			}
 		}
 	}
 
-	if (IFileManager::Get().FileExists(*Tokens.Last()))
+	const FString& OutputFilename = Tokens.Last();
+	const bool bCompressed = OutputFilename.EndsWith(STABLE_CSV_COMPRESSED_EXT);
+	
+	FString CompressedFilename;
+	FString UncompressedFilename;
+	if (bCompressed)
 	{
-		IFileManager::Get().Delete(*Tokens.Last(), false, true);
+		CompressedFilename = OutputFilename;
+		UncompressedFilename = CompressedFilename.LeftChop(STABLE_COMPRESSED_EXT_LEN); // remove the ".compressed"
 	}
-	if (IFileManager::Get().FileExists(*Tokens.Last()))
+	else
 	{
-		UE_LOG(LogShaderPipelineCacheTools, Fatal, TEXT("Could not delete %s"), *Tokens.Last());
+		UncompressedFilename = OutputFilename;
+		CompressedFilename = UncompressedFilename + STABLE_COMPRESSED_EXT;  // add the ".compressed"
 	}
-	FFileHelper::SaveStringArrayToFile(OutputLines, *Tokens.Last());
-	int64 Size = IFileManager::Get().FileSize(*Tokens.Last());
-	if (Size < 1)
+	
+	// delete both compressed and uncompressed files
+	if (IFileManager::Get().FileExists(*UncompressedFilename))
 	{
-		UE_LOG(LogShaderPipelineCacheTools, Fatal, TEXT("Failed to write %s"), *Tokens.Last());
+		IFileManager::Get().Delete(*UncompressedFilename, false, true);
+		if (IFileManager::Get().FileExists(*UncompressedFilename))
+		{
+			UE_LOG(LogShaderPipelineCacheTools, Fatal, TEXT("Could not delete %s"), *UncompressedFilename);
+		}
+	}
+	if (IFileManager::Get().FileExists(*CompressedFilename))
+	{
+		IFileManager::Get().Delete(*CompressedFilename, false, true);
+		if (IFileManager::Get().FileExists(*CompressedFilename))
+		{
+			UE_LOG(LogShaderPipelineCacheTools, Fatal, TEXT("Could not delete %s"), *CompressedFilename);
+		}
 	}
-	UE_LOG(LogShaderPipelineCacheTools, Display, TEXT("Wrote stable PSOs, %d lines (%lldKB) to %s"), OutputLines.Num(), (Size + 1023) / 1024, *Tokens.Last());
 
+	int64 FileSize = SaveStableCSV(OutputFilename, UncomressedOutputLines);
+	if (FileSize < 1)
+	{
+		return 1;
+	}
+	
+	UE_LOG(LogShaderPipelineCacheTools, Display, TEXT("Wrote stable PSOs, %d lines (%.1f KB) to %s"), NumLines, FileSize / 1024.f, *OutputFilename);
 	return 0;
 }
 
@@ -750,8 +922,9 @@ void ParseQuoteComma(const FString& InLine, TArray<FString>& OutParts)
 	}
 }
 
+typedef TFunction<bool(const FString&)> FilenameFilterFN;
 
-void BuildDateSortedListOfFiles(const TArray<FString>& TokenList, FString const & MatchesFileExtension, TArray<FString>& Result)
+void BuildDateSortedListOfFiles(const TArray<FString>& TokenList, FilenameFilterFN FilterFn, TArray<FString>& Result)
 {
 	struct FDateSortableFileRef
 	{
@@ -762,7 +935,7 @@ void BuildDateSortedListOfFiles(const TArray<FString>& TokenList, FString const
 	TArray<FDateSortableFileRef> DateFileList;
 	for (int32 TokenIndex = 0; TokenIndex < TokenList.Num() - 1; TokenIndex++)
 	{
-		if (TokenList[TokenIndex].EndsWith(MatchesFileExtension))
+		if (FilterFn(TokenList[TokenIndex]))
 		{
 			FDateSortableFileRef DateSortEntry;
 			DateSortEntry.SortTime = FDateTime::Now();
@@ -831,8 +1004,13 @@ int32 BuildPSOSC(const TArray<FString>& Tokens)
 	// Get the stable PC files in date order - least to most important(!?)
 	UE_LOG(LogShaderPipelineCacheTools, Display, TEXT("Sorting input stablepc.csv files into chronological order for merge processing...."));
 	
+	FilenameFilterFN ExtenstionFilterFn = [](const FString& Filename)
+	{
+		return Filename.EndsWith(STABLE_CSV_EXT) || Filename.EndsWith(STABLE_CSV_COMPRESSED_EXT);
+	};
+
 	TArray<FString> StablePiplineCacheFiles;
-	BuildDateSortedListOfFiles(Tokens, TEXT(".stablepc.csv"), StablePiplineCacheFiles);
+	BuildDateSortedListOfFiles(Tokens, ExtenstionFilterFn, StablePiplineCacheFiles);
 
 	uint32 MergeCount = 0;
 
@@ -843,7 +1021,7 @@ int32 BuildPSOSC(const TArray<FString>& Tokens)
 		UE_LOG(LogShaderPipelineCacheTools, Display, TEXT("Loading %s...."), *FileName);
 		TArray<FString> SourceFileContents;
 
-		if (!FFileHelper::LoadFileToStringArray(SourceFileContents, *FileName) || SourceFileContents.Num() < 2)
+		if (!LoadStableCSV(FileName, SourceFileContents) || SourceFileContents.Num() < 2)
 		{
 			UE_LOG(LogShaderPipelineCacheTools, Fatal, TEXT("Could not load %s"), *FileName);
 			return 1;
@@ -1161,18 +1339,19 @@ int32 DiffStable(const TArray<FString>& Tokens)
 	TArray<TSet<FString>> Sets;
 	for (int32 TokenIndex = 0; TokenIndex < Tokens.Num(); TokenIndex++)
 	{
-		if (!Tokens[TokenIndex].EndsWith(TEXT(".stablepc.csv")))
+		const FString& Filename = Tokens[TokenIndex];
+		bool bCompressed = Filename.EndsWith(STABLE_CSV_COMPRESSED_EXT);
+		if (!bCompressed && !Filename.EndsWith(STABLE_CSV_EXT))
 		{
 			check(0);
 			continue;
 		}
-
-		UE_LOG(LogShaderPipelineCacheTools, Display, TEXT("Loading %s...."), *Tokens[TokenIndex]);
+			   
+		UE_LOG(LogShaderPipelineCacheTools, Display, TEXT("Loading %s...."), *Filename);
 		TArray<FString> SourceFileContents;
-
-		if (!FFileHelper::LoadFileToStringArray(SourceFileContents, *Tokens[TokenIndex]) || SourceFileContents.Num() < 2)
+		if (LoadStableCSV(Filename, SourceFileContents) || SourceFileContents.Num() < 2)
 		{
-			UE_LOG(LogShaderPipelineCacheTools, Fatal, TEXT("Could not load %s"), *Tokens[TokenIndex]);
+			UE_LOG(LogShaderPipelineCacheTools, Fatal, TEXT("Could not load %s"), *Filename);
 			return 1;
 		}
 
@@ -1211,6 +1390,38 @@ int32 DiffStable(const TArray<FString>& Tokens)
 	return 0;
 }
 
+int32 DecompressCSV(const TArray<FString>& Tokens)
+{
+	TArray<uint8> DecompressedData;
+	for (int32 TokenIndex = 0; TokenIndex < Tokens.Num(); TokenIndex++)
+	{
+		const FString& CompressedFilename = Tokens[TokenIndex];
+		if (!CompressedFilename.EndsWith(STABLE_CSV_COMPRESSED_EXT))
+		{
+			continue;
+		}
+		
+		FString CombinedCSV;
+		DecompressedData.Reset();
+		if (LoadAndDecompressStableCSV(CompressedFilename, DecompressedData))
+		{
+			FMemoryReader MemArchive(DecompressedData);
+			FString LineCSV;
+			while (!MemArchive.AtEnd())
+			{
+				MemArchive << LineCSV;
+				CombinedCSV.Append(LineCSV);
+				CombinedCSV.Append(LINE_TERMINATOR);
+			}
+
+			FString FilenameCSV = CompressedFilename.LeftChop(STABLE_COMPRESSED_EXT_LEN);
+			FFileHelper::SaveStringToFile(CombinedCSV, *FilenameCSV);
+		}
+	}
+
+	return 0;
+}
+
 UShaderPipelineCacheToolsCommandlet::UShaderPipelineCacheToolsCommandlet(const FObjectInitializer& ObjectInitializer)
 	: Super(ObjectInitializer)
 {
@@ -1261,11 +1472,18 @@ int32 UShaderPipelineCacheToolsCommandlet::StaticMain(const FString& Params)
 				}
 			}
 		}
+		else if (Tokens[0] == TEXT("Decompress") && Tokens.Num() >= 2)
+		{
+			Tokens.RemoveAt(0);
+			return DecompressCSV(Tokens);
+		}
 	}
 	
 	UE_LOG(LogShaderPipelineCacheTools, Warning, TEXT("Usage: Dump ShaderCache1.upipelinecache SCLInfo2.scl.csv [...]]\n"));
 	UE_LOG(LogShaderPipelineCacheTools, Warning, TEXT("Usage: Diff ShaderCache1.stablepc.csv ShaderCache1.stablepc.csv [...]]\n"));
 	UE_LOG(LogShaderPipelineCacheTools, Warning, TEXT("Usage: Expand Input1.upipelinecache Dir2/*.upipelinecache InputSCLInfo1.scl.csv Dir2/*.scl.csv InputSCLInfo3.scl.csv [...] Output.stablepc.csv\n"));
 	UE_LOG(LogShaderPipelineCacheTools, Warning, TEXT("Usage: Build Input.stablepc.csv InputDir2/*.stablepc.csv InputSCLInfo1.scl.csv Dir2/*.scl.csv InputSCLInfo3.scl.csv [...] Output.upipelinecache\n"));
+	UE_LOG(LogShaderPipelineCacheTools, Warning, TEXT("Usage: Decompress Input1.stablepc.csv.compressed Input2.stablepc.csv.compressed [...]\n"));
+	UE_LOG(LogShaderPipelineCacheTools, Warning, TEXT("Usage: All commands accept stablepc.csv.compressed instead of stablepc.csv for compressing output\n"));
 	return 0;
 }
diff --git a/Engine/Source/Editor/UnrealEd/Private/CookOnTheFlyServer.cpp b/Engine/Source/Editor/UnrealEd/Private/CookOnTheFlyServer.cpp
--- a/Engine/Source/Editor/UnrealEd/Private/CookOnTheFlyServer.cpp
+++ b/Engine/Source/Editor/UnrealEd/Private/CookOnTheFlyServer.cpp
@@ -6063,11 +6063,14 @@ void UCookOnTheFlyServer::ProcessShaderCodeLibraries(const FString& LibraryName)
 			TargetPlatform->GetAllTargetedShaderFormats(ShaderFormats);
 			for (FName ShaderFormat : ShaderFormats)
 			{
+				// *stablepc.csv or *stablepc.csv.compressed
 				const FString Filename = FString::Printf(TEXT("*%s_%s.stablepc.csv"), *LibraryName, *ShaderFormat.ToString());
 				const FString StablePCPath = FPaths::ProjectDir() / TEXT("Build") / TargetPlatform->IniPlatformName() / TEXT("PipelineCaches") / Filename;
+				const FString StablePCPathCompressed = StablePCPath + TEXT(".compressed");
 
 				TArray<FString> ExpandedFiles;
-				IFileManager::Get().FindFilesRecursive(ExpandedFiles, *FPaths::GetPath(StablePCPath), *FPaths::GetCleanFilename(StablePCPath), true, false);
+				IFileManager::Get().FindFilesRecursive(ExpandedFiles, *FPaths::GetPath(StablePCPath), *FPaths::GetCleanFilename(StablePCPath), true, false, false);
+				IFileManager::Get().FindFilesRecursive(ExpandedFiles, *FPaths::GetPath(StablePCPathCompressed), *FPaths::GetCleanFilename(StablePCPathCompressed), true, false, false);
 				if (!ExpandedFiles.Num())
 				{
 					UE_LOG(LogCook, Display, TEXT("---- NOT Running UShaderPipelineCacheToolsCommandlet for platform %s  shader format %s, no files found at %s"), *TargetPlatformName.ToString(), *ShaderFormat.ToString(), *StablePCPath);
