Fix non-unity
#rnx

diff --git a/Engine/Source/Editor/UnrealEd/Private/Commandlets/ShaderPipelineCacheToolsCommandlet.cpp b/Engine/Source/Editor/UnrealEd/Private/Commandlets/ShaderPipelineCacheToolsCommandlet.cpp
--- a/Engine/Source/Editor/UnrealEd/Private/Commandlets/ShaderPipelineCacheToolsCommandlet.cpp
+++ b/Engine/Source/Editor/UnrealEd/Private/Commandlets/ShaderPipelineCacheToolsCommandlet.cpp
@@ -8,6 +8,7 @@
 #include "Misc/FileHelper.h"
 #include "ShaderPipelineCache.h"
 #include "Serialization/MemoryReader.h"
+#include "Serialization/MemoryWriter.h"
 
 DEFINE_LOG_CATEGORY_STATIC(LogShaderPipelineCacheTools, Log, All);
 
