Follow up crash fix for dragging windows around in editor during autosave

diff --git a/Engine/Source/Editor/UnrealEd/Private/PackageAutoSaver.cpp b/Engine/Source/Editor/UnrealEd/Private/PackageAutoSaver.cpp
--- a/Engine/Source/Editor/UnrealEd/Private/PackageAutoSaver.cpp
+++ b/Engine/Source/Editor/UnrealEd/Private/PackageAutoSaver.cpp
@@ -433,7 +433,7 @@ bool FPackageAutoSaver::CanAutoSave() const
 	const bool bPlayWorldValid = GUnrealEd->PlayWorld != nullptr;
 	const bool bAnyMenusVisible	= FSlateApplication::Get().AnyMenusVisible();
 	const bool bAutomationTesting = GIsAutomationTesting;
-	const bool bIsInteracting = FSlateApplication::Get().HasAnyMouseCaptor() || GUnrealEd->IsUserInteracting() || (bDidInteractRecently && !bAutoSaveNotificationLaunched && !bDelayingDueToFailedSave);
+	const bool bIsInteracting = FSlateApplication::Get().HasAnyMouseCaptor() || FSlateApplication::Get().IsDragDropping() || GUnrealEd->IsUserInteracting() || (bDidInteractRecently && !bAutoSaveNotificationLaunched && !bDelayingDueToFailedSave);
 	const bool bHasGameOrProjectLoaded = FApp::HasProjectName();
 	const bool bAreShadersCompiling = GShaderCompilingManager->IsCompiling();
 	const bool bIsVREditorActive = IVREditorModule::Get().IsVREditorEnabled();	// @todo vreditor: Eventually we should support this while in VR (modal VR progress, with sufficient early warning)
diff --git a/Engine/Source/Runtime/Slate/Private/Framework/Application/SlateApplication.cpp b/Engine/Source/Runtime/Slate/Private/Framework/Application/SlateApplication.cpp
--- a/Engine/Source/Runtime/Slate/Private/Framework/Application/SlateApplication.cpp
+++ b/Engine/Source/Runtime/Slate/Private/Framework/Application/SlateApplication.cpp
@@ -1543,7 +1543,7 @@ void FSlateApplication::PrivateDrawWindows( TSharedPtr<SWindow> DrawOnlyThisWind
 			for( TArray< TSharedRef<SWindow> >::TConstIterator CurrentWindowIt( SlateWindows ); CurrentWindowIt; ++CurrentWindowIt )
 			{
 				const TSharedRef<SWindow>& CurrentWindow = *CurrentWindowIt;
-				if ( CurrentWindow->IsTopmostWindow() )
+				if ( CurrentWindow->GetType() == EWindowType::ToolTip )
 				{
 					DrawWindowAndChildren(CurrentWindow, DrawWindowArgs);
 				}
