Properly set flags on UserDefinedStruct that are being forced to load as part of blueprint compilation - this ensures that instances of said struct are initialized from the default instance

diff --git a/Engine/Source/Editor/UnrealEd/Private/Kismet2/BlueprintEditorUtils.cpp b/Engine/Source/Editor/UnrealEd/Private/Kismet2/BlueprintEditorUtils.cpp
--- a/Engine/Source/Editor/UnrealEd/Private/Kismet2/BlueprintEditorUtils.cpp
+++ b/Engine/Source/Editor/UnrealEd/Private/Kismet2/BlueprintEditorUtils.cpp
@@ -1004,6 +1004,15 @@ struct FRegenerationHelper
 		{
 			Struct->StaticLink(true);
 			ensure(Struct->IsA<UFunction>() || (OldPropertiesSize == Struct->GetPropertiesSize()) || !Struct->HasAnyFlags(RF_LoadCompleted));
+			
+			// UStruct::Link is going to attempt to set the StructFlags, but it has no knowledge of UserDefinedStruct::DefaultStructInstance.
+			// We don't want to set CPF_ZeroConstructor if UserDefinedStruct::DefaultStructInstance has non zero data, so this call
+			// gives UUserDefinedStruct a chance to enforce its invariants. The cooked/EDL path relies on setting struct flags in 
+			// serialize, after link has been called.
+			if(UUserDefinedStruct* AsUDS = Cast<UUserDefinedStruct>(Struct))
+			{
+				AsUDS->UpdateStructFlags();
+			}
 		}
 	}
 
