Disable use of asset registry for package name lookups when clearing the cache for cooking dlc.

diff --git a/Engine/Source/Editor/UnrealEd/Private/CookOnTheFlyServer.cpp b/Engine/Source/Editor/UnrealEd/Private/CookOnTheFlyServer.cpp
--- a/Engine/Source/Editor/UnrealEd/Private/CookOnTheFlyServer.cpp
+++ b/Engine/Source/Editor/UnrealEd/Private/CookOnTheFlyServer.cpp
@@ -6624,7 +6624,7 @@ void UCookOnTheFlyServer::InitializeTargetPlatforms()
 void UCookOnTheFlyServer::TermSandbox()
 {
 	ClearAllCookedData();
-	PackageNameCache->ClearPackageFilenameCache();
+	PackageNameCache->ClearPackageFilenameCache(nullptr);
 	SandboxFile = nullptr;
 }
 
diff --git a/Engine/Source/Editor/UnrealEd/Private/Cooker/PackageNameCache.h b/Engine/Source/Editor/UnrealEd/Private/Cooker/PackageNameCache.h
--- a/Engine/Source/Editor/UnrealEd/Private/Cooker/PackageNameCache.h
+++ b/Engine/Source/Editor/UnrealEd/Private/Cooker/PackageNameCache.h
@@ -43,14 +43,14 @@ struct FPackageNameCache
 
 	const FName*	GetCachedPackageFilenameToPackageFName(const FName& StandardPackageFilename) const;
 
-	void			ClearPackageFilenameCache() const;
+	void			ClearPackageFilenameCache(IAssetRegistry* InAssetRegistry) const;
 	bool			ClearPackageFilenameCacheForPackage(const UPackage* Package) const;
 
 private:
 	bool DoesPackageExist(const FName& PackageName, FString* OutFilename) const;
 	const FCachedPackageFilename& Cache(const FName& PackageName) const;
 
-	IAssetRegistry* AssetRegistry;
+	mutable IAssetRegistry* AssetRegistry;
 
 	mutable TMap<FName, FCachedPackageFilename> PackageFilenameCache; // filename cache (only process the string operations once)
 	mutable TMap<FName, FName>					PackageFilenameToPackageFNameCache;
@@ -167,9 +167,10 @@ const FName* FPackageNameCache::GetCachedPackageFilenameToPackageFName(const FNa
 	return PackageFilenameToPackageFNameCache.Find(StandardPackageFilename);
 }
 
-void FPackageNameCache::ClearPackageFilenameCache() const
+void FPackageNameCache::ClearPackageFilenameCache(IAssetRegistry* InAssetRegistry) const
 {
 	check(IsInGameThread());
 	PackageFilenameCache.Empty();
 	PackageFilenameToPackageFNameCache.Empty();
+	AssetRegistry = InAssetRegistry;
 }
