Fixed notification dissapearing when updating text

diff --git a/Engine/Source/Editor/UnrealEd/Private/PlayLevel.cpp b/Engine/Source/Editor/UnrealEd/Private/PlayLevel.cpp
--- a/Engine/Source/Editor/UnrealEd/Private/PlayLevel.cpp
+++ b/Engine/Source/Editor/UnrealEd/Private/PlayLevel.cpp
@@ -1543,7 +1543,7 @@ public:
 			{
 				GEditor->PlayEditorSound(TEXT("/Engine/EditorSounds/Notifications/CompileFailed_Cue.CompileFailed_Cue"));
 			}
-			else
+			else if (CompletionState == SNotificationItem::CS_Success)
 			{
 				GEditor->PlayEditorSound(TEXT("/Engine/EditorSounds/Notifications/CompileSuccess_Cue.CompileSuccess_Cue"));
 			}
@@ -1551,7 +1551,10 @@ public:
 			TSharedPtr<SNotificationItem> NotificationItem = NotificationItemPtr.Pin();
 			NotificationItem->SetText(Text);
 			NotificationItem->SetCompletionState(CompletionState);
-			NotificationItem->ExpireAndFadeout();
+			if (CompletionState == SNotificationItem::CS_Success || CompletionState == SNotificationItem::CS_Fail)
+			{
+				NotificationItem->ExpireAndFadeout();
+			}
 		}
 	}
 
