Fix crash when using the 'compile all blueprints' commandlet or when otherwise compiling a blueprint without using the compilation manager
Harden PatchNewCDOIntoLinker, avoid crashing when using old compilation path after compile on load has completed

diff --git a/Engine/Source/Editor/UnrealEd/Private/Kismet2/BlueprintEditorUtils.cpp b/Engine/Source/Editor/UnrealEd/Private/Kismet2/BlueprintEditorUtils.cpp
--- a/Engine/Source/Editor/UnrealEd/Private/Kismet2/BlueprintEditorUtils.cpp
+++ b/Engine/Source/Editor/UnrealEd/Private/Kismet2/BlueprintEditorUtils.cpp
@@ -718,7 +718,7 @@ void FBlueprintEditorUtils::PreloadConstructionScript(UBlueprint* Blueprint)
 
 void FBlueprintEditorUtils::PatchNewCDOIntoLinker(UObject* CDO, FLinkerLoad* Linker, int32 ExportIndex, FUObjectSerializeContext* InLoadContext)
 {
-	if( (CDO != nullptr) && (Linker != nullptr) && (ExportIndex != INDEX_NONE) )
+	if( (CDO != nullptr) && (Linker != nullptr) && (ExportIndex != INDEX_NONE) && Linker->ExportMap.Num() != 0 )
 	{
 		// Get rid of the old thing that was in its place
 		UObject* OldCDO = Linker->ExportMap[ExportIndex].Object;
