Fixed crash when the transaction details is opened on a transaction that have referenced to gc'ed objects

diff --git a/Engine/Source/Editor/UnrealEd/Private/EditorTransaction.cpp b/Engine/Source/Editor/UnrealEd/Private/EditorTransaction.cpp
--- a/Engine/Source/Editor/UnrealEd/Private/EditorTransaction.cpp
+++ b/Engine/Source/Editor/UnrealEd/Private/EditorTransaction.cpp
@@ -634,17 +634,23 @@ void FTransaction::FObjectRecord::AddReferencedObjects( FReferenceCollector& Col
 	Collector.AddReferencedObject(Obj);
 	Object.Object = Obj;
 
-	for (FPersistentObjectRef& ReferencedObject : SerializedObject.ReferencedObjects)
+	auto AddSerializedObjectReferences = [](FSerializedObject& InSerializedObject, FReferenceCollector& InCollector)
 	{
-		UObject* RefObj = ReferencedObject.Object;
-		Collector.AddReferencedObject(RefObj);
-		ReferencedObject.Object = RefObj;
-	}
+		for (FPersistentObjectRef& ReferencedObject : InSerializedObject.ReferencedObjects)
+		{
+			UObject* RefObj = ReferencedObject.Object;
+			InCollector.AddReferencedObject(RefObj);
+			ReferencedObject.Object = RefObj;
+		}
 
-	if (SerializedObject.ObjectAnnotation.IsValid())
-	{
-		SerializedObject.ObjectAnnotation->AddReferencedObjects(Collector);
-	}
+		if (InSerializedObject.ObjectAnnotation.IsValid())
+		{
+			InSerializedObject.ObjectAnnotation->AddReferencedObjects(InCollector);
+		}
+	};
+	AddSerializedObjectReferences(SerializedObject, Collector);
+	AddSerializedObjectReferences(SerializedObjectFlip, Collector);
+	AddSerializedObjectReferences(SerializedObjectSnapshot, Collector);
 }
 
 bool FTransaction::FObjectRecord::ContainsPieObject() const
@@ -658,15 +664,31 @@ bool FTransaction::FObjectRecord::ContainsPieObject() const
 		}
 	}
 
-	for (const FPersistentObjectRef& ReferencedObject : SerializedObject.ReferencedObjects)
+	auto SerializedObjectContainPieObjects = [](const FSerializedObject& InSerializedObject) -> bool
 	{
-		const UObject* Obj = ReferencedObject.Object;
-		if( Obj && Obj->GetOutermost()->HasAnyPackageFlags(PKG_PlayInEditor))
+		for (const FPersistentObjectRef& ReferencedObject : InSerializedObject.ReferencedObjects)
 		{
-			return true;
+			const UObject* Obj = ReferencedObject.Object;
+			if (Obj && Obj->GetOutermost()->HasAnyPackageFlags(PKG_PlayInEditor))
+			{
+				return true;
+			}
 		}
-	}
+		return false;
+	};
 
+	if (SerializedObjectContainPieObjects(SerializedObject))
+	{
+		return true;
+	}
+	if (SerializedObjectContainPieObjects(SerializedObjectFlip))
+	{
+		return true;
+	}
+	if (SerializedObjectContainPieObjects(SerializedObjectSnapshot))
+	{
+		return true;
+	}
 	return false;
 }
 
