Merging CL 5921374 from //UE4/Dev-Editor to //UE4/Main

Improve UI material thumbnail generation by using the widget renderer on a slate image with a brush using the material.

diff --git a/Engine/Source/Editor/UnrealEd/Classes/ThumbnailRendering/MaterialInstanceThumbnailRenderer.h b/Engine/Source/Editor/UnrealEd/Classes/ThumbnailRendering/MaterialInstanceThumbnailRenderer.h
--- a/Engine/Source/Editor/UnrealEd/Classes/ThumbnailRendering/MaterialInstanceThumbnailRenderer.h
+++ b/Engine/Source/Editor/UnrealEd/Classes/ThumbnailRendering/MaterialInstanceThumbnailRenderer.h
@@ -30,5 +30,7 @@ class UMaterialInstanceThumbnailRenderer : public UDefaultSizedThumbnailRenderer
 
 private:
 	class FMaterialThumbnailScene* ThumbnailScene;
+	class FWidgetRenderer* WidgetRenderer;
+	struct FSlateMaterialBrush* UIMaterialBrush;
 };
 
diff --git a/Engine/Source/Editor/UnrealEd/Private/ThumbnailRendering/MaterialInstanceThumbnailRenderer.cpp b/Engine/Source/Editor/UnrealEd/Private/ThumbnailRendering/MaterialInstanceThumbnailRenderer.cpp
--- a/Engine/Source/Editor/UnrealEd/Private/ThumbnailRendering/MaterialInstanceThumbnailRenderer.cpp
+++ b/Engine/Source/Editor/UnrealEd/Private/ThumbnailRendering/MaterialInstanceThumbnailRenderer.cpp
@@ -6,11 +6,15 @@
 #include "Materials/MaterialInterface.h"
 #include "SceneView.h"
 #include "ThumbnailHelpers.h"
+#include "Slate/WidgetRenderer.h"
+#include "SlateMaterialBrush.h"
 
 UMaterialInstanceThumbnailRenderer::UMaterialInstanceThumbnailRenderer(const FObjectInitializer& ObjectInitializer)
 	: Super(ObjectInitializer)
+	, ThumbnailScene(nullptr)
+	, WidgetRenderer(nullptr)
+	, UIMaterialBrush(nullptr)
 {
-	ThumbnailScene = nullptr;
 }
 
 void UMaterialInstanceThumbnailRenderer::Draw(UObject* Object, int32 X, int32 Y, uint32 Width, uint32 Height, FRenderTarget* RenderTarget, FCanvas* Canvas)
@@ -18,43 +22,83 @@ void UMaterialInstanceThumbnailRenderer::Draw(UObject* Object, int32 X, int32 Y,
 	UMaterialInterface* MatInst = Cast<UMaterialInterface>(Object);
 	if (MatInst != nullptr)
 	{
-		if ( ThumbnailScene == nullptr || ensure(ThumbnailScene->GetWorld() != nullptr) == false )
+		UMaterial* Mat = MatInst->GetMaterial();
+		if (Mat != nullptr && Mat->IsUIMaterial())
 		{
-			if (ThumbnailScene)
+			if (WidgetRenderer == nullptr)
 			{
-				FlushRenderingCommands();
-				delete ThumbnailScene;
+				const bool bUseGammaCorrection = true;
+				WidgetRenderer = new FWidgetRenderer(bUseGammaCorrection);
+				check(WidgetRenderer);
 			}
-			ThumbnailScene = new FMaterialThumbnailScene();
+
+			if (UIMaterialBrush == nullptr)
+			{
+				UIMaterialBrush = new FSlateMaterialBrush(FVector2D(SlateBrushDefs::DefaultImageSize, SlateBrushDefs::DefaultImageSize));
+				check(UIMaterialBrush);
+			}
+		
+			UIMaterialBrush->SetMaterial(MatInst);
+
+			TSharedPtr<SImage> UIMaterialImage = SNew(SImage).Image(UIMaterialBrush);
+			const FVector2D DrawSize((float)Width, (float)Height);
+			const float DeltaTime = 0.f;
+			WidgetRenderer->DrawWidget(RenderTarget, UIMaterialImage.ToSharedRef(), DrawSize, DeltaTime);
+
+			UIMaterialBrush->SetMaterial(nullptr);
 		}
+		else
+		{
+			if ( ThumbnailScene == nullptr || ensure(ThumbnailScene->GetWorld() != nullptr) == false )
+			{
+				if (ThumbnailScene)
+				{
+					FlushRenderingCommands();
+					delete ThumbnailScene;
+				}
+				ThumbnailScene = new FMaterialThumbnailScene();
+			}
 
-		ThumbnailScene->SetMaterialInterface(MatInst);
-		FSceneViewFamilyContext ViewFamily( FSceneViewFamily::ConstructionValues( RenderTarget, ThumbnailScene->GetScene(), FEngineShowFlags(ESFIM_Game) )
-			.SetWorldTimes(FApp::GetCurrentTime() - GStartTime, FApp::GetDeltaTime(), FApp::GetCurrentTime() - GStartTime));
+			ThumbnailScene->SetMaterialInterface(MatInst);
+			FSceneViewFamilyContext ViewFamily( FSceneViewFamily::ConstructionValues( RenderTarget, ThumbnailScene->GetScene(), FEngineShowFlags(ESFIM_Game) )
+				.SetWorldTimes(FApp::GetCurrentTime() - GStartTime, FApp::GetDeltaTime(), FApp::GetCurrentTime() - GStartTime));
 
-		ViewFamily.EngineShowFlags.DisableAdvancedFeatures();
-		ViewFamily.EngineShowFlags.SetSeparateTranslucency(true);
-		ViewFamily.EngineShowFlags.MotionBlur = 0;
-		ViewFamily.EngineShowFlags.AntiAliasing = 0;
+			ViewFamily.EngineShowFlags.DisableAdvancedFeatures();
+			ViewFamily.EngineShowFlags.SetSeparateTranslucency(true);
+			ViewFamily.EngineShowFlags.MotionBlur = 0;
+			ViewFamily.EngineShowFlags.AntiAliasing = 0;
 
-		ThumbnailScene->GetView(&ViewFamily, X, Y, Width, Height);
+			ThumbnailScene->GetView(&ViewFamily, X, Y, Width, Height);
 
-		if (ViewFamily.Views.Num() > 0)
-		{
-			RenderViewFamily(Canvas, &ViewFamily);
-		}
+			if (ViewFamily.Views.Num() > 0)
+			{
+				RenderViewFamily(Canvas, &ViewFamily);
+			}
 
-		ThumbnailScene->SetMaterialInterface(nullptr);
+			ThumbnailScene->SetMaterialInterface(nullptr);
+		}
 	}
 }
 
 void UMaterialInstanceThumbnailRenderer::BeginDestroy()
 {
-	if ( ThumbnailScene != nullptr )
+	if (ThumbnailScene != nullptr)
 	{
 		delete ThumbnailScene;
 		ThumbnailScene = nullptr;
 	}
 
+	if (WidgetRenderer != nullptr)
+	{
+		BeginCleanup(WidgetRenderer);
+		WidgetRenderer = nullptr;
+	}
+
+	if (UIMaterialBrush != nullptr)
+	{
+		delete UIMaterialBrush;
+		UIMaterialBrush = nullptr;
+	}
+
 	Super::BeginDestroy();
 }
diff --git a/Engine/Source/Runtime/Slate/Public/SlateMaterialBrush.h b/Engine/Source/Runtime/Slate/Public/SlateMaterialBrush.h
--- a/Engine/Source/Runtime/Slate/Public/SlateMaterialBrush.h
+++ b/Engine/Source/Runtime/Slate/Public/SlateMaterialBrush.h
@@ -9,7 +9,7 @@
 #include "Framework/Application/SlateApplication.h"
 
 /**
- * Dynamic rush for referencing a UMaterial.  
+ * Dynamic brush for referencing a UMaterial.  
  *
  * Note: This brush nor the slate renderer holds a strong reference to the material.  You are responsible for maintaining the lifetime of the brush and material object.
  */
@@ -26,4 +26,20 @@ struct FSlateMaterialBrush : public FSlateBrush
 	{
 		ResourceName = FName( *InMaterial.GetFullName() );
 	}
+
+	/**
+	 * Creates and initializes a new instance.
+	 *
+	 * @param InImageSize The material's dimensions.
+	 */
+	FSlateMaterialBrush(const FVector2D& InImageSize)
+		: FSlateBrush(ESlateBrushDrawType::Image, FName(TEXT("None")), FMargin(0), ESlateBrushTileType::NoTile, ESlateBrushImageType::FullColor, InImageSize, FLinearColor::White)
+	{}
+
+	/** Sets the material to use. */
+	void SetMaterial(UMaterialInterface* InMaterial)
+	{
+		SetResourceObject(InMaterial);
+		ResourceName = InMaterial != nullptr ? FName(*InMaterial->GetFullName()) : NAME_None;
+	}
 }; 
diff --git a/Engine/Source/Runtime/SlateRHIRenderer/Private/Slate3DRenderer.cpp b/Engine/Source/Runtime/SlateRHIRenderer/Private/Slate3DRenderer.cpp
--- a/Engine/Source/Runtime/SlateRHIRenderer/Private/Slate3DRenderer.cpp
+++ b/Engine/Source/Runtime/SlateRHIRenderer/Private/Slate3DRenderer.cpp
@@ -126,9 +126,9 @@ void FSlate3DRenderer::DrawWindowToTarget_RenderThread(FRHICommandListImmediate&
 	// Enqueue a command to unlock the draw buffer after all windows have been drawn
 	RenderTargetPolicy->BeginDrawingWindows();
 
-	FTextureRenderTarget2DResource* RenderTargetResource = static_cast<FTextureRenderTarget2DResource*>(Context.RenderTargetResource);
+	FRenderTarget* RenderTargetResource = static_cast<FRenderTarget*>(Context.RenderTargetResource);
 	// Set render target and clear.
-	FTexture2DRHIRef RTTextureRHI = RenderTargetResource->GetTextureRHI();
+	FTexture2DRHIRef RTTextureRHI = RenderTargetResource->GetRenderTargetTexture();
 	InRHICmdList.TransitionResource(EResourceTransitionAccess::EWritable, RTTextureRHI);
 	
 	FRHIRenderPassInfo RPInfo(RTTextureRHI, ERenderTargetActions::Load_Store);
@@ -160,14 +160,14 @@ void FSlate3DRenderer::DrawWindowToTarget_RenderThread(FRHICommandListImmediate&
 				FMatrix ViewOffset = FTranslationMatrix::Make(FVector(DrawOffset, 0));
 				ProjectionMatrix = ViewOffset * ProjectionMatrix;
 
-				FSlateBackBuffer BackBufferTarget(RenderTargetResource->GetTextureRHI(), FIntPoint(RTTextureRHI->GetSizeX(), RTTextureRHI->GetSizeY()));
+				FSlateBackBuffer BackBufferTarget(RenderTargetResource->GetRenderTargetTexture(), FIntPoint(RTTextureRHI->GetSizeX(), RTTextureRHI->GetSizeY()));
 
 				FSlateRenderingParams DrawOptions(ProjectionMatrix, Context.WorldTimeSeconds, Context.DeltaTimeSeconds, Context.RealTimeSeconds);
 				// The scene renderer will handle it in this case
 				DrawOptions.bAllowSwitchVerticalAxis = false;
 				DrawOptions.ViewOffset = DrawOffset;
 
-				FTexture2DRHIRef ColorTarget = RenderTargetResource->GetTextureRHI();
+				FTexture2DRHIRef ColorTarget = RenderTargetResource->GetRenderTargetTexture();
 
 				if (BatchData.IsStencilClippingRequired())
 				{
@@ -196,7 +196,7 @@ void FSlate3DRenderer::DrawWindowToTarget_RenderThread(FRHICommandListImmediate&
 	InRHICmdList.EndRenderPass();
 
 	FSlateEndDrawingWindowsCommand::EndDrawingWindows(InRHICmdList, Context.WindowDrawBuffer, *RenderTargetPolicy);
-	InRHICmdList.CopyToResolveTarget(RenderTargetResource->GetTextureRHI(), RTTextureRHI, FResolveParams());
+	InRHICmdList.CopyToResolveTarget(RenderTargetResource->GetRenderTargetTexture(), RTTextureRHI, FResolveParams());
 
 	ISlate3DRendererPtr Self = SharedThis(this);
 
diff --git a/Engine/Source/Runtime/UMG/Private/Slate/WidgetRenderer.cpp b/Engine/Source/Runtime/UMG/Private/Slate/WidgetRenderer.cpp
--- a/Engine/Source/Runtime/UMG/Private/Slate/WidgetRenderer.cpp
+++ b/Engine/Source/Runtime/UMG/Private/Slate/WidgetRenderer.cpp
@@ -86,7 +86,7 @@ UTextureRenderTarget2D* FWidgetRenderer::CreateTargetFor(FVector2D DrawSize, Tex
 	return nullptr;
 }
 
-void FWidgetRenderer::DrawWidget(UTextureRenderTarget2D* RenderTarget, const TSharedRef<SWidget>& Widget, FVector2D DrawSize, float DeltaTime, bool bDeferRenderTargetUpdate)
+void FWidgetRenderer::DrawWidget(FRenderTarget* RenderTarget, const TSharedRef<SWidget>& Widget, FVector2D DrawSize, float DeltaTime, bool bDeferRenderTargetUpdate)
 {
 	TSharedRef<SVirtualWindow> Window = SNew(SVirtualWindow).Size(DrawSize);
 	TSharedRef<FHittestGrid> HitTestGrid = MakeShareable(new FHittestGrid());
@@ -97,8 +97,13 @@ void FWidgetRenderer::DrawWidget(UTextureRenderTarget2D* RenderTarget, const TSh
 	DrawWindow(RenderTarget, HitTestGrid, Window, 1, DrawSize, DeltaTime, bDeferRenderTargetUpdate);
 }
 
+void FWidgetRenderer::DrawWidget(UTextureRenderTarget2D* RenderTarget, const TSharedRef<SWidget>& Widget, FVector2D DrawSize, float DeltaTime, bool bDeferRenderTargetUpdate)
+{
+	DrawWidget(RenderTarget->GameThread_GetRenderTargetResource(), Widget, DrawSize, DeltaTime, bDeferRenderTargetUpdate);
+}
+
 void FWidgetRenderer::DrawWindow(
-	UTextureRenderTarget2D* RenderTarget,
+	FRenderTarget* RenderTarget,
 	TSharedRef<FHittestGrid> HitTestGrid,
 	TSharedRef<SWindow> Window,
 	float Scale,
@@ -124,6 +129,18 @@ void FWidgetRenderer::DrawWindow(
 	UTextureRenderTarget2D* RenderTarget,
 	TSharedRef<FHittestGrid> HitTestGrid,
 	TSharedRef<SWindow> Window,
+	float Scale,
+	FVector2D DrawSize,
+	float DeltaTime,
+	bool bDeferRenderTargetUpdate)
+{
+	DrawWindow(RenderTarget->GameThread_GetRenderTargetResource(), HitTestGrid, Window, Scale, DrawSize, DeltaTime, bDeferRenderTargetUpdate);
+}
+
+void FWidgetRenderer::DrawWindow(
+	FRenderTarget* RenderTarget,
+	TSharedRef<FHittestGrid> HitTestGrid,
+	TSharedRef<SWindow> Window,
 	FGeometry WindowGeometry,
 	FSlateRect WindowClipRect,
 	float DeltaTime,
@@ -134,8 +151,20 @@ void FWidgetRenderer::DrawWindow(
 }
 
 void FWidgetRenderer::DrawWindow(
-	const FPaintArgs& PaintArgs,
 	UTextureRenderTarget2D* RenderTarget,
+	TSharedRef<FHittestGrid> HitTestGrid,
+	TSharedRef<SWindow> Window,
+	FGeometry WindowGeometry,
+	FSlateRect WindowClipRect,
+	float DeltaTime,
+	bool bDeferRenderTargetUpdate)
+{
+	DrawWindow(RenderTarget->GameThread_GetRenderTargetResource(), HitTestGrid, Window, WindowGeometry, WindowClipRect, DeltaTime, bDeferRenderTargetUpdate);
+}
+
+void FWidgetRenderer::DrawWindow(
+	const FPaintArgs& PaintArgs,
+	FRenderTarget* RenderTarget,
 	TSharedRef<SWindow> Window,
 	FGeometry WindowGeometry,
 	FSlateRect WindowClipRect,
@@ -189,7 +218,7 @@ void FWidgetRenderer::DrawWindow(
 			static_cast<float>(FApp::GetCurrentTime() - GStartTime),
 			static_cast<float>(FApp::GetDeltaTime()),
 			static_cast<float>(FPlatformTime::Seconds() - GStartTime),
-			RenderTarget->GameThread_GetRenderTargetResource(),
+			RenderTarget,
 			Renderer.Get(),
 			bClearTarget
 		};
@@ -198,3 +227,15 @@ void FWidgetRenderer::DrawWindow(
 	}
 #endif // !UE_SERVER
 }
+
+void FWidgetRenderer::DrawWindow(
+	const FPaintArgs& PaintArgs,
+	UTextureRenderTarget2D* RenderTarget,
+	TSharedRef<SWindow> Window,
+	FGeometry WindowGeometry,
+	FSlateRect WindowClipRect,
+	float DeltaTime,
+	bool bDeferRenderTargetUpdate)
+{
+	DrawWindow(PaintArgs, RenderTarget->GameThread_GetRenderTargetResource(), Window, WindowGeometry, WindowClipRect, DeltaTime, bDeferRenderTargetUpdate);
+}
diff --git a/Engine/Source/Runtime/UMG/Public/Slate/WidgetRenderer.h b/Engine/Source/Runtime/UMG/Public/Slate/WidgetRenderer.h
--- a/Engine/Source/Runtime/UMG/Public/Slate/WidgetRenderer.h
+++ b/Engine/Source/Runtime/UMG/Public/Slate/WidgetRenderer.h
@@ -19,6 +19,7 @@ class FSlateDrawBuffer;
 class ISlate3DRenderer;
 class STooltipPresenter;
 class UTextureRenderTarget2D;
+class FRenderTarget;
 
 /**
  * 
@@ -46,6 +47,13 @@ public:
 
 	UTextureRenderTarget2D* DrawWidget(const TSharedRef<SWidget>& Widget, FVector2D DrawSize);
 
+	void DrawWidget(
+		FRenderTarget* RenderTarget,
+		const TSharedRef<SWidget>& Widget,
+		FVector2D DrawSize,
+		float DeltaTime,
+		bool bDeferRenderTargetUpdate = false);
+
 	void DrawWidget(
 		UTextureRenderTarget2D* RenderTarget,
 		const TSharedRef<SWidget>& Widget,
@@ -53,6 +61,15 @@ public:
 		float DeltaTime,
 		bool bDeferRenderTargetUpdate = false);
 
+	void DrawWindow(
+		FRenderTarget* RenderTarget,
+		TSharedRef<FHittestGrid> HitTestGrid,
+		TSharedRef<SWindow> Window,
+		float Scale,
+		FVector2D DrawSize,
+		float DeltaTime,
+		bool bDeferRenderTargetUpdate = false);
+
 	void DrawWindow(
 		UTextureRenderTarget2D* RenderTarget,
 		TSharedRef<FHittestGrid> HitTestGrid,
@@ -62,6 +79,15 @@ public:
 		float DeltaTime,
 		bool bDeferRenderTargetUpdate = false);
 
+	void DrawWindow(
+		FRenderTarget* RenderTarget,
+		TSharedRef<FHittestGrid> HitTestGrid,
+		TSharedRef<SWindow> Window,
+		FGeometry WindowGeometry,
+		FSlateRect WindowClipRect,
+		float DeltaTime,
+		bool bDeferRenderTargetUpdate = false);
+
 	void DrawWindow(
 		UTextureRenderTarget2D* RenderTarget,
 		TSharedRef<FHittestGrid> HitTestGrid,
@@ -71,6 +97,15 @@ public:
 		float DeltaTime,
 		bool bDeferRenderTargetUpdate = false);
 
+	void DrawWindow(
+		const FPaintArgs& PaintArgs,
+		FRenderTarget* RenderTarget,
+		TSharedRef<SWindow> Window,
+		FGeometry WindowGeometry,
+		FSlateRect WindowClipRect,
+		float DeltaTime,
+		bool bDeferRenderTargetUpdate = false);
+
 	void DrawWindow(
 		const FPaintArgs& PaintArgs,
 		UTextureRenderTarget2D* RenderTarget,
