Merging 5921503 from //UE4/Dev-Editor to //UE4/Main
Build fix in MaterialInstanceThumbnailRenderer.cpp

diff --git a/Engine/Source/Editor/UnrealEd/Private/ThumbnailRendering/MaterialInstanceThumbnailRenderer.cpp b/Engine/Source/Editor/UnrealEd/Private/ThumbnailRendering/MaterialInstanceThumbnailRenderer.cpp
--- a/Engine/Source/Editor/UnrealEd/Private/ThumbnailRendering/MaterialInstanceThumbnailRenderer.cpp
+++ b/Engine/Source/Editor/UnrealEd/Private/ThumbnailRendering/MaterialInstanceThumbnailRenderer.cpp
@@ -4,9 +4,11 @@
 #include "Misc/App.h"
 #include "ShowFlags.h"
 #include "Materials/MaterialInterface.h"
+#include "Materials/Material.h"
 #include "SceneView.h"
 #include "ThumbnailHelpers.h"
 #include "Slate/WidgetRenderer.h"
+#include "Widgets/Images/SImage.h"
 #include "SlateMaterialBrush.h"
 
 UMaterialInstanceThumbnailRenderer::UMaterialInstanceThumbnailRenderer(const FObjectInitializer& ObjectInitializer)
