#UE4 When reinstancing blueprints, clear the NeedLoad/NeedPostLoad flags on all subobjects of the old CDO. We no longer need to finish loading these objects as they will be abandoned with the CDO. CDOs that have native default subobjects might leave behind orphaned objects in this list if the native class is changed to no longer include them but they were still saved in the package. Clearing these flags fixes a crash where these orphaned objects were not properly getting loaded anyway, and were causing an assert during GC that a NeedLoad/NeedPostLoad object was being destroyed.

diff --git a/Engine/Source/Editor/UnrealEd/Private/Kismet2/KismetReinstanceUtilities.cpp b/Engine/Source/Editor/UnrealEd/Private/Kismet2/KismetReinstanceUtilities.cpp
--- a/Engine/Source/Editor/UnrealEd/Private/Kismet2/KismetReinstanceUtilities.cpp
+++ b/Engine/Source/Editor/UnrealEd/Private/Kismet2/KismetReinstanceUtilities.cpp
@@ -1567,6 +1567,18 @@ UClass* FBlueprintCompileReinstancer::MoveCDOToNewClass(UClass* OwnerClass, cons
 			OwnerClass->ClassDefaultObject = nullptr;
 			OldCDO->Rename(nullptr, CopyOfOwnerClass->GetOuter(), REN_DoNotDirty | REN_DontCreateRedirectors | REN_ForceNoResetLoaders);
 			CopyOfOwnerClass->ClassDefaultObject = OldCDO;
+
+			// Clear the NeedLoad/NeedPostLoad flags. We no longer need to finish loading these objects as they will be abandoned
+			// with the CDO. CDOs that have native default subobjects might leave behind orphaned objects in this list if the native
+			// class is changed to no longer include them but they were still saved in the package. Clearing these flags fixes a crash
+			// where these orphaned objects were not properly getting loaded anyway, and were causing an assert during GC that
+			// a NeedLoad/NeedPostLoad object was being destroyed.
+			TArray<UObject*> CDOSubObjects;
+			GetObjectsWithOuter(OldCDO, CDOSubObjects, false);
+			for (UObject* CDOSubObject : CDOSubObjects)
+			{
+				CDOSubObject->ClearFlags(RF_NeedLoad | RF_NeedPostLoad);
+			}
 		}
 		OldCDO->SetClass(CopyOfOwnerClass);
 	}
