Apply the Preview Rendering Level to the Play in new window PIE session

diff --git a/Engine/Source/Editor/UnrealEd/Private/PlayLevel.cpp b/Engine/Source/Editor/UnrealEd/Private/PlayLevel.cpp
--- a/Engine/Source/Editor/UnrealEd/Private/PlayLevel.cpp
+++ b/Engine/Source/Editor/UnrealEd/Private/PlayLevel.cpp
@@ -3023,7 +3023,7 @@ UGameInstance* UEditorEngine::CreatePIEGameInstance(int32 InPIEInstance, bool bI
 	FFormatNamedArguments Args;
 	Args.Add( TEXT("GameName"), FText::FromString( FString( WindowTitleOverride.IsEmpty() ? FApp::GetProjectName() : WindowTitleOverride.ToString() ) ) );
 	Args.Add( TEXT("PlatformBits"), FText::FromString( PlatformBitsString ) );
-	Args.Add( TEXT("RHIName"), FText::FromName( LegacyShaderPlatformToShaderFormat( GShaderPlatformForFeatureLevel[GMaxRHIFeatureLevel] ) ) );
+	Args.Add( TEXT("RHIName"), FText::FromName( LegacyShaderPlatformToShaderFormat( GShaderPlatformForFeatureLevel[PreviewFeatureLevel] ) ) );
 
 	const ULevelEditorPlaySettings* PlayInSettings = GetDefault<ULevelEditorPlaySettings>();
 	const EPlayNetMode PlayNetMode = [&PlayInSettings]{ EPlayNetMode NetMode(PIE_Standalone); return (PlayInSettings->GetPlayNetMode(NetMode) ? NetMode : PIE_Standalone); }();
