PR #5809: Modifying UActorFactory GetDefaultActor and GetDefaultActorClass so that the override NewActorClassName is always applied. (Contributed by [REDACTED_PERSON])

diff --git a/Engine/Source/Editor/UnrealEd/Private/Factories/ActorFactory.cpp b/Engine/Source/Editor/UnrealEd/Private/Factories/ActorFactory.cpp
--- a/Engine/Source/Editor/UnrealEd/Private/Factories/ActorFactory.cpp
+++ b/Engine/Source/Editor/UnrealEd/Private/Factories/ActorFactory.cpp
@@ -214,26 +214,30 @@ bool UActorFactory::CanCreateActorFrom( const FAssetData& AssetData, FText& OutE
 
 AActor* UActorFactory::GetDefaultActor( const FAssetData& AssetData )
 {
-	if ( NewActorClassName != TEXT("") )
-	{
-		UE_LOG(LogActorFactory, Log, TEXT("Loading ActorFactory Class %s"), *NewActorClassName);
-		NewActorClass = Cast<UClass>(StaticLoadObject(UClass::StaticClass(), NULL, *NewActorClassName, NULL, LOAD_NoWarn, NULL));
-		NewActorClassName = TEXT("");
-		if ( NewActorClass == NULL )
-		{
-			UE_LOG(LogActorFactory, Log, TEXT("ActorFactory Class LOAD FAILED"));
-		}
-	}
-	return NewActorClass ? NewActorClass->GetDefaultObject<AActor>() : NULL;
+	UClass* DefaultActorClass = GetDefaultActorClass( AssetData );
+	return (DefaultActorClass ? DefaultActorClass->GetDefaultObject<AActor>() : nullptr);
 }
 
 UClass* UActorFactory::GetDefaultActorClass( const FAssetData& AssetData )
 {
-	if ( !NewActorClass )
+	if (!NewActorClassName.IsEmpty())
 	{
-		GetDefaultActor( AssetData );
+		UE_LOG(LogActorFactory, Log, TEXT("Loading ActorFactory Class %s"), *NewActorClassName);
+		UClass* NewActorClassCandidate = Cast<UClass>(StaticLoadObject(UClass::StaticClass(), nullptr, *NewActorClassName, nullptr, LOAD_NoWarn, nullptr));
+		NewActorClassName.Empty();
+		if (NewActorClassCandidate == nullptr)
+		{
+			UE_LOG(LogActorFactory, Warning, TEXT("ActorFactory Class LOAD FAILED - falling back to %s"), *GetNameSafe(*NewActorClass));
+		}
+		else if (!NewActorClassCandidate->IsChildOf(AActor::StaticClass()))
+		{
+			UE_LOG(LogActorFactory, Warning, TEXT("ActorFactory Class loaded class %s doesn't derive from Actor - falling back to %s"), *NewActorClassCandidate->GetName(), *GetNameSafe(*NewActorClass));
+		}
+		else
+		{
+			NewActorClass = NewActorClassCandidate;
+		}
 	}
-
 	return NewActorClass;
 }
 
