Fill in the Optimizations Viewmode Settings menu according to the capabilities of the active Rendering Level

diff --git a/Engine/Source/Editor/UnrealEd/Private/SEditorViewportViewMenu.cpp b/Engine/Source/Editor/UnrealEd/Private/SEditorViewportViewMenu.cpp
--- a/Engine/Source/Editor/UnrealEd/Private/SEditorViewportViewMenu.cpp
+++ b/Engine/Source/Editor/UnrealEd/Private/SEditorViewportViewMenu.cpp
@@ -321,17 +321,14 @@ TSharedRef<SWidget> SEditorViewportViewMenu::GenerateViewMenuContent() const
 
 						Menu.BeginSection("OptimizationViewmodes", LOCTEXT("OptimizationSubMenuHeader", "Optimization Viewmodes"));
 						{
-							if (FeatureLevel == ERHIFeatureLevel::SM5)
+							if (FeatureLevel == ERHIFeatureLevel::SM5 || FeatureLevel == ERHIFeatureLevel::SM4)
 							{
 								Menu.AddMenuEntry(BaseViewportCommands.LightComplexityMode, NAME_None, LOCTEXT("LightComplexityViewModeDisplayName", "Light Complexity"));
 								Menu.AddMenuEntry(BaseViewportCommands.LightmapDensityMode, NAME_None, LOCTEXT("LightmapDensityViewModeDisplayName", "Lightmap Density"));
 								Menu.AddMenuEntry(BaseViewportCommands.StationaryLightOverlapMode, NAME_None, LOCTEXT("StationaryLightOverlapViewModeDisplayName", "Stationary Light Overlap"));
 							}
 
-							if (AllowDebugViewShaderMode(DVSM_ShaderComplexityContainedQuadOverhead, GMaxRHIShaderPlatform, FeatureLevel))
-							{
-								Menu.AddMenuEntry(BaseViewportCommands.ShaderComplexityMode, NAME_None, LOCTEXT("ShaderComplexityViewModeDisplayName", "Shader Complexity"));
-							}
+							Menu.AddMenuEntry(BaseViewportCommands.ShaderComplexityMode, NAME_None, LOCTEXT("ShaderComplexityViewModeDisplayName", "Shader Complexity"));
 
 							if (AllowDebugViewShaderMode(DVSM_ShaderComplexityContainedQuadOverhead, GMaxRHIShaderPlatform, FeatureLevel))
 							{
