Fix cook -iterative bug in UCookOnTheFlyServer::ConvertCookedPathToUncookedPath causing unnecessary recook of game packages for projects located at different folder depth than Engine.
Remove some more string allocations.

diff --git a/Engine/Source/Editor/UnrealEd/Classes/CookOnTheSide/CookOnTheFlyServer.h b/Engine/Source/Editor/UnrealEd/Classes/CookOnTheSide/CookOnTheFlyServer.h
--- a/Engine/Source/Editor/UnrealEd/Classes/CookOnTheSide/CookOnTheFlyServer.h
+++ b/Engine/Source/Editor/UnrealEd/Classes/CookOnTheSide/CookOnTheFlyServer.h
@@ -197,7 +197,10 @@ private:
 	int32 LastCookedPackagesCount = 0;
 	double LastProgressDisplayTime = 0;
 
-	FName ConvertCookedPathToUncookedPath(const FString& SandboxPath, const FString& CookedPath) const;
+	FName ConvertCookedPathToUncookedPath(
+		const FString& SandboxRootDir, const FString& RelativeRootDir,
+		const FString& SandboxProjectDir, const FString& RelativeProjectDir,
+		const FString& CookedPath, FString& OutUncookedPath) const;
 
 	/** Get dependencies for package */
 	const TArray<FName>& GetFullPackageDependencies(const FName& PackageName) const;
@@ -891,9 +894,9 @@ private:
 	* Returns a map of the uncooked file path matches to the cooked file path for each package which exists
 	*
 	* @param UncookedpathToCookedPath out Map of the uncooked path matched with the cooked package which exists
-	* @param SandboxPath path to search for cooked packages in
+	* @param SandboxRootDir root dir to search for cooked packages in
 	*/
-	void GetAllCookedFiles(TMap<FName, FName>& UncookedPathToCookedPath, const FString& SandboxPath);
+	void GetAllCookedFiles(TMap<FName, FName>& UncookedPathToCookedPath, const FString& SandboxRootDir);
 
 	/** Generates asset registry */
 	void GenerateAssetRegistry();
diff --git a/Engine/Source/Editor/UnrealEd/Private/CookOnTheFlyServer.cpp b/Engine/Source/Editor/UnrealEd/Private/CookOnTheFlyServer.cpp
--- a/Engine/Source/Editor/UnrealEd/Private/CookOnTheFlyServer.cpp
+++ b/Engine/Source/Editor/UnrealEd/Private/CookOnTheFlyServer.cpp
@@ -4970,8 +4970,13 @@ bool UCookOnTheFlyServer::SaveCurrentIniSettings(const ITargetPlatform* TargetPl
 
 }
 
-FName UCookOnTheFlyServer::ConvertCookedPathToUncookedPath(const FString& SandboxPath, const FString& CookedPath) const
+FName UCookOnTheFlyServer::ConvertCookedPathToUncookedPath(
+	const FString& SandboxRootDir, const FString& RelativeRootDir,
+	const FString& SandboxProjectDir, const FString& RelativeProjectDir,
+	const FString& CookedPath, FString& OutUncookedPath) const
 {
+	OutUncookedPath.Reset();
+
 	// Check for remapped plugins' cooked content
 	if (PluginsToRemap.Num() > 0 && CookedPath.Contains(REMAPPED_PLUGGINS))
 	{
@@ -4980,52 +4985,81 @@ FName UCookOnTheFlyServer::ConvertCookedPathToUncookedPath(const FString& Sandbo
 		static uint32 RemappedPluginStrLen = FCString::Strlen(REMAPPED_PLUGGINS);
 		// Snip everything up through the RemappedPlugins/ off so we can find the plugin it corresponds to
 		FString PluginPath = CookedPath.RightChop(RemappedIndex + RemappedPluginStrLen + 1);
-		FString FullUncookedPath;
 		// Find the plugin that owns this content
 		for (TSharedRef<IPlugin> Plugin : PluginsToRemap)
 		{
 			if (PluginPath.StartsWith(Plugin->GetName()))
 			{
-				FullUncookedPath = Plugin->GetContentDir();
+				OutUncookedPath = Plugin->GetContentDir();
 				static uint32 ContentStrLen = FCString::Strlen(TEXT("Content/"));
 				// Chop off the pluginName/Content since it's part of the full path
-				FullUncookedPath /= PluginPath.RightChop(Plugin->GetName().Len() + ContentStrLen);
+				OutUncookedPath /= PluginPath.RightChop(Plugin->GetName().Len() + ContentStrLen);
 				break;
 			}
 		}
 
-		if (FullUncookedPath.Len() > 0)
+		if (OutUncookedPath.Len() > 0)
 		{
-			return FName(*FullUncookedPath);
+			return FName(*OutUncookedPath);
 		}
 		// Otherwise fall through to sandbox handling
 	}
 
-	FString UncookedPath;
+	auto BuildUncookedPath = 
+		[&OutUncookedPath](const FString& CookedPath, const FString& CookedRoot, const FString& UncookedRoot)
+	{
+		OutUncookedPath.AppendChars(*UncookedRoot, UncookedRoot.Len());
+		OutUncookedPath.AppendChars(*CookedPath + CookedRoot.Len(), CookedPath.Len() - CookedRoot.Len());
+	};
 
-	if (CookedPath.StartsWith(SandboxPath))
+	if (CookedPath.StartsWith(SandboxRootDir))
 	{
-		UncookedPath = CookedPath.Replace(*SandboxPath, *FPaths::GetRelativePathToRoot());
+		// Optimized CookedPath.StartsWith(SandboxProjectDir) that does not compare all of SandboxRootDir again
+		if (CookedPath.Len() >= SandboxProjectDir.Len() && 
+			0 == FCString::Strnicmp(
+				*CookedPath + SandboxRootDir.Len(),
+				*SandboxProjectDir + SandboxRootDir.Len(),
+				SandboxProjectDir.Len() - SandboxRootDir.Len()))
+		{
+			BuildUncookedPath(CookedPath, SandboxProjectDir, RelativeProjectDir);
+		}
+		else
+		{
+			BuildUncookedPath(CookedPath, SandboxRootDir, RelativeRootDir);
+		}
 	}
 	else
 	{
 		FString FullCookedFilename = FPaths::ConvertRelativePathToFull(CookedPath);
-		UncookedPath = FullCookedFilename.Replace(*SandboxPath, *FPaths::GetRelativePathToRoot());
+		BuildUncookedPath(FullCookedFilename, SandboxRootDir, RelativeRootDir);
 	}
 
-	return FName(*UncookedPath);
+	return FName(*OutUncookedPath);
 }
 
-void UCookOnTheFlyServer::GetAllCookedFiles(TMap<FName, FName>& UncookedPathToCookedPath, const FString& SandboxPath)
+void UCookOnTheFlyServer::GetAllCookedFiles(TMap<FName, FName>& UncookedPathToCookedPath, const FString& SandboxRootDir)
 {
-	IPlatformFile& PlatformFile = FPlatformFileManager::Get().GetPlatformFile();
 	TArray<FString> CookedFiles;
-	FPackageSearchVisitor PackageSearch(CookedFiles);
-	PlatformFile.IterateDirectoryRecursively(*SandboxPath, PackageSearch);
+	{
+		IPlatformFile& PlatformFile = FPlatformFileManager::Get().GetPlatformFile();
+		FPackageSearchVisitor PackageSearch(CookedFiles);
+		PlatformFile.IterateDirectoryRecursively(*SandboxRootDir, PackageSearch);
+	}
+
+	const FString SandboxProjectDir = FPaths::Combine(*SandboxRootDir, FApp::GetProjectName()) + TEXT("/");
+	const FString RelativeRootDir = FPaths::GetRelativePathToRoot();
+	const FString RelativeProjectDir = FPaths::ProjectDir();
+	FString UncookedFilename;
+	UncookedFilename.Reserve(1024);
+
 	for (const FString& CookedFile : CookedFiles)
 	{
 		const FName CookedFName(*CookedFile);
-		const FName UncookedFName = ConvertCookedPathToUncookedPath(SandboxPath, CookedFile);
+		const FName UncookedFName = ConvertCookedPathToUncookedPath(
+			SandboxRootDir, RelativeRootDir,
+			SandboxProjectDir, RelativeProjectDir,
+			CookedFile, UncookedFilename);
+
 		UncookedPathToCookedPath.Add(UncookedFName, CookedFName);
 	}
 }
