Fix cook -iterate bug with incorrect disk sizes and TagsAndValues in cooked asset registries

diff --git a/Engine/Source/Editor/UnrealEd/Private/Commandlets/AssetRegistryGenerator.cpp b/Engine/Source/Editor/UnrealEd/Private/Commandlets/AssetRegistryGenerator.cpp
--- a/Engine/Source/Editor/UnrealEd/Private/Commandlets/AssetRegistryGenerator.cpp
+++ b/Engine/Source/Editor/UnrealEd/Private/Commandlets/AssetRegistryGenerator.cpp
@@ -624,6 +624,45 @@ FAssetPackageData* FAssetRegistryGenerator::GetAssetPackageData(const FName& Pac
 	return State.CreateOrGetAssetPackageData(PackageName);
 }
 
+void FAssetRegistryGenerator::UpdateKeptPackagesDiskData(const TArray<FName>& InKeptPackages)
+{
+	for (const FName& PackageName : InKeptPackages)
+	{
+		// Get mutable PackageData without creating it when it does not exist
+		FAssetPackageData* PackageData =
+			State.GetAssetPackageData(PackageName) ?
+			State.CreateOrGetAssetPackageData(PackageName) :
+			nullptr;
+		const FAssetPackageData* PreviousPackageData =
+			PreviousState.GetAssetPackageData(PackageName);
+
+		if (!PackageData || !PreviousPackageData)
+		{
+			continue;
+		}
+
+		if (PackageData->PackageGuid != PreviousPackageData->PackageGuid)
+		{
+			continue;
+		}
+
+		PackageData->CookedHash = PreviousPackageData->CookedHash;
+		PackageData->DiskSize = PreviousPackageData->DiskSize;
+	}
+}
+
+void FAssetRegistryGenerator::UpdateKeptPackagesAssetData()
+{
+	for (FName PackageName : KeptPackages)
+	{
+		const TArray<const FAssetData*>& PreviousAssetDatas = PreviousState.GetAssetsByPackageName(PackageName);
+		for (int I = 0; I < PreviousAssetDatas.Num(); ++I)
+		{
+			State.UpdateAssetData(*PreviousAssetDatas[I]);
+		}
+	}
+}
+
 void FAssetRegistryGenerator::Initialize(const TArray<FName> &InStartupPackages)
 {
 	StartupPackages.Append(InStartupPackages);
@@ -730,6 +769,14 @@ void FAssetRegistryGenerator::ComputePackageDifferences(TSet<FName>& ModifiedPac
 	}
 }
 
+void FAssetRegistryGenerator::UpdateKeptPackages(const TArray<FName>& InKeptPackages)
+{
+	KeptPackages.Append(InKeptPackages);
+	// Update disk data right away, disk data is only updated when packages are saved, and kept packages are never saved
+	UpdateKeptPackagesDiskData(InKeptPackages);
+	// Delay update of AssetData with TagsAndValues, this data may be modified up until serialization in SaveAssetRegistry
+}
+
 void FAssetRegistryGenerator::BuildChunkManifest(const TSet<FName>& InCookedPackages, const TSet<FName>& InDevelopmentOnlyPackages, FSandboxPlatformFile* InSandboxFile, bool bGenerateStreamingInstallManifest)
 {
 	// If we were asked to generate a streaming install manifest explicitly and we did not have bGenerateNoChunks set, we will generate chunks.
@@ -948,10 +995,11 @@ bool FAssetRegistryGenerator::SaveAssetRegistry(const FString& SandboxPath, bool
 		SaveOptions.DisableFilters();
 	}
 
-	// Flush the asset registry and make sure the asset data is in sync, as it may have been updated during cook
+	// First flush the asset registry and make sure the asset data is in sync, as it may have been updated during cook
 	AssetRegistry.Tick(-1.0f);
-
 	AssetRegistry.InitializeTemporaryAssetRegistryState(State, SaveOptions, true);
+	// Then possibly apply AssetData with TagsAndValues from a previous AssetRegistry for packages kept from a previous cook
+	UpdateKeptPackagesAssetData();
 
 	if (DevelopmentSaveOptions.bSerializeAssetRegistry && bSerializeDevelopmentAssetRegistry)
 	{
diff --git a/Engine/Source/Editor/UnrealEd/Private/CookOnTheFlyServer.cpp b/Engine/Source/Editor/UnrealEd/Private/CookOnTheFlyServer.cpp
--- a/Engine/Source/Editor/UnrealEd/Private/CookOnTheFlyServer.cpp
+++ b/Engine/Source/Editor/UnrealEd/Private/CookOnTheFlyServer.cpp
@@ -5194,6 +5194,8 @@ void UCookOnTheFlyServer::PopulateCookedPackagesFromDisk(const TArray<ITargetPla
 		uint32 NumMarkedFailedSaveKept = 0;
 		uint32 NumPackagesRemoved = 0;
 
+		TArray<FName> KeptPackages;
+
 		for (const auto& CookedPaths : UncookedPathToCookedPath)
 		{
 			const FName CookedFile = CookedPaths.Value;
@@ -5239,6 +5241,7 @@ void UCookOnTheFlyServer::PopulateCookedPackagesFromDisk(const TArray<ITargetPla
 				if (IdenticalCookedPackages.Contains(SourcePackageName))
 				{
 					PackageTracker->CookedPackages.Add(FFilePlatformCookedPackage(UncookedFilename, MoveTemp(PlatformNames), MoveTemp(Succeeded)));
+					KeptPackages.Add(SourcePackageName);
 					++NumPackagesKept;
 				}
 			}
@@ -5277,9 +5280,12 @@ void UCookOnTheFlyServer::PopulateCookedPackagesFromDisk(const TArray<ITargetPla
 			ensure(PackageTracker->CookedPackages.Exists(UncookedFilename, PlatformNames, false) == false);
 
 			PackageTracker->CookedPackages.Add(FFilePlatformCookedPackage(UncookedFilename, MoveTemp(PlatformNames)));
+			KeptPackages.Add(UncookedPackage);
 			++NumMarkedFailedSaveKept;
 		}
 
+		PlatformAssetRegistry->UpdateKeptPackages(KeptPackages);
+
 		UE_LOG(LogCook, Display, TEXT("Iterative cooking summary for %s, \nConsidered: %d, \nFile Hash missmatch: %d, \nPackages Kept: %d, \nPackages failed save kept: %d, \nMissing Cooked Info(expected 0): %d"),
 			*Target->PlatformName(),
 			NumPackagesConsidered, NumPackagesFileHashMismatch,
diff --git a/Engine/Source/Editor/UnrealEd/Public/Commandlets/AssetRegistryGenerator.h b/Engine/Source/Editor/UnrealEd/Public/Commandlets/AssetRegistryGenerator.h
--- a/Engine/Source/Editor/UnrealEd/Public/Commandlets/AssetRegistryGenerator.h
+++ b/Engine/Source/Editor/UnrealEd/Public/Commandlets/AssetRegistryGenerator.h
@@ -56,6 +56,12 @@ public:
 	 */
 	void ComputePackageDifferences(TSet<FName>& ModifiedPackages, TSet<FName>& NewPackages, TSet<FName>& RemovedPackages, TSet<FName>& IdenticalCookedPackages, TSet<FName>& IdenticalUncookedPackages, bool bRecurseModifications, bool bRecurseScriptModifications);
 
+	/**
+	 * Tracks packages that were kept from a previous cook.
+	 * Updates the current asset registry from the previous one for all kept packages.
+	 */
+	void UpdateKeptPackages(const TArray<FName>& InKeptPackages);
+
 	/**
 	 * GenerateChunkManifest 
 	 * generate chunk manifest for the packages passed in using the asset registry to determine dependencies
@@ -160,6 +166,8 @@ private:
 	TSet<FName> CookedPackages;
 	/** List of packages that were filtered out from cooking */
 	TSet<FName> DevelopmentOnlyPackages;
+	/** List of packages that were kept from a previous cook */
+	TArray<FName> KeptPackages;
 	/** Map of Package name to Sandbox Paths */
 	typedef TMap<FName, FString> FChunkPackageSet;
 	/** Holds a reference to asset registry */
@@ -215,6 +223,17 @@ private:
 		uint32		ParentNodeIndex;
 	};
 
+	/**
+	 * Updates disk data with CookedHash and DiskSize from previous asset registry
+	 * for all packages kept from a previous cook.
+	 */
+	void UpdateKeptPackagesDiskData(const TArray<FName>& InKeptPackages);
+
+	/**
+	 * Updates AssetData with TagsAndValues from previous asset registry
+	 * for all packages kept from a previous cook.
+	 */
+	void UpdateKeptPackagesAssetData();
 
 	/**
 	 * Adds a package to chunk manifest
diff --git a/Engine/Source/Runtime/AssetRegistry/Private/AssetRegistryState.cpp b/Engine/Source/Runtime/AssetRegistry/Private/AssetRegistryState.cpp
--- a/Engine/Source/Runtime/AssetRegistry/Private/AssetRegistryState.cpp
+++ b/Engine/Source/Runtime/AssetRegistry/Private/AssetRegistryState.cpp
@@ -1315,6 +1315,14 @@ void FAssetRegistryState::AddAssetData(FAssetData* AssetData)
 	}
 }
 
+void FAssetRegistryState::UpdateAssetData(const FAssetData& NewAssetData)
+{
+	if (FAssetData* AssetData = CachedAssetsByObjectPath.FindRef(NewAssetData.ObjectPath))
+	{
+		UpdateAssetData(AssetData, NewAssetData);
+	}
+}
+
 void FAssetRegistryState::UpdateAssetData(FAssetData* AssetData, const FAssetData& NewAssetData)
 {
 	// Determine if tags need to be remapped
diff --git a/Engine/Source/Runtime/AssetRegistry/Public/AssetRegistryState.h b/Engine/Source/Runtime/AssetRegistry/Public/AssetRegistryState.h
--- a/Engine/Source/Runtime/AssetRegistry/Public/AssetRegistryState.h
+++ b/Engine/Source/Runtime/AssetRegistry/Public/AssetRegistryState.h
@@ -240,6 +240,9 @@ public:
 	/** Adds the asset data to the lookup maps */
 	void AddAssetData(FAssetData* AssetData);
 
+	/** Finds an existing asset data based on object path and updates it with the new value and updates lookup maps */
+	void UpdateAssetData(const FAssetData& NewAssetData);
+
 	/** Updates an existing asset data with the new value and updates lookup maps */
 	void UpdateAssetData(FAssetData* AssetData, const FAssetData& NewAssetData);
 
