Fix the IWCE component naming when duplicating to be consistent with the fixes made to SCS component naming

diff --git a/Engine/Source/Editor/UnrealEd/Private/Kismet2/ComponentEditorUtils.cpp b/Engine/Source/Editor/UnrealEd/Private/Kismet2/ComponentEditorUtils.cpp
--- a/Engine/Source/Editor/UnrealEd/Private/Kismet2/ComponentEditorUtils.cpp
+++ b/Engine/Source/Editor/UnrealEd/Private/Kismet2/ComponentEditorUtils.cpp
@@ -266,18 +266,39 @@ FString FComponentEditorUtils::GenerateValidVariableNameFromAsset(UObject* Asset
 	}
 
 	// Try to create a name without any numerical suffix first
-	FString ComponentInstanceName = AssetName;
+	FString NewName = AssetName;
+
+	auto BuildNewName = [&Counter, &AssetName]()
+	{
+		return FString::Printf(TEXT("%s%d"), *AssetName, Counter++);
+	};
 
 	if (ComponentOwner)
 	{
-		while (!IsComponentNameAvailable(ComponentInstanceName, ComponentOwner))
+		// If a desired name is supplied then walk back and find any numeric suffix so we can increment it nicely
+		int32 Index = AssetName.Len();
+		while (Index > 0 && AssetName[Index - 1] >= '0' && AssetName[Index - 1] <= '9')
+		{
+			--Index;
+		}
+
+		if (Index < AssetName.Len())
 		{
-			// Assign the lowest possible numerical suffix
-			ComponentInstanceName = FString::Printf(TEXT("%s%d"), *AssetName, Counter++);
+			FString NumericSuffix = AssetName.RightChop(Index);
+			Counter = FCString::Atoi(*NumericSuffix);
+			NumericSuffix = FString::Printf(TEXT("%d"), Counter); // Restringify the counter to account for leading 0s that we don't want to remove
+			AssetName.RemoveAt(AssetName.Len() - NumericSuffix.Len(), NumericSuffix.Len(), false);
+			++Counter;
+			NewName = BuildNewName();
+		}
+
+		while (!IsComponentNameAvailable(NewName, ComponentOwner))
+		{
+			NewName = BuildNewName();
 		}
 	}
 
-	return ComponentInstanceName;
+	return NewName;
 }
 
 USceneComponent* FComponentEditorUtils::FindClosestParentInList(UActorComponent* ChildComponent, const TArray<UActorComponent*>& ComponentList)
