#UE4 Downgrade warning/error summary lines from compile messages to Notes from warnings. The warnings/errors themselves raise the issues well enough and this message is more for information.

diff --git a/Engine/Source/Editor/UnrealEd/Private/Kismet2/CompilerResultsLog.cpp b/Engine/Source/Editor/UnrealEd/Private/Kismet2/CompilerResultsLog.cpp
--- a/Engine/Source/Editor/UnrealEd/Private/Kismet2/CompilerResultsLog.cpp
+++ b/Engine/Source/Editor/UnrealEd/Private/Kismet2/CompilerResultsLog.cpp
@@ -241,12 +241,12 @@ void FCompilerResultsLog::InternalLogSummary()
 		{
 			Args.Add(NumErrors); // num errors {4}
 			Args.Add(NumWarnings); // num warnings {5}
-			Warning(*FText::Format(LOCTEXT("CompileFailed", "[{0}] Compile of {1} failed. {4} Fatal Issue(s) {5} Warning(s) [in {3} ms] ({2})"), MoveTemp(Args)).ToString());
+			Note(*FText::Format(LOCTEXT("CompileFailed", "[{0}] Compile of {1} failed. {4} Fatal Issue(s) {5} Warning(s) [in {3} ms] ({2})"), MoveTemp(Args)).ToString());
 		}
 		else if(NumWarnings > 0)
 		{
 			Args.Add(NumWarnings); // num warnings {4}
-			Warning(*FText::Format(LOCTEXT("CompileWarning", "[{0}] Compile of {1} successful, but with {4} Warning(s) [in {3} ms] ({2})"), MoveTemp(Args)).ToString());
+			Note(*FText::Format(LOCTEXT("CompileWarning", "[{0}] Compile of {1} successful, but with {4} Warning(s) [in {3} ms] ({2})"), MoveTemp(Args)).ToString());
 		}
 		else
 		{
