Clarify the failure logging in DiffAssetRegistriesCommandlet to properly indicate whether an old/newpath parameter has not ben provided, or if the path provided failed to resolve properly to an AssetRegistry.bin file.  Prior to this CL, both cases would give an error indicating the former.

diff --git a/Engine/Source/Editor/UnrealEd/Private/Commandlets/DiffAssetRegistriesCommandlet.cpp b/Engine/Source/Editor/UnrealEd/Private/Commandlets/DiffAssetRegistriesCommandlet.cpp
--- a/Engine/Source/Editor/UnrealEd/Private/Commandlets/DiffAssetRegistriesCommandlet.cpp
+++ b/Engine/Source/Editor/UnrealEd/Private/Commandlets/DiffAssetRegistriesCommandlet.cpp
@@ -191,6 +191,11 @@ int32 UDiffAssetRegistriesCommandlet::Main(const FString& FullCommandLine)
 	{
 		FindAssetRegistryPath(*OldPathVal, OldPath);
 	}
+	else
+	{
+		UE_LOG(LogDiffAssets, Error, TEXT("No old path specified \"-oldpath=<>\", use full path to asset registry or build version."));
+		return -1;
+	}
 
 	const FString* NewPathVal = Params.Find(FString(TEXT("NewPath")));
 	if (NewPathVal)
@@ -208,6 +213,11 @@ int32 UDiffAssetRegistriesCommandlet::Main(const FString& FullCommandLine)
 			}
 		}
 	}
+	else
+	{
+		UE_LOG(LogDiffAssets, Error, TEXT("No new path specified \"-newpath=<>\", use full path to asset registry or build version."));
+		return -1;
+	}
 
 	bMatchChangelists = false;
 	if (FParse::Value(*FullCommandLine, TEXT("Branch="), Spec))
@@ -230,12 +240,12 @@ int32 UDiffAssetRegistriesCommandlet::Main(const FString& FullCommandLine)
 
 	if (OldPath.IsEmpty())
 	{
-		UE_LOG(LogDiffAssets, Error, TEXT("No old path specified \"-oldpath=<>\", use full path to asset registry or build version."));
+		UE_LOG(LogDiffAssets, Error, TEXT("Unable to locate AssetRegistry.bin for supplied oldpath (%s), use full path to asset registry or build version."), **OldPathVal);
 		return -1;
 	}
 	if (NewPath.IsEmpty())
 	{
-		UE_LOG(LogDiffAssets, Error, TEXT("No new path specified \"-newpath=<>\", use full path to asset registry or build version."));
+		UE_LOG(LogDiffAssets, Error, TEXT("Unable to locate AssetRegistry.bin for supplied newpath (%s), use full path to asset registry or build version."), **NewPathVal);
 		return -1;
 	}
 
