Fix crash introduced in CL# 4417943

diff --git a/Engine/Source/Editor/UnrealEd/Private/LevelEditorViewport.cpp b/Engine/Source/Editor/UnrealEd/Private/LevelEditorViewport.cpp
--- a/Engine/Source/Editor/UnrealEd/Private/LevelEditorViewport.cpp
+++ b/Engine/Source/Editor/UnrealEd/Private/LevelEditorViewport.cpp
@@ -3386,7 +3386,7 @@ UActorComponent* FLevelEditorViewportClient::FindViewComponentForActor(AActor co
 		for (UActorComponent* Comp : Comps)
 		{
 			FMinimalViewInfo DummyViewInfo;
-			if (Comp->bIsActive && Comp->GetEditorPreviewInfo(/*DeltaTime =*/0.0f, DummyViewInfo))
+			if (Comp && Comp->bIsActive && Comp->GetEditorPreviewInfo(/*DeltaTime =*/0.0f, DummyViewInfo))
 			{
 				if (Comp->IsSelected())
 				{
