Warn the user when trying to add a content pack if Live Coding is enabled.

diff --git a/Engine/Source/Editor/UnrealEd/Private/Factories/PackFactory.cpp b/Engine/Source/Editor/UnrealEd/Private/Factories/PackFactory.cpp
--- a/Engine/Source/Editor/UnrealEd/Private/Factories/PackFactory.cpp
+++ b/Engine/Source/Editor/UnrealEd/Private/Factories/PackFactory.cpp
@@ -36,6 +36,10 @@
 #include "Logging/MessageLog.h"
 #include "Misc/CoreDelegates.h"
 
+#if WITH_LIVE_CODING
+#include "ILiveCodingModule.h"
+#endif
+
 DEFINE_LOG_CATEGORY_STATIC(LogPackFactory, Log, All);
 
 UPackFactory::UPackFactory(const FObjectInitializer& PCIP)
@@ -685,7 +689,19 @@ UObject* UPackFactory::FactoryCreateBinary
 					SOutputLogDialog::Open(NSLOCTEXT("PackFactory", "CreateBinary", "Create binary"), FailReason, FailLog, FText::GetEmpty());
 				}
 
-				if (ConfigParameters.bCompileSource)
+				bool bCompileSource = ConfigParameters.bCompileSource;
+#if WITH_LIVE_CODING
+				if (bCompileSource)
+				{
+					ILiveCodingModule* LiveCoding = FModuleManager::GetModulePtr<ILiveCodingModule>(LIVE_CODING_MODULE_NAME);
+					if (LiveCoding != nullptr && LiveCoding->IsEnabledForSession())
+					{
+						FMessageDialog::Open(EAppMsgType::Ok, NSLOCTEXT("PackFactory", "FailedToCompileNewGameModule", "Unable to compile source code while Live Coding is enabled. Please close the editor and build from your IDE."));
+						bCompileSource = false;
+					}
+				}
+#endif
+				if (bCompileSource)
 				{
 					// Compile the new code, either using the in editor hot-reload (if an existing module), or as a brand new module (if no existing code)
 					IHotReloadInterface& HotReloadSupport = FModuleManager::LoadModuleChecked<IHotReloadInterface>("HotReload");
diff --git a/Engine/Source/Editor/UnrealEd/UnrealEd.Build.cs b/Engine/Source/Editor/UnrealEd/UnrealEd.Build.cs
--- a/Engine/Source/Editor/UnrealEd/UnrealEd.Build.cs
+++ b/Engine/Source/Editor/UnrealEd/UnrealEd.Build.cs
@@ -323,5 +323,10 @@ public class UnrealEd : ModuleRules
 		{
 			PublicDefinitions.Add( "WITH_RECAST=0" );
 		}
-	}
+
+        if (Target.bWithLiveCoding)
+        {
+			PrivateIncludePathModuleNames.Add("LiveCoding");
+        }
+    }
 }
