Copying Dev-Audio to UE-Main from CL 6717121

diff --git a/Engine/Source/Editor/UnrealEd/Private/EditorViewportClient.cpp b/Engine/Source/Editor/UnrealEd/Private/EditorViewportClient.cpp
--- a/Engine/Source/Editor/UnrealEd/Private/EditorViewportClient.cpp
+++ b/Engine/Source/Editor/UnrealEd/Private/EditorViewportClient.cpp
@@ -16,6 +16,7 @@
 #include "Settings/LevelEditorMiscSettings.h"
 #include "Components/DirectionalLightComponent.h"
 #include "Components/BillboardComponent.h"
+#include "Audio/AudioDebug.h"
 #include "Debug/DebugDrawService.h"
 #include "EngineUtils.h"
 #include "Editor.h"
@@ -206,7 +207,7 @@ bool FViewportCameraTransform::UpdateTransition()
 			ViewLocation = FMath::Lerp( StartLocation, DesiredLocation, LerpWeight );
 		}
 
-		
+
 		bIsAnimating = true;
 	}
 
@@ -216,8 +217,8 @@ bool FViewportCameraTransform::UpdateTransition()
 FMatrix FViewportCameraTransform::ComputeOrbitMatrix() const
 {
 	FTransform Transform =
-	FTransform( -LookAt ) * 
-	FTransform( FRotator(0,ViewRotation.Yaw,0) ) * 
+	FTransform( -LookAt ) *
+	FTransform( FRotator(0,ViewRotation.Yaw,0) ) *
 	FTransform( FRotator(0, 0, ViewRotation.Pitch) ) *
 	FTransform( FVector(0,(ViewLocation - LookAt).Size(), 0) );
 
@@ -519,7 +520,7 @@ void FEditorViewportClient::RestoreRealtime(const bool bAllowDisable)
 	{
 		bIsRealtime |= bStoredRealtime;
 		bShowStats |= bStoredShowStats;
-	}	
+	}
 
 	if (bIsRealtime)
 	{
@@ -640,7 +641,7 @@ void FEditorViewportClient::ToggleOrbitCamera( bool bEnableOrbitCamera )
 				ViewTransform.SetRotation(OrbitMatrix.Rotator());
 			}
 		}
-		
+
 		ViewTransform.SetLocation(OrbitMatrix.GetOrigin());
 	}
 }
@@ -665,7 +666,7 @@ void FEditorViewportClient::FocusViewportOnBox( const FBox& BoundingBox, bool bI
 
 		if(!IsOrtho())
 		{
-		   /** 
+		   /**
 			* We need to make sure we are fitting the sphere into the viewport completely, so if the height of the viewport is less
 			* than the width of the viewport, we scale the radius by the aspect ratio in order to compensate for the fact that we have
 			* less visible vertically than horizontally.
@@ -675,7 +676,7 @@ void FEditorViewportClient::FocusViewportOnBox( const FBox& BoundingBox, bool bI
 				Radius *= AspectToUse;
 			}
 
-			/** 
+			/**
 			 * Now that we have a adjusted radius, we are taking half of the viewport's FOV,
 			 * converting it to radians, and then figuring out the camera's distance from the center
 			 * of the bounding sphere using some simple trig.  Once we have the distance, we back up
@@ -698,7 +699,7 @@ void FEditorViewportClient::FocusViewportOnBox( const FBox& BoundingBox, bool bI
 
 			if( !(Viewport->KeyState(EKeys::LeftControl) || Viewport->KeyState(EKeys::RightControl)) )
 			{
-				/** 			
+				/**
 				* We also need to zoom out till the entire volume is in view.  The following block of code first finds the minimum dimension
 				* size of the viewport.  It then calculates backwards from what the view size should be (The radius of the bounding volume),
 				* to find the new OrthoZoom value for the viewport. The 15.0f is a fudge factor.
@@ -720,7 +721,7 @@ void FEditorViewportClient::FocusViewportOnBox( const FBox& BoundingBox, bool bI
 
 //////////////////////////////////////////////////////////////////////////
 //
-// Configures the specified FSceneView object with the view and projection matrices for this viewport. 
+// Configures the specified FSceneView object with the view and projection matrices for this viewport.
 
 FSceneView* FEditorViewportClient::CalcSceneView(FSceneViewFamily* ViewFamily, const EStereoscopicPass StereoPass)
 {
@@ -822,7 +823,7 @@ FSceneView* FEditorViewportClient::CalcSceneView(FSceneViewFamily* ViewFamily, c
 			    GEngine->StereoRenderingDevice->AdjustViewRect( StereoPass, X, Y, SizeX, SizeY );
 		        const FIntRect StereoViewRect = FIntRect( X, Y, X + SizeX, Y + SizeY );
 		        ViewInitOptions.SetViewRectangle( StereoViewRect );
-			
+
 				GEngine->StereoRenderingDevice->CalculateStereoViewOffset( StereoPass, ViewRotation, ViewInitOptions.WorldToMetersScale, ViewInitOptions.ViewOrigin );
 			}
 
@@ -835,7 +836,7 @@ FSceneView* FEditorViewportClient::CalcSceneView(FSceneViewFamily* ViewFamily, c
 				FPlane(1, 0, 0, 0),
 				FPlane(0, 1, 0, 0),
 				FPlane(0, 0, 0, 1));
-    
+
 		    if( bStereoRendering )
 		    {
 			    // @todo vreditor: bConstrainAspectRatio is ignored in this path, as it is in the game client as well currently
@@ -848,7 +849,7 @@ FSceneView* FEditorViewportClient::CalcSceneView(FSceneViewFamily* ViewFamily, c
 			    const float MaxZ = MinZ;
 			    // Avoid zero ViewFOV's which cause divide by zero's in projection matrix
 			    const float MatrixFOV = FMath::Max(0.001f, ViewFOV) * (float)PI / 360.0f;
-    
+
 			    if (bConstrainAspectRatio)
 			    {
 				    if ((bool)ERHIZBuffer::IsInverted)
@@ -878,7 +879,7 @@ FSceneView* FEditorViewportClient::CalcSceneView(FSceneViewFamily* ViewFamily, c
 			    {
 				    float XAxisMultiplier;
 				    float YAxisMultiplier;
-    
+
 				    if (((ViewportSize.X > ViewportSize.Y) && (AspectRatioAxisConstraint == AspectRatio_MajorAxisFOV)) || (AspectRatioAxisConstraint == AspectRatio_MaintainXFOV))
 				    {
 					    //if the viewport is wider than it is tall
@@ -891,7 +892,7 @@ FSceneView* FEditorViewportClient::CalcSceneView(FSceneViewFamily* ViewFamily, c
 					    XAxisMultiplier = ViewportSize.Y / (float)ViewportSize.X;
 					    YAxisMultiplier = 1.0f;
 				    }
-    
+
 				    if ((bool)ERHIZBuffer::IsInverted)
 				    {
 					    ViewInitOptions.ProjectionMatrix = FReversedZPerspectiveMatrix(
@@ -1010,7 +1011,7 @@ FSceneView* FEditorViewportClient::CalcSceneView(FSceneViewFamily* ViewFamily, c
 	if (bStereoRendering)
 	{
 		if (StereoViewStates.Num() <= ViewStateIndex)
-		{ 
+		{
 			StereoViewStates.SetNum(ViewStateIndex + 1);
 		}
 
@@ -1023,7 +1024,7 @@ FSceneView* FEditorViewportClient::CalcSceneView(FSceneViewFamily* ViewFamily, c
 	ViewInitOptions.ViewFamily = ViewFamily;
 	ViewInitOptions.SceneViewStateInterface = ( (StereoPass < eSSP_RIGHT_EYE) ? ViewState.GetReference() : StereoViewStates[ViewStateIndex].GetReference() );
 	ViewInitOptions.StereoPass = StereoPass;
-	
+
 	ViewInitOptions.ViewElementDrawer = this;
 
 	ViewInitOptions.BackgroundColor = GetBackgroundColor();
@@ -1081,7 +1082,7 @@ FSceneView* FEditorViewportClient::CalcSceneView(FSceneViewFamily* ViewFamily, c
 	return View;
 }
 
-/** Determines if the new MoveCanvas movement should be used 
+/** Determines if the new MoveCanvas movement should be used
  * @return - true if we should use the new drag canvas movement.  Returns false for combined object-camera movement and marquee selection
  */
 bool FLevelEditorViewportClient::ShouldUseMoveCanvasMovement()
@@ -1105,7 +1106,7 @@ bool FLevelEditorViewportClient::ShouldUseMoveCanvasMovement()
 		}
 
 		//OBJECT MOVEMENT CODE
-		if ( ( AltDown == false && ShiftDown == false && ( LeftMouseButtonDown ^ RightMouseButtonDown ) ) && 
+		if ( ( AltDown == false && ShiftDown == false && ( LeftMouseButtonDown ^ RightMouseButtonDown ) ) &&
 			( ( GetWidgetMode() == FWidget::WM_Translate && Widget->GetCurrentAxis() != EAxisList::None ) ||
 			( GetWidgetMode() == FWidget::WM_TranslateRotateZ && Widget->GetCurrentAxis() != EAxisList::ZRotation &&  Widget->GetCurrentAxis() != EAxisList::None ) ||
 			( GetWidgetMode() == FWidget::WM_2D && Widget->GetCurrentAxis() != EAxisList::Rotate2D &&  Widget->GetCurrentAxis() != EAxisList::None ) ) )
@@ -1173,7 +1174,7 @@ void FEditorViewportClient::Tick(float DeltaTime)
 		UpdateCameraMovement( DeltaTime );
 
 		UpdateMouseDelta();
-		
+
 		UpdateGestureDelta();
 
 		EndCameraMovement();
@@ -1198,8 +1199,8 @@ void FEditorViewportClient::Tick(float DeltaTime)
 
 	if ( bIsTracking )
 	{
-		// If a mouse button or modifier is pressed we want to assume the user is still in a mode 
-		// they haven't left to perform a non-action in the frame to keep the last used operation 
+		// If a mouse button or modifier is pressed we want to assume the user is still in a mode
+		// they haven't left to perform a non-action in the frame to keep the last used operation
 		// from being reset.
 		const bool LeftMouseButtonDown = Viewport->KeyState(EKeys::LeftMouseButton) ? true : false;
 		const bool MiddleMouseButtonDown = Viewport->KeyState(EKeys::MiddleMouseButton) ? true : false;
@@ -1417,7 +1418,7 @@ EMouseCursor::Type FEditorViewportClient::GetCursor(FViewport* InViewport,int32
 	}
 
 	bool bMoveCanvasMovement = ShouldUseMoveCanvasMovement();
-	
+
 	if (RequiredCursorVisibiltyAndAppearance.bOverrideAppearance &&
 		RequiredCursorVisibiltyAndAppearance.bHardwareCursorVisible)
 	{
@@ -1437,7 +1438,7 @@ EMouseCursor::Type FEditorViewportClient::GetCursor(FViewport* InViewport,int32
 		//We're grabbing the canvas so the icon should look "grippy"
 		MouseCursor = EMouseCursor::GrabHandClosed;
 	}
-	else if (bMoveCanvasMovement && 
+	else if (bMoveCanvasMovement &&
 		bHasMouseMovedSinceClick &&
 		(GetWidgetMode() == FWidget::WM_Translate || GetWidgetMode() == FWidget::WM_TranslateRotateZ || GetWidgetMode() == FWidget::WM_2D))
 	{
@@ -1474,11 +1475,11 @@ EMouseCursor::Type FEditorViewportClient::GetCursor(FViewport* InViewport,int32
 				{
 					SetCurrentWidgetAxis( EAxisList::None );
 					Invalidate( false, false );
-				}			
+				}
 			}
 		}
 	}
-	
+
 	// Allow the viewport interaction to override any previously set mouse cursor
 	UViewportWorldInteraction* WorldInteraction = Cast<UViewportWorldInteraction>(GEditor->GetEditorWorldExtensionsManager()->GetEditorWorldExtensions(GetWorld())->FindExtension(UViewportWorldInteraction::StaticClass()));
 	if (WorldInteraction != nullptr)
@@ -1542,7 +1543,7 @@ void FEditorViewportClient::SetViewportType( ELevelViewportType InViewportType )
 }
 
 void FEditorViewportClient::RotateViewportType()
-{	
+{
 	ViewportType = ViewOptions[ViewOptionIndex];
 
 	// Changing the type may also change the active view mode; re-apply that now
@@ -1583,7 +1584,7 @@ void FEditorViewportClient::UpdateCameraMovement( float DeltaTime )
 		ToggleOrbitCamera(bEnable);
 
 		const bool bIsUsingTrackpad = FSlateApplication::Get().IsUsingTrackpad();
-		
+
 		// Certain keys are only available while the flight camera input mode is active
 		const bool bUsingFlightInput = IsFlightCameraInputModeActive() || bIsUsingTrackpad;
 
@@ -1612,12 +1613,12 @@ void FEditorViewportClient::UpdateCameraMovement( float DeltaTime )
 		CameraUserImpulseData->RotateRollImpulse = 0.0f;
 
 		GestureMoveForwardBackwardImpulse = 0.0f;
-		
+
 		bool bForwardKeyState = false;
 		bool bBackwardKeyState = false;
 		bool bRightKeyState = false;
 		bool bLeftKeyState = false;
-		
+
 		bool bUpKeyState = false;
 		bool bDownKeyState = false;
 		bool bZoomOutKeyState = false;
@@ -1630,7 +1631,7 @@ void FEditorViewportClient::UpdateCameraMovement( float DeltaTime )
 			bBackwardKeyState |= Viewport->KeyState(FViewportNavigationCommands::Get().Backward->GetActiveChord(ChordIndex)->Key);
 			bRightKeyState |= Viewport->KeyState(FViewportNavigationCommands::Get().Right->GetActiveChord(ChordIndex)->Key);
 			bLeftKeyState |= Viewport->KeyState(FViewportNavigationCommands::Get().Left->GetActiveChord(ChordIndex)->Key);
-			
+
 			bUpKeyState |= Viewport->KeyState(FViewportNavigationCommands::Get().Up->GetActiveChord(ChordIndex)->Key);
 			bDownKeyState |= Viewport->KeyState(FViewportNavigationCommands::Get().Down->GetActiveChord(ChordIndex)->Key);
 			bZoomOutKeyState |= Viewport->KeyState(FViewportNavigationCommands::Get().FovZoomOut->GetActiveChord(ChordIndex)->Key);
@@ -1638,7 +1639,7 @@ void FEditorViewportClient::UpdateCameraMovement( float DeltaTime )
 		}
 
 		// Forward/back
-		if( ( bRemapWASDKeys && bForwardKeyState ) || 
+		if( ( bRemapWASDKeys && bForwardKeyState ) ||
 			( bRemapArrowKeys && Viewport->KeyState( EKeys::Up ) ) ||
 			( bUnmodifiedPress && bUseNumpadCameraControl && Viewport->KeyState(EKeys::NumPadEight) ) )
 		{
@@ -1753,7 +1754,7 @@ void FEditorViewportClient::UpdateCameraMovement( float DeltaTime )
 
 #if UE_BUILD_DEBUG
 		// Editor movement is very difficult in debug without this, due to hitching
-		// It is better to freeze movement during a hitch than to fly off past where you wanted to go 
+		// It is better to freeze movement during a hitch than to fly off past where you wanted to go
 		// (considering there will be further hitching trying to get back to where you were)
 		EditorMovementDeltaUpperBound = .15f;
 #endif
@@ -1888,7 +1889,7 @@ bool FEditorViewportClient::CalculateEditorConstrainedViewRect(FSlateRect& OutSa
 	OutSafeFrameRect = FSlateRect(0, 0, SizeX, SizeY);
 	float FixedAspectRatio;
 	bool bSafeFrameActive = GetActiveSafeFrame(FixedAspectRatio);
-	
+
 	if (bSafeFrameActive)
 	{
 		// Get the size of the viewport
@@ -1951,7 +1952,7 @@ void FEditorViewportClient::DrawSafeFrames(FViewport& InViewport, FSceneView& Vi
 					LineItem.Draw(&Canvas, FVector2D(SafeRect.Left, 0), FVector2D(SafeRect.Left, SizeY));
 					LineItem.Draw(&Canvas, FVector2D(SafeRect.Right, 0), FVector2D(SafeRect.Right, SizeY));
 				}
-				
+
 				if (SafeRect.GetSize().Y < SizeY)
 				{
 					DrawSafeFrameQuad(Canvas, FVector2D(SafeRect.Left, 0), FVector2D(SafeRect.Right, SafeRect.Top));
@@ -1970,7 +1971,7 @@ void FEditorViewportClient::DrawSafeFrameQuad( FCanvas &Canvas, FVector2D V1, FV
 	FCanvasUVTri UVTriItem;
 	UVTriItem.V0_Pos = FVector2D(V1.X, V1.Y);
 	UVTriItem.V1_Pos = FVector2D(V2.X, V1.Y);
-	UVTriItem.V2_Pos = FVector2D(V1.X, V2.Y);		
+	UVTriItem.V2_Pos = FVector2D(V1.X, V2.Y);
 	FCanvasTriangleItem TriItem( UVTriItem, GWhiteTexture );
 	UVTriItem.V0_Pos = FVector2D(V2.X, V1.Y);
 	UVTriItem.V1_Pos = FVector2D(V2.X, V2.Y);
@@ -2088,7 +2089,7 @@ void FEditorViewportClient::UpdateMouseDelta()
 			const bool bIsUsingTrackpad = FSlateApplication::Get().IsUsingTrackpad();
 			const bool bIsNonOrbitMiddleMouse = MiddleMouseButtonDown && !IsAltPressed();
 
-			// If a tool is overriding current widget mode behavior, it may need to 
+			// If a tool is overriding current widget mode behavior, it may need to
 			// temporarily set a different widget mode while converting mouse movement.
 			ModeTools->PreConvertMouseMovement(this);
 
@@ -2121,7 +2122,7 @@ void FEditorViewportClient::UpdateMouseDelta()
 						DragStartView = CalcSceneView(DragStartViewFamily);
 					}
 					MouseDeltaTracker->AbsoluteTranslationConvertMouseToDragRot(DragStartView, this, Drag, Rot, Scale);
-				} 
+				}
 				else
 				{
 					MouseDeltaTracker->ConvertMovementDeltaToDragRot(this, DragDelta, Drag, Rot, Scale);
@@ -2131,7 +2132,7 @@ void FEditorViewportClient::UpdateMouseDelta()
 			ModeTools->PostConvertMouseMovement(this);
 
 			const bool bInputHandledByGizmos = InputWidgetDelta( Viewport, CurrentAxis, Drag, Rot, Scale );
-					
+
 			if( !Rot.IsZero() )
 			{
 				Widget->UpdateDeltaRotation();
@@ -2212,7 +2213,7 @@ static bool IsOrbitPanMode( FViewport* Viewport )
 }
 
 static bool IsOrbitZoomMode( FViewport* Viewport )
-{	
+{
 	bool	LeftMouseButton = Viewport->KeyState(EKeys::LeftMouseButton),
 		MiddleMouseButton = Viewport->KeyState(EKeys::MiddleMouseButton),
 		RightMouseButton = Viewport->KeyState(EKeys::RightMouseButton);
@@ -2485,7 +2486,7 @@ void FEditorViewportClient::ChangeBufferVisualizationMode( FName InName )
 
 bool FEditorViewportClient::IsBufferVisualizationModeSelected( FName InName ) const
 {
-	return IsViewModeEnabled( VMI_VisualizeBuffer ) && CurrentBufferVisualizationMode == InName;	
+	return IsViewModeEnabled( VMI_VisualizeBuffer ) && CurrentBufferVisualizationMode == InName;
 }
 
 void FEditorViewportClient::ChangeRayTracingDebugVisualizationMode(FName InName)
@@ -2647,10 +2648,10 @@ bool FEditorViewportClient::InputKey(FViewport* InViewport, int32 ControllerId,
 
 	if( bWasCursorVisible != RequiredCursorVisibiltyAndAppearance.bHardwareCursorVisible || bWasSoftwareCursorVisible != RequiredCursorVisibiltyAndAppearance.bSoftwareCursorVisible )
 	{
-		bHandled = true;	
+		bHandled = true;
 	}
 
-	
+
 	// Compute a view.
 	FSceneViewFamilyContext ViewFamily(FSceneViewFamily::ConstructionValues(
 		InViewport,
@@ -2667,15 +2668,15 @@ bool FEditorViewportClient::InputKey(FViewport* InViewport, int32 ControllerId,
 
 
 	// Start tracking if any mouse button is down and it was a tracking event (MouseButton/Ctrl/Shift/Alt):
-	if ( InputState.IsAnyMouseButtonDown() 
+	if ( InputState.IsAnyMouseButtonDown()
 		&& (Event == IE_Pressed || Event == IE_Released)
-		&& (InputState.IsMouseButtonEvent() || InputState.IsCtrlButtonEvent() || InputState.IsAltButtonEvent() || InputState.IsShiftButtonEvent() ) )		
+		&& (InputState.IsMouseButtonEvent() || InputState.IsCtrlButtonEvent() || InputState.IsAltButtonEvent() || InputState.IsShiftButtonEvent() ) )
 	{
 		StartTrackingDueToInput( InputState, *View );
 		return true;
 	}
 
-	
+
 	// If we are tracking and no mouse button is down and this input event released the mouse button stop tracking and process any clicks if necessary
 	if ( bIsTracking && !InputState.IsAnyMouseButtonDown() && InputState.IsMouseButtonEvent() )
 	{
@@ -2731,7 +2732,7 @@ bool FEditorViewportClient::InputKey(FViewport* InViewport, int32 ControllerId,
 		MouseDeltaTracker->SetExternalMovement();
 		bHandled |= true;
 	}
-	
+
 	//apply the visibility and set the cursor positions
 	ApplyRequiredCursorVisibility( true );
 	return bHandled;
@@ -2791,7 +2792,7 @@ void FEditorViewportClient::StopTracking()
  		HHitProxy* HitProxy = Viewport->GetHitProxy(CachedMouseX,CachedMouseY);
   		CheckHoveredHitProxy(HitProxy);
 
-		bIsTracking = false;	
+		bIsTracking = false;
 	}
 
 	bHasMouseMovedSinceClick = false;
@@ -2815,7 +2816,7 @@ void FEditorViewportClient::StartTrackingDueToInput( const struct FInputEventSta
 	FViewport* InputStateViewport = InputState.GetViewport();
 	FKey Key = InputState.GetKey();
 
-	bool bIsRedundantModifierEvent = 
+	bool bIsRedundantModifierEvent =
 		( InputState.IsAltButtonEvent() && ( ( Event != IE_Released ) == IsAltPressed() ) ) ||
 		( InputState.IsCtrlButtonEvent() && ( ( Event != IE_Released ) == IsCtrlPressed() ) ) ||
 		( InputState.IsShiftButtonEvent() && ( ( Event != IE_Released ) == IsShiftPressed() ) );
@@ -2855,7 +2856,7 @@ void FEditorViewportClient::StartTrackingDueToInput( const struct FInputEventSta
 		}
 
 		// Start new tracking. Potentially reset the widget so that StartTracking can pick a new axis.
-		if ( Widget && ( !bDraggingByHandle || InputState.IsCtrlButtonPressed() ) ) 
+		if ( Widget && ( !bDraggingByHandle || InputState.IsCtrlButtonPressed() ) )
 		{
 			bWidgetAxisControlledByDrag = false;
 			Widget->SetCurrentAxis( EAxisList::None );
@@ -2879,10 +2880,10 @@ void FEditorViewportClient::StartTrackingDueToInput( const struct FInputEventSta
 		//only reset the initial point when the mouse is actually clicked
 		if (InputState.IsAnyMouseButtonDown() && Widget)
 		{
-			Widget->ResetInitialTranslationOffset();				
+			Widget->ResetInitialTranslationOffset();
 		}
 
-		//Don't update the cursor visibility if we don't have focus or mouse capture 
+		//Don't update the cursor visibility if we don't have focus or mouse capture
 		if( InputStateViewport->HasFocus() ||  InputStateViewport->HasMouseCapture())
 		{
 			//Need to call this one more time as the axis variable for the widget has just been updated
@@ -2905,12 +2906,12 @@ void FEditorViewportClient::ProcessClickInViewport( const FInputEventState& Inpu
 
 		const int32	HitX = InputStateViewport->GetMouseX();
 		const int32	HitY = InputStateViewport->GetMouseY();
-		
+
 		// Calc the raw delta from the mouse to detect if there was any movement
 		FVector RawMouseDelta = MouseDeltaTracker->GetRawDelta();
 
-		// Note: We are using raw mouse movement to double check distance moved in low performance situations.  In low performance situations its possible 
-		// that we would get a mouse down and a mouse up before the next tick where GEditor->MouseMovment has not been updated.  
+		// Note: We are using raw mouse movement to double check distance moved in low performance situations.  In low performance situations its possible
+		// that we would get a mouse down and a mouse up before the next tick where GEditor->MouseMovment has not been updated.
 		// In that situation, legitimate drags are incorrectly considered clicks
 		bool bNoMouseMovment = RawMouseDelta.SizeSquared() < MOUSE_CLICK_DRAG_DELTA && GEditor->MouseMovement.SizeSquared() < MOUSE_CLICK_DRAG_DELTA;
 
@@ -2980,7 +2981,7 @@ void FEditorViewportClient::ProcessDoubleClickInViewport( const struct FInputEve
 }
 
 
-/** Determines if the new MoveCanvas movement should be used 
+/** Determines if the new MoveCanvas movement should be used
  * @return - true if we should use the new drag canvas movement.  Returns false for combined object-camera movement and marquee selection
  */
 bool FEditorViewportClient::ShouldUseMoveCanvasMovement() const
@@ -3004,7 +3005,7 @@ bool FEditorViewportClient::ShouldUseMoveCanvasMovement() const
 		}
 
 		//OBJECT MOVEMENT CODE
-		if ( ( AltDown == false && ShiftDown == false && ( LeftMouseButtonDown ^ RightMouseButtonDown ) ) && 
+		if ( ( AltDown == false && ShiftDown == false && ( LeftMouseButtonDown ^ RightMouseButtonDown ) ) &&
 			( ( GetWidgetMode() == FWidget::WM_Translate && Widget->GetCurrentAxis() != EAxisList::None ) ||
 			( GetWidgetMode() == FWidget::WM_TranslateRotateZ && Widget->GetCurrentAxis() != EAxisList::ZRotation &&  Widget->GetCurrentAxis() != EAxisList::None ) ||
 			( GetWidgetMode() == FWidget::WM_2D && Widget->GetCurrentAxis() != EAxisList::Rotate2D &&  Widget->GetCurrentAxis() != EAxisList::None ) ) )
@@ -3076,7 +3077,7 @@ void FEditorViewportClient::DrawAxes(FViewport* InViewport, FCanvas* Canvas, con
 		LineItem.Draw( Canvas, AxisOrigin, AxisEnd );
 		TextItem.Text = LOCTEXT("YAxis","Y");
 		TextItem.Draw( Canvas, FVector2D(AxisEnd.X + 2, AxisEnd.Y - 0.5*YL) );
-		
+
 	}
 
 	if( ( InAxis & EAxisList::Z ) == EAxisList::Z )
@@ -3152,7 +3153,7 @@ void FEditorViewportClient::DrawScaleUnits(FViewport* InViewport, FCanvas* Canva
 	const float SegmentWidthUnits = UnitsPerPixel > 0 ? FMath::Pow(10.f, FMath::RoundToFloat(FMath::LogX(10.f, UnitsPerPixel * ApproxTargetMarkerWidthPx))) : 0.f;
 
 	const FString DisplayText = UnrealUnitsToSiUnits(SegmentWidthUnits);
-	
+
 	UFont* Font = GEngine->GetTinyFont();
 	int32 TextWidth, TextHeight;
 	StringSize(Font, TextWidth, TextHeight, *DisplayText);
@@ -3263,7 +3264,7 @@ void FEditorViewportClient::OnOrthoZoom( const struct FInputEventState& InputSta
 
 	const bool bInvalidateViews = true;
 
-	// Update linked ortho viewport movement based on updated zoom and view location, 
+	// Update linked ortho viewport movement based on updated zoom and view location,
 	UpdateLinkedOrthoViewports( bInvalidateViews );
 
 	const bool bInvalidateHitProxies = true;
@@ -3386,7 +3387,7 @@ bool FEditorViewportClient::InputWidgetDelta(FViewport* InViewport, EAxisList::T
 			ModeTools->SnappedLocation += Drag;
 		}
 
-		// Update visuals of the rotate widget 
+		// Update visuals of the rotate widget
 		ApplyDeltaToRotateWidget(Rot);
 		return true;
 	}
@@ -3472,7 +3473,7 @@ FSceneInterface* FEditorViewportClient::GetScene() const
 	{
 		return World->Scene;
 	}
-	
+
 	return NULL;
 }
 
@@ -3661,7 +3662,7 @@ void FEditorViewportClient::Draw(FViewport* InViewport, FCanvas* Canvas)
 
 	EViewModeIndex CurrentViewMode = GetViewMode();
 	ViewFamily.ViewMode = CurrentViewMode;
-	bool bCanDisableTonemapper = (CurrentViewMode == VMI_VisualizeBuffer && CurrentBufferVisualizationMode != NAME_None) 
+	bool bCanDisableTonemapper = (CurrentViewMode == VMI_VisualizeBuffer && CurrentBufferVisualizationMode != NAME_None)
 								|| (CurrentViewMode == VMI_RayTracingDebug && CurrentRayTracingDebugVisualizationMode != NAME_None);
 
 	EngineShowFlagOverride(ESFIM_Editor, ViewFamily.ViewMode, ViewFamily.EngineShowFlags, bCanDisableTonemapper);
@@ -3747,13 +3748,13 @@ void FEditorViewportClient::Draw(FViewport* InViewport, FCanvas* Canvas)
 		World->ForegroundLineBatcher->Flush();
 	}
 
-	
+
 	// Draw the widget.
 	if (Widget && bShowWidget)
 	{
 		Widget->DrawHUD( Canvas );
 	}
-    
+
 	// Axes indicators
 	if (bDrawAxes && !ViewFamily.EngineShowFlags.Game && !GLevelEditorModeTools().IsViewportUIHidden())
 	{
@@ -3825,15 +3826,15 @@ void FEditorViewportClient::Draw(FViewport* InViewport, FCanvas* Canvas)
 				break;
 			}
 		}
-	}    
-   
+	}
+
 	// NOTE: DebugCanvasObject will be created by UDebugDrawService::Draw() if it doesn't already exist.
 	UDebugDrawService::Draw(ViewFamily.EngineShowFlags, Viewport, View, DebugCanvas);
 	UCanvas* DebugCanvasObject = FindObjectChecked<UCanvas>(GetTransientPackage(),TEXT("DebugCanvasObject"));
 	DebugCanvasObject->Canvas = DebugCanvas;
 	DebugCanvasObject->Init( Viewport->GetSizeXY().X, Viewport->GetSizeXY().Y, View , DebugCanvas);
-    
-    
+
+
 	// Stats display
 	if( IsRealtime() && ShouldShowStats() && DebugCanvas)
 	{
@@ -3944,11 +3945,11 @@ FLinearColor FEditorViewportClient::GetBackgroundColor() const
 }
 
 void FEditorViewportClient::SetCameraSetup(
-	const FVector& LocationForOrbiting, 
-	const FRotator& InOrbitRotation, 
-	const FVector& InOrbitZoom, 
-	const FVector& InOrbitLookAt, 
-	const FVector& InViewLocation, 
+	const FVector& LocationForOrbiting,
+	const FRotator& InOrbitRotation,
+	const FVector& InOrbitZoom,
+	const FVector& InOrbitLookAt,
+	const FVector& InViewLocation,
 	const FRotator &InViewRotation )
 {
 	if( bUsingOrbitCamera )
@@ -3963,7 +3964,7 @@ void FEditorViewportClient::SetCameraSetup(
 		SetViewRotation( InViewRotation );
 	}
 
-	
+
 	// Save settings for toggling between orbit and unlocked camera
 	DefaultOrbitLocation = InViewLocation;
 	DefaultOrbitRotation = InOrbitRotation;
@@ -4008,7 +4009,7 @@ FVector FEditorViewportClient::TranslateDelta( FKey InKey, float InDelta, bool I
 			//only invert x,y if we're moving the camera
 			if( ShouldUseMoveCanvasMovement() )
 			{
-				if(Widget->GetCurrentAxis() == EAxisList::None) 
+				if(Widget->GetCurrentAxis() == EAxisList::None)
 				{
 					X = -X;
 					Y = -Y;
@@ -4115,7 +4116,7 @@ bool FEditorViewportClient::InputAxis(FViewport* InViewport, int32 ControllerId,
 		LightDir.Pitch += -DragY * EditorViewportClient::LightRotSpeed;
 
 		PreviewScene->SetLightDirection( LightDir );
-		
+
 		// Remember that we adjusted it for the visualization
 		MovingPreviewLightTimer = PreviewLightConstants::MovingPreviewLightTimerDuration;
 		MovingPreviewLightSavedScreenPos = FVector2D(LastMouseX, LastMouseY);
@@ -4241,7 +4242,7 @@ bool FEditorViewportClient::InputGesture(FViewport* InViewport, EGestureEvent Ge
 				if( LeftMouseButtonDown )
 				{
 					// Pan left/right/up/down
-					
+
 					CurrentGestureDragDelta.X += AdjustedGestureDelta.X * -FMath::Sin( ViewRotation.Yaw * PI / 180.f );
 					CurrentGestureDragDelta.Y += AdjustedGestureDelta.X *  FMath::Cos( ViewRotation.Yaw * PI / 180.f );
 					CurrentGestureDragDelta.Z += -AdjustedGestureDelta.Y;
@@ -4249,7 +4250,7 @@ bool FEditorViewportClient::InputGesture(FViewport* InViewport, EGestureEvent Ge
 				else
 				{
 					// Change viewing angle
-					
+
 					CurrentGestureRotDelta.Yaw += AdjustGestureCameraRotation( AdjustedGestureDelta.X, 20.0f, 35.0f ) * -0.35f;
 					CurrentGestureRotDelta.Pitch += AdjustGestureCameraRotation( AdjustedGestureDelta.Y, 20.0f, 35.0f ) * 0.35f;
 				}
@@ -4279,9 +4280,9 @@ void FEditorViewportClient::UpdateGestureDelta()
 	if( CurrentGestureDragDelta != FVector::ZeroVector || CurrentGestureRotDelta != FRotator::ZeroRotator )
 	{
 		MoveViewportCamera( CurrentGestureDragDelta, CurrentGestureRotDelta, false );
-		
+
 		Invalidate( true, true );
-		
+
 		CurrentGestureDragDelta = FVector::ZeroVector;
 		CurrentGestureRotDelta = FRotator::ZeroRotator;
 	}
@@ -4456,7 +4457,7 @@ void FEditorViewportClient::ConvertMovementToOrbitDragRot(const FVector& InDelta
 bool FEditorViewportClient::ShouldPanOrDollyCamera() const
 {
 	const bool bIsCtrlDown = IsCtrlPressed();
-	
+
 	const bool bLeftMouseButtonDown = Viewport->KeyState( EKeys::LeftMouseButton );
 	const bool bRightMouseButtonDown = Viewport->KeyState( EKeys::RightMouseButton );
 	const bool bIsMarqueeSelect = IsOrtho() && bLeftMouseButtonDown;
@@ -4510,7 +4511,7 @@ bool FEditorViewportClient::IsFlightCameraInputModeActive() const
 			const bool bMiddleMouseButtonDown = Viewport->KeyState( EKeys::MiddleMouseButton );
 			const bool bRightMouseButtonDown = Viewport->KeyState( EKeys::RightMouseButton );
 			const bool bIsUsingTrackpad = FSlateApplication::Get().IsUsingTrackpad();
-			
+
 			const bool bIsNonOrbitMiddleMouse = bMiddleMouseButtonDown && !IsAltPressed();
 
 			const bool bIsMouseLooking =
@@ -4562,7 +4563,7 @@ void FEditorViewportClient::GetViewportDimensions( FIntPoint& OutOrigin, FIntPoi
 void FEditorViewportClient::UpdateAndApplyCursorVisibility()
 {
 	UpdateRequiredCursorVisibility();
-	ApplyRequiredCursorVisibility();	
+	ApplyRequiredCursorVisibility();
 }
 
 void FEditorViewportClient::UpdateRequiredCursorVisibility()
@@ -4599,7 +4600,7 @@ void FEditorViewportClient::UpdateRequiredCursorVisibility()
 	if (IsOrtho() && bMouseButtonDown && !MouseDeltaTracker->UsingDragTool())
 	{
 		//Translating an object, but NOT moving the camera AND the object (shift)
-		if ( ( AltDown == false && ShiftDown == false && ( LeftMouseButtonDown ^ RightMouseButtonDown ) ) && 
+		if ( ( AltDown == false && ShiftDown == false && ( LeftMouseButtonDown ^ RightMouseButtonDown ) ) &&
 			( ( GetWidgetMode() == FWidget::WM_Translate && Widget->GetCurrentAxis() != EAxisList::None ) ||
 			(  GetWidgetMode() == FWidget::WM_TranslateRotateZ && Widget->GetCurrentAxis() != EAxisList::ZRotation &&  Widget->GetCurrentAxis() != EAxisList::None ) ||
 			( GetWidgetMode() == FWidget::WM_2D && Widget->GetCurrentAxis() != EAxisList::Rotate2D &&  Widget->GetCurrentAxis() != EAxisList::None ) ) )
@@ -4676,8 +4677,8 @@ void FEditorViewportClient::ApplyRequiredCursorVisibility( bool bUpdateSoftwareC
 	{
 		//if we made the software cursor visible set its position
 		if( bOldSoftwareCursorVisibility != Viewport->IsSoftwareCursorVisible() )
-		{			
-			Viewport->SetSoftwareCursorPosition( FVector2D(Viewport->GetMouseX() , Viewport->GetMouseY() ) );			
+		{
+			Viewport->SetSoftwareCursorPosition( FVector2D(Viewport->GetMouseX() , Viewport->GetMouseY() ) );
 		}
 	}
 }
@@ -4726,7 +4727,7 @@ float FEditorViewportClient::GetFarClipPlaneOverride() const
 {
 	return FarPlane;
 }
-	
+
 void FEditorViewportClient::OverrideFarClipPlane(const float InFarPlane)
 {
 	FarPlane = InFarPlane;
@@ -4740,7 +4741,7 @@ float FEditorViewportClient::GetSceneDepthAtLocation(int32 X, int32 Y)
 
 FVector FEditorViewportClient::GetHitProxyObjectLocation(int32 X, int32 Y)
 {
-	// #todo: for now we are just getting the actor and using its location for 
+	// #todo: for now we are just getting the actor and using its location for
 	// depth. in the future we will just sample the depth buffer
 	HHitProxy* const HitProxy = Viewport->GetHitProxy(X, Y);
 	if (HitProxy && HitProxy->IsA(HActor::StaticGetType()))
@@ -4923,7 +4924,7 @@ void FEditorViewportClient::MoveViewportPerspectiveCamera( const FVector& InDrag
 		// normalize to -180 to 180
 		ViewRotation.Pitch = FRotator::NormalizeAxis(ViewRotation.Pitch);
 		// Make sure its withing  +/- 90 degrees.
-		ViewRotation.Pitch = FMath::Clamp( ViewRotation.Pitch, -90.f, 90.f );		
+		ViewRotation.Pitch = FMath::Clamp( ViewRotation.Pitch, -90.f, 90.f );
 	}
 	else
 	{
@@ -4940,7 +4941,7 @@ void FEditorViewportClient::MoveViewportPerspectiveCamera( const FVector& InDrag
 		const FQuat ResultQuat = ViewQuat * PitchQuat;
 
 		//get our correctly rotated ViewRotation
-		ViewRotation = ResultQuat.Rotator();	
+		ViewRotation = ResultQuat.Rotator();
 	}
 
 	const float DistanceToCurrentLookAt = FVector::Dist( GetViewLocation(), GetLookAtLocation() );
@@ -5106,7 +5107,7 @@ void FEditorViewportClient::SetViewMode(EViewModeIndex InViewModeIndex)
 		}
 		// Uncomment this to generate inital viewmode data.
 		// FEditorBuildUtils::CompileViewModeShaders(GetWorld(), InViewModeIndex);
-			 
+
 		PerspViewModeIndex = InViewModeIndex;
 		ApplyViewMode(PerspViewModeIndex, true, EngineShowFlags);
 		bForcingUnlitForNewMap = false;
@@ -5186,7 +5187,7 @@ void FEditorViewportClient::Invalidate(bool bInvalidateChildViews, bool bInvalid
 		if ( bInvalidateChildViews &&
 			ViewState.GetReference()->IsViewParent() )
 		{
-			GEditor->InvalidateChildViewports( ViewState.GetReference(), bInvalidateHitProxies );	
+			GEditor->InvalidateChildViewports( ViewState.GetReference(), bInvalidateHitProxies );
 		}
 	}
 }
@@ -5460,8 +5461,8 @@ void FEditorViewportClient::ProcessScreenShots(FViewport* InViewport)
 		{
 			// Determine the size of the captured viewport data.
 			FIntPoint BitmapSize = CaptureRect.Area() > 0 ? CaptureRect.Size() : InViewport->GetSizeXY();
-			
-			// Determine which region of the captured data we want to save out. If the highres screenshot capture region 
+
+			// Determine which region of the captured data we want to save out. If the highres screenshot capture region
 			// is not valid, we want to save out everything in the viewrect that we just grabbed.
 			FIntRect SourceRect = FIntRect(0, 0, 0, 0);
 			if (GIsHighResScreenshot && bCaptureAreaValid)
@@ -5518,7 +5519,7 @@ void FEditorViewportClient::ProcessScreenShots(FViewport* InViewport)
 				CompletionFuture.Wait();
 			}
 		}
-		
+
 		// Done with the request
 		FScreenshotRequest::Reset();
 		FScreenshotRequest::OnScreenshotRequestProcessed().Broadcast();
@@ -5537,7 +5538,7 @@ void FEditorViewportClient::DrawBoundingBox(FBox &Box, FCanvas* InCanvas, const
 
 	// Project center of bounding box onto screen.
 	const FVector4 ProjBoxCenter = InView->WorldToScreen(BoxCenter);
-		
+
 	// Do nothing if behind camera
 	if( ProjBoxCenter.W > 0.f )
 	{
@@ -5770,15 +5771,15 @@ void FEditorViewportClient::SetEnabledStats(const TArray<FString>& InEnabledStat
 {
 	EnabledStats = InEnabledStats;
 
-#if !UE_BUILD_SHIPPING
-	if (UWorld* MyWorld = GetWorld())
+#if ENABLE_AUDIO_DEBUG
+	if (GEngine)
 	{
-		if (FAudioDevice* AudioDevice = MyWorld->GetAudioDevice())
+		if (FAudioDeviceManager* DeviceManager = GEngine->GetAudioDeviceManager())
 		{
-			AudioDevice->ResolveDesiredStats(this);
+			FAudioDebugger::ResolveDesiredStats(this);
 		}
 	}
-#endif
+#endif // ENABLE_AUDIO_DEBUG
 }
 
 bool FEditorViewportClient::IsStatEnabled(const FString& InName) const
@@ -5939,4 +5940,4 @@ void FViewportNavigationCommands::RegisterCommands()
 	UI_COMMAND(FovZoomOut, "FOV Zoom Out", "Widens the camera FOV", EUserInterfaceActionType::Button, FInputChord(EKeys::Z));
 }
 
-#undef LOCTEXT_NAMESPACE 
+#undef LOCTEXT_NAMESPACE
diff --git a/Engine/Source/Editor/UnrealEd/Private/LevelEditorViewport.cpp b/Engine/Source/Editor/UnrealEd/Private/LevelEditorViewport.cpp
--- a/Engine/Source/Editor/UnrealEd/Private/LevelEditorViewport.cpp
+++ b/Engine/Source/Editor/UnrealEd/Private/LevelEditorViewport.cpp
@@ -1633,7 +1633,6 @@ FLevelEditorViewportClient::FLevelEditorViewportClient(const TSharedPtr<SLevelVi
 	, bWasControlledByOtherViewport(false)
 	, ActorLockedByMatinee(nullptr)
 	, ActorLockedToCamera(nullptr)
-	, SoundShowFlags(ESoundShowFlags::Disabled)
 	, bEditorCameraCut(false)
 	, bWasEditorCameraCut(false)
 	, bApplyCameraSpeedScaleByDistance(true)
diff --git a/Engine/Source/Editor/UnrealEd/Public/LevelEditorViewport.h b/Engine/Source/Editor/UnrealEd/Public/LevelEditorViewport.h
--- a/Engine/Source/Editor/UnrealEd/Public/LevelEditorViewport.h
+++ b/Engine/Source/Editor/UnrealEd/Public/LevelEditorViewport.h
@@ -496,22 +496,6 @@ public:
 		return ActorLockedByMatinee.IsValid();
 	}
 
-	/**
-	 * Get the sound stat flags enabled for this viewport
-	 */
-	virtual ESoundShowFlags::Type GetSoundShowFlags() const override
-	{ 
-		return SoundShowFlags;
-	}
-
-	/**
-	 * Set the sound stat flags enabled for this viewport
-	 */
-	virtual void SetSoundShowFlags(const ESoundShowFlags::Type InSoundShowFlags) override
-	{
-		SoundShowFlags = InSoundShowFlags;
-	}
-
 	void UpdateHoveredObjects( const TSet<FViewportHoverTarget>& NewHoveredObjects );
 
 	/**
@@ -793,9 +777,6 @@ private:
 	TWeakObjectPtr<AActor>	ActorLockedByMatinee;
 	TWeakObjectPtr<AActor>	ActorLockedToCamera;
 
-	/** Those sound stat flags which are enabled on this viewport */
-	ESoundShowFlags::Type	SoundShowFlags;
-
 	/** If true, we switched between two different cameras. Set by matinee, used by the motion blur to invalidate this frames motion vectors */
 	bool					bEditorCameraCut;
 
diff --git a/Engine/Source/Runtime/AudioMixer/Private/AudioMixer.cpp b/Engine/Source/Runtime/AudioMixer/Private/AudioMixer.cpp
--- a/Engine/Source/Runtime/AudioMixer/Private/AudioMixer.cpp
+++ b/Engine/Source/Runtime/AudioMixer/Private/AudioMixer.cpp
@@ -5,6 +5,14 @@
 #include "HAL/RunnableThread.h"
 #include "Misc/ConfigCacheIni.h"
 #include "HAL/ThreadSafeCounter.h"
+#include "ProfilingDebugging/CsvProfiler.h"
+
+
+// Defines the "Audio" category in the CSV profiler.
+// This should only be defined here. Modules who wish to use this category should contain the line
+// 		CSV_DECLARE_CATEGORY_MODULE_EXTERN(AUDIOMIXER_API, Audio);
+// 		
+CSV_DEFINE_CATEGORY_MODULE(AUDIOMIXER_API, Audio, true);
 
 // Command to enable logging to display accurate audio render times
 static int32 LogRenderTimesCVar = 0;
@@ -48,18 +56,6 @@ FAutoConsoleVariableRef CVarOverrunTimeout(
 	TEXT("Amount of time to wait for the render thread to time out before swapping to the null device. \n"),
 	ECVF_Default);
 
-DEFINE_STAT(STAT_AudioMixerRenderAudio);
-DEFINE_STAT(STAT_AudioMixerSourceManagerUpdate);
-DEFINE_STAT(STAT_AudioMixerSourceBuffers);
-DEFINE_STAT(STAT_AudioMixerSourceEffectBuffers);
-DEFINE_STAT(STAT_AudioMixerSourceOutputBuffers);
-DEFINE_STAT(STAT_AudioMixerSubmixes);
-DEFINE_STAT(STAT_AudioMixerSubmixChildren);
-DEFINE_STAT(STAT_AudioMixerSubmixSource);
-DEFINE_STAT(STAT_AudioMixerSubmixEffectProcessing);
-DEFINE_STAT(STAT_AudioMixerMasterReverb);
-DEFINE_STAT(STAT_AudioMixerMasterEQ);
-
 namespace Audio
 {
 	int32 sRenderInstanceIds = 0;
@@ -143,7 +139,7 @@ namespace Audio
 
 	void FOutputBuffer::MixNextBuffer()
  	{
-		SCOPE_CYCLE_COUNTER(STAT_AudioMixerRenderAudio);
+		CSV_SCOPED_TIMING_STAT(Audio, RenderAudio);
 
 		// Zero the buffer
 		FPlatformMemory::Memzero(Buffer.GetData(), Buffer.Num() * sizeof(float));
diff --git a/Engine/Source/Runtime/AudioMixer/Private/AudioMixerDevice.cpp b/Engine/Source/Runtime/AudioMixer/Private/AudioMixerDevice.cpp
--- a/Engine/Source/Runtime/AudioMixer/Private/AudioMixerDevice.cpp
+++ b/Engine/Source/Runtime/AudioMixer/Private/AudioMixerDevice.cpp
@@ -17,6 +17,7 @@
 #include "Runtime/HeadMountedDisplay/Public/IHeadMountedDisplayModule.h"
 #include "Misc/App.h"
 #include "Sound/AudioSettings.h"
+#include "ProfilingDebugging/CsvProfiler.h"
 
 #if WITH_EDITOR
 #include "AudioEditorModule.h"
@@ -30,6 +31,8 @@ FAutoConsoleVariableRef CVarDisableSubmixEQ(
 	TEXT("0: Not Disabled, 1: Disabled"),
 	ECVF_Default);
 
+// Link to "Audio" profiling category
+CSV_DECLARE_CATEGORY_MODULE_EXTERN(AUDIOMIXER_API, Audio);
 
 namespace Audio
 {
@@ -529,7 +532,7 @@ namespace Audio
 		FMixerSubmixWeakPtr MasterSubmix = GetMasterSubmix();
 
 		{
-			SCOPE_CYCLE_COUNTER(STAT_AudioMixerSubmixes);
+			CSV_SCOPED_TIMING_STAT(Audio, Submixes);
 
 			FMixerSubmixPtr MasterSubmixPtr = MasterSubmix.Pin();
 			if (MasterSubmixPtr.IsValid())
@@ -573,6 +576,7 @@ namespace Audio
 
 			FAudioThread::RunCommandOnAudioThread([this]()
 			{
+				CSV_SCOPED_TIMING_STAT(Audio, InitSubmix);
 				InitSoundSubmixes();
 			}, GET_STATID(STAT_InitSoundSubmixes));
 			return;
@@ -988,6 +992,7 @@ namespace Audio
 				FMixerDevice* MixerDevice = this;
 				FAudioThread::RunCommandOnAudioThread([MixerDevice, InSoundSubmix]()
 				{
+					CSV_SCOPED_TIMING_STAT(Audio, RegisterSubmix);
 					MixerDevice->RegisterSoundSubmix(InSoundSubmix);
 				}, GET_STATID(STAT_AudioRegisterSoundSubmix));
 				return;
@@ -1038,6 +1043,7 @@ namespace Audio
 				FMixerDevice* MixerDevice = this;
 				FAudioThread::RunCommandOnAudioThread([MixerDevice, InSoundSubmix]()
 				{
+					CSV_SCOPED_TIMING_STAT(Audio, UnregisterSubmix);
 					MixerDevice->UnregisterSoundSubmix(InSoundSubmix);
 				}, GET_STATID(STAT_AudioUnregisterSoundSubmix));
 				return;
diff --git a/Engine/Source/Runtime/AudioMixer/Private/AudioMixerSource.cpp b/Engine/Source/Runtime/AudioMixer/Private/AudioMixerSource.cpp
--- a/Engine/Source/Runtime/AudioMixer/Private/AudioMixerSource.cpp
+++ b/Engine/Source/Runtime/AudioMixer/Private/AudioMixerSource.cpp
@@ -7,6 +7,10 @@
 #include "IAudioExtensionPlugin.h"
 #include "Sound/AudioSettings.h"
 #include "ContentStreaming.h"
+#include "ProfilingDebugging/CsvProfiler.h"
+
+// Link to "Audio" profiling category
+CSV_DECLARE_CATEGORY_MODULE_EXTERN(AUDIOMIXER_API, Audio);
 
 static int32 DisableHRTFCvar = 0;
 FAutoConsoleVariableRef CVarDisableHRTF(
@@ -32,7 +36,7 @@ namespace Audio
 		, MixerDevice(static_cast<FMixerDevice*>(InAudioDevice))
 		, MixerBuffer(nullptr)
 		, MixerSourceBuffer(nullptr)
-		, MixerSourceVoice(nullptr)		
+		, MixerSourceVoice(nullptr)
 		, PreviousAzimuth(-1.0f)
 		, InitializationState(EMixerSourceInitializationState::NotInitialized)
 		, bPlayedCachedBuffer(false)
@@ -93,7 +97,7 @@ namespace Audio
 		bIsVorbis = WaveInstance->WaveData->bDecompressedFromOgg;
 
 		bIsStoppingVoicesEnabled = ((FAudioDevice*)MixerDevice)->IsStoppingVoicesEnabled();
-		
+
 		bIsStopping = false;
 		bIsEffectTailsDone = true;
 		bIsDone = false;
@@ -101,7 +105,7 @@ namespace Audio
 		FSoundBuffer* SoundBuffer = static_cast<FSoundBuffer*>(MixerBuffer);
 		if (SoundBuffer->NumChannels > 0)
 		{
-			SCOPE_CYCLE_COUNTER(STAT_AudioSourceInitTime);
+			CSV_SCOPED_TIMING_STAT(Audio, InitSources);
 
 			AUDIO_MIXER_CHECK(MixerDevice);
 			MixerSourceVoice = MixerDevice->GetMixerSourceVoice();
@@ -120,6 +124,7 @@ namespace Audio
 			InitParams.SourceVoice = MixerSourceVoice;
 			InitParams.bUseHRTFSpatialization = UseObjectBasedSpatialization();
 			InitParams.bIsAmbisonics = WaveInstance->bIsAmbisonics;
+
 			if (InitParams.bIsAmbisonics)
 			{
 				checkf(InitParams.NumInputChannels == 4, TEXT("Only allow 4 channel source if file is ambisonics format."));
@@ -165,7 +170,7 @@ namespace Audio
 					}
 				}
 
-				// Toggle muting the source if sending only to output bus. 
+				// Toggle muting the source if sending only to output bus.
 				// This can get set even if the source doesn't have bus sends since bus sends can be dynamically enabled.
 				InitParams.bOutputToBusOnly = WaveInstance->bOutputToBusOnly;
 				DynamicBusSendInfos.Reset();
@@ -306,15 +311,18 @@ namespace Audio
 			// Whether or not we're 3D
 			bIs3D = !UseObjectBasedSpatialization() && WaveInstance->GetUseSpatialization() && SoundBuffer->NumChannels < 3;
 
-			// Grab the source's reverb plugin settings 
+			// Grab the source's reverb plugin settings
 			InitParams.SpatializationPluginSettings = UseSpatializationPlugin() ? WaveInstance->SpatializationPluginSettings : nullptr;
 
-			// Grab the source's occlusion plugin settings 
+			// Grab the source's occlusion plugin settings
 			InitParams.OcclusionPluginSettings = UseOcclusionPlugin() ? WaveInstance->OcclusionPluginSettings : nullptr;
 
-			// Grab the source's reverb plugin settings 
+			// Grab the source's reverb plugin settings
 			InitParams.ReverbPluginSettings = UseReverbPlugin() ? WaveInstance->ReverbPluginSettings : nullptr;
 
+			// Grab the source's modulation plugin settings
+			InitParams.ModulationPluginSettings = UseModulationPlugin() ? WaveInstance->ModulationPluginSettings : nullptr;
+
 			// We support reverb
 			SetReverbApplied(true);
 
@@ -381,7 +389,7 @@ namespace Audio
 
 	void FMixerSource::Update()
 	{
-		SCOPE_CYCLE_COUNTER(STAT_AudioUpdateSources);
+		CSV_SCOPED_TIMING_STAT(Audio, UpdateSources);
 
 		LLM_SCOPE(ELLMTag::AudioMixer);
 
@@ -392,6 +400,8 @@ namespace Audio
 
 		++TickCount;
 
+		UpdateModulation();
+
 		UpdatePitch();
 
 		UpdateVolume();
@@ -404,7 +414,9 @@ namespace Audio
 
 		UpdateChannelMaps();
 
-		FSoundSource::DrawDebugInfo();
+#if ENABLE_AUDIO_DEBUG
+		FAudioDebugger::DrawDebugInfo(*this);
+#endif // ENABLE_AUDIO_DEBUG
 	}
 
 	bool FMixerSource::PrepareForInitialization(FWaveInstance* InWaveInstance)
@@ -809,13 +821,25 @@ namespace Audio
 		InitializationState = EMixerSourceInitializationState::NotInitialized;
 	}
 
+	void FMixerSource::UpdateModulation()
+	{
+		check(AudioDevice);
+		check(MixerSourceVoice);
+		check(WaveInstance);
+
+		if (AudioDevice->IsModulationPluginEnabled())
+		{
+			AudioDevice->ModulationInterface->ProcessControls(MixerSourceVoice->GetSourceId(), WaveInstance->SoundModulationControls);
+		}
+	}
+
 	void FMixerSource::UpdatePitch()
 	{
 		AUDIO_MIXER_CHECK(MixerBuffer);
 
 		check(WaveInstance);
 
-		Pitch = WaveInstance->Pitch;
+		Pitch = WaveInstance->GetPitch();
 
 		// Don't apply global pitch scale to UI sounds
 		if (!WaveInstance->bIsUISound)
@@ -838,22 +862,24 @@ namespace Audio
 
 	void FMixerSource::UpdateVolume()
 	{
-		float CurrentVolume;
-		if (AudioDevice->IsAudioDeviceMuted())
-		{
-			CurrentVolume = 0.0f;
-		}
-		else
+		MixerSourceVoice->SetDistanceAttenuation(WaveInstance->GetDistanceAttenuation());
+
+		float CurrentVolume = 0.0f;
+		if (!AudioDevice->IsAudioDeviceMuted())
 		{
-			CurrentVolume = WaveInstance->GetVolume();
-			CurrentVolume *= WaveInstance->GetVolumeApp();
+			// 1. Apply device gain stage(s)
+			CurrentVolume = WaveInstance->ActiveSound->bIsPreviewSound ? 1.0f : AudioDevice->GetMasterVolume();
 			CurrentVolume *= AudioDevice->GetPlatformAudioHeadroom();
+
+			// 2. Apply instance gain stage(s)
+			CurrentVolume *= WaveInstance->GetVolume();
 			CurrentVolume *= WaveInstance->GetDynamicVolume();
+
+			// 3. Apply editor gain stage(s)
 			CurrentVolume = FMath::Clamp<float>(GetDebugVolume(CurrentVolume), 0.0f, MAX_VOLUME);
+			MixerSourceVoice->SetVolume(CurrentVolume);
 		}
-
 		MixerSourceVoice->SetVolume(CurrentVolume);
-		MixerSourceVoice->SetDistanceAttenuation(WaveInstance->GetDistanceAttenuation());
 	}
 
 	void FMixerSource::UpdateSpatialization()
@@ -867,7 +893,7 @@ namespace Audio
 
 	void FMixerSource::UpdateEffects()
 	{
-		// Update the default LPF filter frequency 
+		// Update the default LPF filter frequency
 		SetFilterFrequency();
 
 		if (LastLPFFrequency != LPFFrequency)
@@ -917,7 +943,7 @@ namespace Audio
 					MixerSourceVoice->SetSubmixSendInfo(MixerDevice->GetMasterReverbPluginSubmix(), ReverbSendLevel);
 				}
 			}
-			else 
+			else
 			{
 				// Send the source audio to the master reverb
 				if (MixerDevice->GetMasterReverbSubmix().IsValid())
@@ -1089,7 +1115,7 @@ namespace Audio
 			{
 				UpdateStereoEmitterPositions();
 			}
-			
+
 			if (!UseObjectBasedSpatialization())
 			{
 				float AzimuthOffset = 0.0f;
@@ -1126,7 +1152,7 @@ namespace Audio
 				return true;
 			}
 		}
-				
+
 		if (!OutChannelMap.Num())
 		{
 			MixerDevice->Get2DChannelMap(bIsVorbis, InSubmixChannelType, 2, WaveInstance->bCenterChannelOnly, OutChannelMap);
@@ -1157,9 +1183,9 @@ namespace Audio
 	bool FMixerSource::UseObjectBasedSpatialization() const
 	{
 		return (Buffer->NumChannels <= MixerDevice->MaxChannelsSupportedBySpatializationPlugin &&
-				AudioDevice->IsSpatializationPluginEnabled() &&
-				DisableHRTFCvar == 0 &&
-				WaveInstance->SpatializationMethod == ESoundSpatializationAlgorithm::SPATIALIZATION_HRTF);
+			AudioDevice->IsSpatializationPluginEnabled() &&
+			DisableHRTFCvar == 0 &&
+			WaveInstance->SpatializationMethod == ESoundSpatializationAlgorithm::SPATIALIZATION_HRTF);
 	}
 
 	bool FMixerSource::UseSpatializationPlugin() const
@@ -1172,15 +1198,20 @@ namespace Audio
 	bool FMixerSource::UseOcclusionPlugin() const
 	{
 		return (Buffer->NumChannels == 1 || Buffer->NumChannels == 2) &&
-				AudioDevice->IsOcclusionPluginEnabled() &&
-				WaveInstance->OcclusionPluginSettings != nullptr;
+			AudioDevice->IsOcclusionPluginEnabled() &&
+			WaveInstance->OcclusionPluginSettings != nullptr;
+	}
+
+	bool FMixerSource::UseModulationPlugin() const
+	{
+		return AudioDevice->IsModulationPluginEnabled() &&
+			WaveInstance->ModulationPluginSettings != nullptr;
 	}
 
 	bool FMixerSource::UseReverbPlugin() const
 	{
 		return (Buffer->NumChannels == 1 || Buffer->NumChannels == 2) &&
-				AudioDevice->IsReverbPluginEnabled() &&
-				WaveInstance->ReverbPluginSettings != nullptr;
+			AudioDevice->IsReverbPluginEnabled() &&
+			WaveInstance->ReverbPluginSettings != nullptr;
 	}
-
 }
diff --git a/Engine/Source/Runtime/AudioMixer/Private/AudioMixerSource.h b/Engine/Source/Runtime/AudioMixer/Private/AudioMixerSource.h
--- a/Engine/Source/Runtime/AudioMixer/Private/AudioMixerSource.h
+++ b/Engine/Source/Runtime/AudioMixer/Private/AudioMixerSource.h
@@ -65,6 +65,9 @@ namespace Audio
 		/** Frees any resources for this sound source. */
 		void FreeResources();
 
+		/** Updates modulation parameters set from the modulation plugin. */
+		void UpdateModulation();
+
 		/** Updates the pitch parameter set from the game thread. */
 		void UpdatePitch();
 		
@@ -98,6 +101,9 @@ namespace Audio
 		/** Whether or not to use the spatialization plugin. */
 		bool UseSpatializationPlugin() const;
 
+		/** Whether or not to use the occlusion plugin. */
+		bool UseModulationPlugin() const;
+
 		/** Whether or not to use the occlusion plugin. */
 		bool UseOcclusionPlugin() const;
 
diff --git a/Engine/Source/Runtime/AudioMixer/Private/AudioMixerSourceManager.cpp b/Engine/Source/Runtime/AudioMixer/Private/AudioMixerSourceManager.cpp
--- a/Engine/Source/Runtime/AudioMixer/Private/AudioMixerSourceManager.cpp
+++ b/Engine/Source/Runtime/AudioMixer/Private/AudioMixerSourceManager.cpp
@@ -7,8 +7,10 @@
 #include "AudioMixerSubmix.h"
 #include "IAudioExtensionPlugin.h"
 #include "AudioMixer.h"
+#include "ProfilingDebugging/CsvProfiler.h"
 
-DEFINE_STAT(STAT_AudioMixerHRTF);
+// Link to "Audio" profiling category
+CSV_DECLARE_CATEGORY_MODULE_EXTERN(AUDIOMIXER_API, Audio);
 
 static int32 DisableParallelSourceProcessingCvar = 1;
 FAutoConsoleVariableRef CVarDisableParallelSourceProcessing(
@@ -216,6 +218,7 @@ namespace Audio
 			SourceInfo.bUseHRTFSpatializer = false;
 			SourceInfo.bUseOcclusionPlugin = false;
 			SourceInfo.bUseReverbPlugin = false;
+			SourceInfo.bUseModulationPlugin = false;
 			SourceInfo.bHasStarted = false;
 			SourceInfo.bOutputToBusOnly = false;
 			SourceInfo.bIsVorbis = false;
@@ -441,6 +444,11 @@ namespace Audio
 			MixerDevice->ReverbPluginInterface->OnReleaseSource(SourceId);
 		}
 
+		if (SourceInfo.bUseModulationPlugin)
+		{
+			MixerDevice->ModulationInterface->OnReleaseSource(SourceId);
+		}
+
 		// Delete the source effects
 		SourceInfo.SourceEffectChainId = INDEX_NONE;
 		ResetSourceEffectChain(SourceId);
@@ -490,6 +498,7 @@ namespace Audio
 		SourceInfo.bUseHRTFSpatializer = false;
 		SourceInfo.bUseOcclusionPlugin = false;
 		SourceInfo.bUseReverbPlugin = false;
+		SourceInfo.bUseModulationPlugin = false;
 		SourceInfo.bHasStarted = false;
 		SourceInfo.bOutputToBusOnly = false;
 		SourceInfo.bIsBypassingLPF = false;
@@ -659,6 +668,13 @@ namespace Audio
 				SourceInfo.bUseReverbPlugin = true;
 			}
 
+			// Create the modulation plugin source effect
+			if (InitParams.ModulationPluginSettings != nullptr)
+			{
+				MixerDevice->ModulationInterface->OnInitSource(SourceId, InitParams.AudioComponentUserID, InitParams.NumInputChannels, InitParams.ModulationPluginSettings);
+				SourceInfo.bUseModulationPlugin = true;
+			}
+
 			// Default all sounds to not consider effect chain tails when playing
 			SourceInfo.bEffectTailsDone = true;
 
@@ -1092,7 +1108,6 @@ namespace Audio
 		AudioMixerThreadCommand([this, SourceId, InHPFFrequency]()
 		{
 			AUDIO_MIXER_CHECK_AUDIO_PLAT_THREAD(MixerDevice);
-
 			SourceInfos[SourceId].HighPassFilter.StartFrequencyInterpolation(InHPFFrequency, NumOutputFrames);
 		});
 	}
@@ -1328,7 +1343,7 @@ namespace Audio
 
 	void FMixerSourceManager::ComputeSourceBuffersForIdRange(const bool bGenerateBuses, const int32 SourceIdStart, const int32 SourceIdEnd)
 	{
-		SCOPE_CYCLE_COUNTER(STAT_AudioMixerSourceBuffers);
+		CSV_SCOPED_TIMING_STAT(Audio, SourceBuffers);
 
 		const double AudioRenderThreadTime = MixerDevice->GetAudioRenderThreadTime();
 		const double AudioClockDelta = MixerDevice->GetAudioClockDelta();
@@ -1536,8 +1551,6 @@ namespace Audio
 			return;
 		}
 
-		// Apply the distance attenuation before sending to plugins
-		float* PreDistanceAttenBufferPtr = SourceInfo.PreDistanceAttenuationBuffer.GetData();
 		float* PostDistanceAttenBufferPtr = SourceInfo.SourceBuffer.GetData();
 
 		bool bShouldMixInReverb = false;
@@ -1612,7 +1625,7 @@ namespace Audio
 		// If the source has HRTF processing enabled, run it through the spatializer
 		if (SourceInfo.bUseHRTFSpatializer)
 		{
-			SCOPE_CYCLE_COUNTER(STAT_AudioMixerHRTF);
+			CSV_SCOPED_TIMING_STAT(Audio, HRTF);
 
 			AUDIO_MIXER_CHECK(SpatializationPlugin.IsValid());
 			AUDIO_MIXER_CHECK(SourceInfo.NumInputChannels == 1);
@@ -1961,7 +1974,7 @@ namespace Audio
 
 	void FMixerSourceManager::ComputePostSourceEffectBufferForIdRange(bool bGenerateBuses, const int32 SourceIdStart, const int32 SourceIdEnd)
 	{
-		SCOPE_CYCLE_COUNTER(STAT_AudioMixerSourceEffectBuffers);
+		CSV_SCOPED_TIMING_STAT(Audio, SourceEffectsBuffers);
 
 		const bool bIsDebugModeEnabled = DebugSoloSources.Num() > 0;
 
@@ -2138,7 +2151,7 @@ namespace Audio
 
 	void FMixerSourceManager::ComputeOutputBuffersForIdRange(const bool bGenerateBuses, const int32 SourceIdStart, const int32 SourceIdEnd)
 	{
-		SCOPE_CYCLE_COUNTER(STAT_AudioMixerSourceOutputBuffers);
+		CSV_SCOPED_TIMING_STAT(Audio, SourceOutputBuffers);
 
 		for (int32 SourceId = SourceIdStart; SourceId < SourceIdEnd; ++SourceId)
 		{
@@ -2385,7 +2398,7 @@ namespace Audio
 	{
 		AUDIO_MIXER_CHECK_AUDIO_PLAT_THREAD(MixerDevice);
 
-		SCOPE_CYCLE_COUNTER(STAT_AudioMixerSourceManagerUpdate);
+		CSV_SCOPED_TIMING_STAT(Audio, SourceManagerUpdate);
 
 		if (FPlatformProcess::SupportsMultithreading())
 		{
diff --git a/Engine/Source/Runtime/AudioMixer/Private/AudioMixerSourceVoice.h b/Engine/Source/Runtime/AudioMixer/Private/AudioMixerSourceVoice.h
--- a/Engine/Source/Runtime/AudioMixer/Private/AudioMixerSourceVoice.h
+++ b/Engine/Source/Runtime/AudioMixer/Private/AudioMixerSourceVoice.h
@@ -65,6 +65,9 @@ namespace Audio
 		// Does a faded stop (to avoid discontinuity)
 		void StopFade(int32 NumFrames);
 
+		// Get the source's Id
+		int32 GetSourceId() const { return SourceId; }
+
 		// Queries if the voice is playing
 		bool IsPlaying() const;
 
diff --git a/Engine/Source/Runtime/AudioMixer/Private/AudioMixerSubmix.cpp b/Engine/Source/Runtime/AudioMixer/Private/AudioMixerSubmix.cpp
--- a/Engine/Source/Runtime/AudioMixer/Private/AudioMixerSubmix.cpp
+++ b/Engine/Source/Runtime/AudioMixer/Private/AudioMixerSubmix.cpp
@@ -5,6 +5,10 @@
 #include "AudioMixerSourceVoice.h"
 #include "Sound/SoundSubmix.h"
 #include "Sound/SoundEffectSubmix.h"
+#include "ProfilingDebugging/CsvProfiler.h"
+
+// Link to "Audio" profiling category
+CSV_DECLARE_CATEGORY_MODULE_EXTERN(AUDIOMIXER_API, Audio);
 
 namespace Audio
 {
@@ -591,7 +595,7 @@ namespace Audio
 
 		// Mix all submix audio into this submix's input scratch buffer
 		{
-			SCOPE_CYCLE_COUNTER(STAT_AudioMixerSubmixChildren);
+			CSV_SCOPED_TIMING_STAT(Audio, SubmixChildren);
 
 			// First loop this submix's child submixes mixing in their output into this submix's dry/wet buffers.
 			for (auto& ChildSubmixEntry : ChildSubmixes)
@@ -620,7 +624,7 @@ namespace Audio
 		}
 
 		{
-			SCOPE_CYCLE_COUNTER(STAT_AudioMixerSubmixSource);
+			CSV_SCOPED_TIMING_STAT(Audio, SubmixSource);
 
 			// Loop through this submix's sound sources
 			for (const auto MixerSourceVoiceIter : MixerSourceVoices)
@@ -635,7 +639,7 @@ namespace Audio
 
 		if (EffectSubmixChain.Num() > 0)
 		{
-			SCOPE_CYCLE_COUNTER(STAT_AudioMixerSubmixEffectProcessing);
+			CSV_SCOPED_TIMING_STAT(Audio, SubmixEffectProcessing);
 
 			// Setup the input data buffer
 			FSoundEffectSubmixInputData InputData;
@@ -973,4 +977,4 @@ namespace Audio
 		}
 	}
 
-}
\ No newline at end of file
+}
diff --git a/Engine/Source/Runtime/AudioMixer/Private/Components/SynthComponent.cpp b/Engine/Source/Runtime/AudioMixer/Private/Components/SynthComponent.cpp
--- a/Engine/Source/Runtime/AudioMixer/Private/Components/SynthComponent.cpp
+++ b/Engine/Source/Runtime/AudioMixer/Private/Components/SynthComponent.cpp
@@ -13,7 +13,7 @@ USynthSound::USynthSound(const FObjectInitializer& ObjectInitializer)
 void USynthSound::Init(USynthComponent* InSynthComponent, const int32 InNumChannels, const int32 InSampleRate, const int32 InCallbackSize)
 {
 	OwningSynthComponent = InSynthComponent;
-	bVirtualizeWhenSilent = true;
+	bPlayWhenSilent = true;
 	NumChannels = InNumChannels;
 	NumSamplesToGeneratePerCallback = InCallbackSize;
 	// Turn off async generation in old audio engine on mac.
@@ -61,7 +61,7 @@ int32 USynthSound::OnGeneratePCMAudio(TArray<uint8>& OutAudio, int32 NumSamples)
 		// Use the float scratch buffer instead of the out buffer directly
 		FloatBuffer.Reset();
 		FloatBuffer.AddZeroed(NumSamples * sizeof(float));
-		
+
 		float* FloatBufferDataPtr = FloatBuffer.GetData();
 		int32 NumSamplesGenerated = OwningSynthComponent->OnGeneratePCMAudio(FloatBufferDataPtr, NumSamples);
 
@@ -84,8 +84,8 @@ void USynthSound::OnEndGenerate()
 	OwningSynthComponent->OnEndGenerate();
 }
 
-Audio::EAudioMixerStreamDataFormat::Type USynthSound::GetGeneratedPCMDataFormat() const 
-{ 
+Audio::EAudioMixerStreamDataFormat::Type USynthSound::GetGeneratedPCMDataFormat() const
+{
 	// Only audio mixer supports return float buffers
 	return bAudioMixer ? Audio::EAudioMixerStreamDataFormat::Float : Audio::EAudioMixerStreamDataFormat::Int16;
 }
@@ -179,7 +179,7 @@ void USynthComponent::Initialize(int32 SampleRateOverride)
 		NumChannels = 2;
 		TestSineLeft.Init(SampleRate, 440.0f, 0.5f);
 		TestSineRight.Init(SampleRate, 220.0f, 0.5f);
-#else	
+#else
 		// Initialize the synth component
 		Init(SampleRate);
 
@@ -233,6 +233,7 @@ void USynthComponent::CreateAudioComponent()
 		AudioComponent->bStopWhenOwnerDestroyed = true;
 		AudioComponent->bShouldRemainActiveIfDropped = true;
 		AudioComponent->Mobility = EComponentMobility::Movable;
+		AudioComponent->Modulation = Modulation;
 
 #if WITH_EDITORONLY_DATA
 		AudioComponent->bVisualizeComponent = false;
@@ -386,7 +387,7 @@ void USynthComponent::Start()
 	// We will also ensure that this synth was initialized before attempting to play.
 	Initialize();
 
-	// If there is no Synth USoundBase, we can't start. This can happen if start is called in a cook, a server, or 
+	// If there is no Synth USoundBase, we can't start. This can happen if start is called in a cook, a server, or
 	// if the audio engine is set to "noaudio".
 	// TODO: investigate if this should be handled elsewhere before this point
 	if (!Synth)
diff --git a/Engine/Source/Runtime/AudioMixer/Private/DSP/Envelope.cpp b/Engine/Source/Runtime/AudioMixer/Private/DSP/Envelope.cpp
--- a/Engine/Source/Runtime/AudioMixer/Private/DSP/Envelope.cpp
+++ b/Engine/Source/Runtime/AudioMixer/Private/DSP/Envelope.cpp
@@ -61,6 +61,8 @@ namespace Audio
 
 	void FEnvelope::Start()
 	{
+		bCurrentCycleIsADOnly = SustainGain <= SMALL_NUMBER;
+
 		// Don't reset the envelope if we're in legato mode and we're not in release or off
 		if (bIsLegatoMode && CurrentState != EEnvelopeState::Off && CurrentState != EEnvelopeState::Release)
 		{
@@ -74,18 +76,102 @@ namespace Audio
 		CurrentState = EEnvelopeState::Attack;
 	}
 
+	// logic for one mono note interrupting another mono note (same voice)
+	void FEnvelope::StartLegato(const float InNewDepth)
+	{
+		// Envelope is not being used. Don't do the work (and don't divide by zero)
+		if (Depth <= SMALL_NUMBER && InNewDepth <= SMALL_NUMBER)
+		{
+			return;
+		}
+
+		bCurrentCycleIsADOnly = SustainGain <= SMALL_NUMBER;
+
+		switch (CurrentState)
+		{
+			case EEnvelopeState::Attack:
+			{
+				if (InNewDepth > Depth)
+				{
+					CurrentEnvelopeValue *= Depth / InNewDepth;
+					Depth = InNewDepth;
+
+					bChanged = true;
+				}
+			}
+			break;
+
+			case EEnvelopeState::Decay:
+			{
+				if (InNewDepth > Depth * CurrentEnvelopeValue)
+				{
+					CurrentState = EEnvelopeState::Attack;
+				}
+
+				CurrentEnvelopeValue *= Depth / InNewDepth;
+				Depth = InNewDepth;
+
+				bChanged = true;
+			}
+			break;
+
+			case EEnvelopeState::Sustain:
+			{
+				// new sustain gain is higher
+				if (InNewDepth > Depth * SustainGain)
+				{
+					CurrentEnvelopeValue *= Depth / InNewDepth;
+					Depth = InNewDepth;
+					CurrentState = EEnvelopeState::Attack;
+				}
+
+				bChanged = true;
+			}
+			break;
+
+			case EEnvelopeState::Release:
+			{
+				// "attack up to" a larger new depth or "decay down to" a lower new depth
+				CurrentState = EEnvelopeState::Attack;
+
+				if(InNewDepth < Depth * CurrentEnvelopeValue)
+				{
+					CurrentState = EEnvelopeState::Decay;
+				}
+
+				CurrentEnvelopeValue *= Depth / InNewDepth;
+				Depth = InNewDepth;
+
+				bChanged = true;
+			}
+			break;
+
+			default:
+			{
+				// previous behavior
+				Depth = InNewDepth;
+				Start();
+			}
+		}
+	}
+
 	void FEnvelope::Stop()
 	{
-		// If we're forcing off or if we're already off, then set state to off
 		if (CurrentEnvelopeValue == 0.0f)
 		{
+			// already finished (jump to off)
 			CurrentState = EEnvelopeState::Off;
 		}
-		else
+		else if (!bCurrentCycleIsADOnly)
 		{
-			// If we actually have an envelope value now, go to release
+			// normal envelope mode (jump to release)
 			CurrentState = EEnvelopeState::Release;
 		}
+		else if (CurrentState == EEnvelopeState::Attack)
+		{
+			// AD only envelope mode (jump to decay)
+			CurrentState = EEnvelopeState::Decay;
+		}
 	}
 
 	void FEnvelope::Shutdown()
@@ -167,7 +253,7 @@ namespace Audio
 			AttackData.Offset = (1.0f + AttackData.TCO) * (1.0f - AttackData.Coefficient);
 
 			DecayData.Coefficient = FMath::Exp(-FMath::Loge((1.0f + DecayData.TCO) / DecayData.TCO) / DecayData.TimeSamples);
-			DecayData.Offset = (SustainGain - DecayData.TCO)*(1.0f - DecayData.Coefficient);
+			DecayData.Offset = (bCurrentCycleIsADOnly? 0.0f : SustainGain - DecayData.TCO)*(1.0f - DecayData.Coefficient);
 
 			ReleaseData.Coefficient = FMath::Exp(-FMath::Loge((1.0f + ReleaseData.TCO) / ReleaseData.TCO) / ReleaseData.TimeSamples);
 			ReleaseData.Offset = -ReleaseData.TCO*(1.0f - ReleaseData.Coefficient);
@@ -207,23 +293,40 @@ namespace Audio
 			{
 				// --- render value
 				CurrentEnvelopeValue = DecayData.Offset + CurrentEnvelopeValue * DecayData.Coefficient;
-				if (CurrentEnvelopeValue <= SustainGain || DecayTimeMsec <= 0.0f)
+				if ((CurrentEnvelopeValue <= SustainGain) || DecayTimeMsec <= 0.0f)
 				{
-					CurrentEnvelopeValue = SustainGain;
-					CurrentState = EEnvelopeState::Sustain;
-					break;
+					
+					if(!bCurrentCycleIsADOnly)
+					{
+						CurrentEnvelopeValue = SustainGain;
+						CurrentState = EEnvelopeState::Sustain;
+						break;
+					}
+					else if (CurrentEnvelopeValue <= SMALL_NUMBER)
+					{
+						CurrentState = EEnvelopeState::Off;
+					}
 				}
 			}
 			break;
 
 			case EEnvelopeState::Sustain:
-				// Don't do anything in sustain state
-				break;
+			{
+				// live-update sustain level (to hear changes made during sustain phase)
+				CurrentEnvelopeValue = SustainGain;
+
+				if (bCurrentCycleIsADOnly && SustainGain <= SMALL_NUMBER)
+				{
+					// Check if envelope was being used as AD only
+					CurrentState = EEnvelopeState::Off;
+				}
+			}
+			break;
 
 			case EEnvelopeState::Release:
 			{
 				CurrentEnvelopeValue = ReleaseData.Offset + CurrentEnvelopeValue * ReleaseData.Coefficient;
-				if (CurrentEnvelopeValue <= 0.0f || ReleaseTimeMsec <= 0.0f)
+				if (CurrentEnvelopeValue <= 0.0f || ReleaseTimeMsec <= 0.0f || SustainGain <= SMALL_NUMBER)
 				{
 					CurrentEnvelopeValue = 0.0f;
 					CurrentState = EEnvelopeState::Off;
diff --git a/Engine/Source/Runtime/AudioMixer/Private/DSP/Osc.cpp b/Engine/Source/Runtime/AudioMixer/Private/DSP/Osc.cpp
--- a/Engine/Source/Runtime/AudioMixer/Private/DSP/Osc.cpp
+++ b/Engine/Source/Runtime/AudioMixer/Private/DSP/Osc.cpp
@@ -31,8 +31,7 @@ namespace Audio
 	void IOscBase::Init(const float InSampleRate, const int32 InVoiceId, FModulationMatrix* InMatrix, const int32 ModMatrixStage)
 	{
 		VoiceId = InVoiceId;
-		SampleRate = InSampleRate;
-		Nyquist = 0.5f * SampleRate;
+		SetSampleRate(InSampleRate);
 
 		bChanged = true;
 
@@ -98,6 +97,17 @@ namespace Audio
 		}
 	}
 
+	void IOscBase::SetSampleRate(const float InSampleRate)
+	{
+		if (InSampleRate != SampleRate)
+		{
+			SampleRate = InSampleRate;
+			Nyquist = 0.5f * SampleRate;
+
+			bChanged = true;
+		}
+	}
+
 	void IOscBase::SetSemitones(const float InSemiTone)
 	{
 		if (FreqData.Semitones != InSemiTone)
diff --git a/Engine/Source/Runtime/AudioMixer/Private/DSP/Phaser.cpp b/Engine/Source/Runtime/AudioMixer/Private/DSP/Phaser.cpp
--- a/Engine/Source/Runtime/AudioMixer/Private/DSP/Phaser.cpp
+++ b/Engine/Source/Runtime/AudioMixer/Private/DSP/Phaser.cpp
@@ -44,23 +44,26 @@ namespace Audio
 		LFO.Start();
 
 		// Setup the LFO oscillation ranges for APF cutoff frequencies
-		APFFrequencyRanges[0].X = 16.0f;
-		APFFrequencyRanges[0].Y = 1600.0f;
+		// Use ratios of current sample rate
+		const float SampleRateRatio = SampleRate / 48000.0f;
+		APFFrequencyRanges[0].X = 16.0f * SampleRateRatio;
+		APFFrequencyRanges[0].Y = 1600.0f * SampleRateRatio;
 
-		APFFrequencyRanges[1].X = 33.0f;
-		APFFrequencyRanges[1].Y = 3300.0f;
+		APFFrequencyRanges[1].X = 33.0f * SampleRateRatio;
+		APFFrequencyRanges[1].Y = 3300.0f * SampleRateRatio;
 
-		APFFrequencyRanges[2].X = 48.0f;
-		APFFrequencyRanges[2].Y = 4800.0f;
+		APFFrequencyRanges[2].X = 48.0f * SampleRateRatio;
+		APFFrequencyRanges[2].Y = 4800.0f * SampleRateRatio;
 
-		APFFrequencyRanges[3].X = 98.0f;
-		APFFrequencyRanges[3].Y = 9800.0f;
+		APFFrequencyRanges[3].X = 98.0f * SampleRateRatio;
+		APFFrequencyRanges[3].Y = 9800.0f * SampleRateRatio;
 
-		APFFrequencyRanges[4].X = 160.0f;
-		APFFrequencyRanges[4].Y = 16000.0f;
+		APFFrequencyRanges[4].X = 160.0f * SampleRateRatio;
+		APFFrequencyRanges[4].Y = 16000.0f * SampleRateRatio;
 
-		APFFrequencyRanges[5].X = 220.0f;
-		APFFrequencyRanges[5].Y = 22000.0f;
+		// Make sure final APF frequency won't go above Nyquist
+		APFFrequencyRanges[5].X = 220.0f * SampleRateRatio;
+		APFFrequencyRanges[5].Y = 22000.0f * SampleRateRatio;
 
 		FeedbackFrame[0] = 0.0f;
 		FeedbackFrame[1] = 0.0f;
diff --git a/Engine/Source/Runtime/AudioMixer/Private/Effects/AudioMixerSubmixEffectDynamicsProcessor.cpp b/Engine/Source/Runtime/AudioMixer/Private/Effects/AudioMixerSubmixEffectDynamicsProcessor.cpp
--- a/Engine/Source/Runtime/AudioMixer/Private/Effects/AudioMixerSubmixEffectDynamicsProcessor.cpp
+++ b/Engine/Source/Runtime/AudioMixer/Private/Effects/AudioMixerSubmixEffectDynamicsProcessor.cpp
@@ -1,6 +1,10 @@
 // Copyright 1998-2019 Epic Games, Inc. All Rights Reserved.
 
 #include "SubmixEffects/AudioMixerSubmixEffectDynamicsProcessor.h"
+#include "ProfilingDebugging/CsvProfiler.h"
+
+// Link to "Audio" profiling category
+CSV_DECLARE_CATEGORY_MODULE_EXTERN(AUDIOMIXER_API, Audio);
 
 FSubmixEffectDynamicsProcessor::FSubmixEffectDynamicsProcessor()
 {
@@ -69,6 +73,8 @@ void FSubmixEffectDynamicsProcessor::OnPresetChanged()
 
 void FSubmixEffectDynamicsProcessor::OnProcessAudio(const FSoundEffectSubmixInputData& InData, FSoundEffectSubmixOutputData& OutData)
 {
+	CSV_SCOPED_TIMING_STAT(Audio, SubmixDynamics);
+
 	const Audio::AlignedFloatBuffer& InBuffer = *InData.AudioBuffer;
 	Audio::AlignedFloatBuffer& OutBuffer = *OutData.AudioBuffer;
 
@@ -101,4 +107,4 @@ void FSubmixEffectDynamicsProcessor::OnProcessAudio(const FSoundEffectSubmixInpu
 void USubmixEffectDynamicsProcessorPreset::SetSettings(const FSubmixEffectDynamicsProcessorSettings& InSettings)
 {
 	UpdateSettings(InSettings);
-}
\ No newline at end of file
+}
diff --git a/Engine/Source/Runtime/AudioMixer/Private/Effects/AudioMixerSubmixEffectEQ.cpp b/Engine/Source/Runtime/AudioMixer/Private/Effects/AudioMixerSubmixEffectEQ.cpp
--- a/Engine/Source/Runtime/AudioMixer/Private/Effects/AudioMixerSubmixEffectEQ.cpp
+++ b/Engine/Source/Runtime/AudioMixer/Private/Effects/AudioMixerSubmixEffectEQ.cpp
@@ -4,7 +4,10 @@
 #include "SubmixEffects/AudioMixerSubmixEffectEQ.h"
 #include "Misc/ScopeLock.h"
 #include "AudioMixer.h"
+#include "ProfilingDebugging/CsvProfiler.h"
 
+// Link to "Audio" profiling category
+CSV_DECLARE_CATEGORY_MODULE_EXTERN(AUDIOMIXER_API, Audio);
 
 static bool IsEqual(const FSubmixEffectSubmixEQSettings& Left, const FSubmixEffectSubmixEQSettings& Right)
 {
@@ -89,7 +92,7 @@ void FSubmixEffectSubmixEQ::OnProcessAudio(const FSoundEffectSubmixInputData& In
 {
 	LLM_SCOPE(ELLMTag::AudioMixer);
 
- 	SCOPE_CYCLE_COUNTER(STAT_AudioMixerMasterEQ);
+	CSV_SCOPED_TIMING_STAT(Audio, SubmixEQ);
 
 	// Update parameters that may have been set from game thread
 	UpdateParameters(InData.NumChannels);
@@ -267,4 +270,4 @@ void FSubmixEffectSubmixEQ::UpdateParameters(const int32 InNumOutputChannels)
 void USubmixEffectSubmixEQPreset::SetSettings(const FSubmixEffectSubmixEQSettings& InSettings)
 {
 	UpdateSettings(InSettings);
-}
\ No newline at end of file
+}
diff --git a/Engine/Source/Runtime/AudioMixer/Private/Effects/AudioMixerSubmixEffectReverb.cpp b/Engine/Source/Runtime/AudioMixer/Private/Effects/AudioMixerSubmixEffectReverb.cpp
--- a/Engine/Source/Runtime/AudioMixer/Private/Effects/AudioMixerSubmixEffectReverb.cpp
+++ b/Engine/Source/Runtime/AudioMixer/Private/Effects/AudioMixerSubmixEffectReverb.cpp
@@ -5,28 +5,32 @@
 #include "Sound/ReverbEffect.h"
 #include "Audio.h"
 #include "AudioMixer.h"
+#include "ProfilingDebugging/CsvProfiler.h"
 
-static int32 DisableSubmixReverbCVar = 0;
-static FAutoConsoleVariableRef CVarDisableSubmixReverb(
+// Link to "Audio" profiling category
+CSV_DECLARE_CATEGORY_MODULE_EXTERN(AUDIOMIXER_API, Audio);
+
+static int32 DisableSubmixReverbLegacyCVar = 0;
+static FAutoConsoleVariableRef CVarDisableSubmixReverbLegacy(
 	TEXT("au.DisableReverbSubmixLegacy"),
-	DisableSubmixReverbCVar,
+	DisableSubmixReverbLegacyCVar,
 	TEXT("Disables the reverb submix.\n")
 	TEXT("0: Not Disabled, 1: Disabled"),
 	ECVF_Default);
 
-static int32 EnableReverbStereoFlipForQuadCVar = 0;
-static FAutoConsoleVariableRef CVarReverbStereoFlipForQuad(
+static int32 EnableReverbStereoFlipForQuadLegacyCVar = 0;
+static FAutoConsoleVariableRef CVarReverbStereoFlipForQuadLegacy(
 	TEXT("au.EnableReverbStereoFlipForQuadLegacy"),
-	EnableReverbStereoFlipForQuadCVar,
+	EnableReverbStereoFlipForQuadLegacyCVar,
 	TEXT("Enables doing a stereo flip for quad reverb when in surround.\n")
 	TEXT("0: Not Enabled, 1: Enabled"),
 	ECVF_Default);
 
 
-static int32 DisableQuadReverbCVar = 0;
-static FAutoConsoleVariableRef CVarDisableQuadReverbCVar(
+static int32 DisableQuadReverbLegacyCVar = 0;
+static FAutoConsoleVariableRef CVarDisableQuadReverbLegacy(
 	TEXT("au.DisableQuadReverbLegacy"),
-	DisableQuadReverbCVar,
+	DisableQuadReverbLegacyCVar,
 	TEXT("Disables quad reverb in surround.\n")
 	TEXT("0: Not Disabled, 1: Disabled"),
 	ECVF_Default);
@@ -102,13 +106,13 @@ void FSubmixEffectReverb::OnProcessAudio(const FSoundEffectSubmixInputData& InDa
 	LLM_SCOPE(ELLMTag::AudioMixer);
 
 	check(InData.NumChannels == 2);
- 	if (OutData.NumChannels < 2 || !bIsEnabled || DisableSubmixReverbCVar == 1) 
+ 	if (OutData.NumChannels < 2 || !bIsEnabled || DisableSubmixReverbLegacyCVar == 1) 
 	{
 		// Not supported
 		return;
 	}
 
-	SCOPE_CYCLE_COUNTER(STAT_AudioMixerMasterReverb);
+	CSV_SCOPED_TIMING_STAT(Audio, SubmixReverb);
 
 	UpdateParameters();
 
@@ -126,14 +130,14 @@ void FSubmixEffectReverb::OnProcessAudio(const FSoundEffectSubmixInputData& InDa
 	// 5.1 or higher surround sound. Map stereo output to quad output
 	else if(OutData.NumChannels > 5)
 	{
-		if (DisableQuadReverbCVar == 1)
+		if (DisableQuadReverbLegacyCVar == 1)
 		{
 			for (int32 InSampleIndex = 0, OutSampleIndex = 0; InSampleIndex < InData.AudioBuffer->Num(); InSampleIndex += InData.NumChannels, OutSampleIndex += OutData.NumChannels)
 			{
 				PlateReverb.ProcessAudioFrame(&AudioData[InSampleIndex], InData.NumChannels, &OutAudioData[OutSampleIndex], InData.NumChannels);
 			}
 		}
-		else if (EnableReverbStereoFlipForQuadCVar == 1)
+		else if (EnableReverbStereoFlipForQuadLegacyCVar == 1)
 		{
 			for (int32 InSampleIndex = 0, OutSampleIndex = 0; InSampleIndex < InData.AudioBuffer->Num(); InSampleIndex += InData.NumChannels, OutSampleIndex += OutData.NumChannels)
 			{
diff --git a/Engine/Source/Runtime/AudioMixer/Public/AudioMixer.h b/Engine/Source/Runtime/AudioMixer/Public/AudioMixer.h
--- a/Engine/Source/Runtime/AudioMixer/Public/AudioMixer.h
+++ b/Engine/Source/Runtime/AudioMixer/Public/AudioMixer.h
@@ -8,7 +8,6 @@
 #include "AudioEffect.h"
 #include "AudioMixerTypes.h"
 #include "HAL/Runnable.h"
-#include "Stats/Stats.h"
 #include "Sound/AudioSettings.h"
 #include "Misc/SingleThreadRunnable.h"
 #include "AudioMixerNullDevice.h"
@@ -26,45 +25,6 @@
 #endif
 #endif
 
-// Cycle stats for audio mixer
-DECLARE_STATS_GROUP(TEXT("AudioMixer"), STATGROUP_AudioMixer, STATCAT_Advanced);
-
-// Tracks the time for the full render block 
-DECLARE_CYCLE_STAT_EXTERN(TEXT("Render Audio"), STAT_AudioMixerRenderAudio, STATGROUP_AudioMixer, AUDIOMIXER_API);
-
-// Tracks the time it takes to up the source manager (computes source buffers, source effects, sample rate conversion)
-DECLARE_CYCLE_STAT_EXTERN(TEXT("Source Manager Update"), STAT_AudioMixerSourceManagerUpdate, STATGROUP_AudioMixer, AUDIOMIXER_API);
-
-// The time it takes to compute the source buffers (handle decoding tasks, resampling)
-DECLARE_CYCLE_STAT_EXTERN(TEXT("Source Buffers"), STAT_AudioMixerSourceBuffers, STATGROUP_AudioMixer, AUDIOMIXER_API);
-
-// The time it takes to process the source buffers through their source effects
-DECLARE_CYCLE_STAT_EXTERN(TEXT("Source Effect Buffers"), STAT_AudioMixerSourceEffectBuffers, STATGROUP_AudioMixer, AUDIOMIXER_API);
-
-// The time it takes to apply channel maps and get final pre-submix source buffers
-DECLARE_CYCLE_STAT_EXTERN(TEXT("Source Output Buffers"), STAT_AudioMixerSourceOutputBuffers, STATGROUP_AudioMixer, AUDIOMIXER_API);
-
-// The time it takes to process the subix graph. Process submix effects, mix into the submix buffer, etc.
-DECLARE_CYCLE_STAT_EXTERN(TEXT("Submix Graph"), STAT_AudioMixerSubmixes, STATGROUP_AudioMixer, AUDIOMIXER_API);
-
-// The time it takes to process the subix graph. Process submix effects, mix into the submix buffer, etc.
-DECLARE_CYCLE_STAT_EXTERN(TEXT("Submix Graph Child Processing"), STAT_AudioMixerSubmixChildren, STATGROUP_AudioMixer, AUDIOMIXER_API);
-
-// The time it takes to process the subix graph. Process submix effects, mix into the submix buffer, etc.
-DECLARE_CYCLE_STAT_EXTERN(TEXT("Submix Graph Source Mixing"), STAT_AudioMixerSubmixSource, STATGROUP_AudioMixer, AUDIOMIXER_API);
-
-// The time it takes to process the subix graph. Process submix effects, mix into the submix buffer, etc.
-DECLARE_CYCLE_STAT_EXTERN(TEXT("Submix Graph Effect Processing"), STAT_AudioMixerSubmixEffectProcessing, STATGROUP_AudioMixer, AUDIOMIXER_API);
-
-// The time it takes to process the master reverb.
-DECLARE_CYCLE_STAT_EXTERN(TEXT("Master Reverb"), STAT_AudioMixerMasterReverb, STATGROUP_AudioMixer, AUDIOMIXER_API);
-
-// The time it takes to process the master EQ effect.
-DECLARE_CYCLE_STAT_EXTERN(TEXT("Master EQ"), STAT_AudioMixerMasterEQ, STATGROUP_AudioMixer, AUDIOMIXER_API);
-
-// The time it takes to process the HRTF effect.
-DECLARE_CYCLE_STAT_EXTERN(TEXT("HRTF"), STAT_AudioMixerHRTF, STATGROUP_AudioMixer, AUDIOMIXER_API);
-
 
 // Enable debug checking for audio mixer
 
diff --git a/Engine/Source/Runtime/AudioMixer/Public/Components/SynthComponent.h b/Engine/Source/Runtime/AudioMixer/Public/Components/SynthComponent.h
--- a/Engine/Source/Runtime/AudioMixer/Public/Components/SynthComponent.h
+++ b/Engine/Source/Runtime/AudioMixer/Public/Components/SynthComponent.h
@@ -1,13 +1,13 @@
 // Copyright 1998-2019 Epic Games, Inc. All Rights Reserved.
-
 #pragma once
 
-#include "CoreMinimal.h"
+#include "AudioMixerTypes.h"
 #include "Components/AudioComponent.h"
-#include "UObject/ObjectMacros.h"
+#include "CoreMinimal.h"
 #include "Engine/EngineTypes.h"
+#include "IAudioExtensionPlugin.h"
 #include "Sound/SoundWaveProcedural.h"
-#include "AudioMixerTypes.h"
+#include "UObject/ObjectMacros.h"
 
 #include "SynthComponent.generated.h"
 
@@ -162,6 +162,10 @@ public:
 	UPROPERTY(EditAnywhere, BlueprintReadWrite, Category = Effects, meta = (DisplayName = "Post-Effect Bus Sends"))
 	TArray<FSoundSourceBusSendInfo> BusSends;
 
+	/** Modulation for the sound */
+	UPROPERTY(EditAnywhere, Category = Modulation)
+	FSoundModulation Modulation;
+
 	/** This sound will send its audio output to this list of buses if there are bus instances playing before source effects are processed.  */
 	UPROPERTY(EditAnywhere, BlueprintReadWrite, Category = Effects, meta = (DisplayName = "Pre-Effect Bus Sends"))
 	TArray<FSoundSourceBusSendInfo> PreEffectBusSends;
diff --git a/Engine/Source/Runtime/AudioMixer/Public/DSP/Dsp.h b/Engine/Source/Runtime/AudioMixer/Public/DSP/Dsp.h
--- a/Engine/Source/Runtime/AudioMixer/Public/DSP/Dsp.h
+++ b/Engine/Source/Runtime/AudioMixer/Public/DSP/Dsp.h
@@ -7,7 +7,7 @@
 
 // Macros which can be enabled to cause DSP sample checking
 #if 0
-#define CHECK_SAMPLE(VALUE)
+#define CHECK_SAMPLE(VALUE) 
 #define CHECK_SAMPLE2(VALUE)
 #else
 #define CHECK_SAMPLE(VALUE)  Audio::CheckSample(VALUE)
@@ -16,10 +16,10 @@
 
 namespace Audio
 {
-	// Utility to check for sample clipping. Put breakpoint in conditional to find
+	// Utility to check for sample clipping. Put breakpoint in conditional to find 
 	// DSP code that's not behaving correctly
 	static void CheckSample(float InSample, float Threshold = 0.001f)
-	{
+	{	
 		if (InSample > Threshold || InSample < -Threshold)
 		{
 			UE_LOG(LogTemp, Log, TEXT("SampleValue Was %.2f"), InSample);
@@ -73,7 +73,7 @@ namespace Audio
 		return X2;
 	}
 
-	// Sine approximation using Bhaskara I technique discovered in 7th century.
+	// Sine approximation using Bhaskara I technique discovered in 7th century. 
 	// https://en.wikipedia.org/wiki/Bh%C4%81skara_I
 	static FORCEINLINE float FastSin3(const float X)
 	{
@@ -167,6 +167,24 @@ namespace Audio
 		OutLeft = FMath::Clamp(FastSin(LeftPhase), 0.0f, 1.0f);
 		OutRight = FMath::Clamp(FastSin(RightPhase), 0.0f, 1.0f);
 	}
+ 
+	// This function encodes a stereo Left/Right signal into a stereo Mid/Side signal 
+	static FORCEINLINE void EncodeMidSide(float& LeftChannel, float& RightChannel)
+	{
+		const float Temp = (LeftChannel - RightChannel);
+		//Output
+		LeftChannel = (LeftChannel + RightChannel);
+		RightChannel = Temp;
+	}
+
+	// This function decodes a stereo Mid/Side signal into a stereo Left/Right signal
+	static FORCEINLINE void DecodeMidSide(float& MidChannel, float& SideChannel)
+	{
+		const float Temp = (MidChannel - SideChannel) * 0.5f;
+		//Output
+		MidChannel = (MidChannel + SideChannel) * 0.5f;
+		SideChannel = Temp;
+	}
 
 	// Helper function to get bandwidth from Q
 	static FORCEINLINE float GetBandwidthFromQ(const float InQ)
@@ -187,7 +205,7 @@ namespace Audio
 		return OutQ;
 	}
 
-	// Polynomial interpolation using lagrange polynomials.
+	// Polynomial interpolation using lagrange polynomials. 
 	// https://en.wikipedia.org/wiki/Lagrange_polynomial
 	static FORCEINLINE float LagrangianInterpolation(const TArray<FVector2D> Points, const float Alpha)
 	{
@@ -288,9 +306,9 @@ namespace Audio
 		// Percentage to move toward target value from current value each tick
 		float EaseFactor;
 	};
-
+	
 	// Simple easing function used to help interpolate params
-	class FLinearEase
+	class FLinearEase 
 	{
 	public:
 		FLinearEase()
@@ -339,7 +357,7 @@ namespace Audio
 			++CurrentTick;
 			return CurrentValue;
 		}
-
+		 
 		// Updates the target value without changing the duration or tick data.
 		// Sets the state as if the new value was the target value all along
 		void SetValueInterrupt(const float InValue)
@@ -372,7 +390,7 @@ namespace Audio
 
 			if (DurationTicks == 0)
 			{
-				CurrentValue = InValue;
+				CurrentValue = InValue;			
 			}
 			else
 			{
diff --git a/Engine/Source/Runtime/AudioMixer/Public/DSP/Envelope.h b/Engine/Source/Runtime/AudioMixer/Public/DSP/Envelope.h
--- a/Engine/Source/Runtime/AudioMixer/Public/DSP/Envelope.h
+++ b/Engine/Source/Runtime/AudioMixer/Public/DSP/Envelope.h
@@ -31,6 +31,10 @@ namespace Audio
 		// Start the envelope, puts envelope in attack state
 		virtual void Start();
 
+		// For truly legato envelope logic, we need to know the new sustain gain (if its being changed).
+		virtual void StartLegato(const float InNewDepth);
+		virtual void StartLegato() { StartLegato(Depth); }
+
 		// Stop the envelope, puts in the release state. Can optionally force to off state.
 		virtual void Stop();
 
@@ -122,6 +126,7 @@ namespace Audio
 		float ShutdownDelta;
 		float Depth;
 		float BiasDepth;
+		float OutputGain;
 
 		FEnvData AttackData;
 		FEnvData DecayData;
@@ -136,21 +141,25 @@ namespace Audio
 		FPatchSource BiasedEnvSource;
 
 		// Whether or not we want to simulate analog behavior
-		bool bIsSimulatingAnalog;
+		uint8 bIsSimulatingAnalog:1;
 
 		// Whether or not the envelope is reset back to attack when started
-		bool bIsLegatoMode;
+		uint8 bIsLegatoMode:1;
 
 		// Whether or not the envelope value is set to zero when finished
-		bool bIsRetriggerMode;
+		uint8 bIsRetriggerMode:1;
 
 		// Whether or not this envelope has changed and needs to have values recomputed
-		bool bChanged;
+		uint8 bChanged:1;
 
 		// Inverts the output
-		bool bInvert;
+		uint8 bInvert:1;
 
 		// Bias output inversions
-		bool bBiasInvert;
+		uint8 bBiasInvert:1;
+
+		// tracks if the current envelope was started with sustain at 0.0
+		// (avoids bug where sustain being turned up during decay phase makes note hang)
+		uint8 bCurrentCycleIsADOnly:1;
 	};
 }
diff --git a/Engine/Source/Runtime/AudioMixer/Public/DSP/EventQuantizer.h b/Engine/Source/Runtime/AudioMixer/Public/DSP/EventQuantizer.h
--- a/Engine/Source/Runtime/AudioMixer/Public/DSP/EventQuantizer.h
+++ b/Engine/Source/Runtime/AudioMixer/Public/DSP/EventQuantizer.h
@@ -121,7 +121,7 @@ namespace Audio
 			// The current frame count of this quantization type
 			uint32 FrameCount;
 
-			// The frame duration of this quantization type (how many frames for the event to re-ocurr)
+			// The frame duration of this quantization type (how many frames for the event to reoccur)
 			uint32 EventFrameDuration;
 
 			// Array of events queued for this state
diff --git a/Engine/Source/Runtime/AudioMixer/Public/DSP/LFO.h b/Engine/Source/Runtime/AudioMixer/Public/DSP/LFO.h
--- a/Engine/Source/Runtime/AudioMixer/Public/DSP/LFO.h
+++ b/Engine/Source/Runtime/AudioMixer/Public/DSP/LFO.h
@@ -12,7 +12,7 @@ namespace Audio
 	{
 		enum Type
 		{
-			Sine,
+			Sine = 0,
 			UpSaw,
 			DownSaw,
 			Square,
@@ -29,7 +29,7 @@ namespace Audio
 		enum Type
 		{
 			// Constantly oscillates
-			Sync,
+			Sync = 0,
 
 			// Performs the LFO only once, then stops
 			OneShot,
diff --git a/Engine/Source/Runtime/AudioMixer/Public/DSP/Osc.h b/Engine/Source/Runtime/AudioMixer/Public/DSP/Osc.h
--- a/Engine/Source/Runtime/AudioMixer/Public/DSP/Osc.h
+++ b/Engine/Source/Runtime/AudioMixer/Public/DSP/Osc.h
@@ -49,6 +49,8 @@ namespace Audio
 
 		FOscFrequencyMod(const FOscFrequencyMod& Other)
 			: Scale(Other.Scale)
+			, ExternalMod(Other.ExternalMod)
+			, Mod(Other.Mod)
 			, Detune(Other.Detune)
 			, PitchBend(Other.PitchBend)
 			, Octave(Other.Octave)
@@ -100,8 +102,12 @@ namespace Audio
 		// Returns the frequency of the oscillator
 		float GetFrequency() const { return BaseFreq; }
 
+		// Returns the frequency of the oscillator
+		float GetGain() const { return Gain; }
+
 		void SetCents(const float InCents);
 		void SetOctave(const float InOctave);
+		void SetSampleRate(const float InSampleRate);
 		void SetSemitones(const float InSemiTone);
 		void SetDetune(const float InDetune);
 		void SetPitchBend(const float InPitchBend);
diff --git a/Engine/Source/Runtime/Engine/Classes/Components/AudioComponent.h b/Engine/Source/Runtime/Engine/Classes/Components/AudioComponent.h
--- a/Engine/Source/Runtime/Engine/Classes/Components/AudioComponent.h
+++ b/Engine/Source/Runtime/Engine/Classes/Components/AudioComponent.h
@@ -1,14 +1,13 @@
 // Copyright 1998-2019 Epic Games, Inc. All Rights Reserved.
-
-
 #pragma once
 
+#include "Components/SceneComponent.h"
 #include "CoreMinimal.h"
-#include "UObject/ObjectMacros.h"
 #include "Engine/EngineTypes.h"
-#include "Components/SceneComponent.h"
+#include "IAudioExtensionPlugin.h"
 #include "Sound/SoundAttenuation.h"
 #include "Sound/SoundWave.h"
+#include "UObject/ObjectMacros.h"
 
 #include "AudioComponent.generated.h"
 
@@ -17,6 +16,7 @@ class USoundBase;
 class USoundClass;
 class USoundConcurrency;
 
+
 /** called when we finish playing audio, either because it played to completion or because a Stop() call turned it off early */
 DECLARE_DYNAMIC_MULTICAST_DELEGATE( FOnAudioFinished );
 
@@ -336,6 +336,10 @@ public:
 	/** shadow delegate for non UObject subscribers */
 	FOnAudioMultiEnvelopeValueNative OnAudioMultiEnvelopeValueNative;
 
+	/** Modulation for the sound */
+	UPROPERTY(EditAnywhere, BlueprintReadWrite, Category = Modulation)
+	FSoundModulation Modulation;
+
 	/** Called when subtitles are sent to the SubtitleManager.  Set this delegate if you want to hijack the subtitles for other purposes */
 	UPROPERTY()
 	FOnQueueSubtitles OnQueueSubtitles;
diff --git a/Engine/Source/Runtime/Engine/Classes/Engine/Engine.h b/Engine/Source/Runtime/Engine/Classes/Engine/Engine.h
--- a/Engine/Source/Runtime/Engine/Classes/Engine/Engine.h
+++ b/Engine/Source/Runtime/Engine/Classes/Engine/Engine.h
@@ -24,6 +24,7 @@
 class AMatineeActor;
 class APlayerController;
 class Error;
+class FAudioDeviceManager;
 class FCanvas;
 class FCommonViewportClient;
 class FFineGrainedPerformanceTracker;
@@ -1770,7 +1771,7 @@ private:
 protected:
 
 	/** The audio device manager */
-	class FAudioDeviceManager* AudioDeviceManager;
+	FAudioDeviceManager* AudioDeviceManager;
 
 	/** Audio device handle to the main audio device. */
 	uint32 MainAudioDeviceHandle;
@@ -2430,7 +2431,7 @@ public:
 	bool IsEditor();
 
 	/** @return the audio device manager of the UEngine, this allows the creation and management of multiple audio devices. */
-	class FAudioDeviceManager* GetAudioDeviceManager();
+	FAudioDeviceManager* GetAudioDeviceManager();
 
 	/** @return the main audio device handle used by the engine. */
 	uint32 GetAudioDeviceHandle() const;
@@ -3276,7 +3277,7 @@ public:
 
 private:
 	/**
-	 * Function definition for those stats which have their own render funcsions (or affect another render functions).
+	 * Function definition for those stats which have their own render functions (or affect another render functions).
 	 *
 	 * @param World	The world being drawn to.
 	 * @param ViewportClient The viewport being drawn to.
@@ -3289,7 +3290,7 @@ private:
 	typedef int32 (UEngine::*EngineStatRender)(UWorld* World, FViewport* Viewport, FCanvas* Canvas, int32 X, int32 Y, const FVector* ViewLocation, const FRotator* ViewRotation);
 
 	/**
-	 * Function definition for those stats which have their own toggle funcsions (or toggle other stats).
+	 * Function definition for those stats which have their own toggle functions (or toggle other stats).
 	 *
 	 * @param World	The world being drawn to.
 	 * @param ViewportClient The viewport being drawn to.
@@ -3372,6 +3373,7 @@ private:
 	bool ToggleStatNamedEvents(UWorld* World, FCommonViewportClient* ViewportClient, const TCHAR* Stream = nullptr);
 	bool ToggleStatUnit(UWorld* World, FCommonViewportClient* ViewportClient, const TCHAR* Stream = nullptr);
 #if !UE_BUILD_SHIPPING
+	bool PostStatSoundModulatorHelp(UWorld* World, FCommonViewportClient* ViewportClient, const TCHAR* Stream = nullptr);
 	bool ToggleStatUnitMax(UWorld* World, FCommonViewportClient* ViewportClient, const TCHAR* Stream = nullptr);
 	bool ToggleStatUnitGraph(UWorld* World, FCommonViewportClient* ViewportClient, const TCHAR* Stream = nullptr);
 	bool ToggleStatUnitTime(UWorld* World, FCommonViewportClient* ViewportClient, const TCHAR* Stream = nullptr);
@@ -3380,6 +3382,7 @@ private:
 	bool ToggleStatSoundCues(UWorld* World, FCommonViewportClient* ViewportClient, const TCHAR* Stream = nullptr);
 	bool ToggleStatSounds(UWorld* World, FCommonViewportClient* ViewportClient, const TCHAR* Stream = nullptr);
 	bool ToggleStatSoundMixes(UWorld* World, FCommonViewportClient* ViewportClient, const TCHAR* Stream = nullptr);
+	bool ToggleStatSoundModulators(UWorld* World, FCommonViewportClient* ViewportClient, const TCHAR* Stream = nullptr);
 #endif
 
 	/**
@@ -3405,8 +3408,9 @@ private:
 	int32 RenderStatLevelMap(UWorld* World, FViewport* Viewport, FCanvas* Canvas, int32 X, int32 Y, const FVector* ViewLocation = nullptr, const FRotator* ViewRotation = nullptr);
 	int32 RenderStatUnit(UWorld* World, FViewport* Viewport, FCanvas* Canvas, int32 X, int32 Y, const FVector* ViewLocation = nullptr, const FRotator* ViewRotation = nullptr);
 #if !UE_BUILD_SHIPPING
-	int32 RenderStatReverb(UWorld* World, FViewport* Viewport, FCanvas* Canvas, int32 X, int32 Y, const FVector* ViewLocation = nullptr, const FRotator* ViewRotation = nullptr);
-	int32 RenderStatSoundMixes(UWorld* World, FViewport* Viewport, FCanvas* Canvas, int32 X, int32 Y, const FVector* ViewLocation = nullptr, const FRotator* ViewRotation = nullptr);
+	int32 RenderStatSoundReverb(UWorld* World, FViewport* Viewport, FCanvas* Canvas, int32 X, int32 Y, const FVector* ViewLocation = nullptr, const FRotator* ViewRotation = nullptr);
+ 	int32 RenderStatSoundMixes(UWorld* World, FViewport* Viewport, FCanvas* Canvas, int32 X, int32 Y, const FVector* ViewLocation = nullptr, const FRotator* ViewRotation = nullptr);
+	int32 RenderStatSoundModulators(UWorld* World, FViewport* Viewport, FCanvas* Canvas, int32 X, int32 Y, const FVector* ViewLocation = nullptr, const FRotator* ViewRotation = nullptr);
 	int32 RenderStatSoundWaves(UWorld* World, FViewport* Viewport, FCanvas* Canvas, int32 X, int32 Y, const FVector* ViewLocation = nullptr, const FRotator* ViewRotation = nullptr);
 	int32 RenderStatSoundCues(UWorld* World, FViewport* Viewport, FCanvas* Canvas, int32 X, int32 Y, const FVector* ViewLocation = nullptr, const FRotator* ViewRotation = nullptr);
 	int32 RenderStatSounds(UWorld* World, FViewport* Viewport, FCanvas* Canvas, int32 X, int32 Y, const FVector* ViewLocation = nullptr, const FRotator* ViewRotation = nullptr);
diff --git a/Engine/Source/Runtime/Engine/Classes/Engine/GameViewportClient.h b/Engine/Source/Runtime/Engine/Classes/Engine/GameViewportClient.h
--- a/Engine/Source/Runtime/Engine/Classes/Engine/GameViewportClient.h
+++ b/Engine/Source/Runtime/Engine/Classes/Engine/GameViewportClient.h
@@ -645,22 +645,6 @@ public:
 		return EnabledStats.Contains(InName);
 	}
 
-	/**
-	 * Get the sound stat flags enabled for this viewport
-	 */
-	virtual ESoundShowFlags::Type GetSoundShowFlags() const override
-	{ 
-		return SoundShowFlags;
-	}
-
-	/**
-	 * Set the sound stat flags enabled for this viewport
-	 */
-	virtual void SetSoundShowFlags(const ESoundShowFlags::Type InSoundShowFlags) override
-	{
-		SoundShowFlags = InSoundShowFlags;
-	}
-
 	/**
 	 * Set whether to ignore input.
 	 */
@@ -953,9 +937,6 @@ private:
 	/** A list of all the stat names which are enabled for this viewport (static so they persist between runs) */
 	static TArray<FString> EnabledStats;
 
-	/** Those sound stat flags which are enabled on this viewport */
-	static ESoundShowFlags::Type SoundShowFlags;
-
 	/** Disables splitscreen, useful when game code is in menus, and doesn't want splitscreen on */
 	bool bDisableSplitScreenOverride;
 
diff --git a/Engine/Source/Runtime/Engine/Classes/Sound/AudioSettings.h b/Engine/Source/Runtime/Engine/Classes/Sound/AudioSettings.h
--- a/Engine/Source/Runtime/Engine/Classes/Sound/AudioSettings.h
+++ b/Engine/Source/Runtime/Engine/Classes/Sound/AudioSettings.h
@@ -103,6 +103,8 @@ class ENGINE_API UAudioSettings : public UDeveloperSettings
 	virtual void PostEditChangeChainProperty( struct FPropertyChangedChainEvent& PropertyChangedEvent) override;
 #endif
 
+	virtual void Serialize(FArchive& Ar) override;
+
 	/** The SoundClass assigned to newly created sounds */
 	UPROPERTY(config, EditAnywhere, Category="Audio", meta=(AllowedClasses="SoundClass", DisplayName="Default Sound Class"))
 	FSoftObjectPath DefaultSoundClassName;
@@ -115,7 +117,7 @@ class ENGINE_API UAudioSettings : public UDeveloperSettings
 	UPROPERTY(config, EditAnywhere, Category = "Audio", meta = (AllowedClasses = "SoundConcurrency", DisplayName = "Default Sound Concurrency"))
 	FSoftObjectPath DefaultSoundConcurrencyName;
 
-	/** The SoundMix to use as base when no other system has speciicefied a Base SoundMix */
+	/** The SoundMix to use as base when no other system has specified a Base SoundMix */
 	UPROPERTY(config, EditAnywhere, Category="Audio", meta=(AllowedClasses="SoundMix"))
 	FSoftObjectPath DefaultBaseSoundMix;
 
@@ -156,7 +158,10 @@ class ENGINE_API UAudioSettings : public UDeveloperSettings
 
 	/** Allows sounds to play at 0 volume. */
 	UPROPERTY(config, EditAnywhere, Category = "Quality", AdvancedDisplay)
-	uint32 bAllowVirtualizedSounds:1;
+	uint32 bAllowPlayWhenSilent:1;
+
+	UPROPERTY()
+	uint32 bAllowVirtualizedSounds_DEPRECATED:1;
 
 	/** Disables master EQ effect in the audio DSP graph. */
 	UPROPERTY(config, EditAnywhere, Category = "Quality", AdvancedDisplay)
diff --git a/Engine/Source/Runtime/Engine/Classes/Sound/SampleBuffer.h b/Engine/Source/Runtime/Engine/Classes/Sound/SampleBuffer.h
--- a/Engine/Source/Runtime/Engine/Classes/Sound/SampleBuffer.h
+++ b/Engine/Source/Runtime/Engine/Classes/Sound/SampleBuffer.h
@@ -229,6 +229,21 @@ namespace Audio
 			return *this;
 		}
 
+		// copy from a container of the same element type
+		void CopyFrom(const TArray<SampleType>& InArray, int32 InNumChannels, int32 InSampleRate)
+		{
+			NumSamples = InArray.Num();
+			NumFrames = NumSamples / InNumChannels;
+			NumChannels = InNumChannels;
+			SampleRate = InSampleRate;
+			SampleDuration = ((float)NumFrames) / SampleRate;
+
+			RawPCMData.Reset(NumSamples);
+			RawPCMData.AddZeroed(NumSamples);
+
+			FMemory::Memcpy(RawPCMData.GetData(), InArray.GetData(), NumSamples * sizeof(SampleType));
+		}
+
 		// Append audio data to internal buffer of different sample type of this sample buffer
 		template<class OtherSampleType>
 		void Append(const OtherSampleType* InputBuffer, int32 InNumSamples)
@@ -416,6 +431,89 @@ namespace Audio
 			NumFrames = RawPCMData.Num() / NumChannels;
 			NumSamples = RawPCMData.Num();
 		}
+
+		// InIndex [0.0f, NumSamples - 1.0f]
+		// OutFrame is the multichannel output for one index value
+		// Returns InIndex wrapped between 0.0 and NumFrames
+		float GetAudioFrameAtFractionalIndex(float InIndex, TArray<SampleType>& OutFrame)
+		{
+			InIndex = FMath::Fmod(InIndex, static_cast<float>(NumFrames));
+
+			GetAudioFrameAtFractionalIndexInternal(InIndex, OutFrame);
+
+			return InIndex;
+		}
+
+		// InPhase [0, 1], wrapped, through duration of file (ignores sample rate)
+		// OutFrame is the multichannel output for one phase value
+		// Returns InPhase wrapped between 0.0 and 1.0
+		float GetAudioFrameAtPhase(float InPhase, TArray<SampleType>& OutFrame)
+		{
+			InPhase = FMath::Fmod(InPhase, 1.0f);
+
+			GetAudioFrameAtFractionalIndexInternal(InPhase * NumFrames, OutFrame);
+
+			return InPhase;
+		}
+
+
+		// InTimeSec, get the value of the buffer at the given time (uses sample rate)
+		// OutFrame is the multichannel output for one time value
+		// Returns InTimeSec wrapped between 0.0 and (NumSamples / SampleRate)
+		float GetAudioFrameAtTime(float InTimeSec, TArray<SampleType>& OutFrame)
+		{
+			if (InTimeSec >= SampleDuration)
+			{
+				InTimeSec -= SampleDuration;
+			}
+
+			check(InTimeSec >= 0.0f && InTimeSec <= SampleDuration);
+
+			GetAudioFrameAtFractionalIndexInternal(NumSamples * (InTimeSec / SampleDuration), OutFrame);
+
+			return InTimeSec;
+		}
+
+	private:
+		// Internal implementation. Called by all public GetAudioFrameAt_ _ _ _() functions
+		// public functions do range checking/wrapping and then call this function
+		void GetAudioFrameAtFractionalIndexInternal(float InIndex, TArray<SampleType>& OutFrame)
+		{
+			const float Alpha = FMath::Fmod(InIndex, 1.0f);
+			const int32 WholeThisIndex = FMath::FloorToInt(InIndex);
+			int32 WholeNextIndex = WholeThisIndex + 1;
+
+			// check for interpolation between last and first frames
+			if (WholeNextIndex == NumSamples)
+			{
+				WholeNextIndex = 0;
+			}
+
+			// TODO: if(NumChannels < 4)... do the current (non vectorized) way
+			OutFrame.SetNumUninitialized(NumChannels);
+
+			for (int32 i = 0; i < NumChannels; ++i)
+			{
+				float SampleA, SampleB;
+				
+				if (TIsSame<SampleType, float>::Value)
+				{
+					SampleA = RawPCMData[i * NumChannels + WholeThisIndex];
+					SampleB = RawPCMData[i * NumChannels + WholeNextIndex];
+					OutFrame[i] = FMath::Lerp(SampleA, SampleB, Alpha);
+				}
+				else
+				{
+					SampleA = static_cast<float>(RawPCMData[i * NumChannels + WholeThisIndex]);
+					SampleB = static_cast<float>(RawPCMData[i * NumChannels + WholeNextIndex]);
+					OutFrame[i] = static_cast<SampleType>(FMath::Lerp(SampleA, SampleB, Alpha));
+				}
+			}
+
+			// TODO: else { do vectorized version }
+			// make new function in BufferVectorOperations.cpp
+			// (use FMath::Lerp() overload for VectorRegisters)
+		}
 	};
 
 	// FSampleBuffer is a strictly defined TSampleBuffer that uses the same sample format we use for USoundWaves.
