Change how AssetRegUtil handles file commandline args.

diff --git a/Engine/Source/Editor/UnrealEd/Private/Commandlets/AssetRegUtilCommandlet.cpp b/Engine/Source/Editor/UnrealEd/Private/Commandlets/AssetRegUtilCommandlet.cpp
--- a/Engine/Source/Editor/UnrealEd/Private/Commandlets/AssetRegUtilCommandlet.cpp
+++ b/Engine/Source/Editor/UnrealEd/Private/Commandlets/AssetRegUtilCommandlet.cpp
@@ -430,10 +430,10 @@ int32 UAssetRegUtilCommandlet::Main(const FString& CmdLineParams)
 	AssetRegistry->SearchAllAssets(true);
 	
 	FString ReorderFile;
-	if (FParse::Value(*CmdLineParams, TEXT("ReorderFile="), ReorderFile))
+	if (FParse::Value(*CmdLineParams, TEXT("ReorderFile="), ReorderFile, false))
 	{
 		FString ReorderOutput;
-		if (!FParse::Value(*CmdLineParams, TEXT("ReorderOutput="), ReorderOutput))
+		if (!FParse::Value(*CmdLineParams, TEXT("ReorderOutput="), ReorderOutput, false))
 		{
 			//if nothing specified, base it on the input name
 			ReorderOutput = FPaths::SetExtension(FPaths::SetExtension(ReorderFile, TEXT("")) + TEXT("Reordered"), FPaths::GetExtension(ReorderFile));
