Fixed crash when diffing transaction if only one side has serialized property data

diff --git a/Engine/Source/Editor/UnrealEd/Private/EditorTransaction.cpp b/Engine/Source/Editor/UnrealEd/Private/EditorTransaction.cpp
--- a/Engine/Source/Editor/UnrealEd/Private/EditorTransaction.cpp
+++ b/Engine/Source/Editor/UnrealEd/Private/EditorTransaction.cpp
@@ -473,8 +473,8 @@ void FTransaction::FObjectRecord::Diff( const FTransaction* Owner, const FSerial
 		// Compare the data before the property block to see if something else in the object has changed
 		if (!OutDeltaChange.bHasNonPropertyChanges)
 		{
-			const int32 OldHeaderSize = StartOfOldPropertyBlock;
-			const int32 CurrentHeaderSize = StartOfNewPropertyBlock;
+			const int32 OldHeaderSize = FMath::Min(StartOfOldPropertyBlock, OldSerializedObject.Data.Num());
+			const int32 CurrentHeaderSize = FMath::Min(StartOfNewPropertyBlock, NewSerializedObject.Data.Num());
 
 			bool bIsHeaderIdentical = OldHeaderSize == CurrentHeaderSize;
 			if (bIsHeaderIdentical && CurrentHeaderSize > 0)
@@ -491,8 +491,8 @@ void FTransaction::FObjectRecord::Diff( const FTransaction* Owner, const FSerial
 		// Compare the data after the property block to see if something else in the object has changed
 		if (!OutDeltaChange.bHasNonPropertyChanges)
 		{
-			const int32 OldFooterSize = OldSerializedObject.Data.Num() - EndOfOldPropertyBlock;
-			const int32 CurrentFooterSize = NewSerializedObject.Data.Num() - EndOfNewPropertyBlock;
+			const int32 OldFooterSize = OldSerializedObject.Data.Num() - FMath::Max(EndOfOldPropertyBlock, 0);
+			const int32 CurrentFooterSize = NewSerializedObject.Data.Num() - FMath::Max(EndOfNewPropertyBlock, 0);
 
 			bool bIsFooterIdentical = OldFooterSize == CurrentFooterSize;
 			if (bIsFooterIdentical && CurrentFooterSize > 0)
