Fix for an ensure when reloading a Blueprint asset with the current debug object set to an instance of the Blueprint class.

Change summary:
- Made UBlueprint::CurrentObjectBeingDebugged a transient weak object reference property (so that it can now be serialized during reference collection).
- Removed explicit (manual) reference replacement of UBlueprint::CurrentObjectBeingDebugged in FBlueprintCompileReinstancer::ReplaceInstancesOfClass_Inner().
- Modified FArchiveHasReferences to enable weak object reference value serialization during reference collection.
- Modified FBlueprintCompilationManagerImpl::ReinstanceBatch() to invoke FArchiveReplaceObjectRef in place of FArchiveReplaceOrClearExternalReferences (to avoid a circular path - per DanO).
- Modified ReplaceObjectHelper()/ReplaceActorHelper() and ReplaceInstancesOfClass_Inner() to defer marking objects as PendingKill/destroying Actor instances until after reference replacement.
- Modified FindAndReplaceReferences() to avoid doing reference replacement on subobjects of old instances in addition to the old instances themselves.

diff --git a/Engine/Source/Editor/Kismet/Private/BlueprintCompilationManager.cpp b/Engine/Source/Editor/Kismet/Private/BlueprintCompilationManager.cpp
--- a/Engine/Source/Editor/Kismet/Private/BlueprintCompilationManager.cpp
+++ b/Engine/Source/Editor/Kismet/Private/BlueprintCompilationManager.cpp
@@ -1959,8 +1959,7 @@ void FBlueprintCompilationManagerImpl::ReinstanceBatch(TArray<FReinstancingJob>&
 
 	for(UObject* ArchetypeReferencer : ArchetypeReferencers)
 	{
-		UPackage* NewPackage = ArchetypeReferencer->GetOutermost();
-		FArchiveReplaceOrClearExternalReferences<UObject> ReplaceInCDOAr(ArchetypeReferencer, OldArchetypeToNewArchetype, NewPackage);
+		FArchiveReplaceObjectRef<UObject> ReplaceInCDOAr(ArchetypeReferencer, OldArchetypeToNewArchetype, false, false, false, false, true);
 	}
 }
 
diff --git a/Engine/Source/Editor/UnrealEd/Private/Kismet2/KismetReinstanceUtilities.cpp b/Engine/Source/Editor/UnrealEd/Private/Kismet2/KismetReinstanceUtilities.cpp
--- a/Engine/Source/Editor/UnrealEd/Private/Kismet2/KismetReinstanceUtilities.cpp
+++ b/Engine/Source/Editor/UnrealEd/Private/Kismet2/KismetReinstanceUtilities.cpp
@@ -31,6 +31,7 @@
 #include "BlueprintEditor.h"
 #include "Engine/Selection.h"
 #include "BlueprintEditorSettings.h"
+#include "Algo/RemoveIf.h"
 
 extern COREUOBJECT_API bool GBlueprintUseCompilationManager;
 
@@ -88,6 +89,24 @@ struct FReplaceReferenceHelper
 			BP_SCOPED_COMPILER_EVENT_STAT(EKismetReinstancerStats_FindReferencers);
 
 			Targets = FArchiveHasReferences::GetAllReferencers(SourceObjects, ObjectsThatShouldUseOldStuff);
+
+			// Make sure we don't update properties in old objects, as they
+			// may take ownership of objects referenced in new objects (e.g.
+			// delete components owned by new actors)
+			int32 RangeEnd = Algo::RemoveIf(Targets, [&ObjectsToReplace](UObject* Obj)
+			{
+				for (UObject* ObjectToReplace : ObjectsToReplace)
+				{
+					if (ObjectToReplace == Obj || Obj->IsIn(ObjectToReplace))
+					{
+						return true;
+					}
+				}
+
+				return false;
+			});
+
+			Targets.RemoveAtSwap(RangeEnd, Targets.Num() - RangeEnd, false);
 		}
 
 		{
@@ -95,51 +114,45 @@ struct FReplaceReferenceHelper
 
 			for (UObject* Obj : Targets)
 			{
-				// Make sure we don't update properties in old objects, as they
-				// may take ownership of objects referenced in new objects (e.g.
-				// delete components owned by new actors)
-				if (!ObjectsToReplace.Contains(Obj))
+				// The class for finding and replacing weak references.
+				// We can't relay on "standard" weak references replacement as
+				// it depends on FSoftObjectPath::ResolveObject, which
+				// tries to find the object with the stored path. It is
+				// impossible, cause above we deleted old actors (after
+				// spawning new ones), so during objects traverse we have to
+				// find FSoftObjectPath with the raw given path taken
+				// before deletion of old actors and fix them.
+				class ReferenceReplace : public FArchiveReplaceObjectRef < UObject >
 				{
-					// The class for finding and replacing weak references.
-					// We can't relay on "standard" weak references replacement as
-					// it depends on FSoftObjectPath::ResolveObject, which
-					// tries to find the object with the stored path. It is
-					// impossible, cause above we deleted old actors (after
-					// spawning new ones), so during objects traverse we have to
-					// find FSoftObjectPath with the raw given path taken
-					// before deletion of old actors and fix them.
-					class ReferenceReplace : public FArchiveReplaceObjectRef < UObject >
+				public:
+					ReferenceReplace(UObject* InSearchObject, const TMap<UObject*, UObject*>& InReplacementMap, const TMap<FSoftObjectPath, UObject*>& InWeakReferencesMap)
+						: FArchiveReplaceObjectRef<UObject>(InSearchObject, InReplacementMap, false, false, false, true), WeakReferencesMap(InWeakReferencesMap)
 					{
-					public:
-						ReferenceReplace(UObject* InSearchObject, const TMap<UObject*, UObject*>& InReplacementMap, const TMap<FSoftObjectPath, UObject*>& InWeakReferencesMap)
-							: FArchiveReplaceObjectRef<UObject>(InSearchObject, InReplacementMap, false, false, false, true), WeakReferencesMap(InWeakReferencesMap)
-						{
-							SerializeSearchObject();
-						}
-
-						FArchive& operator<<(FSoftObjectPath& Ref) override
-						{
-							const UObject*const* PtrToObjPtr = WeakReferencesMap.Find(Ref);
-
-							if (PtrToObjPtr != nullptr)
-							{
-								Ref = *PtrToObjPtr;
-							}
+						SerializeSearchObject();
+					}
 
-							return *this;
-						}
+					FArchive& operator<<(FSoftObjectPath& Ref) override
+					{
+						const UObject*const* PtrToObjPtr = WeakReferencesMap.Find(Ref);
 
-						FArchive& operator<<(FSoftObjectPtr& Ref) override
+						if (PtrToObjPtr != nullptr)
 						{
-							return operator<<(Ref.GetUniqueID());
+							Ref = *PtrToObjPtr;
 						}
 
-					private:
-						const TMap<FSoftObjectPath, UObject*>& WeakReferencesMap;
-					};
+						return *this;
+					}
 
-					ReferenceReplace ReplaceAr(Obj, OldToNewInstanceMap, ReinstancedObjectsWeakReferenceMap);
-				}
+					FArchive& operator<<(FSoftObjectPtr& Ref) override
+					{
+						return operator<<(Ref.GetUniqueID());
+					}
+
+				private:
+					const TMap<FSoftObjectPath, UObject*>& WeakReferencesMap;
+				};
+
+				ReferenceReplace ReplaceAr(Obj, OldToNewInstanceMap, ReinstancedObjectsWeakReferenceMap);
 			}
 		}
 	}
@@ -1691,9 +1704,7 @@ static void ReplaceObjectHelper(UObject*& OldObject, UClass* OldClass, UObject*&
 		}
 	}
 
-	OldObject->RemoveFromRoot();
-	OldObject->MarkPendingKill();
-
+	// NOTE: Don't mark the old object as pending kill yet! Otherwise, any weak references to this object will be invalidated too soon and won't get replaced.
 	OldToNewInstanceMap.Add(OldObject, NewUObject);
 
 	if (bIsComponent)
@@ -1850,7 +1861,7 @@ static void ReplaceActorHelper(UObject* OldObject, UClass* OldClass, UObject*& N
 		GEditor->Layers->DisassociateActorFromLayers(OldActor);
 	}
 
-	World->EditorDestroyActor(OldActor, /*bShouldModifyLevel =*/true);
+	// NOTE: Don't destroy the Actor yet! Otherwise, any weak references to this Actor will be invalidated too soon and won't get replaced.
 	OldToNewInstanceMap.Add(OldActor, NewActor);
 }
 
@@ -1905,18 +1916,6 @@ void FBlueprintCompileReinstancer::ReplaceInstancesOfClass_Inner(TMap<UClass*, U
 		OnObjectsReplacedHandle = GEditor->OnObjectsReplaced().AddRaw(&ObjectRemappingHelper, &FObjectRemappingHelper::OnObjectsReplaced);
 	}
 
-	auto UpdateObjectBeingDebugged = [](UObject* InOldObject, UObject* InNewObject)
-	{
-		if (UBlueprint* OldObjBlueprint = Cast<UBlueprint>(InOldObject->GetClass()->ClassGeneratedBy))
-		{
-			const UObject* DebugObj = OldObjBlueprint->GetObjectBeingDebugged(EGetObjectOrWorldBeingDebuggedFlags::IgnorePendingKill);
-			if (DebugObj == InOldObject)
-			{
-				OldObjBlueprint->SetObjectBeingDebugged(InNewObject);
-			}
-		}
-	};
-
 	{
 		TArray<UObject*> ObjectsToReplace;
 
@@ -1967,7 +1966,7 @@ void FBlueprintCompileReinstancer::ReplaceInstancesOfClass_Inner(TMap<UClass*, U
 					{
 						UObject* NewUObject = nullptr;
 						ReplaceObjectHelper(OldObject, OldClass, NewUObject, NewClass, OldToNewInstanceMap, OldToNewNameMap, OldObjIndex, ObjectsToReplace, PotentialEditorsForRefreshing, OwnersToRerunConstructionScript, &FDirectAttachChildrenAccessor::Get, bIsComponent, bArchetypesAreUpToDate);
-						UpdateObjectBeingDebugged(OldObject, NewUObject);
+						
 						ObjectsReplaced.Add(OldObject);
 
 						if (bLogConversions)
@@ -2039,7 +2038,7 @@ void FBlueprintCompileReinstancer::ReplaceInstancesOfClass_Inner(TMap<UClass*, U
 							// Actors that are not in a level cannot be reconstructed, sequencer team decided to reinstance these as normal objects:
 							ReplaceObjectHelper(OldObject, OldClass, NewUObject, NewClass, OldToNewInstanceMap, OldToNewNameMap, OldObjIndex, ObjectsToReplace, PotentialEditorsForRefreshing, OwnersToRerunConstructionScript, &FDirectAttachChildrenAccessor::Get, false, bArchetypesAreUpToDate);
 						}
-						UpdateObjectBeingDebugged(OldObject, NewUObject);
+
 						ObjectsReplaced.Add(OldObject);
 
 						if (bLogConversions)
@@ -2131,6 +2130,22 @@ void FBlueprintCompileReinstancer::ReplaceInstancesOfClass_Inner(TMap<UClass*, U
 				}
 			}
 		}
+
+		Obj->RemoveFromRoot();
+		checkf(!Obj->IsPendingKill(), TEXT("'%s' should not have been marked as PendingKill prior to this point."), *Obj->GetPathName());
+
+		AActor* ObjAsActor = Cast<AActor>(Obj);
+		if (ObjAsActor && ObjAsActor->GetLevel())
+		{
+			// This will take care of marking the object as PendingKill.
+			UWorld* World = ObjAsActor->GetWorld();
+			check(World != nullptr);
+			World->EditorDestroyActor(ObjAsActor, /*bShouldModifyLevel =*/true);
+		}
+		else
+		{
+			Obj->MarkPendingKill();
+		}
 	}
 
 
diff --git a/Engine/Source/Runtime/CoreUObject/Private/Serialization/ArchiveHasReferences.cpp b/Engine/Source/Runtime/CoreUObject/Private/Serialization/ArchiveHasReferences.cpp
--- a/Engine/Source/Runtime/CoreUObject/Private/Serialization/ArchiveHasReferences.cpp
+++ b/Engine/Source/Runtime/CoreUObject/Private/Serialization/ArchiveHasReferences.cpp
@@ -12,6 +12,7 @@ FArchiveHasReferences::FArchiveHasReferences(UObject* InTarget, const TSet<UObje
 	check(InTarget);
 
 	ArIsObjectReferenceCollector = true;
+	ArIsModifyingWeakAndStrongReferences = true;
 	InTarget->Serialize(*this);
 
 	class FArchiveProxyCollector : public FReferenceCollector
diff --git a/Engine/Source/Runtime/Engine/Classes/Engine/Blueprint.h b/Engine/Source/Runtime/Engine/Classes/Engine/Blueprint.h
--- a/Engine/Source/Runtime/Engine/Classes/Engine/Blueprint.h
+++ b/Engine/Source/Runtime/Engine/Classes/Engine/Blueprint.h
@@ -658,9 +658,11 @@ public:
 
 protected:
 	/** Current object being debugged for this blueprint */
+	UPROPERTY(transient, duplicatetransient)
 	TWeakObjectPtr< class UObject > CurrentObjectBeingDebugged;
 
 	/** Current world being debugged for this blueprint */
+	UPROPERTY(transient, duplicatetransient)
 	TWeakObjectPtr< class UWorld > CurrentWorldBeingDebugged;
 
 	/** Delegate called when the debug object is set */
