Fix failure to correctly look up archetype when importing properties for a component in a chlid actor template

diff --git a/Engine/Source/Editor/UnrealEd/Private/EditorObject.cpp b/Engine/Source/Editor/UnrealEd/Private/EditorObject.cpp
--- a/Engine/Source/Editor/UnrealEd/Private/EditorObject.cpp
+++ b/Engine/Source/Editor/UnrealEd/Private/EditorObject.cpp
@@ -473,8 +473,12 @@ static const TCHAR* ImportProperties(
 							UClass* ArchetypeClass = (UClass*)StaticFindObject(UClass::StaticClass(), ANY_PACKAGE, *ObjectClass);
 							if (ArchetypeClass)
 							{
+								ObjectPath = ObjectPath.TrimQuotes();
+
+								UPackage* FindInPackage = (FPackageName::IsShortPackageName(ObjectPath) ? ANY_PACKAGE : nullptr);
+
 								// if we had the class, find the archetype
-								Archetype = StaticFindObject(ArchetypeClass, ANY_PACKAGE, *ObjectPath);
+								Archetype = StaticFindObject(ArchetypeClass, FindInPackage, *ObjectPath);
 							}
 						}
 					}
@@ -539,7 +543,7 @@ static const TCHAR* ImportProperties(
 					if (!bIsOkToReuse)
 					{
 						UE_LOG(LogEditorObject, Log, TEXT("Could not reuse component instance %s, name clash?"), *ComponentTemplate->GetFullName());
-						ComponentTemplate->Rename(); // just abandon the existing component, we are going to create
+						ComponentTemplate->Rename(nullptr, nullptr, REN_DontCreateRedirectors); // just abandon the existing component, we are going to create
 						OldComponent = ComponentTemplate;
 						ComponentTemplate = NULL;
 					}
