[NO NON-PII COMMIT MESSAGE]

diff --git a/Engine/Source/Editor/UnrealEd/Private/FileHelpers.cpp b/Engine/Source/Editor/UnrealEd/Private/FileHelpers.cpp
--- a/Engine/Source/Editor/UnrealEd/Private/FileHelpers.cpp
+++ b/Engine/Source/Editor/UnrealEd/Private/FileHelpers.cpp
@@ -723,6 +723,9 @@ static bool SaveWorld(UWorld* World,
 			}
 		}
 
+		// Mark package as fully loaded, this is usually set implicitly by calling IsFullyLoaded before saving, but that path can get skipped for levels
+		Package->MarkAsFullyLoaded();
+
 		SlowTask.EnterProgressFrame(50);
 
 		// Save package.
diff --git a/Engine/Source/Editor/UnrealEd/Private/PackageTools.cpp b/Engine/Source/Editor/UnrealEd/Private/PackageTools.cpp
--- a/Engine/Source/Editor/UnrealEd/Private/PackageTools.cpp
+++ b/Engine/Source/Editor/UnrealEd/Private/PackageTools.cpp
@@ -172,6 +172,7 @@ UPackageTools::UPackageTools(const FObjectInitializer& ObjectInitializer)
 			check( TopLevelPackage );
 			check( TopLevelPackage->GetOuter() == NULL );
 
+			// Calling IsFullyLoaded() will mark a package as fully loaded if it does not exist on disk
 			if( !TopLevelPackage->IsFullyLoaded() )
 			{	
 				// Ask user to fully load or suppress the message and just fully load
