Fix cooker to filter out plugin content for unsupported target platforms (bug fix for cl 6649730)

(Partial copy-up from dev-core of CookOnTheFlyServer.h/cpp, cl 6910495 and 6925523)

diff --git a/Engine/Source/Editor/UnrealEd/Classes/CookOnTheSide/CookOnTheFlyServer.h b/Engine/Source/Editor/UnrealEd/Classes/CookOnTheSide/CookOnTheFlyServer.h
--- a/Engine/Source/Editor/UnrealEd/Classes/CookOnTheSide/CookOnTheFlyServer.h
+++ b/Engine/Source/Editor/UnrealEd/Classes/CookOnTheSide/CookOnTheFlyServer.h
@@ -689,7 +689,7 @@ private:
 	* Build up a map of unsupported packages per platform that can be checked before saving.
 	*/
 	void DiscoverPlatformSpecificNeverCookPackages(
-		const TArray<FName>& TargetPlatformNames, const TArray<FString>& TargetPlatformStrings);
+		const TArray<FName>& TargetPlatformNames, const TArray<FString>& VanillaPlatformStrings);
 
 	/**
 	* Clean up the sandbox
diff --git a/Engine/Source/Editor/UnrealEd/Private/CookOnTheFlyServer.cpp b/Engine/Source/Editor/UnrealEd/Private/CookOnTheFlyServer.cpp
--- a/Engine/Source/Editor/UnrealEd/Private/CookOnTheFlyServer.cpp
+++ b/Engine/Source/Editor/UnrealEd/Private/CookOnTheFlyServer.cpp
@@ -6718,13 +6718,12 @@ void UCookOnTheFlyServer::InitializeTargetPlatforms()
 }
 
 void UCookOnTheFlyServer::DiscoverPlatformSpecificNeverCookPackages(
-	const TArray<FName>& TargetPlatformNames, const TArray<FString>& TargetPlatformStrings)
+	const TArray<FName>& TargetPlatformNames, const TArray<FString>& VanillaPlatformStrings)
 {
 	TArray<FName> PluginUnsupportedTargetPlatforms;
 	TArray<FAssetData> PluginAssets;
 	FARFilter PluginARFilter;
 	FString PluginPackagePath;
-	static const FString AndroidPlatformName(TEXT("Android"));
 
 	TArray<TSharedRef<IPlugin>> AllContentPlugins = IPluginManager::Get().GetEnabledPluginsWithContent();
 	for (TSharedRef<IPlugin> Plugin : AllContentPlugins)
@@ -6737,27 +6736,11 @@ void UCookOnTheFlyServer::DiscoverPlatformSpecificNeverCookPackages(
 			continue;
 		}
 
-		// Android platforms append a texture compression string to the end of the platform name so we can't do a direct comparison
-		bool bSupportsAndroid = false;
-		for (const FString& PlatformName : Descriptor.SupportedTargetPlatforms)
-		{
-			if (PlatformName.Contains(AndroidPlatformName))
-			{
-				bSupportsAndroid = true;
-				break;
-			}
-		}
-
 		// find any unsupported target platforms for this plugin
 		PluginUnsupportedTargetPlatforms.Reset();
 		for (int32 I = 0, Count = TargetPlatformNames.Num(); I < Count; ++I)
 		{
-			// Partial match for the Android platform family
-			if (bSupportsAndroid && TargetPlatformStrings[I].Contains(AndroidPlatformName))
-			{
-				continue;
-			}
-			else if (!Descriptor.SupportedTargetPlatforms.Contains(TargetPlatformStrings[I]))
+			if (!Descriptor.SupportedTargetPlatforms.Contains(VanillaPlatformStrings[I]))
 			{
 				PluginUnsupportedTargetPlatforms.Add(TargetPlatformNames[I]);
 			}
@@ -6943,17 +6926,19 @@ void UCookOnTheFlyServer::StartCookByTheBook( const FCookByTheBookStartupOptions
 	// build persistent list of all target platform names in CookByTheBookOptions->TargetPlatformNames,
 	// also use temp list of platform strings to discover PlatformSpecificNeverCookPackages
 	{
-		TArray<FString> TargetPlatformStrings;
-		TargetPlatformStrings.Reserve(TargetPlatforms.Num());
+		TArray<FString> VanillaPlatformStrings;
+		VanillaPlatformStrings.Reserve(TargetPlatforms.Num());
 
 		CookByTheBookOptions->TargetPlatformNames.Empty(TargetPlatforms.Num());
 		for (const ITargetPlatform* Platform : TargetPlatforms)
 		{
-			FString& PlatformString = TargetPlatformStrings.Emplace_GetRef(Platform->PlatformName());
-			CookByTheBookOptions->TargetPlatformNames.Emplace(*PlatformString);
+			FString VanillaPlatform;
+			Platform->GetPlatformInfo().VanillaPlatformName.ToString(VanillaPlatform);
+			VanillaPlatformStrings.Emplace(MoveTemp(VanillaPlatform));
+			CookByTheBookOptions->TargetPlatformNames.Emplace(*Platform->PlatformName());
 		}
 
-		DiscoverPlatformSpecificNeverCookPackages(CookByTheBookOptions->TargetPlatformNames, TargetPlatformStrings);
+		DiscoverPlatformSpecificNeverCookPackages(CookByTheBookOptions->TargetPlatformNames, VanillaPlatformStrings);
 	}
 	const TArray<FName>& TargetPlatformNames = CookByTheBookOptions->TargetPlatformNames;
 
