Fix for an ensure when reloading a Blueprint asset with the current debug object set to an instance of the Blueprint class.

Change summary:
- Modified UpdateObjectBeingDebugged() inside of FBlueprintCompileReinstancer_Inner() to validate that both outer BPs match before setting the new instance. This is a low-risk fix to address the regression in 4.23; a more in-depth reinstancer fix (that requires more testing) will replace this in 4.24.

#rnx

diff --git a/Engine/Source/Editor/UnrealEd/Private/Kismet2/KismetReinstanceUtilities.cpp b/Engine/Source/Editor/UnrealEd/Private/Kismet2/KismetReinstanceUtilities.cpp
--- a/Engine/Source/Editor/UnrealEd/Private/Kismet2/KismetReinstanceUtilities.cpp
+++ b/Engine/Source/Editor/UnrealEd/Private/Kismet2/KismetReinstanceUtilities.cpp
@@ -1909,10 +1909,17 @@ void FBlueprintCompileReinstancer::ReplaceInstancesOfClass_Inner(TMap<UClass*, U
 	{
 		if (UBlueprint* OldObjBlueprint = Cast<UBlueprint>(InOldObject->GetClass()->ClassGeneratedBy))
 		{
-			const UObject* DebugObj = OldObjBlueprint->GetObjectBeingDebugged(EGetObjectOrWorldBeingDebuggedFlags::IgnorePendingKill);
-			if (DebugObj == InOldObject)
+			// For now, don't update the object if the outer BP assets don't match (e.g. after a reload). Otherwise, it will
+			// trigger an ensure() in SetObjectBeginDebugged(). This will be replaced with a better solution in a future release.
+			if (OldObjBlueprint == Cast<UBlueprint>(InNewObject->GetClass()->ClassGeneratedBy))
 			{
-				OldObjBlueprint->SetObjectBeingDebugged(InNewObject);
+				// The old object may already be PendingKill, but we still want to check the current
+				// ptr value for a match. Otherwise, the selection will get cleared after every compile.
+				const UObject* DebugObj = OldObjBlueprint->GetObjectBeingDebugged(EGetObjectOrWorldBeingDebuggedFlags::IgnorePendingKill);
+				if (DebugObj == InOldObject)
+				{
+					OldObjBlueprint->SetObjectBeingDebugged(InNewObject);
+				}
 			}
 		}
 	};
