Added clarification text for Blueprint Compilation Error dialog not opening immediately.

diff --git a/Engine/Source/Editor/UnrealEd/Private/Dialogs/CustomDialog.cpp b/Engine/Source/Editor/UnrealEd/Private/Dialogs/CustomDialog.cpp
--- a/Engine/Source/Editor/UnrealEd/Private/Dialogs/CustomDialog.cpp
+++ b/Engine/Source/Editor/UnrealEd/Private/Dialogs/CustomDialog.cpp
@@ -63,6 +63,7 @@ void SCustomDialog::Construct(const FArguments& InArgs)
 				.AutoWidth()
 				.VAlign(VAlign_Center)
 				.HAlign(HAlign_Left)
+				.Padding(0, 0, 8, 0)
 				[
 					SNew(SImage)
 					.Image(ImageBrush)
diff --git a/Engine/Source/Editor/UnrealEd/Private/PlayLevel.cpp b/Engine/Source/Editor/UnrealEd/Private/PlayLevel.cpp
--- a/Engine/Source/Editor/UnrealEd/Private/PlayLevel.cpp
+++ b/Engine/Source/Editor/UnrealEd/Private/PlayLevel.cpp
@@ -2289,15 +2289,12 @@ static bool ShowBlueprintErrorDialog( TArray<UBlueprint*> ErroredBlueprints )
 		}
 	};
 
-	TSharedRef<SVerticalBox> DialogContents = SNew( SVerticalBox )
+	TSharedRef<SVerticalBox> DialogContents = SNew(SVerticalBox)
 		+ SVerticalBox::Slot()
+		.Padding(0, 0, 0, 16)
 		[
-			SNew( STextBlock )
-			.Text( NSLOCTEXT( "PlayInEditor", "PrePIE_BlueprintErrors", "One or more blueprints has an unresolved compiler error, are you sure you want to Play in Editor?" ) )
-		]
-		+ SVerticalBox::Slot()
-		[
-			SNew( STextBlock )
+			SNew(STextBlock)
+			.Text(NSLOCTEXT("PlayInEditor", "PrePIE_BlueprintErrors", "One or more blueprints has an unresolved compiler error, are you sure you want to Play in Editor?"))
 		];
 
 	for (UBlueprint* Blueprint : ErroredBlueprints)
@@ -2316,6 +2313,13 @@ static bool ShowBlueprintErrorDialog( TArray<UBlueprint*> ErroredBlueprints )
 			];
 	}
 
+	DialogContents->AddSlot()
+		.Padding(0, 16, 0, 0)
+		[
+			SNew(STextBlock)
+			.Text(NSLOCTEXT("PlayInEditor", "PrePIE_BlueprintErrorsDelayedOpen", "Clicked blueprints will open once this dialog is closed."))
+		];
+
 	FText DialogTitle = NSLOCTEXT("PlayInEditor", "PrePIE_BlueprintErrorsTitle", "Blueprint Compilation Errors");
 
 	FText OKText = NSLOCTEXT("PlayInEditor", "PrePIE_OkText", "Play in Editor");
