Unified the code used by the editor and project launcher, to fix incorrect editor executable being used for cooking when using the project launcher.

diff --git a/Engine/Source/Developer/LauncherServices/Private/LauncherServicesModule.cpp b/Engine/Source/Developer/LauncherServices/Private/LauncherServicesModule.cpp
--- a/Engine/Source/Developer/LauncherServices/Private/LauncherServicesModule.cpp
+++ b/Engine/Source/Developer/LauncherServices/Private/LauncherServicesModule.cpp
@@ -64,6 +64,37 @@ public:
 		return LauncherServicesSDKNotInstalled.Broadcast(PlatformName, DocLink);
 	}
 
+	virtual FString GetExecutableForCommandlets() const
+	{
+		FString ExecutableName;
+#if WITH_EDITOR
+		ExecutableName = FString(FPlatformProcess::ExecutablePath());
+#if PLATFORM_WINDOWS
+		// turn UE4editor into UE4editor-cmd
+		if (ExecutableName.EndsWith(".exe", ESearchCase::IgnoreCase) && !FPaths::GetBaseFilename(ExecutableName).EndsWith("-cmd", ESearchCase::IgnoreCase))
+		{
+			FString NewExeName = ExecutableName.Left(ExecutableName.Len() - 4) + "-Cmd.exe";
+			if (FPaths::FileExists(NewExeName))
+			{
+				ExecutableName = NewExeName;
+			}
+		}
+#elif PLATFORM_MAC
+		// turn UE4editor into UE4editor-cmd
+		if (!FPaths::GetBaseFilename(ExecutableName).EndsWith("-cmd", ESearchCase::IgnoreCase))
+		{
+			FString NewExeName = ExecutableName + "-Cmd";
+			if (FPaths::FileExists(NewExeName))
+			{
+				ExecutableName = NewExeName;
+			}
+		}
+#endif
+#endif
+		return ExecutableName;
+
+	}
+
 private:
 	
 	/** Event to be called when the editor tried to use a platform, but it wasn't installed. */
diff --git a/Engine/Source/Developer/LauncherServices/Private/Profiles/LauncherProfile.h b/Engine/Source/Developer/LauncherServices/Private/Profiles/LauncherProfile.h
--- a/Engine/Source/Developer/LauncherServices/Private/Profiles/LauncherProfile.h
+++ b/Engine/Source/Developer/LauncherServices/Private/Profiles/LauncherProfile.h
@@ -1857,43 +1857,8 @@ public:
 		PackagingMode = ELauncherProfilePackagingModes::DoNotPackage;
 
 		// default UAT settings
-		EditorExe = FPlatformProcess::ExecutableName(false);
-		if (EditorExe.Contains(TEXT("Editor")))
-		{
-#if PLATFORM_WINDOWS
-			// turn UE4editor into UE4editor-cmd
-			if (EditorExe.EndsWith(".exe", ESearchCase::IgnoreCase) && !FPaths::GetBaseFilename(EditorExe).EndsWith("-cmd", ESearchCase::IgnoreCase))
-			{
-				FString NewExeName = EditorExe.Left(EditorExe.Len() - 4) + "-Cmd.exe";
-				if (FPaths::FileExists(NewExeName))
-				{
-					EditorExe = NewExeName;
-				}
-				else
-				{
-					EditorExe.Empty();
-				}
-			}
-#elif PLATFORM_MAC
-			// turn UE4editor into UE4editor-cmd
-			if (!FPaths::GetBaseFilename(EditorExe).EndsWith("-cmd", ESearchCase::IgnoreCase))
-			{
-				FString NewExeName = EditorExe + "-Cmd";
-				if (FPaths::FileExists(NewExeName))
-				{
-					EditorExe = NewExeName;
-				}
-				else
-				{
-					EditorExe.Empty();
-				}
-			}
-#endif
-		}
-		else
-		{
-			EditorExe.Empty();
-		}
+		ILauncherServicesModule& LauncherServicesModule = FModuleManager::GetModuleChecked<ILauncherServicesModule>(TEXT("LauncherServices"));
+		EditorExe = LauncherServicesModule.GetExecutableForCommandlets();
 
 		bNotForLicensees = false;
 
diff --git a/Engine/Source/Developer/LauncherServices/Public/ILauncherServicesModule.h b/Engine/Source/Developer/LauncherServices/Public/ILauncherServicesModule.h
--- a/Engine/Source/Developer/LauncherServices/Public/ILauncherServicesModule.h
+++ b/Engine/Source/Developer/LauncherServices/Public/ILauncherServicesModule.h
@@ -65,6 +65,11 @@ public:
 	virtual FLauncherServicesSDKNotInstalled& OnLauncherServicesSDKNotInstalled( ) = 0;
 	virtual void BroadcastLauncherServicesSDKNotInstalled(const FString& PlatformName, const FString& DocLink) = 0;
 
+	/**
+	 * Determines the executable that should be used for executing commandlets
+	 */
+	virtual FString GetExecutableForCommandlets() const = 0;
+
 public:
 
 	/** Virtual destructor. */
diff --git a/Engine/Source/Editor/UnrealEd/Private/UnrealEdMisc.cpp b/Engine/Source/Editor/UnrealEd/Private/UnrealEdMisc.cpp
--- a/Engine/Source/Editor/UnrealEd/Private/UnrealEdMisc.cpp
+++ b/Engine/Source/Editor/UnrealEd/Private/UnrealEdMisc.cpp
@@ -82,6 +82,7 @@
 #include "IVREditorModule.h"
 #include "ILauncherPlatform.h"
 #include "LauncherPlatformModule.h"
+#include "ILauncherServicesModule.h"
 
 #define USE_UNIT_TESTS 0
 
@@ -1825,29 +1826,8 @@ bool FUnrealEdMisc::GetURL( const TCHAR* InKey, FString& OutURL, const bool bChe
 
 FString FUnrealEdMisc::GetExecutableForCommandlets() const
 {
-	FString ExecutableName = FString(FPlatformProcess::ExecutablePath());
-#if PLATFORM_WINDOWS
-	// turn UE4editor into UE4editor-cmd
-	if(ExecutableName.EndsWith(".exe", ESearchCase::IgnoreCase) && !FPaths::GetBaseFilename(ExecutableName).EndsWith("-cmd", ESearchCase::IgnoreCase))
-	{
-		FString NewExeName = ExecutableName.Left(ExecutableName.Len() - 4) + "-Cmd.exe";
-		if (FPaths::FileExists(NewExeName))
-		{
-			ExecutableName = NewExeName;
-		}
-	}
-#elif PLATFORM_MAC
-	// turn UE4editor into UE4editor-cmd
-	if (!FPaths::GetBaseFilename(ExecutableName).EndsWith("-cmd", ESearchCase::IgnoreCase))
-	{
-		FString NewExeName = ExecutableName + "-Cmd";
-		if (FPaths::FileExists(NewExeName))
-		{
-			ExecutableName = NewExeName;
-		}
-	}
-#endif
-	return ExecutableName;
+	ILauncherServicesModule& LauncherServicesModule = FModuleManager::LoadModuleChecked<ILauncherServicesModule>(TEXT("LauncherServices"));
+	return LauncherServicesModule.GetExecutableForCommandlets();
 }
 
 void FUnrealEdMisc::OpenMarketplace(const FString& CustomLocation)
