Fix cooker to filter out plugin content for unsupported target platforms (bug fix for cl 6649730 and 6910495).
Make sure to use the correct platform string (FPlatformInfo::UBTTargetId, i.e. the executable platform) when querying FPluginDescriptor::SupportedTargetPlatforms.

(Partial copy-up from dev-core of CookOnTheFlyServer.h/cpp, cl 7039440)

diff --git a/Engine/Source/Editor/UnrealEd/Classes/CookOnTheSide/CookOnTheFlyServer.h b/Engine/Source/Editor/UnrealEd/Classes/CookOnTheSide/CookOnTheFlyServer.h
--- a/Engine/Source/Editor/UnrealEd/Classes/CookOnTheSide/CookOnTheFlyServer.h
+++ b/Engine/Source/Editor/UnrealEd/Classes/CookOnTheSide/CookOnTheFlyServer.h
@@ -689,7 +689,7 @@ private:
 	* Build up a map of unsupported packages per platform that can be checked before saving.
 	*/
 	void DiscoverPlatformSpecificNeverCookPackages(
-		const TArray<FName>& TargetPlatformNames, const TArray<FString>& VanillaPlatformStrings);
+		const TArray<FName>& TargetPlatformNames, const TArray<FString>& UBTPlatformStrings);
 
 	/**
 	* Clean up the sandbox
diff --git a/Engine/Source/Editor/UnrealEd/Private/CookOnTheFlyServer.cpp b/Engine/Source/Editor/UnrealEd/Private/CookOnTheFlyServer.cpp
--- a/Engine/Source/Editor/UnrealEd/Private/CookOnTheFlyServer.cpp
+++ b/Engine/Source/Editor/UnrealEd/Private/CookOnTheFlyServer.cpp
@@ -6718,7 +6718,7 @@ void UCookOnTheFlyServer::InitializeTargetPlatforms()
 }
 
 void UCookOnTheFlyServer::DiscoverPlatformSpecificNeverCookPackages(
-	const TArray<FName>& TargetPlatformNames, const TArray<FString>& VanillaPlatformStrings)
+	const TArray<FName>& TargetPlatformNames, const TArray<FString>& UBTPlatformStrings)
 {
 	TArray<FName> PluginUnsupportedTargetPlatforms;
 	TArray<FAssetData> PluginAssets;
@@ -6740,7 +6740,7 @@ void UCookOnTheFlyServer::DiscoverPlatformSpecificNeverCookPackages(
 		PluginUnsupportedTargetPlatforms.Reset();
 		for (int32 I = 0, Count = TargetPlatformNames.Num(); I < Count; ++I)
 		{
-			if (!Descriptor.SupportedTargetPlatforms.Contains(VanillaPlatformStrings[I]))
+			if (!Descriptor.SupportedTargetPlatforms.Contains(UBTPlatformStrings[I]))
 			{
 				PluginUnsupportedTargetPlatforms.Add(TargetPlatformNames[I]);
 			}
@@ -6924,21 +6924,21 @@ void UCookOnTheFlyServer::StartCookByTheBook( const FCookByTheBookStartupOptions
 	}
 	
 	// build persistent list of all target platform names in CookByTheBookOptions->TargetPlatformNames,
-	// also use temp list of platform strings to discover PlatformSpecificNeverCookPackages
+	// also use temp list of UBT platform strings to discover PlatformSpecificNeverCookPackages
 	{
-		TArray<FString> VanillaPlatformStrings;
-		VanillaPlatformStrings.Reserve(TargetPlatforms.Num());
+		TArray<FString> UBTPlatformStrings;
+		UBTPlatformStrings.Reserve(TargetPlatforms.Num());
 
 		CookByTheBookOptions->TargetPlatformNames.Empty(TargetPlatforms.Num());
 		for (const ITargetPlatform* Platform : TargetPlatforms)
 		{
-			FString VanillaPlatform;
-			Platform->GetPlatformInfo().VanillaPlatformName.ToString(VanillaPlatform);
-			VanillaPlatformStrings.Emplace(MoveTemp(VanillaPlatform));
+			FString UBTPlatformName;
+			Platform->GetPlatformInfo().UBTTargetId.ToString(UBTPlatformName);
+			UBTPlatformStrings.Emplace(MoveTemp(UBTPlatformName));
 			CookByTheBookOptions->TargetPlatformNames.Emplace(*Platform->PlatformName());
 		}
 
-		DiscoverPlatformSpecificNeverCookPackages(CookByTheBookOptions->TargetPlatformNames, VanillaPlatformStrings);
+		DiscoverPlatformSpecificNeverCookPackages(CookByTheBookOptions->TargetPlatformNames, UBTPlatformStrings);
 	}
 	const TArray<FName>& TargetPlatformNames = CookByTheBookOptions->TargetPlatformNames;
 
