Fixing lookup pin names being displayed to the user
#rnx

diff --git a/Engine/Source/Editor/UnrealEd/Private/MaterialGraphNode.cpp b/Engine/Source/Editor/UnrealEd/Private/MaterialGraphNode.cpp
--- a/Engine/Source/Editor/UnrealEd/Private/MaterialGraphNode.cpp
+++ b/Engine/Source/Editor/UnrealEd/Private/MaterialGraphNode.cpp
@@ -553,7 +553,7 @@ void UMaterialGraphNode::CreateInputPins()
 		{
 			// Makes sure pin has a name for lookup purposes but user will never see it
 			NewPin->PinName = CreateUniquePinName(TEXT("Input"));
-			NewPin->PinFriendlyName = FText::GetEmpty();
+			NewPin->PinFriendlyName = FText::FromString(TEXT(" "));;
 		}
 	}
 }
@@ -607,7 +607,7 @@ void UMaterialGraphNode::CreateOutputPins()
 		{
 			// Makes sure pin has a name for lookup purposes but user will never see it
 			NewPin->PinName = CreateUniquePinName(TEXT("Output"));
-			NewPin->PinFriendlyName = FText::GetEmpty();
+			NewPin->PinFriendlyName = FText::FromString(TEXT(" "));;
 		}
 	}
 }
diff --git a/Engine/Source/Runtime/Landscape/Private/LandscapeGrass.cpp b/Engine/Source/Runtime/Landscape/Private/LandscapeGrass.cpp
--- a/Engine/Source/Runtime/Landscape/Private/LandscapeGrass.cpp
+++ b/Engine/Source/Runtime/Landscape/Private/LandscapeGrass.cpp
@@ -1042,7 +1042,7 @@ public:
 //
 // UMaterialExpressionLandscapeGrassOutput
 //
-FName UMaterialExpressionLandscapeGrassOutput::PinDefaultName = TEXT("Pin");
+FName UMaterialExpressionLandscapeGrassOutput::PinDefaultName = TEXT("Input");
 
 UMaterialExpressionLandscapeGrassOutput::UMaterialExpressionLandscapeGrassOutput(const FObjectInitializer& ObjectInitializer)
 : Super(ObjectInitializer)
