Fixed the actor transform widget not updating its position when selected actors move during simulation. Edigrated from Dev-Editor.

diff --git a/Engine/Source/Editor/UnrealEd/Private/EditorModeManager.cpp b/Engine/Source/Editor/UnrealEd/Private/EditorModeManager.cpp
--- a/Engine/Source/Editor/UnrealEd/Private/EditorModeManager.cpp
+++ b/Engine/Source/Editor/UnrealEd/Private/EditorModeManager.cpp
@@ -159,6 +159,49 @@ bool FEditorModeTools::SelectionHasSceneComponent() const
 	return bSelectionHasSceneComponent;
 }
 
+void FEditorModeTools::UpdateModeWidgetLocation()
+{
+	bool Reached = false;
+
+	double starttime = FPlatformTime::Seconds();
+
+	USelection* SelectedActors = GetSelectedActors();
+
+	if (SelectedActors && SelectedActors->Num() > 0)
+	{
+		AActor* Actor = CastChecked<AActor>(SelectedActors->GetSelectedObject(SelectedActors->Num() - 1));
+
+		if (Actor)
+		{
+			FVector CurrentActorLocation = Actor->GetActorLocation();
+
+			if (Actor == LastSelectedActor)
+			{
+				if (!(LastSelectedActorLocation - CurrentActorLocation).IsNearlyZero())
+				{
+					for (const auto& Mode : ActiveModes)
+					{
+						if (Mode->UsesTransformWidget())
+						{
+							SetPivotLocation(CurrentActorLocation, false);
+							LastSelectedActorLocation = CurrentActorLocation;
+
+							Reached = true;
+
+							break;
+						}
+					}
+
+				}
+			}
+			else
+			{
+				LastSelectedActor = Actor;
+			}
+		}
+	}
+}
+
 void FEditorModeTools::OnEditorSelectionChanged(UObject* NewSelection)
 {
 	if(NewSelection == GetSelectedActors())
@@ -847,6 +890,8 @@ void FEditorModeTools::Tick( FEditorViewportClient* ViewportClient, float DeltaT
 		const TSharedPtr<FEdMode>& Mode = ActiveModes[ ModeIndex ];
 		Mode->Tick( ViewportClient, DeltaTime );
 	}
+
+	UpdateModeWidgetLocation();
 }
 
 /** Notifies all active modes of any change in mouse movement */
diff --git a/Engine/Source/Editor/UnrealEd/Public/EditorModeManager.h b/Engine/Source/Editor/UnrealEd/Public/EditorModeManager.h
--- a/Engine/Source/Editor/UnrealEd/Public/EditorModeManager.h
+++ b/Engine/Source/Editor/UnrealEd/Public/EditorModeManager.h
@@ -99,6 +99,10 @@ protected:
 private:
 	void RebuildModeToolBar();
 	void SpawnOrUpdateModeToolbar();
+
+	/** Updates the widget's position for changes in location during simulation */
+	void UpdateModeWidgetLocation();
+
 public:
 
 	/**
@@ -531,6 +535,11 @@ protected:
 
 	/** if true the current selection has a scene component */
 	bool bSelectionHasSceneComponent;
+
+	/** Compare for transformation of selected actor to mode widget */
+	TWeakObjectPtr<AActor>	LastSelectedActor;
+	FVector					LastSelectedActorLocation;
+
 private:
 	struct FEdModeToolbarRow
 	{
