Integrating CL 6895543 and CL 6910851

diff --git a/Engine/Source/Editor/UnrealEd/Private/Fbx/FbxAnimationExport.cpp b/Engine/Source/Editor/UnrealEd/Private/Fbx/FbxAnimationExport.cpp
--- a/Engine/Source/Editor/UnrealEd/Private/Fbx/FbxAnimationExport.cpp
+++ b/Engine/Source/Editor/UnrealEd/Private/Fbx/FbxAnimationExport.cpp
@@ -23,6 +23,43 @@ DEFINE_LOG_CATEGORY_STATIC(LogFbxAnimationExport, Log, All);
 namespace UnFbx
 {
 
+	bool FFbxExporter::SetupAnimStack(const UAnimSequence* AnimSeq)
+	{
+		if (AnimSeq->SequenceLength == 0.f)
+		{
+			// something is wrong
+			return false;
+		}
+
+		const float FrameRate = FMath::TruncToFloat(((AnimSeq->GetRawNumberOfFrames() - 1) / AnimSeq->SequenceLength) + 0.5f);
+		//Configure the scene time line
+		{
+			FbxGlobalSettings& SceneGlobalSettings = Scene->GetGlobalSettings();
+			double CurrentSceneFrameRate = FbxTime::GetFrameRate(SceneGlobalSettings.GetTimeMode());
+			if (!bSceneGlobalTimeLineSet || FrameRate > CurrentSceneFrameRate)
+			{
+				FbxTime::EMode ComputeTimeMode = FbxTime::ConvertFrameRateToTimeMode(FrameRate);
+				FbxTime::SetGlobalTimeMode(ComputeTimeMode, ComputeTimeMode == FbxTime::eCustom ? FrameRate : 0.0);
+				SceneGlobalSettings.SetTimeMode(ComputeTimeMode);
+				if (ComputeTimeMode == FbxTime::eCustom)
+				{
+					SceneGlobalSettings.SetCustomFrameRate(FrameRate);
+				}
+				bSceneGlobalTimeLineSet = true;
+			}
+		}
+
+		// set time correctly
+		FbxTime ExportedStartTime, ExportedStopTime;
+		ExportedStartTime.SetSecondDouble(0.f);
+		ExportedStopTime.SetSecondDouble(AnimSeq->SequenceLength);
+
+		FbxTimeSpan ExportedTimeSpan;
+		ExportedTimeSpan.Set(ExportedStartTime, ExportedStopTime);
+		AnimStack->SetLocalTimeSpan(ExportedTimeSpan);
+
+		return true;
+	}
 
 void FFbxExporter::ExportAnimSequenceToFbx(const UAnimSequence* AnimSeq,
 									 const USkeletalMesh* SkelMesh,
@@ -38,41 +75,12 @@ void FFbxExporter::ExportAnimSequenceToFbx(const UAnimSequence* AnimSeq,
 
 	USkeleton* Skeleton = AnimSeq->GetSkeleton();
 
-	if (AnimSeq->SequenceLength == 0.f || Skeleton == nullptr)
+	if (Skeleton == nullptr || !SetupAnimStack(AnimSeq))
 	{
 		// something is wrong
 		return;	
 	}
 
-	const float FrameRate = FMath::TruncToFloat(((AnimSeq->GetRawNumberOfFrames() - 1) / AnimSeq->SequenceLength) + 0.5f);
-	//Configure the scene time line
-	{
-		FbxGlobalSettings& SceneGlobalSettings = Scene->GetGlobalSettings();
-		double CurrentSceneFrameRate = FbxTime::GetFrameRate(SceneGlobalSettings.GetTimeMode());
-		if (!bSceneGlobalTimeLineSet || FrameRate > CurrentSceneFrameRate)
-		{
-			FbxTime::EMode ComputeTimeMode = FbxTime::ConvertFrameRateToTimeMode(FrameRate);
-			FbxTime::SetGlobalTimeMode(ComputeTimeMode, ComputeTimeMode == FbxTime::eCustom ? FrameRate : 0.0);
-			SceneGlobalSettings.SetTimeMode(ComputeTimeMode);
-			if (ComputeTimeMode == FbxTime::eCustom)
-			{
-				SceneGlobalSettings.SetCustomFrameRate(FrameRate);
-			}
-			bSceneGlobalTimeLineSet = true;
-		}
-	}
-
-	// set time correctly
-	FbxTime ExportedStartTime, ExportedStopTime;
-	
-
-	ExportedStartTime.SetSecondDouble(0.f);
-	ExportedStopTime.SetSecondDouble(AnimSeq->SequenceLength);
-
-	FbxTimeSpan ExportedTimeSpan;
-	ExportedTimeSpan.Set(ExportedStartTime, ExportedStopTime);
-	AnimStack->SetLocalTimeSpan(ExportedTimeSpan);
-
 	//Prepare root anim curves data to be exported
 	TArray<FName> AnimCurveNames;
 	TArray<SmartName::UID_Type> AnimCurveUIDs;
@@ -85,21 +93,28 @@ void FFbxExporter::ExportAnimSequenceToFbx(const UAnimSequence* AnimSeq,
 		{
 			AnimCurveMapping->FillNameArray(AnimCurveNames);
 			AnimCurveMapping->FillUidArray(AnimCurveUIDs);
-		}
-		
-		if (AnimCurveNames.Num() > 0)
-		{
+
+			const UFbxExportOption* ExportOptions = GetExportOptions();
+			const bool bExportMorphTargetCurvesInMesh = ExportOptions && ExportOptions->bExportPreviewMesh && ExportOptions->bExportMorphTargets;
+
 			for (auto AnimCurveName : AnimCurveNames)
 			{
-				FbxProperty AnimCurveFbxProp = FbxProperty::Create(BoneNodes[0], FbxDoubleDT, TCHAR_TO_ANSI(*AnimCurveName.ToString()));
-				AnimCurveFbxProp.ModifyFlag(FbxPropertyFlags::eAnimatable, true);
-				AnimCurveFbxProp.ModifyFlag(FbxPropertyFlags::eUserDefined, true);
-				FbxAnimCurve* AnimFbxCurve = AnimCurveFbxProp.GetCurve(InAnimLayer, true);
-				CustomCurveMap.Add(AnimCurveName, AnimFbxCurve);
+				const FCurveMetaData* CurveMetaData = AnimCurveMapping->GetCurveMetaData(AnimCurveName);
+
+				//Only export the custom curve if it is not used in a MorphTarget that will be exported latter on.
+				if(!(bExportMorphTargetCurvesInMesh && CurveMetaData && CurveMetaData->Type.bMorphtarget))
+				{
+					FbxProperty AnimCurveFbxProp = FbxProperty::Create(BoneNodes[0], FbxDoubleDT, TCHAR_TO_ANSI(*AnimCurveName.ToString()));
+					AnimCurveFbxProp.ModifyFlag(FbxPropertyFlags::eAnimatable, true);
+					AnimCurveFbxProp.ModifyFlag(FbxPropertyFlags::eUserDefined, true);
+					FbxAnimCurve* AnimFbxCurve = AnimCurveFbxProp.GetCurve(InAnimLayer, true);
+					CustomCurveMap.Add(AnimCurveName, AnimFbxCurve);
+				}
 			}
 		}
 	}
 
+	ExportCustomAnimCurvesToFbx(CustomCurveMap, AnimSeq, AnimStartOffset, AnimEndOffset, AnimPlayRate, StartTime);
 
 	// Add the animation data to the bone nodes
 	for(int32 BoneIndex = 0; BoneIndex < BoneNodes.Num(); ++BoneIndex)
@@ -123,19 +138,6 @@ void FFbxExporter::ExportAnimSequenceToFbx(const UAnimSequence* AnimSeq,
 		Curves[7] = CurrentBoneNode->LclScaling.GetCurve(InAnimLayer, FBXSDK_CURVENODE_COMPONENT_Y, true);
 		Curves[8] = CurrentBoneNode->LclScaling.GetCurve(InAnimLayer, FBXSDK_CURVENODE_COMPONENT_Z, true);
 
-		float AnimTime					= AnimStartOffset;
-		float AnimEndTime				= (AnimSeq->SequenceLength - AnimEndOffset);
-		// Subtracts 1 because NumFrames includes an initial pose for 0.0 second
-		double TimePerKey				= (AnimSeq->SequenceLength / (AnimSeq->GetRawNumberOfFrames()-1));
-		const float AnimTimeIncrement	= TimePerKey * AnimPlayRate;
-		uint32 AnimFrameIndex = 0;
-
-		FbxTime ExportTime;
-		ExportTime.SetSecondDouble(StartTime);
-
-		FbxTime ExportTimeIncrement;
-		ExportTimeIncrement.SetSecondDouble( TimePerKey );
-
 		int32 BoneTreeIndex = Skeleton->GetSkeletonBoneIndexFromMeshBoneIndex(SkelMesh, BoneIndex);
 		int32 BoneTrackIndex = Skeleton->GetRawAnimationTrackIndex(BoneTreeIndex, AnimSeq);
 		if(BoneTrackIndex == INDEX_NONE)
@@ -149,10 +151,7 @@ void FFbxExporter::ExportAnimSequenceToFbx(const UAnimSequence* AnimSeq,
 			Curve->KeyModifyBegin();
 		}
 
-		bool bLastKey = false;
-		// Step through each frame and add the bone's transformation data
-		while (!bLastKey)
-		{
+		auto ExportLambda = [&](float AnimTime, FbxTime ExportTime, bool bLastKey) {
 			FTransform BoneAtom;
 			AnimSeq->GetBoneTransform(BoneAtom, BoneTrackIndex, AnimTime, true);
 
@@ -160,34 +159,7 @@ void FFbxExporter::ExportAnimSequenceToFbx(const UAnimSequence* AnimSeq,
 			FbxVector4 Rotation = Converter.ConvertToFbxRot(BoneAtom.GetRotation().Euler());
 			FbxVector4 Scale = Converter.ConvertToFbxScale(BoneAtom.GetScale3D());
 			FbxVector4 Vectors[3] = { Translation, Rotation, Scale };
-		
-			//Add custom curve keys only to the root bone
-			if (BoneIndex == 0 && CustomCurveMap.Num() > 0)
-			{
-				FBlendedCurve BlendedCurve;
-				BlendedCurve.InitFrom(&AnimCurveUIDs);
-				AnimSeq->EvaluateCurveData(BlendedCurve, AnimTime, true);
-				if (BlendedCurve.IsValid())
-				{
-					//Loop over the custom curves and add the actual keys
-					for (auto CustomCurve : CustomCurveMap)
-					{
-						CustomCurve.Value->KeyModifyBegin();
-						SmartName::UID_Type NameUID = Skeleton->GetUIDByName(USkeleton::AnimCurveMappingName, CustomCurve.Key);
-						if (NameUID != SmartName::MaxUID)
-						{
-							float CurveValueAtTime = BlendedCurve.Get(NameUID);
-							int32 KeyIndex = CustomCurve.Value->KeyAdd(ExportTime);
-							CustomCurve.Value->KeySetValue(KeyIndex, CurveValueAtTime);
-						}
-					}
 
-				}
-			}
-
-			int32 lKeyIndex;
-			bLastKey = (AnimTime + KINDA_SMALL_NUMBER) > AnimEndTime;
-			
 			// Loop over each curve and channel to set correct values
 			for (uint32 CurveIndex = 0; CurveIndex < 3; ++CurveIndex)
 			{
@@ -195,7 +167,7 @@ void FFbxExporter::ExportAnimSequenceToFbx(const UAnimSequence* AnimSeq,
 				{
 					uint32 OffsetCurveIndex = (CurveIndex * 3) + ChannelIndex;
 
-					lKeyIndex = Curves[OffsetCurveIndex]->KeyAdd(ExportTime);
+					int32 lKeyIndex = Curves[OffsetCurveIndex]->KeyAdd(ExportTime);
 					Curves[OffsetCurveIndex]->KeySetValue(lKeyIndex, Vectors[CurveIndex][ChannelIndex]);
 					Curves[OffsetCurveIndex]->KeySetInterpolation(lKeyIndex, bLastKey ? FbxAnimCurveDef::eInterpolationConstant : FbxAnimCurveDef::eInterpolationCubic);
 
@@ -203,29 +175,97 @@ void FFbxExporter::ExportAnimSequenceToFbx(const UAnimSequence* AnimSeq,
 					{
 						Curves[OffsetCurveIndex]->KeySetConstantMode(lKeyIndex, FbxAnimCurveDef::eConstantStandard);
 					}
-				}				
+				}
 			}
+		};
 
-			ExportTime += ExportTimeIncrement;
-			AnimFrameIndex++;
-			AnimTime = AnimStartOffset + ((float)AnimFrameIndex * AnimTimeIncrement);
-		}
+		IterateInsideAnimSequence(AnimSeq, AnimStartOffset, AnimEndOffset, AnimPlayRate, StartTime, ExportLambda);
 
 		for (FbxAnimCurve* Curve : Curves)
 		{
 			Curve->KeyModifyEnd();
 		}
+	}
+}
 
-		if (BoneIndex == 0 && CustomCurveMap.Num() > 0)
+void FFbxExporter::ExportCustomAnimCurvesToFbx(const TMap<FName, FbxAnimCurve*>& CustomCurves, const UAnimSequence* AnimSeq, 
+	float AnimStartOffset, float AnimEndOffset, float AnimPlayRate, float StartTime, float ValueScale)
+{
+	// stack allocator for extracting curve
+	FMemMark Mark(FMemStack::Get());
+	const USkeleton* Skeleton = AnimSeq->GetSkeleton();
+	const FSmartNameMapping* SmartNameMapping = Skeleton ? Skeleton->GetSmartNameContainer(USkeleton::AnimCurveMappingName) : nullptr;
+
+	if (!Skeleton || !SmartNameMapping || !SetupAnimStack(AnimSeq))
+	{
+		//Something is wrong.
+		return;
+	}
+
+	TArray<SmartName::UID_Type> AnimCurveUIDs;
+	SmartNameMapping->FillUidArray(AnimCurveUIDs);
+
+	for (auto CustomCurve : CustomCurves)
+	{
+		CustomCurve.Value->KeyModifyBegin();
+	}
+	
+	auto ExportLambda = [&](float AnimTime, FbxTime ExportTime, bool bLastKey) {
+		FBlendedCurve BlendedCurve;
+		BlendedCurve.InitFrom(&AnimCurveUIDs);
+		AnimSeq->EvaluateCurveData(BlendedCurve, AnimTime, true);
+		if (BlendedCurve.IsValid())
 		{
-			for (auto CustomCurve : CustomCurveMap)
+			//Loop over the custom curves and add the actual keys
+			for (auto CustomCurve : CustomCurves)
 			{
-				CustomCurve.Value->KeyModifyEnd();
+				SmartName::UID_Type NameUID = Skeleton->GetUIDByName(USkeleton::AnimCurveMappingName, CustomCurve.Key);
+				if (NameUID != SmartName::MaxUID)
+				{
+					float CurveValueAtTime = BlendedCurve.Get(NameUID) * ValueScale;
+					int32 KeyIndex = CustomCurve.Value->KeyAdd(ExportTime);
+					CustomCurve.Value->KeySetValue(KeyIndex, CurveValueAtTime);
+				}
 			}
 		}
+	};
+
+	IterateInsideAnimSequence(AnimSeq, AnimStartOffset, AnimEndOffset, AnimPlayRate, StartTime, ExportLambda);
+
+	for (auto CustomCurve : CustomCurves)
+	{
+		CustomCurve.Value->KeyModifyEnd();
 	}
 }
 
+void FFbxExporter::IterateInsideAnimSequence(const UAnimSequence* AnimSeq, float AnimStartOffset, float AnimEndOffset, float AnimPlayRate, float StartTime, TFunctionRef<void(float, FbxTime, bool)> IterationLambda)
+{
+	float AnimTime = AnimStartOffset;
+	float AnimEndTime = (AnimSeq->SequenceLength - AnimEndOffset);
+	// Subtracts 1 because NumFrames includes an initial pose for 0.0 second
+	double TimePerKey = (AnimSeq->SequenceLength / (AnimSeq->GetRawNumberOfFrames() - 1));
+	const float AnimTimeIncrement = TimePerKey * AnimPlayRate;
+	uint32 AnimFrameIndex = 0;
+
+	FbxTime ExportTime;
+	ExportTime.SetSecondDouble(StartTime);
+
+	FbxTime ExportTimeIncrement;
+	ExportTimeIncrement.SetSecondDouble(TimePerKey);
+
+	// Step through each frame and add custom curve data
+	bool bLastKey = false;
+	while (!bLastKey)
+	{
+		bLastKey = (AnimTime + KINDA_SMALL_NUMBER) > AnimEndTime;
+
+		IterationLambda(AnimTime, ExportTime, bLastKey);
+
+		ExportTime += ExportTimeIncrement;
+		AnimFrameIndex++;
+		AnimTime = AnimStartOffset + ((float)AnimFrameIndex * AnimTimeIncrement);
+	}
+}
 
 // The curve code doesn't differentiate between angles and other data, so an interpolation from 179 to -179
 // will cause the bone to rotate all the way around through 0 degrees.  So here we make a second pass over the 
@@ -350,7 +390,7 @@ FbxNode* FFbxExporter::ExportAnimSequence( const UAnimSequence* AnimSeq, const U
 					double LodScreenSize = (double)(10.0f / SkelMesh->GetLODInfo(CurrentLodIndex)->ScreenSize.Default);
 					FbxLodGroupAttribute->AddThreshold(LodScreenSize);
 				}
-				FbxNode* FbxActorLOD = CreateMesh(SkelMesh, *FbxLODNodeName, CurrentLodIndex);
+				FbxNode* FbxActorLOD = CreateMesh(SkelMesh, *FbxLODNodeName, CurrentLodIndex, AnimSeq);
 				if (FbxActorLOD)
 				{
 					MeshRootNode->AddChild(FbxActorLOD);
@@ -366,7 +406,7 @@ FbxNode* FFbxExporter::ExportAnimSequence( const UAnimSequence* AnimSeq, const U
 		}
 		else
 		{
-			MeshRootNode = CreateMesh(SkelMesh, *MeshNodeName, 0);
+			MeshRootNode = CreateMesh(SkelMesh, *MeshNodeName, 0, AnimSeq);
 			if (MeshRootNode)
 			{
 				TmpNodeNoTransform->AddChild(MeshRootNode);
diff --git a/Engine/Source/Editor/UnrealEd/Private/Fbx/FbxSkeletalMeshExport.cpp b/Engine/Source/Editor/UnrealEd/Private/Fbx/FbxSkeletalMeshExport.cpp
--- a/Engine/Source/Editor/UnrealEd/Private/Fbx/FbxSkeletalMeshExport.cpp
+++ b/Engine/Source/Editor/UnrealEd/Private/Fbx/FbxSkeletalMeshExport.cpp
@@ -103,7 +103,7 @@ void FFbxExporter::GetSkeleton(FbxNode* RootNode, TArray<FbxNode*>& BoneNodes)
 /**
  * Adds an Fbx Mesh to the FBX scene based on the data in the given FSkeletalMeshLODModel
  */
-FbxNode* FFbxExporter::CreateMesh(const USkeletalMesh* SkelMesh, const TCHAR* MeshName, int32 LODIndex)
+FbxNode* FFbxExporter::CreateMesh(const USkeletalMesh* SkelMesh, const TCHAR* MeshName, int32 LODIndex, const UAnimSequence* AnimSeq)
 {
 	const FSkeletalMeshModel* SkelMeshResource = SkelMesh->GetImportedModel();
 	if (!SkelMeshResource->LODModels.IsValidIndex(LODIndex))
@@ -253,12 +253,15 @@ FbxNode* FFbxExporter::CreateMesh(const USkeletalMesh* SkelMesh, const TCHAR* Me
 
 	if (GetExportOptions()->bExportMorphTargets)
 	{
+		const FSmartNameMapping* SmartNameMapping = SkelMesh->Skeleton->GetSmartNameContainer(USkeleton::AnimCurveMappingName);
+		TMap<FName, FbxAnimCurve*> BlendShapeCurvesMap;
+
 		for (UMorphTarget* MorphTarget : SkelMesh->MorphTargets)
 		{
 			// The original BlendShape Name was not saved during import, so we need to come up with a new one.
 			const FString BlendShapeName(TEXT("blendShape_") + MorphTarget->GetName());
 			FbxBlendShape* BlendShape = FbxBlendShape::Create(Mesh, TCHAR_TO_UTF8(*BlendShapeName));
-			Mesh->AddDeformer(BlendShape);
+			int32 DeformerIndex = Mesh->AddDeformer(BlendShape);
 			FbxBlendShapeChannel* BlendShapeChannel = FbxBlendShapeChannel::Create(BlendShape, TCHAR_TO_UTF8(*MorphTarget->GetName()));
 
 			if (BlendShape->AddBlendShapeChannel(BlendShapeChannel))
@@ -284,8 +287,24 @@ FbxNode* FFbxExporter::CreateMesh(const USkeletalMesh* SkelMesh, const TCHAR* Me
 				}
 
 				BlendShapeChannel->AddTargetShape(Shape);
+				FName MorphTargetName = MorphTarget->GetFName();
+				if (AnimSeq && SmartNameMapping && SmartNameMapping->GetCurveMetaData(MorphTargetName) && SmartNameMapping->GetCurveMetaData(MorphTargetName)->Type.bMorphtarget)
+				{
+					FbxAnimCurve* AnimCurve = Mesh->GetShapeChannel(DeformerIndex, 0, AnimLayer, true);
+					BlendShapeCurvesMap.Add(MorphTargetName, AnimCurve);
+				}
 			}
 		}
+
+		if(AnimSeq && BlendShapeCurvesMap.Num() > 0)
+		{
+			ExportCustomAnimCurvesToFbx(BlendShapeCurvesMap, AnimSeq,
+				0.f,		// AnimStartOffset
+				0.f,		// AnimEndOffset
+				1.f,		// AnimPlayRate
+				0.f,		// StartTime
+				100.f);		// ValueScale, for some reason we need to scale BlendShape curves by a factor of 100.
+		}
 	}
 
 	FbxNode* MeshNode = FbxNode::Create(Scene, TCHAR_TO_UTF8(MeshName));
diff --git a/Engine/Source/Editor/UnrealEd/Private/FbxExporter.h b/Engine/Source/Editor/UnrealEd/Private/FbxExporter.h
--- a/Engine/Source/Editor/UnrealEd/Private/FbxExporter.h
+++ b/Engine/Source/Editor/UnrealEd/Private/FbxExporter.h
@@ -369,8 +369,12 @@ private:
 
 	/**
 	 * Adds an Fbx Mesh to the FBX scene based on the data in the given FSkeletalMeshLODModel
+	 * @param SkelMesh			The SkeletalMesh we are exporting
+	 * @param MeshName			The SkeletalMesh name
+	 * @param LODIndex			The mesh LOD index we are exporting
+	 * @param AnimSeq			If an AnimSeq is provided and are exporting MorphTarget, the MorphTarget animation will be exported as well.
 	 */
-	FbxNode* CreateMesh(const USkeletalMesh* SkelMesh, const TCHAR* MeshName, int32 LODIndex);
+	FbxNode* CreateMesh(const USkeletalMesh* SkelMesh, const TCHAR* MeshName, int32 LODIndex, const UAnimSequence* AnimSeq = nullptr);
 
 	/**
 	 * Adds Fbx Clusters necessary to skin a skeletal mesh to the bones in the BoneNodes list
@@ -390,12 +394,26 @@ private:
 	/** Export SkeletalMeshComponent */
 	void ExportSkeletalMeshComponent(USkeletalMeshComponent* SkelMeshComp, const TCHAR* MeshName, FbxNode* ActorRootNode, bool bSaveAnimSeq = true);
 
+	/** Initializing the AnimStack playrate from the AnimSequence */
+	bool SetupAnimStack(const UAnimSequence* AnimSeq);
+
 	/**
 	 * Add the given animation sequence as rotation and translation tracks to the given list of bone nodes
 	 */
 	void ExportAnimSequenceToFbx(const UAnimSequence* AnimSeq, const USkeletalMesh* SkelMesh, TArray<FbxNode*>& BoneNodes, FbxAnimLayer* AnimLayer,
 		float AnimStartOffset, float AnimEndOffset, float AnimPlayRate, float StartTime);
 
+	/**
+	 * Add the custom Curve data to the FbxAnimCurves passed in parameter by matching their name to the skeletal mesh custom curves.
+	 */
+	void ExportCustomAnimCurvesToFbx(const TMap<FName, FbxAnimCurve*>& CustomCurves, const UAnimSequence* AnimSeq,
+		float AnimStartOffset, float AnimEndOffset, float AnimPlayRate, float StartTime, float ValueScale = 1.f);
+
+	/**
+	 * Used internally to reuse the AnimSequence iteration code when exporting various kind of curves.
+	 */
+	void IterateInsideAnimSequence(const UAnimSequence* AnimSeq, float AnimStartOffset, float AnimEndOffset, float AnimPlayRate, float StartTime, TFunctionRef<void(float, FbxTime, bool)> IterationLambda);
+
 	/** 
 	 * The curve code doesn't differentiate between angles and other data, so an interpolation from 179 to -179
 	 * will cause the bone to rotate all the way around through 0 degrees.  So here we make a second pass over the 
