[NO NON-PII COMMIT MESSAGE]

diff --git a/Engine/Source/Editor/UnrealEd/Private/Factories/EditorFactories.cpp b/Engine/Source/Editor/UnrealEd/Private/Factories/EditorFactories.cpp
--- a/Engine/Source/Editor/UnrealEd/Private/Factories/EditorFactories.cpp
+++ b/Engine/Source/Editor/UnrealEd/Private/Factories/EditorFactories.cpp
@@ -5131,11 +5131,14 @@ void FCustomizableTextObjectFactory::ClearObjectNameUsage(UObject* InParent, FNa
 	}
 
 	// If there is already another object in the same scope with this name, rename it.
-	if( Found )
+	while (Found)
 	{
 		check(Found->GetOuter() == InParent);
 
 		Found->Rename(nullptr, nullptr, REN_DontCreateRedirectors);
+
+		// It's possible after undo for there to be multiple objects with the same name in the way, rename all of them
+		Found = FindObject<UObject>(InParent, *InName.ToString());
 	}
 }
 
