Fixed ensure about a module being loaded from outside the game thread

#rnx

diff --git a/Engine/Source/Editor/UnrealEd/Classes/Commandlets/GatherTextCommandletBase.h b/Engine/Source/Editor/UnrealEd/Classes/Commandlets/GatherTextCommandletBase.h
--- a/Engine/Source/Editor/UnrealEd/Classes/Commandlets/GatherTextCommandletBase.h
+++ b/Engine/Source/Editor/UnrealEd/Classes/Commandlets/GatherTextCommandletBase.h
@@ -58,6 +58,7 @@ class UNREALED_API UGatherTextCommandletBase : public UCommandlet
 
 public:
 	virtual void Initialize( const TSharedRef< FLocTextHelper >& InGatherManifestHelper, const TSharedPtr< FLocalizationSCC >& InSourceControlInfo );
+	virtual void BeginDestroy() override;
 
 	// Wrappers for extracting config values
 	bool GetBoolFromConfig( const TCHAR* Section, const TCHAR* Key, bool& OutValue, const FString& Filename );
diff --git a/Engine/Source/Editor/UnrealEd/Private/Commandlets/GatherTextCommandlet.cpp b/Engine/Source/Editor/UnrealEd/Private/Commandlets/GatherTextCommandlet.cpp
--- a/Engine/Source/Editor/UnrealEd/Private/Commandlets/GatherTextCommandlet.cpp
+++ b/Engine/Source/Editor/UnrealEd/Private/Commandlets/GatherTextCommandlet.cpp
@@ -7,6 +7,7 @@
 #include "HAL/FileManager.h"
 #include "UObject/Class.h"
 #include "UObject/Package.h"
+#include "UObject/GCObjectScopeGuard.h"
 #include "SourceControlHelpers.h"
 
 DEFINE_LOG_CATEGORY_STATIC(LogGatherTextCommandlet, Log, All);
@@ -179,7 +180,8 @@ int32 UGatherTextCommandlet::ProcessGatherConfig(const FString& GatherTextConfig
 
 		UGatherTextCommandletBase* Commandlet = NewObject<UGatherTextCommandletBase>(GetTransientPackage(), CommandletClass);
 		check(Commandlet);
-		Commandlet->AddToRoot();
+		
+		FGCObjectScopeGuard CommandletGCGuard(Commandlet);
 		Commandlet->Initialize( CommandletGatherManifestHelper, CommandletSourceControlInfo );
 
 		// Execute the commandlet.
diff --git a/Engine/Source/Editor/UnrealEd/Private/Commandlets/GatherTextCommandletBase.cpp b/Engine/Source/Editor/UnrealEd/Private/Commandlets/GatherTextCommandletBase.cpp
--- a/Engine/Source/Editor/UnrealEd/Private/Commandlets/GatherTextCommandletBase.cpp
+++ b/Engine/Source/Editor/UnrealEd/Private/Commandlets/GatherTextCommandletBase.cpp
@@ -42,6 +42,14 @@ void UGatherTextCommandletBase::Initialize( const TSharedRef< FLocTextHelper >&
 	}
 }
 
+void UGatherTextCommandletBase::BeginDestroy()
+{
+	Super::BeginDestroy();
+
+	GatherManifestHelper.Reset();
+	SourceControlInfo.Reset();
+}
+
 void UGatherTextCommandletBase::CreateCustomEngine(const FString& Params)
 {
 	GEngine = GEditor = NULL;//Force a basic default engine. 
