Swap CastChecked for Cast to avoid exception and can fall through to nullptr test. #rb none #fyi Alexander.Sanoja #jira UE-75921

diff --git a/Engine/Source/Editor/UnrealEd/Private/EditorModeManager.cpp b/Engine/Source/Editor/UnrealEd/Private/EditorModeManager.cpp
--- a/Engine/Source/Editor/UnrealEd/Private/EditorModeManager.cpp
+++ b/Engine/Source/Editor/UnrealEd/Private/EditorModeManager.cpp
@@ -169,7 +169,7 @@ void FEditorModeTools::UpdateModeWidgetLocation()
 
 	if (SelectedActors && SelectedActors->Num() > 0)
 	{
-		AActor* Actor = CastChecked<AActor>(SelectedActors->GetSelectedObject(SelectedActors->Num() - 1));
+		AActor* Actor = Cast<AActor>(SelectedActors->GetSelectedObject(SelectedActors->Num() - 1));
 
 		if (Actor)
 		{
