Fixed propagation of bIsEditorPreview in level serialization or through duplication in PIE.

diff --git a/Engine/Source/Editor/UnrealEd/Private/Factories/ActorFactory.cpp b/Engine/Source/Editor/UnrealEd/Private/Factories/ActorFactory.cpp
--- a/Engine/Source/Editor/UnrealEd/Private/Factories/ActorFactory.cpp
+++ b/Engine/Source/Editor/UnrealEd/Private/Factories/ActorFactory.cpp
@@ -311,7 +311,7 @@ AActor* UActorFactory::SpawnActor( UObject* Asset, ULevel* InLevel, const FTrans
 		SpawnInfo.ObjectFlags = InObjectFlags;
 		SpawnInfo.Name = Name;
 #if WITH_EDITOR
-		SpawnInfo.bTemporaryEditorActor = GEditor->bIsSimulatingInEditor ? FLevelEditorViewportClient::IsDroppingPreviewActor(): true;
+		SpawnInfo.bTemporaryEditorActor = FLevelEditorViewportClient::IsDroppingPreviewActor();
 #endif
 		return InLevel->OwningWorld->SpawnActor( DefaultActor->GetClass(), &Transform, SpawnInfo );
 	}
diff --git a/Engine/Source/Runtime/Engine/Classes/GameFramework/Actor.h b/Engine/Source/Runtime/Engine/Classes/GameFramework/Actor.h
--- a/Engine/Source/Runtime/Engine/Classes/GameFramework/Actor.h
+++ b/Engine/Source/Runtime/Engine/Classes/GameFramework/Actor.h
@@ -689,7 +689,7 @@ public:
 	uint8 bHiddenEd:1;
 
 	/** True if this actor is the preview actor dragged out of the content browser */
-	UPROPERTY()
+	UPROPERTY(Transient)
 	uint8 bIsEditorPreviewActor:1;
 
 	/** Whether this actor is hidden by the layer browser. */
