Removed unrequired shader compilation when changing preview feature level.
Texture streaming accuracy visualization now show the current state of data, and doesn't need to build the texture streaming as a prestep.

diff --git a/Engine/Source/Editor/MaterialEditor/Private/MaterialEditingLibrary.cpp b/Engine/Source/Editor/MaterialEditor/Private/MaterialEditingLibrary.cpp
--- a/Engine/Source/Editor/MaterialEditor/Private/MaterialEditingLibrary.cpp
+++ b/Engine/Source/Editor/MaterialEditor/Private/MaterialEditingLibrary.cpp
@@ -27,7 +27,7 @@
 #include "Misc/RuntimeErrors.h"
 #include "SceneTypes.h"
 #include "AssetRegistryModule.h"
-
+#include "DebugViewModeHelpers.h"
 
 DEFINE_LOG_CATEGORY_STATIC(LogMaterialEditingLibrary, Warning, All);
 
@@ -650,6 +650,7 @@ void UMaterialEditingLibrary::RecompileMaterial(UMaterial* Material)
 
 		UMaterialEditingLibrary::RebuildMaterialInstanceEditors(Material);
 
+		ClearDebugViewMaterials(Material);
 		FMaterialEditorUtilities::BuildTextureStreamingData(Material);
 	}
 }
diff --git a/Engine/Source/Editor/MaterialEditor/Private/MaterialEditorUtilities.cpp b/Engine/Source/Editor/MaterialEditor/Private/MaterialEditorUtilities.cpp
--- a/Engine/Source/Editor/MaterialEditor/Private/MaterialEditorUtilities.cpp
+++ b/Engine/Source/Editor/MaterialEditor/Private/MaterialEditorUtilities.cpp
@@ -688,22 +688,12 @@ void FMaterialEditorUtilities::BuildTextureStreamingData(UMaterialInterface* Upd
 		FScopedSlowTask SlowTask(2.f, (LOCTEXT("MaterialEditorUtilities_UpdatingTextureStreamingData", "Updating Texture Streaming Data")));
 		SlowTask.MakeDialog(true);
 
-		TSet<UMaterialInterface*> Materials;
-		Materials.Add(UpdatedMaterial);
-
 		// Clear the build data.
 		const TArray<FMaterialTextureInfo> EmptyTextureStreamingData;
+		UpdatedMaterial->SetTextureStreamingData(EmptyTextureStreamingData);
 
-		// Here we also update the parents as we just want to save the delta between the hierarchy.
-		// Since the instance may only override partially the parent params, we try to find what the child has overridden.
-		UMaterialInstance* MaterialInstance = Cast<UMaterialInstance>(UpdatedMaterial);
-		while (MaterialInstance)
-		{
-			// Clear the data in case the build is canceled.
-			MaterialInstance->SetTextureStreamingData(EmptyTextureStreamingData);
-			Materials.Add(MaterialInstance);
-			MaterialInstance = Cast<UMaterialInstance>(MaterialInstance->Parent);
-		};
+		TSet<UMaterialInterface*> Materials;
+		Materials.Add(UpdatedMaterial);
 
 		// Here we need a full rebuild since the shader changed. Although don't wait for previous shaders to fasten the process.
 		if (CompileDebugViewModeShaders(DVSM_OutputMaterialTextureScales, QualityLevel, FeatureLevel, true, false, Materials, &SlowTask))
diff --git a/Engine/Source/Editor/MaterialEditor/Private/MaterialInstanceEditor.cpp b/Engine/Source/Editor/MaterialEditor/Private/MaterialInstanceEditor.cpp
--- a/Engine/Source/Editor/MaterialEditor/Private/MaterialInstanceEditor.cpp
+++ b/Engine/Source/Editor/MaterialEditor/Private/MaterialInstanceEditor.cpp
@@ -43,6 +43,7 @@
 #include "MaterialStats.h"
 #include "MaterialEditingLibrary.h"
 #include "Widgets/Layout/SScrollBox.h"
+#include "DebugViewModeHelpers.h"
 
 #define LOCTEXT_NAMESPACE "MaterialInstanceEditor"
 
@@ -614,6 +615,7 @@ FMaterialInstanceEditor::~FMaterialInstanceEditor()
 		UPackage* Package = MaterialEditorInstance->SourceInstance->GetOutermost();
 		if (Package && Package->IsDirty() && Package != GetTransientPackage())
 		{
+			ClearDebugViewMaterials(MaterialEditorInstance->SourceInstance);
 			FMaterialEditorUtilities::BuildTextureStreamingData(MaterialEditorInstance->SourceInstance);
 		}
 	}
diff --git a/Engine/Source/Editor/UnrealEd/Classes/Editor/EditorEngine.h b/Engine/Source/Editor/UnrealEd/Classes/Editor/EditorEngine.h
--- a/Engine/Source/Editor/UnrealEd/Classes/Editor/EditorEngine.h
+++ b/Engine/Source/Editor/UnrealEd/Classes/Editor/EditorEngine.h
@@ -3086,17 +3086,6 @@ protected:
 	void LoadEditorFeatureLevel();
 	void SaveEditorFeatureLevel();
 
-	/** For some platforms (e.g. mobiles), when running in editor mode, we emulate the shaders functionality on available running GPU (e.g. DirectX),
-	 *  but when displaying the shader complexity we need to be able compile and extract statistics (instruction count) from the real shaders that
-	 *  will be compiled when the game will run on the specific platform. Thus (if compiler available) we perform an 'offline' shader compilation step,
-	 *  extract the needed statistics and transfer them to the emulated editor running shaders.
-	 *  This function will be called from OnSceneMaterialsModified()
-	 *
-	 * @param	bForceUpdate	When true, view mode shaders are always updated for worlds displaying shader complexity materials.
-	 *                          When false, view mode shaders are rebuilt only when emulating a shader platform.
-	 */
-	void UpdateShaderComplexityMaterials(bool bForceUpdate);
-
 	/** Utility function that can determine whether some input world is using materials who's shaders are emulated in the editor */
 	bool IsEditorShaderPlatformEmulated(UWorld* World);
 
diff --git a/Engine/Source/Editor/UnrealEd/Private/EditorBuildUtils.cpp b/Engine/Source/Editor/UnrealEd/Private/EditorBuildUtils.cpp
--- a/Engine/Source/Editor/UnrealEd/Private/EditorBuildUtils.cpp
+++ b/Engine/Source/Editor/UnrealEd/Private/EditorBuildUtils.cpp
@@ -1088,6 +1088,66 @@ EDebugViewShaderMode ViewModeIndexToDebugViewShaderMode(EViewModeIndex SelectedV
 	}
 }
 
+void FEditorBuildUtils::UpdateTextureStreamingMaterialBindings(UWorld* InWorld)
+{
+	const EMaterialQualityLevel::Type QualityLevel = EMaterialQualityLevel::High;
+	const ERHIFeatureLevel::Type FeatureLevel = GMaxRHIFeatureLevel;
+
+	CollectGarbage( GARBAGE_COLLECTION_KEEPFLAGS );
+
+	TSet<UMaterialInterface*> Materials;
+	if (GetUsedMaterialsInWorld(InWorld, Materials, nullptr))
+	{
+		// Flush renderthread since we are about to update the material streaming data.
+		FlushRenderingCommands();
+
+		for (UMaterialInterface* MaterialInterface : Materials)
+		{
+			if (!MaterialInterface)
+			{
+				continue;
+			}
+
+			TArray<UTexture*> UsedTextures;
+			TArray< TArray<int32> > UsedIndices;
+			MaterialInterface->GetUsedTexturesAndIndices(UsedTextures, UsedIndices, QualityLevel, FeatureLevel);
+			MaterialInterface->SortTextureStreamingData(true, false);
+
+			MaterialInterface->TextureStreamingDataMissingEntries.Empty();
+
+			for (int32 UsedIndex = 0; UsedIndex < UsedTextures.Num(); ++UsedIndex)
+			{
+				if (UsedTextures[UsedIndex])
+				{
+					TArray<FMaterialTextureInfo>& MaterialData = MaterialInterface->GetTextureStreamingData();
+
+					int32 LowerIndex = INDEX_NONE;
+					int32 HigherIndex = INDEX_NONE;
+					if (MaterialInterface->FindTextureStreamingDataIndexRange(UsedTextures[UsedIndex]->GetFName(), LowerIndex, HigherIndex))
+					{
+						// Here we expect every entry in UsedIndices to match one of the entry in the range.
+						for (int32 SubIndex = 0; LowerIndex <= HigherIndex && SubIndex < UsedIndices[UsedIndex].Num(); ++LowerIndex, ++SubIndex)
+						{
+							MaterialData[LowerIndex].TextureIndex = UsedIndices[UsedIndex][SubIndex];
+						}
+					}
+					else // If the texture is missing in the material data, add it ass missing
+					{
+						FMaterialTextureInfo MissingInfo;
+						MissingInfo.TextureName = UsedTextures[UsedIndex]->GetFName();
+						for (int32 SubIndex = 0; SubIndex < UsedIndices[UsedIndex].Num(); ++SubIndex)
+						{
+							MissingInfo.TextureIndex = UsedIndices[UsedIndex][SubIndex];
+							MaterialInterface->TextureStreamingDataMissingEntries.Add(MissingInfo);
+						}
+					}
+				}
+			}
+		}
+	}
+}
+
+
 bool FEditorBuildUtils::EditorBuildTextureStreaming(UWorld* InWorld, EViewModeIndex SelectedViewMode)
 {
 	if (!InWorld) return false;
@@ -1277,67 +1337,6 @@ bool FEditorBuildUtils::EditorBuildMaterialTextureStreamingData(UPackage* Packag
 	return bAnyPackagesDirtied;
 }
 
-bool FEditorBuildUtils::CompileViewModeShaders(UWorld* InWorld, EViewModeIndex SelectedViewMode)
-{
-	if (!InWorld)
-	{
-		return false;
-	}
-
-	// This process is incredibly slow for large projects even if we end up compiling no shaders
-	// but isn't essential and we will lazily compile shaders afterwards. This upfront compile
-	// was found to cause minute-long stalls every time the view mode was activated on a simulated
-	// feature level so the check for !bExtractComplexityStats has been removed for shader complexity.
-	//const EShaderPlatform ShaderPlatform = GetFeatureLevelShaderPlatform(FeatureLevel);
-	//bool bIsSimulated = IsSimulatedPlatform(ShaderPlatform);
-	//bool bExtractComplexityStats = bIsSimulated && (SelectedViewMode == VMI_ShaderComplexity || SelectedViewMode == VMI_ShaderComplexityWithQuadOverdraw);
-	EDebugViewShaderMode DebugViewMode = ViewModeIndexToDebugViewShaderMode(SelectedViewMode);
-	if (DebugViewMode == DVSM_None)
-	{
-		return false;
-	}
-
-	const ERHIFeatureLevel::Type FeatureLevel = InWorld->FeatureLevel;
-	const EMaterialQualityLevel::Type QualityLevel = GetCachedScalabilityCVars().MaterialQualityLevel;;
-
-	FScopedSlowTask CompileShaderTask(3.f, LOCTEXT("CompileMissingViewModeShaders", "Compiling Missing ViewMode Shaders")); // { Get Used Materials, Sync Pending Shader, Wait for Compilation }
-	CompileShaderTask.MakeDialog(false); // Can't cancel because get compiled anyway if shaders are missing.
-
-	CollectGarbage(GARBAGE_COLLECTION_KEEPFLAGS);
-
-	TSet<UMaterialInterface*> Materials;
-	if (!GetUsedMaterialsInWorld(InWorld, Materials, &CompileShaderTask))
-	{
-		return false;
-	}
-
-	if (Materials.Num())
-	{
-		// As mentioned above this is now only done if the SelectedViewMode is VMI_RequiredTextureResolution so removing the unused code
-/*		if (SelectedViewMode == VMI_ShaderComplexity || SelectedViewMode == VMI_ShaderComplexityWithQuadOverdraw)
-		{
-			if (!CompileShadersComplexityViewMode(QualityLevel, FeatureLevel, Materials, CompileShaderTask))
-			{
-				return false;
-			}
-		}
-		else if (SelectedViewMode == VMI_RequiredTextureResolution)*/
-		{
-			if (!CompileDebugViewModeShaders(DebugViewMode, QualityLevel, FeatureLevel, false, true, Materials, &CompileShaderTask))
-			{
-				return false;
-			}
-		}
-	}
-	else
-	{
-		CompileShaderTask.EnterProgressFrame();
-	}
-
-	CollectGarbage(GARBAGE_COLLECTION_KEEPFLAGS);
-	return true;
-}
-
 /** classed used to compile shaders for a specific (mobile) platform and copy the number of instruction to the editor-emulated (mobile) platform */
 class FMaterialOfflineCompilation : public FMaterialResource
 {
diff --git a/Engine/Source/Editor/UnrealEd/Private/EditorEngine.cpp b/Engine/Source/Editor/UnrealEd/Private/EditorEngine.cpp
--- a/Engine/Source/Editor/UnrealEd/Private/EditorEngine.cpp
+++ b/Engine/Source/Editor/UnrealEd/Private/EditorEngine.cpp
@@ -7319,36 +7319,8 @@ bool UEditorEngine::IsOfflineShaderCompilerAvailable(UWorld* World)
 	return FMaterialStatsUtils::IsPlatformOfflineCompilerAvailable(RealPlatform);
 }
 
-void UEditorEngine::UpdateShaderComplexityMaterials(bool bForceUpdate)
-{
-	TSet<UWorld *> WorldSet;
-
-	for(FLevelEditorViewportClient* ViewportClient : LevelViewportClients)
-	{
-		auto ViewMode = ViewportClient->GetViewMode();
-		if (ViewMode == EViewModeIndex::VMI_ShaderComplexity || ViewMode == EViewModeIndex::VMI_ShaderComplexityWithQuadOverdraw)
-		{
-			WorldSet.Add(ViewportClient->GetWorld());
-		}
-	}
-
-	for (auto* SomeWorld : WorldSet)
-	{
-		bool bShadersEmulated = IsEditorShaderPlatformEmulated(SomeWorld);
-		if (bShadersEmulated || bForceUpdate)
-		{
-			bool bOfflineCompilerAvailable = IsOfflineShaderCompilerAvailable(SomeWorld);
-			if (bOfflineCompilerAvailable || bForceUpdate)
-			{
-				FEditorBuildUtils::CompileViewModeShaders(SomeWorld, VMI_ShaderComplexity);
-			}
-		}
-	}
-}
-
 void UEditorEngine::OnSceneMaterialsModified()
 {
-	UpdateShaderComplexityMaterials(false);
 }
 
 void UEditorEngine::SetMaterialsFeatureLevel(const ERHIFeatureLevel::Type InFeatureLevel)
@@ -7387,14 +7359,6 @@ void UEditorEngine::SetMaterialsFeatureLevel(const ERHIFeatureLevel::Type InFeat
 	GShaderCompilingManager->ProcessAsyncResults(false, true);
 
 	PreviewFeatureLevelChanged.Broadcast(InFeatureLevel);
-
-	// The feature level changed, so existing debug view materials are invalid and need to be rebuilt.
-	// This process must follow the PreviewFeatureLevelChanged event, because any listeners need
-	// opportunity to switch to the new feature level first.
-	void ClearDebugViewMaterials(UMaterialInterface*);
-	ClearDebugViewMaterials(nullptr);
-
-	UpdateShaderComplexityMaterials(true);
 }
 
 void UEditorEngine::SetFeatureLevelPreview(const ERHIFeatureLevel::Type InPreviewFeatureLevel)
@@ -7466,14 +7430,6 @@ void UEditorEngine::ToggleFeatureLevelPreview()
 	ERHIFeatureLevel::Type NewPreviewFeatureLevel = Settings->bIsFeatureLevelPreviewActive ? PreviewFeatureLevel : GMaxRHIFeatureLevel;
 
 	PreviewFeatureLevelChanged.Broadcast(NewPreviewFeatureLevel);
-	
-	// The feature level changed, so existing debug view materials are invalid and need to be rebuilt.
-	// This process must follow the PreviewFeatureLevelChanged event, because any listeners need
-	// opportunity to switch to the new feature level first.
-	void ClearDebugViewMaterials(UMaterialInterface*);
-	ClearDebugViewMaterials(nullptr);
-
-	UpdateShaderComplexityMaterials(true);
 
 	GEditor->RedrawAllViewports();
 	
diff --git a/Engine/Source/Editor/UnrealEd/Private/EditorViewportClient.cpp b/Engine/Source/Editor/UnrealEd/Private/EditorViewportClient.cpp
--- a/Engine/Source/Editor/UnrealEd/Private/EditorViewportClient.cpp
+++ b/Engine/Source/Editor/UnrealEd/Private/EditorViewportClient.cpp
@@ -5125,12 +5125,10 @@ void FEditorViewportClient::SetViewMode(EViewModeIndex InViewModeIndex)
 
 	if (IsPerspective())
 	{
-		if (InViewModeIndex == VMI_PrimitiveDistanceAccuracy || InViewModeIndex == VMI_MeshUVDensityAccuracy || InViewModeIndex == VMI_MaterialTextureScaleAccuracy)
+		if (InViewModeIndex == VMI_MaterialTextureScaleAccuracy)
 		{
-			FEditorBuildUtils::EditorBuildTextureStreaming(GetWorld(), InViewModeIndex);
+			FEditorBuildUtils::UpdateTextureStreamingMaterialBindings(GetWorld());
 		}
-		// Uncomment this to generate inital viewmode data.
-		// FEditorBuildUtils::CompileViewModeShaders(GetWorld(), InViewModeIndex);
 
 		PerspViewModeIndex = InViewModeIndex;
 		ApplyViewMode(PerspViewModeIndex, true, EngineShowFlags);
diff --git a/Engine/Source/Editor/UnrealEd/Private/EditorViewportCommands.cpp b/Engine/Source/Editor/UnrealEd/Private/EditorViewportCommands.cpp
--- a/Engine/Source/Editor/UnrealEd/Private/EditorViewportCommands.cpp
+++ b/Engine/Source/Editor/UnrealEd/Private/EditorViewportCommands.cpp
@@ -280,6 +280,18 @@ static void AppendTextureStreamingInfoToMenu(const UMaterialInterface* MaterialI
 			}
 		}
 	}
+
+	for (const FMaterialTextureInfo& MissingEntry : MaterialInterface->TextureStreamingDataMissingEntries)
+	{
+		if (bSingleMaterial)
+		{
+			DataPerTextureIndex.FindOrAdd(MissingEntry.TextureIndex).AddUnique(FString::Printf(TEXT("Missing : %s"), *MissingEntry.TextureName.ToString()));
+		}
+		else
+		{
+			DataPerTextureIndex.FindOrAdd(MissingEntry.TextureIndex).AddUnique(FString::Printf(TEXT("Missing : %s.%s"), *MaterialInterface->GetName(), *MissingEntry.TextureName.ToString()));
+		}
+	}
 }
 
 static void AppendMaterialInfoToMenu(const UMaterialInterface* MaterialInterface, ERHIFeatureLevel::Type FeatureLevel, const FString& MenuName, TMap<int32, TArray<FString> >& DataPerTextureIndex, TMap<FName, TArray<FString> >& DataPerTextureName)
diff --git a/Engine/Source/Editor/UnrealEd/Public/EditorBuildUtils.h b/Engine/Source/Editor/UnrealEd/Public/EditorBuildUtils.h
--- a/Engine/Source/Editor/UnrealEd/Public/EditorBuildUtils.h
+++ b/Engine/Source/Editor/UnrealEd/Public/EditorBuildUtils.h
@@ -124,6 +124,13 @@ public:
 	 */
 	static UNREALED_API bool EditorBuild( UWorld* InWorld, FName Id, const bool bAllowLightingDialog = true );
 
+	/**
+	 * Update every used material texture binding (see FMaterialTextureInfo::TextureIndex)
+	 *
+	 * @param	InWorld				WorldContext
+	 */
+	static UNREALED_API void UpdateTextureStreamingMaterialBindings( UWorld* InWorld );
+
 	/**
 	 * Perform an editor build for texture streaming
 	 *
@@ -144,16 +151,6 @@ public:
 	*/
 	static UNREALED_API bool EditorBuildMaterialTextureStreamingData(UPackage* Package);
 
-	/**
-	 * Compiled required shader for the given viewmode
-	 *
-	 * @param	InWorld				WorldContext to find out which shader to build.
-	 * @param	InViewMode			Viewmode to build the data for.
-	 *
-	 * @return	true if the build completed successfully; false if it did not (or was manually canceled)
-	 */
-	static UNREALED_API bool CompileViewModeShaders(UWorld* InWorld, EViewModeIndex InViewMode);
-
 	/** 
 	* check if navigation build was was triggered from editor as user request
 	*
diff --git a/Engine/Source/Runtime/Engine/Classes/Materials/MaterialInterface.h b/Engine/Source/Runtime/Engine/Classes/Materials/MaterialInterface.h
--- a/Engine/Source/Runtime/Engine/Classes/Materials/MaterialInterface.h
+++ b/Engine/Source/Runtime/Engine/Classes/Materials/MaterialInterface.h
@@ -259,6 +259,10 @@ public:
 	//~ End IInterface_AssetUserData Interface
 
 #if WITH_EDITORONLY_DATA
+	/** List of all used but missing texture indices in TextureStreamingData. Used for visualization / debugging only. */
+	UPROPERTY(transient)
+	TArray<FMaterialTextureInfo> TextureStreamingDataMissingEntries;
+
 	/** The mesh used by the material editor to preview the material.*/
 	UPROPERTY(EditAnywhere, Category=Previewing, meta=(AllowedClasses="StaticMesh,SkeletalMesh", ExactClass="true"))
 	FSoftObjectPath PreviewMesh;
@@ -828,6 +832,10 @@ public:
 	FORCEINLINE bool HasTextureStreamingData() const { return TextureStreamingData.Num() != 0; }
 	/** Accessor to the data. */
 	FORCEINLINE const TArray<FMaterialTextureInfo>& GetTextureStreamingData() const { return TextureStreamingData; }
+	FORCEINLINE TArray<FMaterialTextureInfo>& GetTextureStreamingData() { return TextureStreamingData; }
+	/** Find entries within TextureStreamingData that match the given name. */
+	ENGINE_API bool FindTextureStreamingDataIndexRange(FName TextureName, int32& LowerIndex, int32& HigherIndex) const;
+
 	/** Set new texture streaming data. */
 	ENGINE_API void SetTextureStreamingData(const TArray<FMaterialTextureInfo>& InTextureStreamingData);
 
@@ -842,8 +850,6 @@ public:
 
 	ENGINE_API virtual void PreSave(const class ITargetPlatform* TargetPlatform) override;
 
-protected:
-
 	/**
 	* Sort the texture streaming data by names to accelerate search. Only sorts if required.
 	*
@@ -852,14 +858,13 @@ protected:
 	*/
 	ENGINE_API void SortTextureStreamingData(bool bForceSort, bool bFinalSort);
 
+protected:
+
 	/** Returns a bitfield indicating which feature levels should be compiled for rendering. GMaxRHIFeatureLevel is always present */
 	ENGINE_API uint32 GetFeatureLevelsToCompileForRendering() const;
 
 	void UpdateMaterialRenderProxy(FMaterialRenderProxy& Proxy);
 
-	/** Find entries within TextureStreamingData that match the given name. */
-	bool FindTextureStreamingDataIndexRange(FName TextureName, int32& LowerIndex, int32& HigherIndex) const;
-
 private:
 	/**
 	 * Post loads all default materials.
diff --git a/Engine/Source/Runtime/Engine/Private/Materials/MaterialInterface.cpp b/Engine/Source/Runtime/Engine/Private/Materials/MaterialInterface.cpp
--- a/Engine/Source/Runtime/Engine/Private/Materials/MaterialInterface.cpp
+++ b/Engine/Source/Runtime/Engine/Private/Materials/MaterialInterface.cpp
@@ -573,7 +573,18 @@ void UMaterialInterface::SortTextureStreamingData(bool bForceSort, bool bFinalSo
 		}
 
 		// Sort by name to be compatible with FindTextureStreamingDataIndexRange
-		TextureStreamingData.Sort([](const FMaterialTextureInfo& Lhs, const FMaterialTextureInfo& Rhs) { return Lhs.TextureName.LexicalLess(Rhs.TextureName); });
+		TextureStreamingData.Sort([](const FMaterialTextureInfo& Lhs, const FMaterialTextureInfo& Rhs) 
+		{ 
+#if WITH_EDITORONLY_DATA
+			// Sort by register indices when the name are the same, as when initially added in the streaming data.
+			if (Lhs.TextureName == Rhs.TextureName)
+			{
+				return Lhs.TextureIndex < Rhs.TextureIndex;
+
+			}
+#endif
+			return Lhs.TextureName.LexicalLess(Rhs.TextureName); 
+		});
 		bTextureStreamingDataSorted = true;
 	}
 #endif
@@ -615,7 +626,9 @@ void UMaterialInterface::SetTextureStreamingData(const TArray<FMaterialTextureIn
 {
 	TextureStreamingData = InTextureStreamingData;
 #if WITH_EDITORONLY_DATA
+	bTextureStreamingDataSorted = false;
 	TextureStreamingDataVersion = InTextureStreamingData.Num() ? MATERIAL_TEXTURE_STREAMING_DATA_VERSION : 0;
+	TextureStreamingDataMissingEntries.Empty();
 #endif
 	SortTextureStreamingData(true, false);
 }
diff --git a/Engine/Source/Runtime/Engine/Private/SkeletalMesh.cpp b/Engine/Source/Runtime/Engine/Private/SkeletalMesh.cpp
--- a/Engine/Source/Runtime/Engine/Private/SkeletalMesh.cpp
+++ b/Engine/Source/Runtime/Engine/Private/SkeletalMesh.cpp
@@ -4543,16 +4543,25 @@ bool FSkeletalMeshSceneProxy::GetMaterialTextureScales(int32 LODIndex, int32 Sec
 		{
 			// This is thread safe because material texture data is only updated while the renderthread is idle.
 			for (const FMaterialTextureInfo TextureData : Material->GetTextureStreamingData())
-	{
+			{
 				const int32 TextureIndex = TextureData.TextureIndex;
 				if (TextureData.IsValid(true))
-		{
+				{
 					OneOverScales[TextureIndex / 4][TextureIndex % 4] = 1.f / TextureData.SamplingScale;
 					UVChannelIndices[TextureIndex / 4][TextureIndex % 4] = TextureData.UVChannelIndex;
-		}
-	}
+				}
+			}
+			for (const FMaterialTextureInfo TextureData : Material->TextureStreamingDataMissingEntries)
+			{
+				const int32 TextureIndex = TextureData.TextureIndex;
+				if (TextureIndex >= 0 && TextureIndex < TEXSTREAM_MAX_NUM_TEXTURES_PER_MATERIAL)
+				{
+					OneOverScales[TextureIndex / 4][TextureIndex % 4] = 1.f;
+					UVChannelIndices[TextureIndex / 4][TextureIndex % 4] = 0;
+				}
+			}
 			return true;
-	}
+		}
 	}
 	return false;
 }
diff --git a/Engine/Source/Runtime/Engine/Private/StaticMeshRender.cpp b/Engine/Source/Runtime/Engine/Private/StaticMeshRender.cpp
--- a/Engine/Source/Runtime/Engine/Private/StaticMeshRender.cpp
+++ b/Engine/Source/Runtime/Engine/Private/StaticMeshRender.cpp
@@ -705,6 +705,15 @@ bool FStaticMeshSceneProxy::GetMaterialTextureScales(int32 LODIndex, int32 Secti
 					UVChannelIndices[TextureIndex / 4][TextureIndex % 4] = TextureData.UVChannelIndex;
 				}
 			}
+			for (const FMaterialTextureInfo TextureData : Material->TextureStreamingDataMissingEntries)
+			{
+				const int32 TextureIndex = TextureData.TextureIndex;
+				if (TextureIndex >= 0 && TextureIndex < TEXSTREAM_MAX_NUM_TEXTURES_PER_MATERIAL)
+				{
+					OneOverScales[TextureIndex / 4][TextureIndex % 4] = 1.f;
+					UVChannelIndices[TextureIndex / 4][TextureIndex % 4] = 0;
+				}
+			}
 			return true;
 		}
 	}
diff --git a/Engine/Source/Runtime/Engine/Public/DebugViewModeHelpers.h b/Engine/Source/Runtime/Engine/Public/DebugViewModeHelpers.h
--- a/Engine/Source/Runtime/Engine/Public/DebugViewModeHelpers.h
+++ b/Engine/Source/Runtime/Engine/Public/DebugViewModeHelpers.h
@@ -46,6 +46,7 @@ ENGINE_API int32 GetNumActorsInWorld(UWorld* InWorld);
 ENGINE_API bool GetUsedMaterialsInWorld(UWorld* InWorld, OUT TSet<UMaterialInterface*>& OutMaterials, FSlowTask* Task);
 ENGINE_API bool CompileDebugViewModeShaders(EDebugViewShaderMode Mode, EMaterialQualityLevel::Type QualityLevel, ERHIFeatureLevel::Type FeatureLevel, bool bFullRebuild, bool bWaitForPreviousShaders, TSet<UMaterialInterface*>& Materials, FSlowTask* ProgressTask);
 ENGINE_API void UpdateDebugViewModeShaders();
+ENGINE_API void ClearDebugViewMaterials(UMaterialInterface* InMaterialInterface);
 
 ENGINE_API bool WaitForShaderCompilation(const FText& Message, FSlowTask* ProgressTask);
 
