Asset tabs will now use the icon for the generated class instead of the generic blueprint icon. Fixed up issue where widget blueprint icons were no longer appearing. Copy of CL# 6979887 from Dev-Editor

diff --git a/Engine/Source/Editor/EditorStyle/Private/SlateEditorStyle.cpp b/Engine/Source/Editor/EditorStyle/Private/SlateEditorStyle.cpp
--- a/Engine/Source/Editor/EditorStyle/Private/SlateEditorStyle.cpp
+++ b/Engine/Source/Editor/EditorStyle/Private/SlateEditorStyle.cpp
@@ -6414,7 +6414,7 @@ void FSlateEditorStyle::FStyle::SetupClassIconsAndThumbnails()
 			TEXT("UserDefinedCaptureProtocol"),
 			TEXT("UserDefinedEnum"),
 			TEXT("UserDefinedStruct"),
-			TEXT("WidgetBlueprint"),
+			TEXT("UserWidget"),
 			TEXT("WindDirectionalSource"),
 			TEXT("World"),
 			TEXT("Cube"),
diff --git a/Engine/Source/Editor/UnrealEd/Private/Toolkits/AssetEditorToolkit.cpp b/Engine/Source/Editor/UnrealEd/Private/Toolkits/AssetEditorToolkit.cpp
--- a/Engine/Source/Editor/UnrealEd/Private/Toolkits/AssetEditorToolkit.cpp
+++ b/Engine/Source/Editor/UnrealEd/Private/Toolkits/AssetEditorToolkit.cpp
@@ -617,8 +617,16 @@ const FSlateBrush* FAssetEditorToolkit::GetDefaultTabIcon() const
 	{
 		if (Object)
 		{
+			UClass* IconClass = Object->GetClass();
+
+			if (IconClass->IsChildOf<UBlueprint>())
+			{
+				UBlueprint* Blueprint = Cast<UBlueprint>(Object);
+				IconClass = Blueprint->GeneratedClass;
+			}
+
 			// Find the first object that has a valid brush
-			const FSlateBrush* ThisAssetBrush = FSlateIconFinder::FindIconBrushForClass(Object->GetClass());
+			const FSlateBrush* ThisAssetBrush = FSlateIconFinder::FindIconBrushForClass(IconClass);
 			if (ThisAssetBrush != nullptr)
 			{
 				IconBrush = ThisAssetBrush;
