Adjusted Scale Gizmo's universal handle hitbox.

diff --git a/Engine/Source/Editor/UnrealEd/Private/UnrealWidget.cpp b/Engine/Source/Editor/UnrealEd/Private/UnrealWidget.cpp
--- a/Engine/Source/Editor/UnrealEd/Private/UnrealWidget.cpp
+++ b/Engine/Source/Editor/UnrealEd/Private/UnrealWidget.cpp
@@ -22,6 +22,8 @@ static const float OUTER_AXIS_CIRCLE_RADIUS = 56.0f;
 static const float ROTATION_TEXT_RADIUS = 75.0f;
 static const int32 AXIS_CIRCLE_SIDES = 24;
 static const float ARCALL_RELATIVE_INNER_SIZE = 0.75f;
+static const float AXIS_LENGTH_SCALE_OFFSET = 5.0f;
+
 /*
  *  Simple struct used to create and group data related to the current window's / viewport's space,
  *  orientation, and scale.
@@ -341,7 +343,7 @@ void FWidget::Render( const FSceneView* View,FPrimitiveDrawInterface* PDI, FEdit
 /**
  * Draws an arrow head line for a specific axis.
  */
-void FWidget::Render_Axis( const FSceneView* View, FPrimitiveDrawInterface* PDI, EAxisList::Type InAxis, FMatrix& InMatrix, UMaterialInterface* InMaterial, const FLinearColor& InColor, FVector2D& OutAxisDir, const FVector& InScale, bool bDrawWidget, bool bCubeHead )
+void FWidget::Render_Axis(const FSceneView* View, FPrimitiveDrawInterface* PDI, EAxisList::Type InAxis, FMatrix& InMatrix, UMaterialInterface* InMaterial, const FLinearColor& InColor, FVector2D& OutAxisDir, const FVector& InScale, bool bDrawWidget, bool bCubeHead, float AxisLengthOffset)
 {
 	FMatrix AxisRotation = FMatrix::Identity;
 	if( InAxis == EAxisList::Y )
@@ -371,10 +373,10 @@ void FWidget::Render_Axis( const FSceneView* View, FPrimitiveDrawInterface* PDI,
 		const bool bDisabled = EditorModeTools ? (EditorModeTools->IsDefaultModeActive() && GEditor->HasLockedActors() ) : false;
 		PDI->SetHitProxy( new HWidgetAxis( InAxis, bDisabled) );
 
-		const float AxisLength = AXIS_LENGTH + GetDefault<ULevelEditorViewportSettings>()->TransformWidgetSizeAdjustment;
+		const float AxisLength = AXIS_LENGTH + GetDefault<ULevelEditorViewportSettings>()->TransformWidgetSizeAdjustment - (AxisLengthOffset * 2);
 		const float HalfHeight = AxisLength/2.0f;
 		const float CylinderRadius = 1.2f;
-		const FVector Offset( 0,0,HalfHeight );
+		const FVector Offset(0, 0, HalfHeight + AxisLengthOffset);
 
 		switch( InAxis )
 		{
@@ -398,14 +400,14 @@ void FWidget::Render_Axis( const FSceneView* View, FPrimitiveDrawInterface* PDI,
 		if ( bCubeHead )
 		{
 			const float CubeHeadOffset = 3.0f;
-			FVector RootPos(AxisLength + CubeHeadOffset, 0, 0);
+			FVector RootPos(AxisLength + CubeHeadOffset + AxisLengthOffset, 0, 0);
 
 			Render_Cube(PDI, (FTranslationMatrix(RootPos) * ArrowToWorld) * FScaleMatrix(FlattenScale), InMaterial, FVector(4.0f));
 		}
 		else
 		{
 			const float ConeHeadOffset = 12.0f;
-			FVector RootPos(AxisLength + ConeHeadOffset, 0, 0);
+			FVector RootPos(AxisLength + ConeHeadOffset + AxisLengthOffset, 0, 0);
 
 			float Angle = FMath::DegreesToRadians( PI * 5 );
 			DrawCone(PDI, ( FScaleMatrix(-13) * FTranslationMatrix(RootPos) * ArrowToWorld ) * FScaleMatrix(FlattenScale), Angle, Angle, 32, false, FColor::White, InMaterial->GetRenderProxy(), SDPG_Foreground);
@@ -821,7 +823,7 @@ void FWidget::Render_Rotate( const FSceneView* View,FPrimitiveDrawInterface* PDI
 /**
  * Draws the scaling widget.
  */
-void FWidget::Render_Scale( const FSceneView* View,FPrimitiveDrawInterface* PDI, FEditorViewportClient* ViewportClient, const FVector& InLocation, bool bDrawWidget )
+void FWidget::Render_Scale( const FSceneView* View,FPrimitiveDrawInterface* PDI, FEditorViewportClient* ViewportClient, const FVector& InLocation, bool bDrawWidget)
 {
 	// Figure out axis colors
 	const FLinearColor& XColor = ( CurrentAxis&EAxisList::X ? (FLinearColor)CurrentColor : AxisColorX );
@@ -846,17 +848,17 @@ void FWidget::Render_Scale( const FSceneView* View,FPrimitiveDrawInterface* PDI,
 	// Draw the axis lines with cube heads	
     if (Space.ShouldDrawAxisX(DrawAxis))
     {
-        Render_Axis(View, PDI, EAxisList::X, WidgetMatrix, XMaterial, XColor, XAxisDir, UniformScale, bDrawWidget, true);
+        Render_Axis(View, PDI, EAxisList::X, WidgetMatrix, XMaterial, XColor, XAxisDir, UniformScale, bDrawWidget, true, AXIS_LENGTH_SCALE_OFFSET);
     }
 
     if (Space.ShouldDrawAxisY(DrawAxis))
     {
-        Render_Axis(View, PDI, EAxisList::Y, WidgetMatrix, YMaterial, YColor, YAxisDir, UniformScale, bDrawWidget, true);
+        Render_Axis(View, PDI, EAxisList::Y, WidgetMatrix, YMaterial, YColor, YAxisDir, UniformScale, bDrawWidget, true, AXIS_LENGTH_SCALE_OFFSET);
     }
 
     if (Space.ShouldDrawAxisZ(DrawAxis))
     {
-        Render_Axis(View, PDI, EAxisList::Z, WidgetMatrix, ZMaterial, ZColor, ZAxisDir, UniformScale, bDrawWidget, true);
+        Render_Axis(View, PDI, EAxisList::Z, WidgetMatrix, ZMaterial, ZColor, ZAxisDir, UniformScale, bDrawWidget, true, AXIS_LENGTH_SCALE_OFFSET);
     }
 
 	// Draw grabber handles and center cube
diff --git a/Engine/Source/Editor/UnrealEd/Public/UnrealWidget.h b/Engine/Source/Editor/UnrealEd/Public/UnrealWidget.h
--- a/Engine/Source/Editor/UnrealEd/Public/UnrealWidget.h
+++ b/Engine/Source/Editor/UnrealEd/Public/UnrealWidget.h
@@ -62,7 +62,7 @@ public:
 	 * Draws an arrow head line for a specific axis.
 	 * @param	bCubeHead		[opt] If true, render a cube at the axis tips.  If false (the default), render a cone.
 	 */
-	void Render_Axis(const FSceneView* View, FPrimitiveDrawInterface* PDI, EAxisList::Type InAxis, FMatrix& InMatrix, UMaterialInterface* InMaterial, const FLinearColor& InColor, FVector2D& OutAxisDir, const FVector& InScale, bool bDrawWidget, bool bCubeHead=false);
+	void Render_Axis(const FSceneView* View, FPrimitiveDrawInterface* PDI, EAxisList::Type InAxis, FMatrix& InMatrix, UMaterialInterface* InMaterial, const FLinearColor& InColor, FVector2D& OutAxisDir, const FVector& InScale, bool bDrawWidget, bool bCubeHead=false, float AxisLengthOffset = 0);
 
 	/**
 	 * Draws a cube
