Amortize the cost of AddUnique when serializing transaction
  -1.7 sec to 132 ms for biggest SaveToTransactionBuffer when deleting an actor from a 60 000 actors scene.
  -4.9 sec to 3.2 sec to delete one actor from a 60 000 actors scene.
Remove time spent waiting for user input while deleting actors from a scene so the number actually represents computing time

From dev-enterprise[REDACTED_HANDLE]

diff --git a/Engine/Source/Editor/UnrealEd/Classes/Editor/Transactor.h b/Engine/Source/Editor/UnrealEd/Classes/Editor/Transactor.h
--- a/Engine/Source/Editor/UnrealEd/Classes/Editor/Transactor.h
+++ b/Engine/Source/Editor/UnrealEd/Classes/Editor/Transactor.h
@@ -361,6 +361,11 @@ protected:
 					ObjectMap.Add(SerializedObject.ReferencedObjects[ObjIndex].Get(), ObjIndex);
 				}
 
+				for(int32 NameIndex = 0; NameIndex < SerializedObject.ReferencedNames.Num(); ++NameIndex)
+				{
+					NameMap.Add(SerializedObject.ReferencedNames[NameIndex], NameIndex);
+				}
+
 				this->SetWantBinaryPropertySerialization(bWantBinarySerialization);
 				this->SetIsSaving(true);
 				this->SetIsTransacting(true);
@@ -430,7 +435,17 @@ protected:
 			}
 			FArchive& operator<<( class FName& N ) override
 			{
-				int32 NameIndex = SerializedObject.ReferencedNames.AddUnique(N);
+				int32 NameIndex = INDEX_NONE;
+				const int32 * NameIndexPtr = NameMap.Find(N);
+				if (NameIndexPtr)
+				{
+					NameIndex = *NameIndexPtr;
+				}
+				else
+				{
+					NameIndex = SerializedObject.ReferencedNames.Add(N);
+					NameMap.Add(N, NameIndex);
+				}
 
 				// Track this name index in the serialized data
 				{
@@ -444,7 +459,7 @@ protected:
 			FArchive& operator<<( class UObject*& Res ) override
 			{
 				int32 ObjectIndex = INDEX_NONE;
-				int32* ObjIndexPtr = ObjectMap.Find(Res);
+				const int32* ObjIndexPtr = ObjectMap.Find(Res);
 				if(ObjIndexPtr)
 				{
 					ObjectIndex = *ObjIndexPtr;
@@ -466,6 +481,7 @@ protected:
 			}
 			FSerializedObject& SerializedObject;
 			TMap<UObject*, int32> ObjectMap;
+			TMap<FName, int32> NameMap;
 			FCachedPropertyKey CachedSerializedTaggedPropertyKey;
 			int64 Offset;
 		};
diff --git a/Engine/Source/Editor/UnrealEd/Private/EditorActor.cpp b/Engine/Source/Editor/UnrealEd/Private/EditorActor.cpp
--- a/Engine/Source/Editor/UnrealEd/Private/EditorActor.cpp
+++ b/Engine/Source/Editor/UnrealEd/Private/EditorActor.cpp
@@ -703,6 +703,9 @@ bool UUnrealEdEngine::edactDeleteSelected( UWorld* InWorld, bool bVerifyDeletion
 
 	const double StartSeconds = FPlatformTime::Seconds();
 
+	// Aggregate the time we waited on user input to remove it from the total time
+	double DialogWaitingSeconds = 0;
+
 	FSlateApplication::Get().CancelDragDrop();
 
 	if (GetSelectedComponentCount() > 0)
@@ -949,7 +952,12 @@ bool UUnrealEdEngine::edactDeleteSelected( UWorld* InWorld, bool bVerifyDeletion
 						FText::FromString(Actor->GetActorLabel()), FText::FromString(ActorReferenceString));
 				}
 
+				const double DialogStartSeconds = FPlatformTime::Seconds();
+
 				int32 Result = FMessageDialog::Open(MessageType, ConfirmDelete);
+
+				DialogWaitingSeconds += FPlatformTime::Seconds() - DialogStartSeconds;
+
 				if (Result == EAppReturnType::YesAll)
 				{
 					bRequestedDeleteAllByLevel |= bReferencedByLevelScript;
@@ -1075,7 +1083,7 @@ bool UUnrealEdEngine::edactDeleteSelected( UWorld* InWorld, bool bVerifyDeletion
 		ULevel::LevelDirtiedEvent.Broadcast();
 	}
 
-	UE_LOG(LogEditorActor, Log,  TEXT("Deleted %d Actors (%3.3f secs)"), DeleteCount, FPlatformTime::Seconds() - StartSeconds );
+	UE_LOG(LogEditorActor, Log,  TEXT("Deleted %d Actors (%3.3f secs)"), DeleteCount, (FPlatformTime::Seconds() - StartSeconds) - DialogWaitingSeconds);
 
 	return true;
 }
@@ -2604,4 +2612,3 @@ void UUnrealEdEngine::edactAlignVertices()
 PRAGMA_ENABLE_OPTIMIZATION
 
 #undef LOCTEXT_NAMESPACE
-
