Avoid serializing all actors visibility individually in a transaction since they can easily be recomputed on demand instead
    - Properly restore actors visibility when map is loading
    - Properly recompute actors visibility on undo/redo actions
    -12.5s to 59ms for toggling layer visibility on a layer containing 35 000 elements

From dev-enterprise[REDACTED_HANDLE]

diff --git a/Engine/Source/Editor/Layers/Private/LayerCollectionViewModel.cpp b/Engine/Source/Editor/Layers/Private/LayerCollectionViewModel.cpp
--- a/Engine/Source/Editor/Layers/Private/LayerCollectionViewModel.cpp
+++ b/Engine/Source/Editor/Layers/Private/LayerCollectionViewModel.cpp
@@ -191,6 +191,8 @@ void FLayerCollectionViewModel::OnFilterChanged()
 
 void FLayerCollectionViewModel::Refresh()
 {
+	WorldLayers->UpdateAllActorsVisibility( true, true );
+
 	OnLayersChanged( ELayersAction::Reset, NULL, NAME_None );
 }
 
diff --git a/Engine/Source/Editor/UnrealEd/Private/Layers/Layers.cpp b/Engine/Source/Editor/UnrealEd/Private/Layers/Layers.cpp
--- a/Engine/Source/Editor/UnrealEd/Private/Layers/Layers.cpp
+++ b/Engine/Source/Editor/UnrealEd/Private/Layers/Layers.cpp
@@ -125,6 +125,16 @@ bool FLayers::InitializeNewActorLayers( const TWeakObjectPtr< AActor >& Actor )
 		AddActorToStats( Layer, Actor);
 	}
 
+	// update per-view visibility info
+	UpdateActorAllViewsVisibility(Actor);
+
+	// update general actor visibility
+	bool bActorModified = false;
+	bool bActorSelectionChanged = false;
+	const bool bActorNotifySelectionChange = true;
+	const bool bActorRedrawViewports = false;
+	UpdateActorVisibility( Actor, bActorSelectionChanged, bActorModified, bActorNotifySelectionChange, bActorRedrawViewports );
+
 	return Actor->Layers.Num() > 0;
 }
 
@@ -757,8 +767,6 @@ bool FLayers::UpdateActorVisibility( const TWeakObjectPtr< AActor >& Actor, bool
 		// If the actor is also hidden
 		if( Actor->bHiddenEdLayer )
 		{
-			Actor->Modify();
-
 			// Actors that don't belong to any layer shouldn't be hidden
 			Actor->bHiddenEdLayer = false;
 			Actor->MarkComponentsRenderStateDirty();
@@ -782,7 +790,6 @@ bool FLayers::UpdateActorVisibility( const TWeakObjectPtr< AActor >& Actor, bool
 		{
 			if ( Actor->bHiddenEdLayer )
 			{
-				Actor->Modify();
 				Actor->bHiddenEdLayer = false;
 				Actor->MarkComponentsRenderStateDirty();
 				bOutActorModified = true;
@@ -805,7 +812,6 @@ bool FLayers::UpdateActorVisibility( const TWeakObjectPtr< AActor >& Actor, bool
 	{
 		if ( !Actor->bHiddenEdLayer )
 		{
-			Actor->Modify();
 			Actor->bHiddenEdLayer = true;
 			Actor->MarkComponentsRenderStateDirty();
 			bOutActorModified = true;
diff --git a/Engine/Source/Runtime/Engine/Classes/GameFramework/Actor.h b/Engine/Source/Runtime/Engine/Classes/GameFramework/Actor.h
--- a/Engine/Source/Runtime/Engine/Classes/GameFramework/Actor.h
+++ b/Engine/Source/Runtime/Engine/Classes/GameFramework/Actor.h
@@ -693,7 +693,7 @@ public:
 	uint8 bIsEditorPreviewActor:1;
 
 	/** Whether this actor is hidden by the layer browser. */
-	UPROPERTY()
+	UPROPERTY(Transient)
 	uint8 bHiddenEdLayer:1;
 
 	/** Whether this actor is hidden by the level browser. */
