Provide caching for FindViewComponentForActor to improve editor responsiveness when a selected actor contains many others.
    - Add GetAttachedActor overload which takes a functor so we can iterate more effectively and abort iteration if needed.
    - 5 sec to 90ms per frame (55x improvement) when selected actor contains 60 000 others.
    - 43 sec to 2s (21x improvement) for CaptureSceneThumbnail at the end of datasmith import for a 60 000 actors scene.

From dev-enterprise[REDACTED_HANDLE],7063185

diff --git a/Engine/Source/Editor/UnrealEd/Private/LevelEditorViewport.cpp b/Engine/Source/Editor/UnrealEd/Private/LevelEditorViewport.cpp
--- a/Engine/Source/Editor/UnrealEd/Private/LevelEditorViewport.cpp
+++ b/Engine/Source/Editor/UnrealEd/Private/LevelEditorViewport.cpp
@@ -85,6 +85,8 @@ DEFINE_LOG_CATEGORY(LogEditorViewport);
 static const float MIN_ACTOR_BOUNDS_EXTENT	= 1.0f;
 
 TArray< TWeakObjectPtr< AActor > > FLevelEditorViewportClient::DropPreviewActors;
+TMap< TObjectKey< AActor >, TWeakObjectPtr< UActorComponent > > FLevelEditorViewportClient::ViewComponentForActorCache;
+
 bool FLevelEditorViewportClient::bIsDroppingPreviewActor;
 
 /** Static: List of objects we're hovering over */
@@ -2182,6 +2184,9 @@ void FLevelEditorViewportClient::Tick(float DeltaTime)
 	}
 	bWasEditorCameraCut = bEditorCameraCut;
 
+	// Gives FindViewComponentForActor a chance to refresh once every Tick.
+	ViewComponentForActorCache.Reset();
+
 	FEditorViewportClient::Tick(DeltaTime);
 
 	// Update the preview mesh for the preview mesh mode. 
@@ -3412,8 +3417,23 @@ void FLevelEditorViewportClient::MoveCameraToLockedActor()
 
 UActorComponent* FLevelEditorViewportClient::FindViewComponentForActor(AActor const* Actor)
 {
-	TSet<AActor const*> CheckedActors;
-	return FindViewComponentForActor(Actor, CheckedActors);
+	UActorComponent* PreviewComponent = nullptr;
+	if (Actor)
+	{
+		const TWeakObjectPtr<UActorComponent> * CachedComponent = ViewComponentForActorCache.Find(Actor);
+		if (CachedComponent != nullptr)
+		{
+			PreviewComponent = CachedComponent->Get();
+		}
+		else
+		{
+			TSet<AActor const*> CheckedActors;
+			PreviewComponent = FindViewComponentForActor(Actor, CheckedActors);
+			ViewComponentForActorCache.Add(Actor, PreviewComponent);
+		}
+	}
+
+	return PreviewComponent;
 }
 
 UActorComponent* FLevelEditorViewportClient::FindViewComponentForActor(AActor const* Actor, TSet<AActor const*>& CheckedActors)
@@ -3425,7 +3445,6 @@ UActorComponent* FLevelEditorViewportClient::FindViewComponentForActor(AActor co
 		// see if actor has a component with preview capabilities (prioritize camera components)
 		const TSet<UActorComponent*>& Comps = Actor->GetComponents();
 
-		bool bChoseCamComponent = false;
 		for (UActorComponent* Comp : Comps)
 		{
 			FMinimalViewInfo DummyViewInfo;
@@ -3438,11 +3457,6 @@ UActorComponent* FLevelEditorViewportClient::FindViewComponentForActor(AActor co
 				}
 				else if (PreviewComponent)
 				{
-					if (bChoseCamComponent)
-					{
-						continue;
-					}
-
 					UCameraComponent* AsCamComp = Cast<UCameraComponent>(Comp);
 					if (AsCamComp != nullptr)
 					{
@@ -3458,17 +3472,19 @@ UActorComponent* FLevelEditorViewportClient::FindViewComponentForActor(AActor co
 		// we will just return the first one.
 		if (PreviewComponent == nullptr)
 		{
-			TArray<AActor*> AttachedActors;
-			Actor->GetAttachedActors(AttachedActors);
-			for (AActor* AttachedActor : AttachedActors)
-			{
-				UActorComponent* const Comp = FindViewComponentForActor(AttachedActor, CheckedActors);
-				if (Comp)
+			Actor->ForEachAttachedActors(
+				[&](AActor * AttachedActor) -> bool
 				{
-					PreviewComponent = Comp;
-					break;
+					UActorComponent* const Comp = FindViewComponentForActor(AttachedActor, CheckedActors);
+					if (Comp)
+					{
+						PreviewComponent = Comp;
+						return false; /* stops iteration */
+					}
+
+					return true; /* continue iteration */
 				}
-			}
+			);
 		}
 	}
 
diff --git a/Engine/Source/Editor/UnrealEd/Public/LevelEditorViewport.h b/Engine/Source/Editor/UnrealEd/Public/LevelEditorViewport.h
--- a/Engine/Source/Editor/UnrealEd/Public/LevelEditorViewport.h
+++ b/Engine/Source/Editor/UnrealEd/Public/LevelEditorViewport.h
@@ -9,6 +9,7 @@
 #include "Camera/CameraComponent.h"
 #include "UnrealWidget.h"
 #include "EditorViewportClient.h"
+#include "UObject/ObjectKey.h"
 
 struct FAssetData;
 class FCanvas;
@@ -778,6 +779,9 @@ private:
 	TWeakObjectPtr<AActor>	ActorLockedByMatinee;
 	TWeakObjectPtr<AActor>	ActorLockedToCamera;
 
+	/** Caching for expensive FindViewComponentForActor. Invalidated once per Tick. */
+	static TMap<TObjectKey<AActor>, TWeakObjectPtr<UActorComponent>> ViewComponentForActorCache;
+
 	/** If true, we switched between two different cameras. Set by matinee, used by the motion blur to invalidate this frames motion vectors */
 	bool					bEditorCameraCut;
 
diff --git a/Engine/Source/Runtime/CoreUObject/Public/UObject/ObjectKey.h b/Engine/Source/Runtime/CoreUObject/Public/UObject/ObjectKey.h
--- a/Engine/Source/Runtime/CoreUObject/Public/UObject/ObjectKey.h
+++ b/Engine/Source/Runtime/CoreUObject/Public/UObject/ObjectKey.h
@@ -17,7 +17,7 @@ public:
 	}
 
 	/** Construct from an object pointer */
-	FORCEINLINE FObjectKey(const UObject *Object)
+	FORCEINLINE FObjectKey(const UObject* Object)
 		: ObjectIndex(INDEX_NONE)
 		, ObjectSerialNumber(0)
 	{
@@ -30,13 +30,13 @@ public:
 	}
 
 	/** Compare this key with another */
-	FORCEINLINE bool operator==(const FObjectKey &Other) const
+	FORCEINLINE bool operator==(const FObjectKey& Other) const
 	{
 		return ObjectIndex == Other.ObjectIndex && ObjectSerialNumber == Other.ObjectSerialNumber;
 	}
 
 	/** Compare this key with another */
-	FORCEINLINE bool operator!=(const FObjectKey &Other) const
+	FORCEINLINE bool operator!=(const FObjectKey& Other) const
 	{
 		return ObjectIndex != Other.ObjectIndex || ObjectSerialNumber != Other.ObjectSerialNumber;
 	}
@@ -65,3 +65,50 @@ private:
 	int32		ObjectIndex;
 	int32		ObjectSerialNumber;
 };
+
+/** TObjectKey is a strongly typed, immutable, copyable key which can be used to uniquely identify an object for the lifetime of the application */
+template<typename InElementType>
+class TObjectKey
+{
+public:
+	typedef InElementType ElementType;
+
+	/** Default constructor */
+	FORCEINLINE TObjectKey() = default;
+
+	/** Construct from an object pointer */
+	FORCEINLINE TObjectKey(const ElementType* Object)
+		: ObjectKey(Object)
+	{
+	}
+
+	/** Compare this key with another */
+	FORCEINLINE bool operator==(const TObjectKey& Other) const
+	{
+		return ObjectKey == Other.ObjectKey;
+	}
+
+	/** Compare this key with another */
+	FORCEINLINE bool operator!=(const TObjectKey& Other) const
+	{
+		return ObjectKey != Other.ObjectKey;
+	}
+
+	//** Hash function */
+	friend uint32 GetTypeHash(const TObjectKey& Key)
+	{
+		return GetTypeHash(Key.ObjectKey);
+	}
+
+	/**
+	 * Attempt to access the object from which this key was constructed.
+	 * @return The object used to construct this key, or nullptr if it is no longer valid
+	 */
+	InElementType* ResolveObjectPtr() const
+	{
+		return (InElementType*)ObjectKey.ResolveObjectPtr();
+	}
+
+private:
+	FObjectKey ObjectKey;
+};
\ No newline at end of file
diff --git a/Engine/Source/Runtime/Engine/Classes/GameFramework/Actor.h b/Engine/Source/Runtime/Engine/Classes/GameFramework/Actor.h
--- a/Engine/Source/Runtime/Engine/Classes/GameFramework/Actor.h
+++ b/Engine/Source/Runtime/Engine/Classes/GameFramework/Actor.h
@@ -2517,6 +2517,9 @@ public:
 	UFUNCTION(BlueprintPure, Category = "Utilities")
 	FName GetAttachParentSocketName() const;
 
+	/** Call a functor for Actors which are attached directly to a component in this actor. Functor should return true to carry on, false to abort. */
+	void ForEachAttachedActors(TFunctionRef<bool(class AActor*)> Functor) const;
+	
 	/** Find all Actors which are attached directly to a component in this actor */
 	UFUNCTION(BlueprintPure, Category = "Utilities")
 	void GetAttachedActors(TArray<AActor*>& OutActors) const;
diff --git a/Engine/Source/Runtime/Engine/Private/Actor.cpp b/Engine/Source/Runtime/Engine/Private/Actor.cpp
--- a/Engine/Source/Runtime/Engine/Private/Actor.cpp
+++ b/Engine/Source/Runtime/Engine/Private/Actor.cpp
@@ -1798,9 +1798,8 @@ FName AActor::GetAttachParentSocketName() const
 	return NAME_None;
 }
 
-void AActor::GetAttachedActors(TArray<class AActor*>& OutActors) const
+void AActor::ForEachAttachedActors(TFunctionRef<bool(class AActor*)> Functor) const
 {
-	OutActors.Reset();
 	if (RootComponent != nullptr)
 	{
 		// Current set of components to check
@@ -1812,35 +1811,36 @@ void AActor::GetAttachedActors(TArray<class AActor*>& OutActors) const
 		CompsToCheck.Push(RootComponent);
 
 		// While still work left to do
-		while(CompsToCheck.Num() > 0)
+		while (CompsToCheck.Num() > 0)
 		{
 			// Get the next off the queue
 			const bool bAllowShrinking = false;
 			USceneComponent* SceneComp = CompsToCheck.Pop(bAllowShrinking);
 
 			// Add it to the 'checked' set, should not already be there!
-			if (!CheckedComps.Contains(SceneComp))
-			{
-				CheckedComps.Add(SceneComp);
+			CheckedComps.Add(SceneComp);
 
-				AActor* CompOwner = SceneComp->GetOwner();
-				if (CompOwner != nullptr)
+			AActor* CompOwner = SceneComp->GetOwner();
+			if (CompOwner != nullptr)
+			{
+				if (CompOwner != this)
 				{
-					if (CompOwner != this)
+					// If this component has a different owner, call the callback and stop if told.
+					if (!Functor(CompOwner))
 					{
-						// If this component has a different owner, add that owner to our output set and do nothing more
-						OutActors.AddUnique(CompOwner);
+						// The functor wants us to abort
+						return;
 					}
-					else
+				}
+				else
+				{
+					// This component is owned by us, we need to add its children
+					for (USceneComponent* ChildComp : SceneComp->GetAttachChildren())
 					{
-						// This component is owned by us, we need to add its children
-						for (USceneComponent* ChildComp : SceneComp->GetAttachChildren())
+						// Add any we have not explored yet to the set to check
+						if (ChildComp != nullptr && !CheckedComps.Contains(ChildComp))
 						{
-							// Add any we have not explored yet to the set to check
-							if ((ChildComp != nullptr) && !CheckedComps.Contains(ChildComp))
-							{
-								CompsToCheck.Push(ChildComp);
-							}
+							CompsToCheck.Push(ChildComp);
 						}
 					}
 				}
@@ -1849,6 +1849,12 @@ void AActor::GetAttachedActors(TArray<class AActor*>& OutActors) const
 	}
 }
 
+void AActor::GetAttachedActors(TArray<class AActor*>& OutActors) const
+{
+	OutActors.Reset();
+	ForEachAttachedActors([&OutActors](AActor * Actor) { OutActors.AddUnique(Actor); return true; });
+}
+
 bool AActor::ActorHasTag(FName Tag) const
 {
 	return (Tag != NAME_None) && Tags.Contains(Tag);
