Reduce time to clear existing world with lots of redirectors
  - Skip validation for package already vetted for unloading
  - From ~30h (estimated) to 5.4s

From dev-enterprise[REDACTED_HANDLE]

diff --git a/Engine/Source/Editor/UnrealEd/Private/EditorEngine.cpp b/Engine/Source/Editor/UnrealEd/Private/EditorEngine.cpp
--- a/Engine/Source/Editor/UnrealEd/Private/EditorEngine.cpp
+++ b/Engine/Source/Editor/UnrealEd/Private/EditorEngine.cpp
@@ -2159,6 +2159,12 @@ void UEditorEngine::Cleanse( bool ClearSelection, bool Redraw, const FText& Tran
 		{
 			UPackage* RedirectorPackage = RedirIt->GetOutermost();
 
+			if (PackagesToUnload.Find(RedirectorPackage))
+			{
+				// Package was already marked to unload
+				continue;
+			}
+
 			if (RedirectorPackage == GetTransientPackage())
 			{
 				RedirIt->ClearFlags(FlagsToClear);
