WorldFactory now responsible for initializing the builder brush, so that it happens in every path which creates a UWorld.

diff --git a/Engine/Source/Editor/UnrealEd/Private/EditorServer.cpp b/Engine/Source/Editor/UnrealEd/Private/EditorServer.cpp
--- a/Engine/Source/Editor/UnrealEd/Private/EditorServer.cpp
+++ b/Engine/Source/Editor/UnrealEd/Private/EditorServer.cpp
@@ -2272,9 +2272,6 @@ UWorld* UEditorEngine::NewMap()
 		Context.World()->GetDefaultBrush()->SetActorLocation(FVector::ZeroVector, false);
 	}
 
-	// Make the builder brush a small 256x256x256 cube so its visible.
-	InitBuilderBrush( Context.World() );
-
 	// Let navigation system know we're done creating new world
 	FNavigationSystem::AddNavigationSystemToWorld(*Context.World(), FNavigationSystemRunMode::EditorMode);
 
diff --git a/Engine/Source/Editor/UnrealEd/Private/Factories/WorldFactory.cpp b/Engine/Source/Editor/UnrealEd/Private/Factories/WorldFactory.cpp
--- a/Engine/Source/Editor/UnrealEd/Private/Factories/WorldFactory.cpp
+++ b/Engine/Source/Editor/UnrealEd/Private/Factories/WorldFactory.cpp
@@ -1,9 +1,8 @@
 // Copyright 1998-2019 Epic Games, Inc. All Rights Reserved.
 
 #include "Factories/WorldFactory.h"
+#include "Editor.h"
 #include "UObject/Package.h"
-#include "Engine/Level.h"
-#include "Editor/EditorEngine.h"
 #include "ThumbnailRendering/WorldThumbnailInfo.h"
 #include "EditorClassUtils.h"
 
@@ -29,6 +28,7 @@ UObject* UWorldFactory::FactoryCreateNew(UClass* Class, UObject* InParent, FName
 	// Create a new world.
 	const bool bAddToRoot = false;
 	UWorld* NewWorld = UWorld::CreateWorld(WorldType, bInformEngineOfWorld, Name, Cast<UPackage>(InParent), bAddToRoot, FeatureLevel);
+	GEditor->InitBuilderBrush(NewWorld);
 	NewWorld->SetFlags(Flags);
 	NewWorld->ThumbnailInfo = NewObject<UWorldThumbnailInfo>(NewWorld, NAME_None, RF_Transactional);
 
diff --git a/Engine/Source/Editor/UnrealEd/Private/UnrealEdEngine.cpp b/Engine/Source/Editor/UnrealEd/Private/UnrealEdEngine.cpp
--- a/Engine/Source/Editor/UnrealEd/Private/UnrealEdEngine.cpp
+++ b/Engine/Source/Editor/UnrealEd/Private/UnrealEdEngine.cpp
@@ -108,8 +108,6 @@ void UUnrealEdEngine::Init(IEngineLoop* InEngineLoop)
 
 	// Register for notification of volume changes
 	AVolume::GetOnVolumeShapeChangedDelegate().AddStatic(&FBSPOps::HandleVolumeShapeChanged);
-	//
-	InitBuilderBrush( GWorld );
 
 	// Iterate over all always fully loaded packages and load them.
 	if (!IsRunningCommandlet())
