Fixes significant performance impact when opening a Cascade Particle system before opening Sequencer. Sequencer was causing Cascade viewports to endlessly do an expensive re-render.

diff --git a/Engine/Source/Editor/Cascade/Private/CascadeEmitterCanvasClient.h b/Engine/Source/Editor/Cascade/Private/CascadeEmitterCanvasClient.h
--- a/Engine/Source/Editor/Cascade/Private/CascadeEmitterCanvasClient.h
+++ b/Engine/Source/Editor/Cascade/Private/CascadeEmitterCanvasClient.h
@@ -29,6 +29,11 @@ public:
 	virtual void Draw(FViewport* Viewport, FCanvas* Canvas) override;
 	virtual bool InputKey(FViewport* Viewport, int32 ControllerId, FKey Key, EInputEvent Event, float AmountDepressed = 1.0f, bool bGamepad = false) override;
 	virtual void CapturedMouseMove(FViewport* Viewport, int32 X, int32 Y) override;
+	virtual void RequestRealTimeFrames(uint64 NumRealTimeFrames = 1) override
+	{
+		// Intentionally do nothing. We don't want to render this viewport except under our own very
+		// specific rules (this isn't a real editor viewport and doesn't call the super on draw anyways)
+	}
 
 	/** Returns the ratio of the size of the particle emitters to the size of the viewport */
 	float GetViewportVerticalScrollBarRatio() const;
diff --git a/Engine/Source/Editor/Sequencer/Private/LevelEditorSequencerIntegration.cpp b/Engine/Source/Editor/Sequencer/Private/LevelEditorSequencerIntegration.cpp
--- a/Engine/Source/Editor/Sequencer/Private/LevelEditorSequencerIntegration.cpp
+++ b/Engine/Source/Editor/Sequencer/Private/LevelEditorSequencerIntegration.cpp
@@ -411,7 +411,10 @@ void FLevelEditorSequencerIntegration::OnSequencerEvaluated()
 	}
 
 	// Request a single real-time frame to be rendered to ensure that we tick the world and update the viewport
-	for (FEditorViewportClient* LevelVC : GEditor->GetAllViewportClients())
+	// ToDo: This attempts to redraw all viewports (blueprint windows, cascade particles, etc.) so it's not the most efficient
+	// but we can't just use GetLevelViewportClients() because of Actor Sequences. For now, we've disabled the realtime frames
+	// on the most expensive viewport (cascade).
+ 	for (FEditorViewportClient* LevelVC : GEditor->GetAllViewportClients())
 	{
 		if (LevelVC && !LevelVC->IsRealtime())
 		{
@@ -426,9 +429,6 @@ void FLevelEditorSequencerIntegration::OnSequencerEvaluated()
 
 	// If realtime is off, this needs to be called to update the pivot location when scrubbing.
 	GUnrealEd->UpdatePivotLocationForSelection();
-
-	// Redraw
-	FEditorSupportDelegates::RedrawAllViewports.Broadcast();
 }
 
 void FLevelEditorSequencerIntegration::OnBeginDeferUpdates()
diff --git a/Engine/Source/Editor/UnrealEd/Private/EditorViewportClient.cpp b/Engine/Source/Editor/UnrealEd/Private/EditorViewportClient.cpp
--- a/Engine/Source/Editor/UnrealEd/Private/EditorViewportClient.cpp
+++ b/Engine/Source/Editor/UnrealEd/Private/EditorViewportClient.cpp
@@ -354,7 +354,7 @@ FEditorViewportClient::FEditorViewportClient(FEditorModeTools* InModeTools, FPre
 	, CurrentGestureRotDelta(FRotator::ZeroRotator)
 	, GestureMoveForwardBackwardImpulse(0.0f)
 	, bForceAudioRealtime(false)
-	, RealTimeFrameCount(0)
+	, RealTimeUntilFrameNumber(0)
 	, bIsRealtime(false)
 	, bStoredRealtime(false)
 	, bStoredShowStats(false)
@@ -3606,11 +3606,6 @@ void FEditorViewportClient::SetupViewForRendering(FSceneViewFamily& ViewFamily,
 
 void FEditorViewportClient::Draw(FViewport* InViewport, FCanvas* Canvas)
 {
-	if (RealTimeFrameCount > 0)
-	{
-		--RealTimeFrameCount;
-	}
-
 	FViewport* ViewportBackup = Viewport;
 	Viewport = InViewport ? InViewport : Viewport;
 
diff --git a/Engine/Source/Editor/UnrealEd/Public/EditorViewportClient.h b/Engine/Source/Editor/UnrealEd/Public/EditorViewportClient.h
--- a/Engine/Source/Editor/UnrealEd/Public/EditorViewportClient.h
+++ b/Engine/Source/Editor/UnrealEd/Public/EditorViewportClient.h
@@ -300,7 +300,7 @@ public:
 	/** @return		True if viewport is in realtime mode, false otherwise. */
 	bool IsRealtime() const
 	{ 
-		return bIsRealtime || RealTimeFrameCount != 0;
+		return bIsRealtime || GFrameCounter < RealTimeUntilFrameNumber;
 	}
 
 	/**
@@ -309,9 +309,9 @@ public:
 	 * this can be used to ensure that not only the viewport renders a frame, but also that the world ticks
 	 * @param NumExtraFrames 		The number of extra real time frames to draw
 	 */
-	void RequestRealTimeFrames(uint32 NumRealTimeFrames = 1)
+	virtual void RequestRealTimeFrames(uint64 NumRealTimeFrames = 1)
 	{
-		RealTimeFrameCount = FMath::Max(NumRealTimeFrames, RealTimeFrameCount);
+		RealTimeUntilFrameNumber = FMath::Max(GFrameCounter + NumRealTimeFrames, RealTimeUntilFrameNumber);
 	}
 
 	/**
@@ -1569,8 +1569,8 @@ protected:
 	/** If true, force this viewport to use real time audio regardless of other settings */
 	bool bForceAudioRealtime;
 
-	/** Counter to force real-time mode for a number of frames. Overrides bIsRealtime when non-zero. */
-	uint32 RealTimeFrameCount;
+	/** Overrides bIsRealtime until GFrameCounter is >= this number. Used to force viewports to draw a number of frames even if they are otherwise non-realtime viewports. */
+	uint64 RealTimeUntilFrameNumber;
 
 	/** if the viewport is currently realtime */
 	bool bIsRealtime;
