Specific HLOD exclusion for object does not work without setting 'Reuse clusters from previous level'
#fix Make sure HLOD exclusion works in all cases, required to re-evaluate previously discarded actors

diff --git a/Engine/Source/Editor/UnrealEd/Private/HierarchicalLOD.cpp b/Engine/Source/Editor/UnrealEd/Private/HierarchicalLOD.cpp
--- a/Engine/Source/Editor/UnrealEd/Private/HierarchicalLOD.cpp
+++ b/Engine/Source/Editor/UnrealEd/Private/HierarchicalLOD.cpp
@@ -429,6 +429,20 @@ void FHierarchicalLODBuilder::InitializeClusters(ULevel* InLevel, const int32 LO
 			TArray<AActor*> Actors;
 
 			Actors.Append(LODLevelLODActors[LODIdx - 1]);
+						
+			// Re-evaluate rejected actors
+			for (AActor* Actor : RejectedActorsInLevel)
+			{
+				if (ShouldGenerateCluster(Actor, bPreviewBuild, LODIdx))
+				{
+					ValidStaticMeshActorsInLevel.Add(Actor);
+				}
+			}			
+			RejectedActorsInLevel.RemoveAll([this](AActor* Actor)
+			{
+				return ValidStaticMeshActorsInLevel.Contains(Actor);
+			});
+
 			Actors.Append(ValidStaticMeshActorsInLevel);
 
 			// first we generate graph with 2 pair nodes
