Fixing UI material thumbnails
#rnx

diff --git a/Engine/Source/Editor/UnrealEd/Private/ThumbnailHelpers.cpp b/Engine/Source/Editor/UnrealEd/Private/ThumbnailHelpers.cpp
--- a/Engine/Source/Editor/UnrealEd/Private/ThumbnailHelpers.cpp
+++ b/Engine/Source/Editor/UnrealEd/Private/ThumbnailHelpers.cpp
@@ -359,6 +359,11 @@ void FMaterialThumbnailScene::SetMaterialInterface(UMaterialInterface* InMateria
 	PreviewActor->GetStaticMeshComponent()->RecreateRenderState_Concurrent();
 }
 
+bool FMaterialThumbnailScene::ShouldSetSeparateTranslucency(class UMaterialInterface* InMaterial) const
+{
+	return InMaterial->GetMaterialResource(GMaxRHIFeatureLevel)->IsTranslucencyAfterDOFEnabled();
+}
+
 void FMaterialThumbnailScene::GetViewMatrixParameters(const float InFOVDegrees, FVector& OutOrigin, float& OutOrbitPitch, float& OutOrbitYaw, float& OutOrbitZoom) const
 {
 	check(PreviewActor);
diff --git a/Engine/Source/Editor/UnrealEd/Private/ThumbnailRendering/MaterialInstanceThumbnailRenderer.cpp b/Engine/Source/Editor/UnrealEd/Private/ThumbnailRendering/MaterialInstanceThumbnailRenderer.cpp
--- a/Engine/Source/Editor/UnrealEd/Private/ThumbnailRendering/MaterialInstanceThumbnailRenderer.cpp
+++ b/Engine/Source/Editor/UnrealEd/Private/ThumbnailRendering/MaterialInstanceThumbnailRenderer.cpp
@@ -65,11 +65,11 @@ void UMaterialInstanceThumbnailRenderer::Draw(UObject* Object, int32 X, int32 Y,
 			FSceneViewFamilyContext ViewFamily( FSceneViewFamily::ConstructionValues( RenderTarget, ThumbnailScene->GetScene(), FEngineShowFlags(ESFIM_Game) )
 				.SetWorldTimes(FApp::GetCurrentTime() - GStartTime, FApp::GetDeltaTime(), FApp::GetCurrentTime() - GStartTime));
 
+
 			ViewFamily.EngineShowFlags.DisableAdvancedFeatures();
-			ViewFamily.EngineShowFlags.SetSeparateTranslucency(true);
+			ViewFamily.EngineShowFlags.SetSeparateTranslucency(ThumbnailScene->ShouldSetSeparateTranslucency(MatInst));
 			ViewFamily.EngineShowFlags.MotionBlur = 0;
 			ViewFamily.EngineShowFlags.AntiAliasing = 0;
-
 			ThumbnailScene->GetView(&ViewFamily, X, Y, Width, Height);
 
 			if (ViewFamily.Views.Num() > 0)
diff --git a/Engine/Source/Editor/UnrealEd/Public/ThumbnailHelpers.h b/Engine/Source/Editor/UnrealEd/Public/ThumbnailHelpers.h
--- a/Engine/Source/Editor/UnrealEd/Public/ThumbnailHelpers.h
+++ b/Engine/Source/Editor/UnrealEd/Public/ThumbnailHelpers.h
@@ -76,6 +76,8 @@ public:
 	/** Sets the material to use in the next GetView() */
 	void SetMaterialInterface(class UMaterialInterface* InMaterial);
 
+	bool ShouldSetSeparateTranslucency(class UMaterialInterface* InMaterial) const;
+
 protected:
 	// FThumbnailPreviewScene implementation
 	virtual void GetViewMatrixParameters(const float InFOVDegrees, FVector& OutOrigin, float& OutOrbitPitch, float& OutOrbitYaw, float& OutOrbitZoom) const override;
