Curve Editor: Test curve validity in hit test

diff --git a/Engine/Source/Editor/UnrealEd/Private/SCurveEditor.cpp b/Engine/Source/Editor/UnrealEd/Private/SCurveEditor.cpp
--- a/Engine/Source/Editor/UnrealEd/Private/SCurveEditor.cpp
+++ b/Engine/Source/Editor/UnrealEd/Private/SCurveEditor.cpp
@@ -3061,7 +3061,8 @@ TSharedPtr<FCurveViewModel> SCurveEditor::HitTestCurves(  const FGeometry& InMyG
 		{
 
 			FRealCurve* Curve = CurveViewModel->CurveInfo.CurveToEdit;
-			if(Curve != NULL)
+
+			if (IsValidCurve(Curve))
 			{
 				float Time		 = ScaleInfo.LocalXToInput(HitPosition.X);
 				float KeyScreenY = ScaleInfo.OutputToLocalY(Curve->Eval(Time));
