Improve dialog title when reimporting Skin Weight data for a missing file
#fix added Yes/No popup dialog for informing user about missing Skin Weights file (path), to make it more clear what is going on (specifically during Base Mesh reimport)

diff --git a/Engine/Source/Editor/Persona/Private/SkinWeightProfileHelpers.cpp b/Engine/Source/Editor/Persona/Private/SkinWeightProfileHelpers.cpp
--- a/Engine/Source/Editor/Persona/Private/SkinWeightProfileHelpers.cpp
+++ b/Engine/Source/Editor/Persona/Private/SkinWeightProfileHelpers.cpp
@@ -134,9 +134,13 @@ void FSkinWeightProfileHelpers::ReimportSkinWeightProfileLOD(USkeletalMesh* InSk
 			}
 			else
 			{
-				// Otherwise let the user pick a new path
-				const FString PickedFileName = FLODUtilities::PickSkinWeightFBXPath(LODIndex);
-				bResult = FLODUtilities::ImportAlternateSkinWeight(InSkeletalMesh, PickedFileName, LODIndex, InProfileName, true);
+				FText WarningMessage = FText::Format(LOCTEXT("Warning_SkinWeightsFileMissing", "Previous file {0} containing Skin Weight data for LOD {1} could not be found, do you want to specify a new path?"), FText::FromString(PathName), LODIndex);				
+				if (EAppReturnType::Yes == FMessageDialog::Open(EAppMsgType::YesNo, WarningMessage))
+				{
+					// Otherwise let the user pick a new path
+					const FString PickedFileName = FLODUtilities::PickSkinWeightFBXPath(LODIndex);
+					bResult = FLODUtilities::ImportAlternateSkinWeight(InSkeletalMesh, PickedFileName, LODIndex, InProfileName, true);
+				}
 			}
 
 			if (bResult)
diff --git a/Engine/Source/Editor/UnrealEd/Private/LODUtilities.cpp b/Engine/Source/Editor/UnrealEd/Private/LODUtilities.cpp
--- a/Engine/Source/Editor/UnrealEd/Private/LODUtilities.cpp
+++ b/Engine/Source/Editor/UnrealEd/Private/LODUtilities.cpp
@@ -44,6 +44,7 @@
 #include "Framework/Application/SlateApplication.h"
 
 
+#define LOCTEXT_NAMESPACE "LODUtilities"
 
 DEFINE_LOG_CATEGORY_STATIC(LogLODUtilities, Log, All);
 
@@ -2035,10 +2036,14 @@ bool FLODUtilities::ReimportAlternateSkinWeight(USkeletalMesh* SkeletalMesh, int
 			}
 			else
 			{
-				const FString PickedFileName = FLODUtilities::PickSkinWeightFBXPath(TargetLODIndex);
-				if (!PickedFileName.IsEmpty() && FPaths::FileExists(PickedFileName))
+				FText WarningMessage = FText::Format(LOCTEXT("Warning_SkinWeightsFileMissing", "Previous file {0} containing Skin Weight data for LOD {1} could not be found, do you want to specify a new path?"), FText::FromString(PathName), TargetLODIndex);
+				if (EAppReturnType::Yes == FMessageDialog::Open(EAppMsgType::YesNo, WarningMessage))
 				{
-					bResult |= FLODUtilities::ImportAlternateSkinWeight(SkeletalMesh, PickedFileName, TargetLODIndex, ProfileInfo.Name, false);
+					const FString PickedFileName = FLODUtilities::PickSkinWeightFBXPath(TargetLODIndex);
+					if (!PickedFileName.IsEmpty() && FPaths::FileExists(PickedFileName))
+					{
+						bResult |= FLODUtilities::ImportAlternateSkinWeight(SkeletalMesh, PickedFileName, TargetLODIndex, ProfileInfo.Name, false);
+					}
 				}
 			}
 		}
@@ -2190,4 +2195,6 @@ FString FLODUtilities::PickSkinWeightFBXPath(int32 LODIndex)
 	}
 
 	return PickedFileName;
-}
\ No newline at end of file
+}
+
+#undef LOCTEXT_NAMESPACE // "LODUtilities"
\ No newline at end of file
