Slate Global Invalidation Mode - Fixing the PIE new window to not always be volatile, and allow the mode to work for new PIE Windows, and fixing the console to be visible.  Also fixing a case SlateDebugging visuals that need to check for having multiple invalidation roots.

diff --git a/Engine/Source/Editor/UnrealEd/Private/PlayLevel.cpp b/Engine/Source/Editor/UnrealEd/Private/PlayLevel.cpp
--- a/Engine/Source/Editor/UnrealEd/Private/PlayLevel.cpp
+++ b/Engine/Source/Editor/UnrealEd/Private/PlayLevel.cpp
@@ -3055,6 +3055,53 @@ void UEditorEngine::RequestLateJoin()
 	GetMultipleInstancePositions(0, NextX, NextY);
 }
 
+class SPIEViewport : public SViewport
+{
+	SLATE_BEGIN_ARGS(SPIEViewport)
+		: _Content()
+		, _RenderDirectlyToWindow(false)
+		, _EnableStereoRendering(false)
+	{
+		_Clipping = EWidgetClipping::ClipToBoundsAlways;
+	}
+
+		SLATE_DEFAULT_SLOT(FArguments, Content)
+
+		/**
+		 * Whether or not to render directly to the window's backbuffer or an offscreen render target that is applied to the window later
+		 * Rendering to an offscreen target is the most common option in the editor where there may be many frames which this viewport's interface may wish to not re-render but use a cached buffer instead
+		 * Rendering directly to the backbuffer is the most common option in the game where you want to update each frame without the cost of writing to an intermediate target first.
+		 */
+		SLATE_ARGUMENT(bool, RenderDirectlyToWindow)
+
+		/** Whether or not to enable stereo rendering. */
+		SLATE_ARGUMENT(bool, EnableStereoRendering )
+
+	SLATE_END_ARGS()
+
+	void Construct(const FArguments& InArgs)
+	{
+		SViewport::Construct(
+			SViewport::FArguments()
+			.EnableGammaCorrection(false) // Gamma correction in the game is handled in post processing in the scene renderer
+			.RenderDirectlyToWindow(InArgs._RenderDirectlyToWindow)
+			.EnableStereoRendering(InArgs._EnableStereoRendering)
+			[
+				InArgs._Content.Widget
+			]
+		);
+	}
+
+	virtual void Tick(const FGeometry& AllottedGeometry, const double InCurrentTime, const float InDeltaTime) override
+	{
+		SViewport::Tick(AllottedGeometry, InCurrentTime, InDeltaTime);
+
+		// Rather than binding the attribute we're going to poll it in tick, otherwise we will make this widget volatile, and it therefore
+		// wont be possible to cache it or its children in GSlateEnableGlobalInvalidation mode.
+		SetEnabled(FSlateApplication::Get().GetNormalExecutionAttribute().Get());
+	}
+};
+
 UGameInstance* UEditorEngine::CreatePIEGameInstance(int32 InPIEInstance, bool bInSimulateInEditor, bool bAnyBlueprintErrors, bool bStartInSpectatorMode, bool bRunAsDedicated, bool bPlayStereoscopic, float PIEStartTime)
 {
 	// create a new GameInstance
@@ -3165,7 +3212,7 @@ UGameInstance* UEditorEngine::CreatePIEGameInstance(int32 InPIEInstance, bool bI
 	}
 
 	// For play in editor, this is the viewport widget where the game is being displayed
-	TSharedPtr<SViewport> PieViewportWidget;
+	TSharedPtr<SPIEViewport> PieViewportWidget;
 
 	// Initialize the viewport client.
 	UGameViewportClient* ViewportClient = NULL;
@@ -3333,9 +3380,7 @@ UGameInstance* UEditorEngine::CreatePIEGameInstance(int32 InPIEInstance, bool bI
 				const bool bIgnoreTextureAlpha = (PropagateAlpha != EAlphaChannelMode::AllowThroughTonemapper);
 
 				PieViewportWidget = 
-					SNew( SViewport )
-						.IsEnabled(FSlateApplication::Get().GetNormalExecutionAttribute())
-						.EnableGammaCorrection( false )// Gamma correction in the game is handled in post processing in the scene renderer
+					SNew( SPIEViewport )
 						.RenderDirectlyToWindow( bRenderDirectlyToWindow )
 						.EnableStereoRendering( bEnableStereoRendering )
 						.IgnoreTextureAlpha(bIgnoreTextureAlpha)
@@ -3359,6 +3404,7 @@ UGameInstance* UEditorEngine::CreatePIEGameInstance(int32 InPIEInstance, bool bI
 
 				ViewportClient->SetViewportOverlayWidget( PieWindow, ViewportOverlayWidgetRef );
 				ViewportClient->SetGameLayerManager(GameLayerManagerRef);
+
 				bool bShouldMinimizeRootWindow = bPlayStereoscopic && GEngine->XRSystem.IsValid() && GetDefault<ULevelEditorPlaySettings>()->ShouldMinimizeEditorOnVRPIE;
 				// Set up a notification when the window is closed so we can clean up PIE
 				{
@@ -3418,10 +3464,12 @@ UGameInstance* UEditorEngine::CreatePIEGameInstance(int32 InPIEInstance, bool bI
 				// Create a new viewport that the viewport widget will use to render the game
 				SlatePlayInEditorSession.SlatePlayInEditorWindowViewport = MakeShareable( new FSceneViewport( ViewportClient, PieViewportWidget ) );
 
+				GameLayerManagerRef->SetSceneViewport(SlatePlayInEditorSession.SlatePlayInEditorWindowViewport.Get());
+
 				const bool bShouldGameGetMouseControl = GetDefault<ULevelEditorPlaySettings>()->GameGetsMouseControl || (bPlayStereoscopic && GEngine->XRSystem.IsValid());
 				SlatePlayInEditorSession.SlatePlayInEditorWindowViewport->SetPlayInEditorGetsMouseControl(bShouldGameGetMouseControl);
 				PieViewportWidget->SetViewportInterface( SlatePlayInEditorSession.SlatePlayInEditorWindowViewport.ToSharedRef() );
-				
+
 				FSlateApplication::Get().RegisterViewport(PieViewportWidget.ToSharedRef());
 				
 				SlatePlayInEditorSession.SlatePlayInEditorWindow = PieWindow;
diff --git a/Engine/Source/Runtime/SlateCore/Private/Debugging/SlateDebugging.cpp b/Engine/Source/Runtime/SlateCore/Private/Debugging/SlateDebugging.cpp
--- a/Engine/Source/Runtime/SlateCore/Private/Debugging/SlateDebugging.cpp
+++ b/Engine/Source/Runtime/SlateCore/Private/Debugging/SlateDebugging.cpp
@@ -88,7 +88,7 @@ struct FInvalidatedWidgetDrawer
 
 			FSlateDrawElement::MakeBox(
 				ElementList,
-				MyState.OutgoingLayerId+1,
+				MyState.OutgoingLayerId + 1,
 				MyState.AllottedGeometry.ToPaintGeometry(),
 				WhiteBrush,
 				ESlateDrawEffect::None,
@@ -268,7 +268,10 @@ void FSlateDebugging::WidgetInvalidated(FSlateInvalidationRoot& InvalidationRoot
 	{
 		int32 Index = WidgetProxy.Index;
 	
-		FInvalidatedWidgetDrawer* Drawer = InvalidatedWidgetDrawers.FindByPredicate([Index](FInvalidatedWidgetDrawer& InDrawer) { return InDrawer.ProxyHandle.GetIndex() == Index; });
+		FInvalidatedWidgetDrawer* Drawer = InvalidatedWidgetDrawers.FindByPredicate([&InvalidationRoot, Index](FInvalidatedWidgetDrawer& InDrawer) { 
+			return InDrawer.ProxyHandle.GetInvalidationRoot() == &InvalidationRoot && InDrawer.ProxyHandle.GetIndex() == Index;
+		});
+
 		if (!Drawer) 
 		{
 			Drawer = &InvalidatedWidgetDrawers.Emplace_GetRef(FWidgetProxyHandle(InvalidationRoot, WidgetProxy.Index));
