Update the SPIEViewport to handle forwarding IgnoreTextureAlpha.

diff --git a/Engine/Source/Editor/UnrealEd/Private/PlayLevel.cpp b/Engine/Source/Editor/UnrealEd/Private/PlayLevel.cpp
--- a/Engine/Source/Editor/UnrealEd/Private/PlayLevel.cpp
+++ b/Engine/Source/Editor/UnrealEd/Private/PlayLevel.cpp
@@ -3061,6 +3061,7 @@ class SPIEViewport : public SViewport
 		: _Content()
 		, _RenderDirectlyToWindow(false)
 		, _EnableStereoRendering(false)
+		, _IgnoreTextureAlpha(true)
 	{
 		_Clipping = EWidgetClipping::ClipToBoundsAlways;
 	}
@@ -3077,6 +3078,12 @@ class SPIEViewport : public SViewport
 		/** Whether or not to enable stereo rendering. */
 		SLATE_ARGUMENT(bool, EnableStereoRendering )
 
+		/**
+		 * If true, the viewport's texture alpha is ignored when performing blending.  In this case only the viewport tint opacity is used
+		 * If false, the texture alpha is used during blending
+		 */
+		SLATE_ARGUMENT( bool, IgnoreTextureAlpha )
+
 	SLATE_END_ARGS()
 
 	void Construct(const FArguments& InArgs)
@@ -3086,6 +3093,7 @@ class SPIEViewport : public SViewport
 			.EnableGammaCorrection(false) // Gamma correction in the game is handled in post processing in the scene renderer
 			.RenderDirectlyToWindow(InArgs._RenderDirectlyToWindow)
 			.EnableStereoRendering(InArgs._EnableStereoRendering)
+			.IgnoreTextureAlpha(InArgs._IgnoreTextureAlpha)
 			[
 				InArgs._Content.Widget
 			]
