Apply platform scalability settings when changing Preview Rendering Level

diff --git a/Engine/Source/Editor/LevelEditor/Private/LevelEditor.cpp b/Engine/Source/Editor/LevelEditor/Private/LevelEditor.cpp
--- a/Engine/Source/Editor/LevelEditor/Private/LevelEditor.cpp
+++ b/Engine/Source/Editor/LevelEditor/Private/LevelEditor.cpp
@@ -1757,36 +1757,46 @@ void FLevelEditorModule::BindGlobalLevelEditorCommands()
 		FIsActionButtonVisible::CreateStatic(FLevelEditorActionCallbacks::IsPreviewModeButtonVisible));
 
 	ActionList.MapAction(
-		Commands.PreviewPlatformOverride_DefaultES2,
-		FExecuteAction::CreateStatic(&FLevelEditorActionCallbacks::SetPreviewPlatform, FName(), ERHIFeatureLevel::ES2),
+		Commands.PreviewPlatformOverride_SM5,
+		FExecuteAction::CreateStatic(&FLevelEditorActionCallbacks::SetPreviewPlatform, FPreviewPlatformInfo(ERHIFeatureLevel::SM5, NAME_None, false)),
 		FCanExecuteAction(),
-		FIsActionChecked::CreateStatic(&FLevelEditorActionCallbacks::IsPreviewPlatformChecked, FName(), ERHIFeatureLevel::ES2));
+		FIsActionChecked::CreateStatic(&FLevelEditorActionCallbacks::IsPreviewPlatformChecked, FPreviewPlatformInfo(ERHIFeatureLevel::SM5, NAME_None)));
+
 	ActionList.MapAction(
-		Commands.PreviewPlatformOverride_AndroidGLES2,
-		FExecuteAction::CreateStatic(&FLevelEditorActionCallbacks::SetPreviewPlatform, LegacyShaderPlatformToShaderFormat(SP_OPENGL_ES2_ANDROID), ERHIFeatureLevel::ES2),
+		Commands.PreviewPlatformOverride_SM4,
+		FExecuteAction::CreateStatic(&FLevelEditorActionCallbacks::SetPreviewPlatform, FPreviewPlatformInfo(ERHIFeatureLevel::SM4, NAME_None, true)),
 		FCanExecuteAction(),
-		FIsActionChecked::CreateStatic(&FLevelEditorActionCallbacks::IsPreviewPlatformChecked, LegacyShaderPlatformToShaderFormat(SP_OPENGL_ES2_ANDROID), ERHIFeatureLevel::ES2));
+		FIsActionChecked::CreateStatic(&FLevelEditorActionCallbacks::IsPreviewPlatformChecked, FPreviewPlatformInfo(ERHIFeatureLevel::SM4, NAME_None)));
 
 	ActionList.MapAction(
-		Commands.PreviewPlatformOverride_DefaultES31,
-		FExecuteAction::CreateStatic(&FLevelEditorActionCallbacks::SetPreviewPlatform, FName(), ERHIFeatureLevel::ES3_1),
+		Commands.PreviewPlatformOverride_HTML5,
+		FExecuteAction::CreateStatic(&FLevelEditorActionCallbacks::SetPreviewPlatform, FPreviewPlatformInfo(ERHIFeatureLevel::ES2, NAME_None, true)),
 		FCanExecuteAction(),
-		FIsActionChecked::CreateStatic(&FLevelEditorActionCallbacks::IsPreviewPlatformChecked, FName(), ERHIFeatureLevel::ES3_1));
+		FIsActionChecked::CreateStatic(&FLevelEditorActionCallbacks::IsPreviewPlatformChecked, FPreviewPlatformInfo(ERHIFeatureLevel::ES2, NAME_None)));
+
+	ActionList.MapAction(
+		Commands.PreviewPlatformOverride_AndroidGLES2,
+		FExecuteAction::CreateStatic(&FLevelEditorActionCallbacks::SetPreviewPlatform, FPreviewPlatformInfo(ERHIFeatureLevel::ES2, LegacyShaderPlatformToShaderFormat(SP_OPENGL_ES2_ANDROID), true)),
+		FCanExecuteAction(),
+		FIsActionChecked::CreateStatic(&FLevelEditorActionCallbacks::IsPreviewPlatformChecked, FPreviewPlatformInfo(ERHIFeatureLevel::ES2, LegacyShaderPlatformToShaderFormat(SP_OPENGL_ES2_ANDROID))));
+
 	ActionList.MapAction(
 		Commands.PreviewPlatformOverride_AndroidGLES31,
-		FExecuteAction::CreateStatic(&FLevelEditorActionCallbacks::SetPreviewPlatform, LegacyShaderPlatformToShaderFormat(SP_OPENGL_ES3_1_ANDROID), ERHIFeatureLevel::ES3_1),
+		FExecuteAction::CreateStatic(&FLevelEditorActionCallbacks::SetPreviewPlatform, FPreviewPlatformInfo(ERHIFeatureLevel::ES3_1, LegacyShaderPlatformToShaderFormat(SP_OPENGL_ES3_1_ANDROID), true)),
 		FCanExecuteAction(),
-		FIsActionChecked::CreateStatic(&FLevelEditorActionCallbacks::IsPreviewPlatformChecked, LegacyShaderPlatformToShaderFormat(SP_OPENGL_ES3_1_ANDROID), ERHIFeatureLevel::ES3_1));
+		FIsActionChecked::CreateStatic(&FLevelEditorActionCallbacks::IsPreviewPlatformChecked, FPreviewPlatformInfo(ERHIFeatureLevel::ES3_1, LegacyShaderPlatformToShaderFormat(SP_OPENGL_ES3_1_ANDROID))));
+
 	ActionList.MapAction(
 		Commands.PreviewPlatformOverride_AndroidVulkanES31,
-		FExecuteAction::CreateStatic(&FLevelEditorActionCallbacks::SetPreviewPlatform, LegacyShaderPlatformToShaderFormat(SP_VULKAN_ES3_1_ANDROID), ERHIFeatureLevel::ES3_1),
+		FExecuteAction::CreateStatic(&FLevelEditorActionCallbacks::SetPreviewPlatform, FPreviewPlatformInfo(ERHIFeatureLevel::ES3_1, LegacyShaderPlatformToShaderFormat(SP_VULKAN_ES3_1_ANDROID), true)),
 		FCanExecuteAction(),
-		FIsActionChecked::CreateStatic(&FLevelEditorActionCallbacks::IsPreviewPlatformChecked, LegacyShaderPlatformToShaderFormat(SP_VULKAN_ES3_1_ANDROID), ERHIFeatureLevel::ES3_1));
+		FIsActionChecked::CreateStatic(&FLevelEditorActionCallbacks::IsPreviewPlatformChecked, FPreviewPlatformInfo(ERHIFeatureLevel::ES3_1, LegacyShaderPlatformToShaderFormat(SP_VULKAN_ES3_1_ANDROID))));
+
 	ActionList.MapAction(
 		Commands.PreviewPlatformOverride_IOSMetalES31,
-		FExecuteAction::CreateStatic(&FLevelEditorActionCallbacks::SetPreviewPlatform, LegacyShaderPlatformToShaderFormat(SP_METAL), ERHIFeatureLevel::ES3_1),
+		FExecuteAction::CreateStatic(&FLevelEditorActionCallbacks::SetPreviewPlatform, FPreviewPlatformInfo(ERHIFeatureLevel::ES3_1, LegacyShaderPlatformToShaderFormat(SP_METAL), true)),
 		FCanExecuteAction(),
-		FIsActionChecked::CreateStatic(&FLevelEditorActionCallbacks::IsPreviewPlatformChecked, LegacyShaderPlatformToShaderFormat(SP_METAL), ERHIFeatureLevel::ES3_1));
+		FIsActionChecked::CreateStatic(&FLevelEditorActionCallbacks::IsPreviewPlatformChecked, FPreviewPlatformInfo(ERHIFeatureLevel::ES3_1, LegacyShaderPlatformToShaderFormat(SP_METAL))));
 
 	ActionList.MapAction(
 		Commands.OpenMergeActor,
@@ -1813,16 +1823,6 @@ void FLevelEditorModule::BindGlobalLevelEditorCommands()
 		FCanExecuteAction(),
 		FIsActionChecked::CreateStatic(FLevelEditorActionCallbacks::GeometryCollection_IsChecked)
 		);
-
-
-	for (int32 i = 0; i < ERHIFeatureLevel::Num; ++i)
-	{
-		ActionList.MapAction(
-			Commands.FeatureLevelPreview[i],
-			FExecuteAction::CreateStatic(&FLevelEditorActionCallbacks::SetFeatureLevelPreview, (ERHIFeatureLevel::Type)i),
-			FCanExecuteAction::CreateStatic(&FLevelEditorActionCallbacks::IsFeatureLevelPreviewAvailable, (ERHIFeatureLevel::Type)i),
-			FIsActionChecked::CreateStatic(&FLevelEditorActionCallbacks::IsFeatureLevelPreviewChecked, (ERHIFeatureLevel::Type)i));
-	}
 }
 	
 #undef LOCTEXT_NAMESPACE
diff --git a/Engine/Source/Editor/LevelEditor/Private/LevelEditorActions.cpp b/Engine/Source/Editor/LevelEditor/Private/LevelEditorActions.cpp
--- a/Engine/Source/Editor/LevelEditor/Private/LevelEditorActions.cpp
+++ b/Engine/Source/Editor/LevelEditor/Private/LevelEditorActions.cpp
@@ -583,7 +583,7 @@ bool FLevelEditorActionCallbacks::IsFeatureLevelPreviewEnabled()
 	{
 		return false;
 	}
-	if (GEditor->PreviewFeatureLevel == ERHIFeatureLevel::SM5)
+	if (GEditor->PreviewPlatform.PreviewFeatureLevel == ERHIFeatureLevel::SM5)
 	{
 		return true;
 	}
@@ -592,7 +592,7 @@ bool FLevelEditorActionCallbacks::IsFeatureLevelPreviewEnabled()
 
 bool FLevelEditorActionCallbacks::IsFeatureLevelPreviewActive()
 {
-	if (GEditor->PreviewFeatureLevel == ERHIFeatureLevel::SM5)
+	if (GEditor->PreviewPlatform.PreviewFeatureLevel == ERHIFeatureLevel::SM5)
 	{
 		return false;
 	}
@@ -601,39 +601,17 @@ bool FLevelEditorActionCallbacks::IsFeatureLevelPreviewActive()
 
 bool FLevelEditorActionCallbacks::IsPreviewModeButtonVisible()
 {
-	return GEditor->PreviewFeatureLevel != ERHIFeatureLevel::SM5;
+	return GEditor->PreviewPlatform.PreviewFeatureLevel != ERHIFeatureLevel::SM5;
 }
 
-void FLevelEditorActionCallbacks::SetPreviewPlatform(FName MaterialQualityPlatform, ERHIFeatureLevel::Type PreviewFeatureLevel)
+void FLevelEditorActionCallbacks::SetPreviewPlatform(FPreviewPlatformInfo NewPreviewPlatform)
 {
-	GEditor->SetPreviewPlatform(MaterialQualityPlatform, PreviewFeatureLevel);
+	GEditor->SetPreviewPlatform(NewPreviewPlatform, true);
 }
 
-bool FLevelEditorActionCallbacks::IsPreviewPlatformChecked(FName InMaterialQualityPlatform, ERHIFeatureLevel::Type InPreviewFeatureLevel)
+bool FLevelEditorActionCallbacks::IsPreviewPlatformChecked(FPreviewPlatformInfo PreviewPlatform)
 {
-	const FName& PreviewPlatform = UMaterialShaderQualitySettings::Get()->GetPreviewPlatform();
-	return PreviewPlatform == InMaterialQualityPlatform && InPreviewFeatureLevel == GEditor->PreviewFeatureLevel;
-}
-
-void FLevelEditorActionCallbacks::SetFeatureLevelPreview(ERHIFeatureLevel::Type InPreviewFeatureLevel)
-{
-	// When called through SMenuEntryBlock::OnClicked(), the popup menus are not dismissed when
-	// clicking on a checkbox, but they are dismissed when clicking on a button. We need the popup
-	// menus to go away, or SetFeaturePlatform() is unable to display a progress dialog. Force
-	// the dismissal here.
-	FSlateApplication::Get().DismissAllMenus();
-	
-	GEditor->SetPreviewPlatform(NAME_None, InPreviewFeatureLevel);
-}
-
-bool FLevelEditorActionCallbacks::IsFeatureLevelPreviewChecked(ERHIFeatureLevel::Type InPreviewFeatureLevel)
-{
-	return GEditor->PreviewFeatureLevel == InPreviewFeatureLevel;
-}
-
-bool FLevelEditorActionCallbacks::IsFeatureLevelPreviewAvailable(ERHIFeatureLevel::Type InPreviewFeatureLevel)
-{
-	return GShaderPlatformForFeatureLevel[InPreviewFeatureLevel] != SP_NumPlatforms;
+	return GEditor->PreviewPlatform.Matches(PreviewPlatform);
 }
 
 void FLevelEditorActionCallbacks::ConfigureLightingBuildOptions( const FLightingBuildOptions& Options )
@@ -3418,15 +3396,14 @@ void FLevelEditorCommands::RegisterCommands()
 
 	UI_COMMAND(ToggleFeatureLevelPreview, "Preview Mode Toggle", "Toggles the Preview Mode on or off for the currently selected Preview target", EUserInterfaceActionType::ToggleButton, FInputChord());
 
+	UI_COMMAND(PreviewPlatformOverride_SM5, "Shader Model 5", "DirectX 11, OpenGL 4.3+, PS4, XB1", EUserInterfaceActionType::Check, FInputChord());
+	UI_COMMAND(PreviewPlatformOverride_SM4, "Shader Model 4", "DirectX 10, OpenGL 3.3+", EUserInterfaceActionType::Check, FInputChord());
+	UI_COMMAND(PreviewPlatformOverride_HTML5, "HTML5", "HTML5 preview.", EUserInterfaceActionType::Check, FInputChord());
 	UI_COMMAND(PreviewPlatformOverride_AndroidGLES2, "Android ES2", "Mobile preview using Android's quality settings.", EUserInterfaceActionType::Check, FInputChord());
-	UI_COMMAND(PreviewPlatformOverride_DefaultES2, "HTML5", "HTML5 preview.", EUserInterfaceActionType::Check, FInputChord());
-
-	UI_COMMAND(PreviewPlatformOverride_DefaultES31, "Default High-End Mobile", "Use default mobile settings (no quality overrides).", EUserInterfaceActionType::Check, FInputChord());
 	UI_COMMAND(PreviewPlatformOverride_AndroidGLES31, "Android ES 3.1", "Mobile preview using Android ES3.1 quality settings.", EUserInterfaceActionType::Check, FInputChord());
 	UI_COMMAND(PreviewPlatformOverride_AndroidVulkanES31, "Android Vulkan", "Mobile preview using Android Vulkan quality settings.", EUserInterfaceActionType::Check, FInputChord());
 	UI_COMMAND(PreviewPlatformOverride_IOSMetalES31, "iOS", "Mobile preview using iOS material quality settings.", EUserInterfaceActionType::Check, FInputChord());
 
-
 	UI_COMMAND( ConnectToSourceControl, "Connect to Source Control...", "Opens a dialog to connect to source control.", EUserInterfaceActionType::Button, FInputChord());
 	UI_COMMAND( ChangeSourceControlSettings, "Change Source Control Settings...", "Opens a dialog to change source control settings.", EUserInterfaceActionType::Button, FInputChord());
 	UI_COMMAND( CheckOutModifiedFiles, "Check Out Modified Files...", "Opens a dialog to check out any assets which have been modified.", EUserInterfaceActionType::Button, FInputChord());
@@ -3436,37 +3413,6 @@ void FLevelEditorCommands::RegisterCommands()
 	UI_COMMAND(GeometryCollectionSelectNone, "Deselect All Geometry In Hierarchy", "Deselect all geometry in hierarchy", EUserInterfaceActionType::Button, FInputChord());
 	UI_COMMAND(GeometryCollectionSelectInverseGeometry, "Select Inverse Geometry In Hierarchy", "Select inverse geometry in hierarchy", EUserInterfaceActionType::Button, FInputChord());
 
-	static const FText FeatureLevelLabels[ERHIFeatureLevel::Num] = 
-	{
-		NSLOCTEXT("LevelEditorCommands", "FeatureLevelPreviewType_ES2", "Mobile / HTML5"),
-		NSLOCTEXT("LevelEditorCommands", "FeatureLevelPreviewType_ES31", "High-End Mobile"),
-		NSLOCTEXT("LevelEditorCommands", "FeatureLevelPreviewType_SM4", "Shader Model 4"),
-		NSLOCTEXT("LevelEditorCommands", "FeatureLevelPreviewType_SM5", "Shader Model 5"),
-	};
-
-	static const FText FeatureLevelToolTips[ERHIFeatureLevel::Num] = 
-	{
-		NSLOCTEXT("LevelEditorCommands", "FeatureLevelPreviewTooltip_ES2", "OpenGLES 2"),
-		NSLOCTEXT("LevelEditorCommands", "FeatureLevelPreviewTooltip_ES3", "OpenGLES 3.1, Metal, Vulkan"),
-		NSLOCTEXT("LevelEditorCommands", "FeatureLevelPreviewTooltip_SM4", "DirectX 10, OpenGL 3.3+"),
-		NSLOCTEXT("LevelEditorCommands", "FeatureLevelPreviewTooltip_SM5", "DirectX 11, OpenGL 4.3+, PS4, XB1"),
-	};
-
-	for (int32 i = 0; i < ERHIFeatureLevel::Num; ++i)
-	{
-		FName Name;
-		GetFeatureLevelName((ERHIFeatureLevel::Type)i, Name);
-
-		FeatureLevelPreview[i] =
-			FUICommandInfoDecl(
-			this->AsShared(),
-			Name,
-			FeatureLevelLabels[i],
-			FeatureLevelToolTips[i])
-			.UserInterfaceType(EUserInterfaceActionType::Check)
-			.DefaultChord(FInputChord());
-	}
-
 	UI_COMMAND(OpenMergeActor, "Merge Actors", "Opens the Merge Actor panel", EUserInterfaceActionType::Button, FInputChord());
 }
 
diff --git a/Engine/Source/Editor/LevelEditor/Private/LevelEditorToolBar.cpp b/Engine/Source/Editor/LevelEditor/Private/LevelEditorToolBar.cpp
--- a/Engine/Source/Editor/LevelEditor/Private/LevelEditorToolBar.cpp
+++ b/Engine/Source/Editor/LevelEditor/Private/LevelEditorToolBar.cpp
@@ -1281,14 +1281,7 @@ TSharedRef< SWidget > FLevelEditorToolBar::MakeLevelEditorToolBar( const TShared
 
 			static FText GetPreviewModeText()
 			{
-				UMaterialShaderQualitySettings* MaterialShaderQualitySettings = UMaterialShaderQualitySettings::Get();
-				const FName& PreviewPlatform = MaterialShaderQualitySettings->GetPreviewPlatform();
-
-				EShaderPlatform ShaderPlatform = ShaderFormatToLegacyShaderPlatform(PreviewPlatform);
-				if (ShaderPlatform == SP_NumPlatforms)
-				{
-					ShaderPlatform = GetFeatureLevelShaderPlatform(GEditor->PreviewFeatureLevel);
-				}
+				EShaderPlatform ShaderPlatform = ShaderFormatToLegacyShaderPlatform(GEditor->PreviewPlatform.PreviewShaderPlatformName);
 
 				switch (ShaderPlatform)
 				{
@@ -1298,7 +1291,7 @@ TSharedRef< SWidget > FLevelEditorToolBar::MakeLevelEditorToolBar( const TShared
 					}
 				}
 
-				switch (GEditor->PreviewFeatureLevel)
+				switch (GEditor->PreviewPlatform.PreviewFeatureLevel)
 				{
 					case ERHIFeatureLevel::SM4:
 					{
@@ -1319,23 +1312,17 @@ TSharedRef< SWidget > FLevelEditorToolBar::MakeLevelEditorToolBar( const TShared
 				}
 			}
 
-
 			static FText GetPreviewModeTooltip()
 			{
-				UMaterialShaderQualitySettings* MaterialShaderQualitySettings = UMaterialShaderQualitySettings::Get();
-				const FName& PreviewPlatform = MaterialShaderQualitySettings->GetPreviewPlatform();
-
-				EShaderPlatform PreviewShaderPlatform = ShaderFormatToLegacyShaderPlatform(PreviewPlatform);
-				if (PreviewShaderPlatform == SP_NumPlatforms)
-				{
-					PreviewShaderPlatform = GetFeatureLevelShaderPlatform(GEditor->PreviewFeatureLevel);
-				}
+				EShaderPlatform PreviewShaderPlatform = GEditor->PreviewPlatform.PreviewShaderPlatformName != NAME_None ?
+					ShaderFormatToLegacyShaderPlatform(GEditor->PreviewPlatform.PreviewShaderPlatformName) :
+					GetFeatureLevelShaderPlatform(GEditor->PreviewPlatform.PreviewFeatureLevel);
 
 				EShaderPlatform MaxRHIFeatureLevelPlatform = GetFeatureLevelShaderPlatform(GMaxRHIFeatureLevel);
 
 				{
-                    const FText& RenderingAsPlatformName = GetFriendlyShaderPlatformName(GWorld->FeatureLevel == GMaxRHIFeatureLevel ? MaxRHIFeatureLevelPlatform : PreviewShaderPlatform);
-                    const FText& SwitchToPlatformName = GetFriendlyShaderPlatformName(GWorld->FeatureLevel == GMaxRHIFeatureLevel ? PreviewShaderPlatform : MaxRHIFeatureLevelPlatform);
+					const FText& RenderingAsPlatformName = GetFriendlyShaderPlatformName(GEditor->PreviewPlatform.bPreviewFeatureLevelActive ? PreviewShaderPlatform : MaxRHIFeatureLevelPlatform);
+                    const FText& SwitchToPlatformName = GetFriendlyShaderPlatformName(GEditor->PreviewPlatform.bPreviewFeatureLevelActive ? MaxRHIFeatureLevelPlatform : PreviewShaderPlatform);
                     if (GWorld->FeatureLevel == GMaxRHIFeatureLevel)
                     {
                         return FText::Format(LOCTEXT("PreviewModeViewingAsSwitchTo", "Viewing {0}. Click to preview {1}."), RenderingAsPlatformName, SwitchToPlatformName);
@@ -1349,13 +1336,11 @@ TSharedRef< SWidget > FLevelEditorToolBar::MakeLevelEditorToolBar( const TShared
 
 			static FSlateIcon GetPreviewModeIcon()
 			{
-				UMaterialShaderQualitySettings* MaterialShaderQualitySettings = UMaterialShaderQualitySettings::Get();
-				const FName& PreviewPlatform = MaterialShaderQualitySettings->GetPreviewPlatform();
+				EShaderPlatform ShaderPlatform = ShaderFormatToLegacyShaderPlatform(GEditor->PreviewPlatform.PreviewShaderPlatformName);
 
-				EShaderPlatform ShaderPlatform = ShaderFormatToLegacyShaderPlatform(PreviewPlatform);
 				if (ShaderPlatform == SP_NumPlatforms)
 				{
-					ShaderPlatform = GetFeatureLevelShaderPlatform(GEditor->PreviewFeatureLevel);
+					ShaderPlatform = GetFeatureLevelShaderPlatform(GEditor->PreviewPlatform.PreviewFeatureLevel);
 				}
 				switch (ShaderPlatform)
 				{
@@ -1380,7 +1365,7 @@ TSharedRef< SWidget > FLevelEditorToolBar::MakeLevelEditorToolBar( const TShared
 						return FSlateIcon(FEditorStyle::GetStyleSetName(), GEditor->IsFeatureLevelPreviewActive() ? "LevelEditor.PreviewMode.HTML5.Enabled" : "LevelEditor.PreviewMode.HTML5.Disabled");
 					}
 				}
-				switch (GEditor->PreviewFeatureLevel)
+				switch (GEditor->PreviewPlatform.PreviewFeatureLevel)
 				{
 					case ERHIFeatureLevel::SM4:
 					{
@@ -1829,10 +1814,10 @@ static void MakeShaderModelPreviewMenu( FMenuBuilder& MenuBuilder )
 	MenuBuilder.BeginSection("EditorPreviewMode", LOCTEXT("EditorPreviewModeDevices", "Preview Devices"));
 
 	// SM5
-	MenuBuilder.AddMenuEntry(FLevelEditorCommands::Get().FeatureLevelPreview[ERHIFeatureLevel::SM5]);
+	MenuBuilder.AddMenuEntry(FLevelEditorCommands::Get().PreviewPlatformOverride_SM5);
 
 	// SM4
-	MenuBuilder.AddMenuEntry(FLevelEditorCommands::Get().FeatureLevelPreview[ERHIFeatureLevel::SM4]);
+	MenuBuilder.AddMenuEntry(FLevelEditorCommands::Get().PreviewPlatformOverride_SM4);
 
 	// Android
 	bool bAndroidBuildForES31 = false;
@@ -1865,7 +1850,7 @@ static void MakeShaderModelPreviewMenu( FMenuBuilder& MenuBuilder )
 	}
 
 	// HTML5
-	MenuBuilder.AddMenuEntry(FLevelEditorCommands::Get().PreviewPlatformOverride_DefaultES2);
+	MenuBuilder.AddMenuEntry(FLevelEditorCommands::Get().PreviewPlatformOverride_HTML5);
 
     MenuBuilder.EndSection();
 
diff --git a/Engine/Source/Editor/LevelEditor/Public/LevelEditorActions.h b/Engine/Source/Editor/LevelEditor/Public/LevelEditorActions.h
--- a/Engine/Source/Editor/LevelEditor/Public/LevelEditorActions.h
+++ b/Engine/Source/Editor/LevelEditor/Public/LevelEditorActions.h
@@ -590,14 +590,12 @@ public:
 	TSharedPtr< FUICommandInfo > MaterialQualityLevel_Medium;
 	TSharedPtr< FUICommandInfo > MaterialQualityLevel_High;
 
-	TSharedPtr< FUICommandInfo > FeatureLevelPreview[ERHIFeatureLevel::Num];
-	
 	TSharedPtr< FUICommandInfo > ToggleFeatureLevelPreview;
 
-	TSharedPtr< FUICommandInfo > PreviewPlatformOverride_DefaultES2;
+	TSharedPtr< FUICommandInfo > PreviewPlatformOverride_SM5;
+	TSharedPtr< FUICommandInfo > PreviewPlatformOverride_SM4;
+	TSharedPtr< FUICommandInfo > PreviewPlatformOverride_HTML5;
 	TSharedPtr< FUICommandInfo > PreviewPlatformOverride_AndroidGLES2;
-
-	TSharedPtr< FUICommandInfo > PreviewPlatformOverride_DefaultES31;
 	TSharedPtr< FUICommandInfo > PreviewPlatformOverride_AndroidGLES31;
 	TSharedPtr< FUICommandInfo > PreviewPlatformOverride_AndroidVulkanES31;
 	TSharedPtr< FUICommandInfo > PreviewPlatformOverride_IOSMetalES31;
@@ -816,11 +814,8 @@ public:
 	static bool IsFeatureLevelPreviewEnabled();
 	static bool IsFeatureLevelPreviewActive();
 	static bool IsPreviewModeButtonVisible();
-	static void SetPreviewPlatform(FName MaterialQualityPlatform,ERHIFeatureLevel::Type PreviewFeatureLevel);
-	static bool IsPreviewPlatformChecked(FName MaterialQualityPlatform, ERHIFeatureLevel::Type PreviewFeatureLevel);
-	static void SetFeatureLevelPreview(ERHIFeatureLevel::Type InFeatureLevel);
-	static bool IsFeatureLevelPreviewChecked(ERHIFeatureLevel::Type InFeatureLevel);
-	static bool IsFeatureLevelPreviewAvailable(ERHIFeatureLevel::Type InFeatureLevel);
+	static void SetPreviewPlatform(FPreviewPlatformInfo NewPreviewPlatform);
+	static bool IsPreviewPlatformChecked(FPreviewPlatformInfo NewPreviewPlatform);
 	static void GeometryCollection_SelectAllGeometry();
 	static void GeometryCollection_SelectNone();
 	static void GeometryCollection_SelectInverseGeometry();
diff --git a/Engine/Source/Editor/MaterialEditor/Private/SMaterialEditorViewport.cpp b/Engine/Source/Editor/MaterialEditor/Private/SMaterialEditorViewport.cpp
--- a/Engine/Source/Editor/MaterialEditor/Private/SMaterialEditorViewport.cpp
+++ b/Engine/Source/Editor/MaterialEditor/Private/SMaterialEditorViewport.cpp
@@ -394,7 +394,7 @@ bool SMaterialEditor3DPreviewViewport::SetPreviewAsset(UObject* InAsset)
 	// Add the new component to the scene
 	if (PreviewMeshComponent != nullptr)
 	{
-		if (GEditor->PreviewFeatureLevel <= ERHIFeatureLevel::ES3_1)
+		if (GEditor->PreviewPlatform.GetEffectivePreviewFeatureLevel() <= ERHIFeatureLevel::ES3_1)
 		{
 			PreviewMeshComponent->SetMobility(EComponentMobility::Static);
 		}
diff --git a/Engine/Source/Editor/Persona/Private/PersonaToolkit.cpp b/Engine/Source/Editor/Persona/Private/PersonaToolkit.cpp
--- a/Engine/Source/Editor/Persona/Private/PersonaToolkit.cpp
+++ b/Engine/Source/Editor/Persona/Private/PersonaToolkit.cpp
@@ -119,7 +119,7 @@ void FPersonaToolkit::CreatePreviewScene(const FPersonaToolkitArgs& PersonaToolk
 
 			// Create the preview component
 			UDebugSkelMeshComponent* SkeletalMeshComponent = NewObject<UDebugSkelMeshComponent>(Actor);
-			if (GEditor->PreviewFeatureLevel <= ERHIFeatureLevel::ES3_1)
+			if (GEditor->PreviewPlatform.GetEffectivePreviewFeatureLevel() <= ERHIFeatureLevel::ES3_1)
 			{
 				SkeletalMeshComponent->SetMobility(EComponentMobility::Static);
 			}
diff --git a/Engine/Source/Editor/StaticMeshEditor/Private/SStaticMeshEditorViewport.cpp b/Engine/Source/Editor/StaticMeshEditor/Private/SStaticMeshEditorViewport.cpp
--- a/Engine/Source/Editor/StaticMeshEditor/Private/SStaticMeshEditorViewport.cpp
+++ b/Engine/Source/Editor/StaticMeshEditor/Private/SStaticMeshEditorViewport.cpp
@@ -58,7 +58,7 @@ void SStaticMeshEditorViewport::Construct(const FArguments& InArgs)
 	SEditorViewport::Construct( SEditorViewport::FArguments() );
 
 	PreviewMeshComponent = NewObject<UStaticMeshComponent>(GetTransientPackage(), NAME_None, RF_Transient );
-	if (GEditor->PreviewFeatureLevel <= ERHIFeatureLevel::ES3_1)
+	if (GEditor->PreviewPlatform.GetEffectivePreviewFeatureLevel() <= ERHIFeatureLevel::ES3_1)
 	{
 		PreviewMeshComponent->SetMobility(EComponentMobility::Static);
 	}
@@ -278,7 +278,7 @@ void SStaticMeshEditorViewport::UpdatePreviewMesh(UStaticMesh* InStaticMesh, boo
 	}
 
 	PreviewMeshComponent = NewObject<UStaticMeshComponent>();
-	if (GEditor->PreviewFeatureLevel <= ERHIFeatureLevel::ES3_1)
+	if (GEditor->PreviewPlatform.GetEffectivePreviewFeatureLevel() <= ERHIFeatureLevel::ES3_1)
 	{
 		PreviewMeshComponent->SetMobility(EComponentMobility::Static);
 	}
diff --git a/Engine/Source/Editor/UnrealEd/Classes/Editor/EditorEngine.h b/Engine/Source/Editor/UnrealEd/Classes/Editor/EditorEngine.h
--- a/Engine/Source/Editor/UnrealEd/Classes/Editor/EditorEngine.h
+++ b/Engine/Source/Editor/UnrealEd/Classes/Editor/EditorEngine.h
@@ -26,6 +26,7 @@
 
 #include "EditorSubsystem.h"
 #include "Subsystems/SubsystemCollection.h"
+#include "RHI.h"
 
 #include "EditorEngine.generated.h"
 
@@ -352,6 +353,50 @@ struct FPlayInEditorOverrides
 	TOptional<int32> NumberOfClients;
 };
 
+struct FPreviewPlatformInfo
+{
+	FPreviewPlatformInfo()
+	:	PreviewFeatureLevel(ERHIFeatureLevel::SM5)
+	,	bPreviewFeatureLevelActive(false)
+	{}
+
+	FPreviewPlatformInfo(ERHIFeatureLevel::Type InFeatureLevel, FName InPreviewShaderPlatformName = NAME_None, bool InbPreviewFeatureLevelActive = false)
+	:	PreviewFeatureLevel(InFeatureLevel)
+	,	PreviewShaderPlatformName(InPreviewShaderPlatformName)
+	,	bPreviewFeatureLevelActive(InbPreviewFeatureLevelActive)
+	{}
+
+	/** The feature level we should use when loading or creating a new world */
+	ERHIFeatureLevel::Type PreviewFeatureLevel;
+	
+	/** The shader platform to preview, or NAME_None if there is no shader preview platform */
+	FName PreviewShaderPlatformName;
+
+	/** Is feature level preview currently active */
+	bool bPreviewFeatureLevelActive;
+
+	/** Checks if two FPreviewPlatformInfos are for the same preview platform. Note, this does NOT compare the bPreviewFeatureLevelActive flag */
+	bool Matches(const FPreviewPlatformInfo& Other) const
+	{
+		return PreviewFeatureLevel == Other.PreviewFeatureLevel && PreviewShaderPlatformName == Other.PreviewShaderPlatformName;
+	}
+
+	/** Convert platform name like "Android", or NAME_None if none is set or the preview feature level is not active */
+	FName GetEffectivePreviewPlatformName() const
+	{
+		return (PreviewShaderPlatformName != NAME_None && bPreviewFeatureLevelActive) ?
+			ShaderPlatformToPlatformName(ShaderFormatToLegacyShaderPlatform(PreviewShaderPlatformName)) :
+			NAME_None;
+	}
+
+	/** returns the preview feature level if active, or GMaxRHIFeatureLevel otherwise */
+	ERHIFeatureLevel::Type GetEffectivePreviewFeatureLevel() const
+	{
+		return bPreviewFeatureLevelActive ? PreviewFeatureLevel : GMaxRHIFeatureLevel;
+	}
+
+};
+
 /**
  * Engine that drives the Editor.
  * Separate from UGameEngine because it may have much different functionality than desired for an instance of a game itself.
@@ -638,9 +683,9 @@ public:
 	/** The feature level we should use when loading or creating a new world */
 	ERHIFeatureLevel::Type DefaultWorldFeatureLevel;
 
-	/** The feature level we should use when loading or creating a new world */
-	ERHIFeatureLevel::Type PreviewFeatureLevel;
-
+	/** The feature level and platform we should use when loading or creating a new world */
+	FPreviewPlatformInfo PreviewPlatform;
+	
 	/** A delegate that is called when the preview feature level changes. Primarily used to switch a viewport's feature level. */
 	DECLARE_MULTICAST_DELEGATE_OneParam(FPreviewFeatureLevelChanged, ERHIFeatureLevel::Type);
 	FPreviewFeatureLevelChanged PreviewFeatureLevelChanged;
@@ -3052,7 +3097,7 @@ public:
 	void OnSceneMaterialsModified();
 
 	/** Call this function to change the feature level and to override the material quality platform of the editor and PIE worlds */
-	void SetPreviewPlatform(const FName MaterialQualityPlatform, ERHIFeatureLevel::Type InPreviewFeatureLevel, const bool bSaveSettings = true);
+	void SetPreviewPlatform(const FPreviewPlatformInfo& NewPreviewPlatform, bool bSaveSettings);
 
 	/** Toggle the feature level preview */
 	void ToggleFeatureLevelPreview();
@@ -3073,14 +3118,6 @@ public:
 	FPreviewFeatureLevelChanged& OnPreviewFeatureLevelChanged() { return PreviewFeatureLevelChanged; }
 
 protected:
-	/** Call this function to change the feature level of the editor and PIE worlds */
-	void SetFeatureLevelPreview(const ERHIFeatureLevel::Type InPreviewFeatureLevel);
-
-	/** call this function to change the feature level for all materials */
-	void SetMaterialsFeatureLevel(const ERHIFeatureLevel::Type InPreviewFeatureLevel);
-
-	/** call this to recompile the materials */
-	void AllMaterialsCacheResourceShadersForRendering(ERHIFeatureLevel::Type InPreviewFeatureLevel);
 
 	/** Function pair used to save and restore the global feature level */
 	void LoadEditorFeatureLevel();
diff --git a/Engine/Source/Editor/UnrealEd/Classes/Editor/EditorPerProjectUserSettings.h b/Engine/Source/Editor/UnrealEd/Classes/Editor/EditorPerProjectUserSettings.h
--- a/Engine/Source/Editor/UnrealEd/Classes/Editor/EditorPerProjectUserSettings.h
+++ b/Engine/Source/Editor/UnrealEd/Classes/Editor/EditorPerProjectUserSettings.h
@@ -149,17 +149,17 @@ public:
 	UPROPERTY(config)
 	int32 MaterialQualityLevel;
 
-	/** The name of the shader platform that will be used in the editor */
+	/** The feature level we should use when loading or creating a new world */
 	UPROPERTY(config)
-	FName PreviewShaderPlatformName;
+	int32 PreviewFeatureLevel;
 
-	/** whether or not the above PreviewShaderPlatformName is a platform that needs to override the material settings in UMaterialShaderQualitySettings */
+	/** The shader platform to preview, or NAME_None if there is no shader preview platform */
 	UPROPERTY(config)
-	bool bIsMaterialQualityOverridePlatform;
+	FName PreviewShaderFormatName;
 
-	/** whether or not the Feature Level Preview is active or whether it is GMaxRHIFeatureLevel is being shown instead */
+	/** Is feature level preview currently active */
 	UPROPERTY(config)
-	bool bIsFeatureLevelPreviewActive;
+	bool bPreviewFeatureLevelActive;
 
 public:
 
diff --git a/Engine/Source/Editor/UnrealEd/Private/EditorEngine.cpp b/Engine/Source/Editor/UnrealEd/Private/EditorEngine.cpp
--- a/Engine/Source/Editor/UnrealEd/Private/EditorEngine.cpp
+++ b/Engine/Source/Editor/UnrealEd/Private/EditorEngine.cpp
@@ -89,6 +89,7 @@
 #include "UObject/ReferenceChainSearch.h"
 #include "HAL/PlatformApplicationMisc.h"
 #include "IMediaModule.h"
+#include "Scalability.h"
 
 // needed for the RemotePropagator
 #include "AudioDevice.h"
@@ -348,12 +349,12 @@ UEditorEngine::UEditorEngine(const FObjectInitializer& ObjectInitializer)
 	bIsEndingPlay = false;
 	NumOnlinePIEInstances = 0;
 	DefaultWorldFeatureLevel = GMaxRHIFeatureLevel;
-	PreviewFeatureLevel = DefaultWorldFeatureLevel;
+	PreviewPlatform = FPreviewPlatformInfo(DefaultWorldFeatureLevel);
 
 	FCoreDelegates::OnFeatureLevelDisabled.AddLambda([this](int RHIType, const FName& PreviewPlatformName)
 		{
 			ERHIFeatureLevel::Type FeatureLevelTypeToDisable = (ERHIFeatureLevel::Type)RHIType;
-			if (PreviewFeatureLevel == FeatureLevelTypeToDisable)
+			if (PreviewPlatform.PreviewFeatureLevel == FeatureLevelTypeToDisable)
 			{
 				UMaterialShaderQualitySettings* MaterialShaderQualitySettings = UMaterialShaderQualitySettings::Get();
 				if (MaterialShaderQualitySettings->GetPreviewPlatform() != PreviewPlatformName)
@@ -361,7 +362,7 @@ UEditorEngine::UEditorEngine(const FObjectInitializer& ObjectInitializer)
 					return;
 				}
 				
-				SetPreviewPlatform(FName(), ERHIFeatureLevel::SM5);
+				SetPreviewPlatform(FPreviewPlatformInfo(ERHIFeatureLevel::SM5), false);
 			}
 		});
 		
@@ -7331,99 +7332,96 @@ void UEditorEngine::OnSceneMaterialsModified()
 {
 }
 
-void UEditorEngine::SetMaterialsFeatureLevel(const ERHIFeatureLevel::Type InFeatureLevel)
+void UEditorEngine::SetPreviewPlatform(const FPreviewPlatformInfo& NewPreviewPlatform, bool bSaveSettings)
 {
-	FScopedSlowTask SlowTask(100.f, NSLOCTEXT("Engine", "UpdatingMaterialsMessage", "Updating Materials"), true);
-	SlowTask.Visibility = ESlowTaskVisibility::ForceVisible;
-	SlowTask.MakeDialog();
-
-	//invalidate global bound shader states so they will be created with the new shaders the next time they are set (in SetGlobalBoundShaderState)
-	for (TLinkedList<FGlobalBoundShaderStateResource*>::TIterator It(FGlobalBoundShaderStateResource::GetGlobalBoundShaderStateList()); It; It.Next())
-	{
-		BeginUpdateResourceRHI(*It);
-	}
-
-	FGlobalComponentReregisterContext RecreateComponents;
-	FlushRenderingCommands();
-
-	// Clear all required global feature levels, we only require the preview feature level.
-	for (uint32 i = (uint32)ERHIFeatureLevel::ES2; i < (uint32)ERHIFeatureLevel::Num; i++)
+#if RHI_RAYTRACING
+	if (IsRayTracingEnabled())
 	{
-		UMaterialInterface::SetGlobalRequiredFeatureLevel((ERHIFeatureLevel::Type)i, false);
+		if (PreviewPlatform.PreviewFeatureLevel != ERHIFeatureLevel::SM5)
+		{
+			UE_LOG(LogEditor, Warning, TEXT("Preview feature level is incompatible with ray tracing, defaulting to Shader Model 5"));
+			PreviewPlatform.PreviewFeatureLevel = ERHIFeatureLevel::SM5;
+		}
+		return;
 	}
+#endif
 
-	UMaterialInterface::SetGlobalRequiredFeatureLevel(InFeatureLevel, true);
+	// If we have specified a MaterialQualityPlatform ensure its feature level matches the requested feature level.
+	check(NewPreviewPlatform.PreviewShaderPlatformName.IsNone() || GetMaxSupportedFeatureLevel(ShaderFormatToLegacyShaderPlatform(NewPreviewPlatform.PreviewShaderPlatformName)) == NewPreviewPlatform.PreviewFeatureLevel);
 
-	SlowTask.EnterProgressFrame(35.0f);
-	UMaterial::AllMaterialsCacheResourceShadersForRendering(true);
-	
-	SlowTask.EnterProgressFrame(35.0f);
-	UMaterialInstance::AllMaterialsCacheResourceShadersForRendering(true);
+	const bool bChangedPreviewShaderPlatform = NewPreviewPlatform.PreviewShaderPlatformName != PreviewPlatform.PreviewShaderPlatformName;
+	const bool bChangedFeatureLevel = NewPreviewPlatform.PreviewFeatureLevel != PreviewPlatform.PreviewFeatureLevel ||
+		NewPreviewPlatform.bPreviewFeatureLevelActive != NewPreviewPlatform.bPreviewFeatureLevelActive;
+	const ERHIFeatureLevel::Type EffectiveFeatureLevel = NewPreviewPlatform.GetEffectivePreviewFeatureLevel();
 
-	SlowTask.EnterProgressFrame(15.0f, NSLOCTEXT("Engine", "SlowTaskGlobalShaderMapMessage", "Compiling global shaders"));
-	CompileGlobalShaderMap(InFeatureLevel);
 
-	SlowTask.EnterProgressFrame(15.0f, NSLOCTEXT("Engine", "SlowTaskFinalizingMessage", "Finalizing"));
-	GShaderCompilingManager->ProcessAsyncResults(false, true);
+	// Record the new preview platform
+	PreviewPlatform = NewPreviewPlatform;
 
-	PreviewFeatureLevelChanged.Broadcast(InFeatureLevel);
-}
+	// If we changed the preview platform, we need to update the material quality settings
+	if (bChangedPreviewShaderPlatform)
+	{
+		UMaterialShaderQualitySettings* MaterialShaderQualitySettings = UMaterialShaderQualitySettings::Get();
+		MaterialShaderQualitySettings->SetPreviewPlatform(PreviewPlatform.PreviewShaderPlatformName);
+	}
 
-void UEditorEngine::SetFeatureLevelPreview(const ERHIFeatureLevel::Type InPreviewFeatureLevel)
-{
-	if (PreviewFeatureLevel != InPreviewFeatureLevel)
+	if (bChangedFeatureLevel)
 	{
-		// Record this feature level as we want to use it for all subsequent level creation and loading
-		PreviewFeatureLevel = InPreviewFeatureLevel;
+		FScopedSlowTask SlowTask(100.f, NSLOCTEXT("Engine", "ChangingPreviewPlatform", "Changing Preview Platform"), true);
+		SlowTask.Visibility = ESlowTaskVisibility::ForceVisible;
+		SlowTask.MakeDialog();
 
-		// first change the feature level for global/shared resources
-		SetMaterialsFeatureLevel(InPreviewFeatureLevel);
-	}
-}
+		//invalidate global bound shader states so they will be created with the new shaders the next time they are set (in SetGlobalBoundShaderState)
+		for (TLinkedList<FGlobalBoundShaderStateResource*>::TIterator It(FGlobalBoundShaderStateResource::GetGlobalBoundShaderStateList()); It; It.Next())
+		{
+			BeginUpdateResourceRHI(*It);
+		}
 
-void UEditorEngine::AllMaterialsCacheResourceShadersForRendering(ERHIFeatureLevel::Type InPreviewFeatureLevel)
-{
-	FGlobalComponentRecreateRenderStateContext Recreate;
-	FlushRenderingCommands();
-	UMaterial::AllMaterialsCacheResourceShadersForRendering(true);
-	UMaterialInstance::AllMaterialsCacheResourceShadersForRendering(true);
-	PreviewFeatureLevelChanged.Broadcast(InPreviewFeatureLevel);
-}
+		FGlobalComponentReregisterContext RecreateComponents;
+		FlushRenderingCommands();
 
-void UEditorEngine::SetPreviewPlatform(const FName MaterialQualityPlatform, ERHIFeatureLevel::Type InPreviewFeatureLevel, const bool bSaveSettings/* = true*/)
-{
-#if RHI_RAYTRACING
-	if (IsRayTracingEnabled())
-	{
-		if (PreviewFeatureLevel != ERHIFeatureLevel::SM5)
+		// Set only require the preview feature level and the max feature level. The Max feature level is required for the toggle feature.
+		for (uint32 i = (uint32)ERHIFeatureLevel::ES2; i < (uint32)ERHIFeatureLevel::Num; i++)
 		{
-			UE_LOG(LogEditor, Warning, TEXT("Preview feature level is incompatible with ray tracing, defaulting to Shader Model 5"));
-			PreviewFeatureLevel = ERHIFeatureLevel::SM5;
+			ERHIFeatureLevel::Type FeatureLevel = (ERHIFeatureLevel::Type)i;
+			UMaterialInterface::SetGlobalRequiredFeatureLevel(FeatureLevel, FeatureLevel == PreviewPlatform.PreviewFeatureLevel || FeatureLevel == GMaxRHIFeatureLevel);
 		}
-	}
-#endif
-	// If we have specified a MaterialQualityPlatform ensure its feature level matches the requested feature level.
-	check(MaterialQualityPlatform.IsNone() || GetMaxSupportedFeatureLevel(ShaderFormatToLegacyShaderPlatform(MaterialQualityPlatform)) == InPreviewFeatureLevel);
 
-	UMaterialShaderQualitySettings* MaterialShaderQualitySettings = UMaterialShaderQualitySettings::Get();
-	const FName InitialPreviewPlatform = MaterialShaderQualitySettings->GetPreviewPlatform();
-	MaterialShaderQualitySettings->SetPreviewPlatform(MaterialQualityPlatform);
+		SlowTask.EnterProgressFrame(35.0f);
+		UMaterial::AllMaterialsCacheResourceShadersForRendering(true);
 
-	// Force activation of the feature level preview, because it may have been inactive before.
- 	auto* Settings = GetMutableDefault<UEditorPerProjectUserSettings>();
-	Settings->bIsFeatureLevelPreviewActive = true;
+		SlowTask.EnterProgressFrame(35.0f);
+		UMaterialInstance::AllMaterialsCacheResourceShadersForRendering(true);
 
-	if (PreviewFeatureLevel != InPreviewFeatureLevel)
-	{
-		// a new feature level will recompile the materials and apply the effect of any 'material quality platform'
-		SetFeatureLevelPreview(InPreviewFeatureLevel);
+		SlowTask.EnterProgressFrame(15.0f, NSLOCTEXT("Engine", "SlowTaskGlobalShaderMapMessage", "Compiling global shaders"));
+		CompileGlobalShaderMap(PreviewPlatform.PreviewFeatureLevel);
+
+		SlowTask.EnterProgressFrame(15.0f, NSLOCTEXT("Engine", "SlowTaskFinalizingMessage", "Finalizing"));
+		GShaderCompilingManager->ProcessAsyncResults(false, true);
+
+		DefaultWorldFeatureLevel = EffectiveFeatureLevel;
+		PreviewFeatureLevelChanged.Broadcast(EffectiveFeatureLevel);
+
+		// The feature level changed, so existing debug view materials are invalid and need to be rebuilt.
+		// This process must follow the PreviewFeatureLevelChanged event, because any listeners need
+		// opportunity to switch to the new feature level first.
+		void ClearDebugViewMaterials(UMaterialInterface*);
+		ClearDebugViewMaterials(nullptr);
+
+		UpdateShaderComplexityMaterials(true);
 	}
-	else if (InitialPreviewPlatform != MaterialQualityPlatform)
+	else if (bChangedPreviewShaderPlatform)
 	{
 		// Rebuild materials if we have the same feature level but a different 'material quality platform'
-		AllMaterialsCacheResourceShadersForRendering(InPreviewFeatureLevel);
+		FGlobalComponentRecreateRenderStateContext Recreate;
+		FlushRenderingCommands();
+		UMaterial::AllMaterialsCacheResourceShadersForRendering(true);
+		UMaterialInstance::AllMaterialsCacheResourceShadersForRendering(true);
+		PreviewFeatureLevelChanged.Broadcast(EffectiveFeatureLevel);
 	}
 
+	Scalability::ChangeScalabilityPreviewPlatform(PreviewPlatform.GetEffectivePreviewPlatformName());
+
 	if (bSaveSettings)
 	{
 		SaveEditorFeatureLevel();
@@ -7432,13 +7430,15 @@ void UEditorEngine::SetPreviewPlatform(const FName MaterialQualityPlatform, ERHI
 
 void UEditorEngine::ToggleFeatureLevelPreview()
 {
- 	auto* Settings = GetMutableDefault<UEditorPerProjectUserSettings>();
- 	Settings->bIsFeatureLevelPreviewActive ^= 1;
+ 	PreviewPlatform.bPreviewFeatureLevelActive ^= 1;
 
-	ERHIFeatureLevel::Type NewPreviewFeatureLevel = Settings->bIsFeatureLevelPreviewActive ? PreviewFeatureLevel : GMaxRHIFeatureLevel;
+	ERHIFeatureLevel::Type NewPreviewFeatureLevel = PreviewPlatform.GetEffectivePreviewFeatureLevel();
 
+	DefaultWorldFeatureLevel = NewPreviewFeatureLevel;
 	PreviewFeatureLevelChanged.Broadcast(NewPreviewFeatureLevel);
 
+	Scalability::ChangeScalabilityPreviewPlatform(PreviewPlatform.GetEffectivePreviewPlatformName());
+
 	GEditor->RedrawAllViewports();
 	
 	SaveEditorFeatureLevel();
@@ -7446,64 +7446,34 @@ void UEditorEngine::ToggleFeatureLevelPreview()
 
 bool UEditorEngine::IsFeatureLevelPreviewEnabled() const
 {
-	return PreviewFeatureLevel != GMaxRHIFeatureLevel;
+	return PreviewPlatform.PreviewFeatureLevel != GMaxRHIFeatureLevel || PreviewPlatform.PreviewShaderPlatformName != NAME_None;
 }
 
 bool UEditorEngine::IsFeatureLevelPreviewActive() const
 {
- 	auto* Settings = GetMutableDefault<UEditorPerProjectUserSettings>();
- 	return Settings->bIsFeatureLevelPreviewActive;
+ 	return PreviewPlatform.bPreviewFeatureLevelActive;
 }
 
 ERHIFeatureLevel::Type UEditorEngine::GetActiveFeatureLevelPreviewType() const
 {
- 	auto* Settings = GetMutableDefault<UEditorPerProjectUserSettings>();
-	return Settings->bIsFeatureLevelPreviewActive ? PreviewFeatureLevel : GMaxRHIFeatureLevel;
+	return PreviewPlatform.bPreviewFeatureLevelActive ? PreviewPlatform.PreviewFeatureLevel : GMaxRHIFeatureLevel;
 }
 
 void UEditorEngine::LoadEditorFeatureLevel()
 {
- 	auto* Settings = GetMutableDefault<UEditorPerProjectUserSettings>();
-	auto* QualitySettings = UMaterialShaderQualitySettings::Get();
-
-	EShaderPlatform ShaderPlatform = ShaderFormatToLegacyShaderPlatform(Settings->PreviewShaderPlatformName);
-	if (ShaderPlatform != SP_NumPlatforms)
+	auto* Settings = GetMutableDefault<UEditorPerProjectUserSettings>();
+	if (Settings->PreviewFeatureLevel >= 0 && Settings->PreviewFeatureLevel < (int32)ERHIFeatureLevel::Num)
 	{
-		FName MaterialQualityPlatform = NAME_None;
-		const ERHIFeatureLevel::Type FeatureLevel = GetMaxSupportedFeatureLevel(ShaderPlatform);
-
-		if (Settings->bIsMaterialQualityOverridePlatform)
-		{
-			MaterialQualityPlatform = Settings->PreviewShaderPlatformName;
-			UMaterialShaderQualitySettings::Get()->GetShaderPlatformQualitySettings(Settings->PreviewShaderPlatformName);
-		}
-
-		// Save off bIsFeatureLevelPreviewActive, because SetPreviewPlatform() will force it on.
-		// We want it to be what the user chose.
-		bool bSaveIsFeatureLevelPreviewActive = Settings->bIsFeatureLevelPreviewActive;
-		SetPreviewPlatform(MaterialQualityPlatform, FeatureLevel, false);
-		Settings->bIsFeatureLevelPreviewActive = bSaveIsFeatureLevelPreviewActive;
+		SetPreviewPlatform(FPreviewPlatformInfo((ERHIFeatureLevel::Type)Settings->PreviewFeatureLevel, Settings->PreviewShaderFormatName, Settings->bPreviewFeatureLevelActive), false);
 	}
 }
 
 void UEditorEngine::SaveEditorFeatureLevel()
 {
 	auto* Settings = GetMutableDefault<UEditorPerProjectUserSettings>();
-
-	const FName& PreviewPlatormName = UMaterialShaderQualitySettings::Get()->GetPreviewPlatform();
-	if (PreviewPlatormName == NAME_None)
-	{
-		Settings->bIsMaterialQualityOverridePlatform = false;
-
-		const EShaderPlatform ShaderPlatform = GetFeatureLevelShaderPlatform(PreviewFeatureLevel);
-		Settings->PreviewShaderPlatformName = LegacyShaderPlatformToShaderFormat(ShaderPlatform);
-	}
-	else
-	{
-		Settings->bIsMaterialQualityOverridePlatform = true;
-		Settings->PreviewShaderPlatformName = PreviewPlatormName;
-	}
-
+	Settings->PreviewFeatureLevel = (int32)PreviewPlatform.PreviewFeatureLevel;
+	Settings->PreviewShaderFormatName = PreviewPlatform.PreviewShaderPlatformName;
+	Settings->bPreviewFeatureLevelActive = PreviewPlatform.bPreviewFeatureLevelActive;
 	Settings->PostEditChange();
 }
 
diff --git a/Engine/Source/Editor/UnrealEd/Private/PlayLevel.cpp b/Engine/Source/Editor/UnrealEd/Private/PlayLevel.cpp
--- a/Engine/Source/Editor/UnrealEd/Private/PlayLevel.cpp
+++ b/Engine/Source/Editor/UnrealEd/Private/PlayLevel.cpp
@@ -3131,7 +3131,7 @@ UGameInstance* UEditorEngine::CreatePIEGameInstance(int32 InPIEInstance, bool bI
 	GameInstanceParams.bSimulateInEditor = bInSimulateInEditor;
 	GameInstanceParams.bStartInSpectatorMode = bStartInSpectatorMode;
 	GameInstanceParams.bRunAsDedicated = bRunAsDedicated;
-	GameInstanceParams.WorldFeatureLevel = PreviewFeatureLevel;
+	GameInstanceParams.WorldFeatureLevel = PreviewPlatform.GetEffectivePreviewFeatureLevel();
 
 	const FGameInstancePIEResult InitializeResult = GameInstance->InitializeForPlayInEditor(InPIEInstance, GameInstanceParams);
 	if (!InitializeResult.IsSuccess())
@@ -3163,7 +3163,7 @@ UGameInstance* UEditorEngine::CreatePIEGameInstance(int32 InPIEInstance, bool bI
 	FFormatNamedArguments Args;
 	Args.Add( TEXT("GameName"), FText::FromString( FString( WindowTitleOverride.IsEmpty() ? FApp::GetProjectName() : WindowTitleOverride.ToString() ) ) );
 	Args.Add( TEXT("PlatformBits"), FText::FromString( PlatformBitsString ) );
-	Args.Add( TEXT("RHIName"), FText::FromName( LegacyShaderPlatformToShaderFormat( GShaderPlatformForFeatureLevel[PreviewFeatureLevel] ) ) );
+	Args.Add( TEXT("RHIName"), FText::FromName( LegacyShaderPlatformToShaderFormat( GShaderPlatformForFeatureLevel[PreviewPlatform.PreviewFeatureLevel] ) ) );
 
 	const ULevelEditorPlaySettings* PlayInSettings = GetDefault<ULevelEditorPlaySettings>();
 	const EPlayNetMode PlayNetMode = [&PlayInSettings]{ EPlayNetMode NetMode(PIE_Standalone); return (PlayInSettings->GetPlayNetMode(NetMode) ? NetMode : PIE_Standalone); }();
diff --git a/Engine/Source/Editor/UnrealEd/Private/Settings/EditorPerProjectUserSettings.cpp b/Engine/Source/Editor/UnrealEd/Private/Settings/EditorPerProjectUserSettings.cpp
--- a/Engine/Source/Editor/UnrealEd/Private/Settings/EditorPerProjectUserSettings.cpp
+++ b/Engine/Source/Editor/UnrealEd/Private/Settings/EditorPerProjectUserSettings.cpp
@@ -26,6 +26,9 @@ UEditorPerProjectUserSettings::UEditorPerProjectUserSettings(const FObjectInitia
 	SwarmNumOfConcurrentJobs = 16;
 	SwarmMaxUploadChunkSizeInMB = 100;
 	SwarmIntermediateFolder = FPaths::ConvertRelativePathToFull(FPaths::ProjectIntermediateDir() + TEXT("Simplygon/"));
+	PreviewFeatureLevel = (int32)ERHIFeatureLevel::SM5;
+	PreviewShaderFormatName = NAME_None;
+	bPreviewFeatureLevelActive = false;
 }
 
 void UEditorPerProjectUserSettings::PostInitProperties()
diff --git a/Engine/Source/Runtime/Core/Private/Misc/ConfigCacheIni.cpp b/Engine/Source/Runtime/Core/Private/Misc/ConfigCacheIni.cpp
--- a/Engine/Source/Runtime/Core/Private/Misc/ConfigCacheIni.cpp
+++ b/Engine/Source/Runtime/Core/Private/Misc/ConfigCacheIni.cpp
@@ -4078,6 +4078,25 @@ void ApplyCVarSettingsFromIni(const TCHAR* InSectionName, const TCHAR* InIniFile
 	}
 }
 
+void ForEachCVarInSectionFromIni(const TCHAR* InSectionName, const TCHAR* InIniFilename, TFunction<void(IConsoleVariable* CVar, const FString& KeyString, const FString& ValueString)> InEvaluationFunction)
+{
+	if (FConfigSection* Section = GConfig->GetSectionPrivate(InSectionName, false, true, InIniFilename))
+	{
+		for (FConfigSectionMap::TConstIterator It(*Section); It; ++It)
+		{
+			const FString& KeyString = It.Key().GetPlainNameString();
+			const FString& ValueString = ConvertValueFromHumanFriendlyValue(*It.Value().GetValue());
+
+			IConsoleVariable* CVar = IConsoleManager::Get().FindConsoleVariable(*KeyString);
+			if (CVar)
+			{
+				InEvaluationFunction(CVar, KeyString, ValueString);
+			}
+		}
+	}
+}
+
+
 void ApplyCVarSettingsGroupFromIni(const TCHAR* InSectionBaseName, int32 InGroupNumber, const TCHAR* InIniFilename, uint32 SetBy)
 {
 	// Lookup the config section for this section and group number
diff --git a/Engine/Source/Runtime/Core/Public/Misc/ConfigCacheIni.h b/Engine/Source/Runtime/Core/Public/Misc/ConfigCacheIni.h
--- a/Engine/Source/Runtime/Core/Public/Misc/ConfigCacheIni.h
+++ b/Engine/Source/Runtime/Core/Public/Misc/ConfigCacheIni.h
@@ -853,22 +853,10 @@ private:
 	EConfigCacheType Type;
 };
 
-/**
- * Helper function to read the contents of an ini file and a specified group of cvar parameters, where sections in the ini file are marked [InName@InGroupNumber]
- * @param InSectionBaseName - The base name of the section to apply cvars from (i.e. the bit before the @)
- * @param InGroupNumber - The group number required
- * @param InIniFilename - The ini filename
- * @param SetBy anything in ECVF_LastSetMask e.g. ECVF_SetByScalability
- */
+UE_DEPRECATED(4.24, "This functionality to generate Scalability@Level section string has been moved to Scalability.cpp. Explictly construct section you need manually.")
 CORE_API void ApplyCVarSettingsGroupFromIni(const TCHAR* InSectionBaseName, int32 InGroupNumber, const TCHAR* InIniFilename, uint32 SetBy);
 
-/**
-* Helper function to read the contents of an ini file and a specified group of cvar parameters, where sections in the ini file are marked [InName@TagName]
-* @param InSectionBaseName - The base name of the section to apply cvars from (i.e. the bit before the @)
-* @param InSectionTag - The group name required. e.g. 'Cine'
-* @param InIniFilename - The ini filename
-* @param SetBy anything in ECVF_LastSetMask e.g. ECVF_SetByScalability
-*/
+UE_DEPRECATED(4.24, "This functionality to generate Scalability@Level section string has been moved to Scalability.cpp. Explictly construct section you need manually.")
 CORE_API void ApplyCVarSettingsGroupFromIni(const TCHAR* InSectionBaseName, const TCHAR* InSectionTag, const TCHAR* InIniFilename, uint32 SetBy);
 
 /**
@@ -879,7 +867,13 @@ CORE_API void ApplyCVarSettingsGroupFromIni(const TCHAR* InSectionBaseName, cons
  */
 CORE_API void ApplyCVarSettingsFromIni(const TCHAR* InSectionBaseName, const TCHAR* InIniFilename, uint32 SetBy, bool bAllowCheating = false);
 
-
+/**
+ * Helper function to operate a user defined function for each CVar key/value pair in the specified section in an ini file
+ * @param InSectionName - The name of the section to apply cvars from
+ * @param InIniFilename - The ini filename
+ * @param InEvaluationFunction - The evaluation function to be called for each key/value pair
+ */
+CORE_API void ForEachCVarInSectionFromIni(const TCHAR* InSectionName, const TCHAR* InIniFilename, TFunction<void(IConsoleVariable* CVar, const FString& KeyString, const FString& ValueString)> InEvaluationFunction);
 
 /**
  * CVAR Ini history records all calls to ApplyCVarSettingsFromIni and can re run them 
diff --git a/Engine/Source/Runtime/Engine/Private/Scalability.cpp b/Engine/Source/Runtime/Engine/Private/Scalability.cpp
--- a/Engine/Source/Runtime/Engine/Private/Scalability.cpp
+++ b/Engine/Source/Runtime/Engine/Private/Scalability.cpp
@@ -256,22 +256,130 @@ static void InferCurrentQualityLevel(const FString& InGroupName, int32& OutQuali
 	}
 }
 
-static void SetGroupQualityLevel(const TCHAR* InGroupName, int32 InQualityLevel, int32 InNumLevels)
+
+FString GetScalabilitySectionString(const TCHAR* InGroupName, int32 InQualityLevel, int32 InNumLevels)
 {
 	check(InNumLevels > 0);
 	const int32 MaxLevel = InNumLevels - 1;
 	InQualityLevel = FMath::Clamp(InQualityLevel, 0, MaxLevel);
+	FString Result;
+
+	if (InQualityLevel == MaxLevel)
+	{
+		Result = FString::Printf(TEXT("%s@Cine"), InGroupName);
+	}
+	else
+	{
+		Result = FString::Printf(TEXT("%s@%d"), InGroupName, InQualityLevel);
+	}
+
+	return Result;
+}
+
+#if WITH_EDITOR
+FName PlatformScalabilityName; // The name of the current platform scalability, or NAME_None if none is active
+FString PlatformScalabilityIniFilename;
+TMap<IConsoleVariable*, FString> PlatformScalabilityCVarBackup;
+TSet<const IConsoleVariable*> PlatformScalabilityCVarWhitelist;
+TSet<const IConsoleVariable*> PlatformScalabilityCVarBlacklist;
+
+void UndoPlatformScalability()
+{
+	// Reapply all CVars to the previous values  
+	for (auto It = PlatformScalabilityCVarBackup.CreateIterator(); It; ++It)
+	{
+		It.Key()->Set(*It.Value(), ECVF_SetByScalability);
+	}
+	PlatformScalabilityCVarBackup.Empty();
+}
+
+void ApplyScalabilityGroupFromPlatformIni(const TCHAR* InSectionName, const TCHAR* InIniFilename)
+{
+	UE_LOG(LogConfig, Log, TEXT("Applying CVar settings from Section [%s] File [%s]"), InSectionName, InIniFilename);
+
+	TFunction<void(IConsoleVariable*, const FString&, const FString&)> Func = [&](IConsoleVariable* CVar, const FString& KeyString, const FString& ValueString)
+	{
+		// Check the blacklist and whitelist
+		if (!PlatformScalabilityCVarBlacklist.Contains(CVar) && (PlatformScalabilityCVarWhitelist.Num() == 0 || PlatformScalabilityCVarWhitelist.Contains(CVar)))
+		{
+			// Backup cvar we're going to overwrite with the platform specific
+			if (!PlatformScalabilityCVarBackup.Contains(CVar))
+			{
+				PlatformScalabilityCVarBackup.Add(CVar, CVar->GetString());
+			}
+
+			// Apply the platform override
+			UE_LOG(LogConfig, Log, TEXT("Setting CVar [[%s:%s]]"), *KeyString, *ValueString);
+			CVar->Set(*ValueString, ECVF_SetByScalability);
+		}			   
+	};
+			
+	ForEachCVarInSectionFromIni(InSectionName, InIniFilename, Func);
+}
 
+void ChangeScalabilityPreviewPlatform(FName NewPlatformScalabilityName)
+{
+	if (PlatformScalabilityName != NAME_None)
+	{
+		// restore any modified CVar values and reapply the scalability settings for the default Editor platform
+		UndoPlatformScalability();
+		PlatformScalabilityName = NAME_None;
+		FQualityLevels State = Scalability::GetQualityLevels();
+		Scalability::SetQualityLevels(State);
+	}
+
+	if (NewPlatformScalabilityName != NAME_None)
+	{
+		PlatformScalabilityName = NewPlatformScalabilityName;
+		FString PlatformString = NewPlatformScalabilityName.ToString();
+		FConfigCacheIni::LoadGlobalIniFile(PlatformScalabilityIniFilename, TEXT("Scalability"), *PlatformString, true);
+
+		// load blacklist and whitelist of cvars we can set when previewing this platform
+		PlatformScalabilityCVarWhitelist.Empty();
+		TArray<FString> WhitelistCVarNames;
+		GConfig->GetArray(TEXT("ScalabilityPreview"), TEXT("WhitelistCVars"), WhitelistCVarNames, *PlatformScalabilityIniFilename);
+		for (const FString& CVarName : WhitelistCVarNames)
+		{
+			const IConsoleVariable* CVar = IConsoleManager::Get().FindConsoleVariable(*CVarName);
+			if (CVar)
+			{
+				PlatformScalabilityCVarWhitelist.Add(CVar);
+			}
+		}
+		PlatformScalabilityCVarBlacklist.Empty();
+		TArray<FString> BlacklistCVarNames;
+		GConfig->GetArray(TEXT("ScalabilityPreview"), TEXT("BlacklistCVars"), BlacklistCVarNames, *PlatformScalabilityIniFilename);
+		for (const FString& CVarName : BlacklistCVarNames)
+		{
+			const IConsoleVariable* CVar = IConsoleManager::Get().FindConsoleVariable(*CVarName);
+			if (CVar)
+			{
+				PlatformScalabilityCVarBlacklist.Add(CVar);
+			}
+		}
+		
+		// apply scalability
+		FQualityLevels State = Scalability::GetQualityLevels();
+		Scalability::SetQualityLevels(State);
+	}
+}
+#endif
+
+static void SetGroupQualityLevel(const TCHAR* InGroupName, int32 InQualityLevel, int32 InNumLevels)
+{
 //	UE_LOG(LogConsoleResponse, Display, TEXT("  %s %d"), InGroupName, InQualityLevel);
+	FString Section = GetScalabilitySectionString(InGroupName, InQualityLevel, InNumLevels);
 
-	if (InQualityLevel == MaxLevel)
+#if WITH_EDITOR
+	if (PlatformScalabilityName != NAME_None)
 	{
-		ApplyCVarSettingsGroupFromIni(InGroupName, TEXT("Cine"), *GScalabilityIni, ECVF_SetByScalability);
+		ApplyScalabilityGroupFromPlatformIni(*Section, *PlatformScalabilityIniFilename);
 	}
 	else
+#endif
 	{
-		ApplyCVarSettingsGroupFromIni(InGroupName, InQualityLevel, *GScalabilityIni, ECVF_SetByScalability);
-	}	
+		ApplyCVarSettingsFromIni(*Section, *GScalabilityIni, ECVF_SetByScalability);
+	}
 }
 
 float GetResolutionScreenPercentage()
@@ -813,5 +921,12 @@ FQualityLevels GetQualityLevelCounts()
 	return Result;
 }
 
+void LoadPlatformScalability(FString PlatformName)
+{
+	
+
+}
+
+
 }
 
diff --git a/Engine/Source/Runtime/Engine/Public/Scalability.h b/Engine/Source/Runtime/Engine/Public/Scalability.h
--- a/Engine/Source/Runtime/Engine/Public/Scalability.h
+++ b/Engine/Source/Runtime/Engine/Public/Scalability.h
@@ -160,4 +160,9 @@ namespace Scalability
 
 	/** Returns the current screen percentage */
 	ENGINE_API float GetResolutionScreenPercentage();
+
+#if WITH_EDITOR
+	/** Set an Editor preview scalability platform */
+	void ENGINE_API ChangeScalabilityPreviewPlatform(FName NewPlatformScalabilityName);
+#endif;
 }
diff --git a/Engine/Source/Runtime/Landscape/Private/Landscape.cpp b/Engine/Source/Runtime/Landscape/Private/Landscape.cpp
--- a/Engine/Source/Runtime/Landscape/Private/Landscape.cpp
+++ b/Engine/Source/Runtime/Landscape/Private/Landscape.cpp
@@ -1375,7 +1375,13 @@ FPrimitiveSceneProxy* ULandscapeComponent::CreateSceneProxy()
 	}
 	else // i.e. (FeatureLevel <= ERHIFeatureLevel::ES3_1)
 	{
-		if (PlatformData.HasValidRuntimeData())
+		if (!PlatformData.HasValidRuntimeData())
+		{
+#if WITH_EDITOR
+			CheckGenerateLandscapePlatformData(false, nullptr);
+#endif
+		}
+		else
 		{
 			Proxy = new FLandscapeComponentSceneProxyMobile(this);
 		}
diff --git a/Engine/Source/Runtime/RHI/Private/RHI.cpp b/Engine/Source/Runtime/RHI/Private/RHI.cpp
--- a/Engine/Source/Runtime/RHI/Private/RHI.cpp
+++ b/Engine/Source/Runtime/RHI/Private/RHI.cpp
@@ -684,34 +684,61 @@ static FName NAME_PLATFORM_IOS(TEXT("IOS"));
 static FName NAME_PLATFORM_MAC(TEXT("Mac"));
 static FName NAME_PLATFORM_SWITCH(TEXT("Switch"));
 static FName NAME_PLATFORM_TVOS(TEXT("TVOS"));
+static FName NAME_PLATFORM_HTML5(TEXT("HTML5"));
+static FName NAME_PLATFORM_LUMIN(TEXT("Lumin"));
+
 
 // @todo platplug: This is still here, only being used now by UMaterialShaderQualitySettings::GetOrCreatePlatformSettings
 // since I have moved the other uses to FindTargetPlatformWithSupport
 // But I'd like to delete it anyway!
 FName ShaderPlatformToPlatformName(EShaderPlatform Platform)
 {
-	switch(Platform)
+	switch (Platform)
 	{
+	case SP_PCD3D_SM5:
+	case SP_OPENGL_SM4:
+	case SP_OPENGL_PCES2:
 	case SP_PCD3D_SM4:
-	case SP_PCD3D_SM5: return NAME_PLATFORM_WINDOWS;
-	case SP_PS4: return NAME_PLATFORM_PS4;
-	case SP_XBOXONE_D3D12: return NAME_PLATFORM_XBOXONE;
+	case SP_OPENGL_SM5:
+	case SP_PCD3D_ES2:
+	case SP_PCD3D_ES3_1:
+	case SP_OPENGL_PCES3_1:
+	case SP_VULKAN_PCES3_1:
+	case SP_VULKAN_SM4:
+	case SP_VULKAN_SM5:
+		return NAME_PLATFORM_WINDOWS;
+	case SP_PS4:
+		return NAME_PLATFORM_PS4;
+	case SP_XBOXONE_D3D12:
+		return NAME_PLATFORM_XBOXONE;
+	case SP_OPENGL_ES2_ANDROID:
+	case SP_OPENGL_ES31_EXT:
+	case SP_VULKAN_ES3_1_ANDROID:
 	case SP_OPENGL_ES3_1_ANDROID:
-	case SP_VULKAN_ES3_1_ANDROID: return NAME_PLATFORM_ANDROID;
+		return NAME_PLATFORM_ANDROID;
+	case SP_OPENGL_ES2_WEBGL:
+		return NAME_PLATFORM_HTML5;
+	case SP_OPENGL_ES2_IOS:
 	case SP_METAL:
 	case SP_METAL_MRT:
-        return NAME_PLATFORM_IOS;
-	case SP_METAL_TVOS:
-	case SP_METAL_MRT_TVOS: return NAME_PLATFORM_TVOS;
+		return NAME_PLATFORM_IOS;
 	case SP_METAL_SM5:
 	case SP_METAL_SM5_NOTESS:
 	case SP_METAL_MACES3_1:
 	case SP_METAL_MACES2:
-	case SP_METAL_MRT_MAC: return NAME_PLATFORM_MAC;
+	case SP_METAL_MRT_MAC:
+		return NAME_PLATFORM_MAC;
 	case SP_SWITCH:
-	case SP_SWITCH_FORWARD: return NAME_PLATFORM_SWITCH;
-
-	default: return FName();
+	case SP_SWITCH_FORWARD:
+		return NAME_PLATFORM_SWITCH;
+	case SP_VULKAN_SM5_LUMIN:
+	case SP_VULKAN_ES3_1_LUMIN:
+		return NAME_PLATFORM_LUMIN;
+	case SP_METAL_TVOS:
+	case SP_METAL_MRT_TVOS:
+		return NAME_PLATFORM_TVOS;
+	default:
+		return NAME_None;
 	}
 }
 
