Fix merge error breaking CIS

diff --git a/Engine/Source/Editor/UnrealEd/Private/EditorEngine.cpp b/Engine/Source/Editor/UnrealEd/Private/EditorEngine.cpp
--- a/Engine/Source/Editor/UnrealEd/Private/EditorEngine.cpp
+++ b/Engine/Source/Editor/UnrealEd/Private/EditorEngine.cpp
@@ -7407,8 +7407,6 @@ void UEditorEngine::SetPreviewPlatform(const FPreviewPlatformInfo& NewPreviewPla
 		// opportunity to switch to the new feature level first.
 		void ClearDebugViewMaterials(UMaterialInterface*);
 		ClearDebugViewMaterials(nullptr);
-
-		UpdateShaderComplexityMaterials(true);
 	}
 	else if (bChangedPreviewShaderPlatform)
 	{
