Fix typo causing static analysis failure

diff --git a/Engine/Source/Editor/UnrealEd/Private/EditorEngine.cpp b/Engine/Source/Editor/UnrealEd/Private/EditorEngine.cpp
--- a/Engine/Source/Editor/UnrealEd/Private/EditorEngine.cpp
+++ b/Engine/Source/Editor/UnrealEd/Private/EditorEngine.cpp
@@ -7351,7 +7351,7 @@ void UEditorEngine::SetPreviewPlatform(const FPreviewPlatformInfo& NewPreviewPla
 
 	const bool bChangedPreviewShaderPlatform = NewPreviewPlatform.PreviewShaderPlatformName != PreviewPlatform.PreviewShaderPlatformName;
 	const bool bChangedFeatureLevel = NewPreviewPlatform.PreviewFeatureLevel != PreviewPlatform.PreviewFeatureLevel ||
-		NewPreviewPlatform.bPreviewFeatureLevelActive != NewPreviewPlatform.bPreviewFeatureLevelActive;
+		NewPreviewPlatform.bPreviewFeatureLevelActive != PreviewPlatform.bPreviewFeatureLevelActive;
 	const ERHIFeatureLevel::Type EffectiveFeatureLevel = NewPreviewPlatform.GetEffectivePreviewFeatureLevel();
 
 
