Ensure that objects created via duplication in PIE are correctly constructed

diff --git a/Engine/Source/Editor/UnrealEd/Private/Factories/EditorFactories.cpp b/Engine/Source/Editor/UnrealEd/Private/Factories/EditorFactories.cpp
--- a/Engine/Source/Editor/UnrealEd/Private/Factories/EditorFactories.cpp
+++ b/Engine/Source/Editor/UnrealEd/Private/Factories/EditorFactories.cpp
@@ -881,11 +881,6 @@ UObject* ULevelFactory::FactoryCreateText
 						SpawnInfo.Name = ActorUniqueName;
 						SpawnInfo.Template = Archetype;
 						SpawnInfo.SpawnCollisionHandlingOverride = ESpawnActorCollisionHandlingMethod::AlwaysSpawn;
-						if (GEditor->bIsSimulatingInEditor)
-						{
-							// During SIE, we don't want to run construction scripts on a BP until it is completely constructed
-							SpawnInfo.bDeferConstruction = true;
-						}
 						AActor* NewActor = World->SpawnActor( TempClass, nullptr, nullptr, SpawnInfo );
 						
 						if( NewActor )
