Avoid marking a blueprint as modified when the compiler renames a graph

diff --git a/Engine/Source/Editor/UnrealEd/Private/Kismet2/BlueprintEditorUtils.cpp b/Engine/Source/Editor/UnrealEd/Private/Kismet2/BlueprintEditorUtils.cpp
--- a/Engine/Source/Editor/UnrealEd/Private/Kismet2/BlueprintEditorUtils.cpp
+++ b/Engine/Source/Editor/UnrealEd/Private/Kismet2/BlueprintEditorUtils.cpp
@@ -2804,7 +2804,7 @@ void FBlueprintEditorUtils::RenameGraph(UEdGraph* Graph, const FString& NewNameS
 			}
 		}
 
-		if (!Blueprint->bIsRegeneratingOnLoad)
+		if (!Blueprint->bIsRegeneratingOnLoad && !Blueprint->bBeingCompiled)
 		{
 			FBlueprintEditorUtils::MarkBlueprintAsStructurallyModified(Blueprint);
 		}
