PR #5096: UE-64019: Use compare function for tooltip FText instead of != operator (Contributed by [REDACTED_PERSON])

diff --git a/Engine/Source/Editor/UnrealEd/Private/Kismet2/StructureEditorUtils.cpp b/Engine/Source/Editor/UnrealEd/Private/Kismet2/StructureEditorUtils.cpp
--- a/Engine/Source/Editor/UnrealEd/Private/Kismet2/StructureEditorUtils.cpp
+++ b/Engine/Source/Editor/UnrealEd/Private/Kismet2/StructureEditorUtils.cpp
@@ -645,7 +645,7 @@ FString FStructureEditorUtils::GetTooltip(const UUserDefinedStruct* Struct)
 bool FStructureEditorUtils::ChangeTooltip(UUserDefinedStruct* Struct, const FString& InTooltip)
 {
 	UUserDefinedStructEditorData* StructEditorData = Struct ? Cast<UUserDefinedStructEditorData>(Struct->EditorData) : NULL;
-	if (StructEditorData && (InTooltip != StructEditorData->ToolTip))
+	if (StructEditorData && StructEditorData->ToolTip.Compare(InTooltip) != 0)
 	{
 		const FScopedTransaction Transaction(LOCTEXT("ChangeTooltip", "Change UDS Tooltip"));
 		StructEditorData->Modify();
@@ -668,7 +668,7 @@ FString FStructureEditorUtils::GetVariableTooltip(const UUserDefinedStruct* Stru
 bool FStructureEditorUtils::ChangeVariableTooltip(UUserDefinedStruct* Struct, FGuid VarGuid, const FString& InTooltip)
 {
 	FStructVariableDescription* VarDesc = GetVarDescByGuid(Struct, VarGuid);
-	if (VarDesc && (InTooltip != VarDesc->ToolTip))
+	if (VarDesc && VarDesc->ToolTip.Compare(InTooltip) != 0)
 	{
 		const FScopedTransaction Transaction(LOCTEXT("ChangeVariableTooltip", "Change UDS Variable Tooltip"));
 		ModifyStructData(Struct);
