PR #5076: Correct icon when dragging an event dispatcher variable (Contributed by [REDACTED_PERSON])

diff --git a/Engine/Source/Editor/UnrealEd/Private/Kismet2/BlueprintEditorUtils.cpp b/Engine/Source/Editor/UnrealEd/Private/Kismet2/BlueprintEditorUtils.cpp
--- a/Engine/Source/Editor/UnrealEd/Private/Kismet2/BlueprintEditorUtils.cpp
+++ b/Engine/Source/Editor/UnrealEd/Private/Kismet2/BlueprintEditorUtils.cpp
@@ -8924,6 +8924,10 @@ const FSlateBrush* FBlueprintEditorUtils::GetIconFromPin( const FEdGraphPinType&
 			IconBrush = FEditorStyle::GetBrush(TEXT("Kismet.VariableList.SetTypeIcon"));
 		}
 	}
+	else if (PinType.PinCategory == UEdGraphSchema_K2::PC_MCDelegate)
+	{
+		IconBrush = FEditorStyle::GetBrush(TEXT("GraphEditor.Delegate_16x"));
+	}
 	else if( PinSubObject )
 	{
 		UClass* VarClass = FindObject<UClass>(ANY_PACKAGE, *PinSubObject->GetName());
