Copied 4369800 Anousack.Kitisa
Import/export FBX metadata on SkeletalMesh and Skeleton assets.

diff --git a/Engine/Source/Editor/UnrealEd/Private/Fbx/FbxFactory.cpp b/Engine/Source/Editor/UnrealEd/Private/Fbx/FbxFactory.cpp
--- a/Engine/Source/Editor/UnrealEd/Private/Fbx/FbxFactory.cpp
+++ b/Engine/Source/Editor/UnrealEd/Private/Fbx/FbxFactory.cpp
@@ -549,7 +549,7 @@ UObject* UFbxFactory::FactoryCreateFile
 							{
 								break;
 							}
-						
+
 							TArray<FbxNode*> SkelMeshNodeArray;
 							for (int32 j = 0; j < NodeArray.Num(); j++)
 							{
@@ -581,9 +581,13 @@ UObject* UFbxFactory::FactoryCreateFile
 							{
 								FName OutputName = FbxImporter->MakeNameForMesh(Name.ToString(), SkelMeshNodeArray[0]);
 
+								TArray<FbxNode*> SkeletonNodeArray;
+								FbxImporter->FillFbxSkeletonArray(RootNodeToImport, SkeletonNodeArray);
+
 								UnFbx::FFbxImporter::FImportSkeletalMeshArgs ImportSkeletalMeshArgs;
 								ImportSkeletalMeshArgs.InParent = InParent;
 								ImportSkeletalMeshArgs.NodeArray = SkelMeshNodeArray;
+								ImportSkeletalMeshArgs.BoneNodeArray = SkeletonNodeArray;
 								ImportSkeletalMeshArgs.Name = OutputName;
 								ImportSkeletalMeshArgs.Flags = Flags;
 								ImportSkeletalMeshArgs.TemplateImportData = ImportUI->SkeletalMeshImportData;
diff --git a/Engine/Source/Editor/UnrealEd/Private/Fbx/FbxMainExport.cpp b/Engine/Source/Editor/UnrealEd/Private/Fbx/FbxMainExport.cpp
--- a/Engine/Source/Editor/UnrealEd/Private/Fbx/FbxMainExport.cpp
+++ b/Engine/Source/Editor/UnrealEd/Private/Fbx/FbxMainExport.cpp
@@ -3415,7 +3415,7 @@ FbxNode* FFbxExporter::ExportCollisionMesh(const UStaticMesh* StaticMesh, const
 }
 #endif
 
-void ExportObjectMetadata(const UObject* ObjectToExport, FbxNode* Node)
+void FFbxExporter::ExportObjectMetadata(const UObject* ObjectToExport, FbxNode* Node)
 {
 	if (ObjectToExport && Node)
 	{
@@ -3431,6 +3431,9 @@ void ExportObjectMetadata(const UObject* ObjectToExport, FbxNode* Node)
 				FString TagAsString = MetadataIt.Key.ToString();
 				if (TagAsString.RemoveFromStart(MetadataPrefix))
 				{
+					// Remaining tag follows the format NodeName.PropertyName, so replace '.' with '_'
+					TagAsString.ReplaceInline(TEXT("."), TEXT("_"));
+
 					FbxProperty Property = FbxProperty::Create(Node, FbxStringDT, TCHAR_TO_UTF8(*TagAsString));
 					FbxString ValueString(TCHAR_TO_UTF8(*MetadataIt.Value));
 
diff --git a/Engine/Source/Editor/UnrealEd/Private/Fbx/FbxMainImport.cpp b/Engine/Source/Editor/UnrealEd/Private/Fbx/FbxMainImport.cpp
--- a/Engine/Source/Editor/UnrealEd/Private/Fbx/FbxMainImport.cpp
+++ b/Engine/Source/Editor/UnrealEd/Private/Fbx/FbxMainImport.cpp
@@ -25,6 +25,7 @@
 #include "EngineAnalytics.h"
 #include "AnalyticsEventAttribute.h"
 #include "Interfaces/IAnalyticsProvider.h"
+#include "UObject/MetaData.h"
 #include "UObject/UObjectGlobals.h"
 #include "UObject/Package.h"
 #include "AssetRegistryModule.h"
@@ -1699,6 +1700,19 @@ void FFbxImporter::FillFbxMeshArray(FbxNode* Node, TArray<FbxNode*>& outMeshArra
 	}
 }
 
+void FFbxImporter::FillFbxSkeletonArray(FbxNode* Node, TArray<FbxNode*>& OutMeshArray)
+{
+	if (Node->GetSkeleton())
+	{
+		OutMeshArray.Add(Node);
+	}
+
+	for (int32 ChildIndex = 0; ChildIndex < Node->GetChildCount(); ++ChildIndex)
+	{
+		FillFbxSkeletonArray(Node->GetChild(ChildIndex), OutMeshArray);
+	}
+}
+
 void FFbxImporter::FillFbxMeshAndLODGroupArray(FbxNode* Node, TArray<FbxNode*>& outLODGroupArray, TArray<FbxNode*>& outMeshArray)
 {
 	// Is this node an LOD group
@@ -2781,6 +2795,111 @@ FbxNode* FFbxImporter::RetrieveObjectFromName(const TCHAR* ObjectName, FbxNode*
 	return nullptr;
 }
 
+FString GetFbxPropertyStringValue(const FbxProperty& Property)
+{
+	FString ValueStr(TEXT("Unsupported type"));
+
+	FbxDataType DataType = Property.GetPropertyDataType();
+	switch (DataType.GetType())
+	{
+	case eFbxBool:
+	{
+		FbxBool BoolValue = Property.Get<FbxBool>();
+		ValueStr = LexToString(BoolValue);
+	}
+	break;
+	case eFbxInt:
+	{
+		FbxInt IntValue = Property.Get<FbxInt>();
+		ValueStr = LexToString(IntValue);
+	}
+	break;
+	case eFbxEnum:
+	{
+		FbxEnum EnumValue = Property.Get<FbxEnum>();
+		ValueStr = LexToString(EnumValue);
+	}
+	break;
+	case eFbxFloat:
+	{
+		FbxFloat FloatValue = Property.Get<FbxFloat>();
+		ValueStr = LexToString(FloatValue);
+	}
+	break;
+	case eFbxDouble:
+	{
+		FbxDouble DoubleValue = Property.Get<FbxDouble>();
+		ValueStr = LexToString(DoubleValue);
+	}
+	break;
+	case eFbxDouble2:
+	{
+		FbxDouble2 Vec = Property.Get<FbxDouble2>();
+		ValueStr = FString::Printf(TEXT("(%f, %f, %f, %f)"), Vec[0], Vec[1]);
+	}
+	break;
+	case eFbxDouble3:
+	{
+		FbxDouble3 Vec = Property.Get<FbxDouble3>();
+		ValueStr = FString::Printf(TEXT("(%f, %f, %f)"), Vec[0], Vec[1], Vec[2]);
+	}
+	break;
+	case eFbxDouble4:
+	{
+		FbxDouble4 Vec = Property.Get<FbxDouble4>();
+		ValueStr = FString::Printf(TEXT("(%f, %f, %f, %f)"), Vec[0], Vec[1], Vec[2], Vec[3]);
+	}
+	break;
+	case eFbxString:
+	{
+		FbxString StringValue = Property.Get<FbxString>();
+		ValueStr = UTF8_TO_TCHAR(StringValue.Buffer());
+	}
+	break;
+	default:
+		break;
+	}
+	return ValueStr;
+}
+
+void FFbxImporter::ImportNodeCustomProperties(UObject* Object, FbxNode* Node)
+{
+	if (!Object || !Node)
+	{
+		return;
+	}
+
+	// Import all custom user-defined FBX properties from the FBX node to the object metadata
+	FbxProperty CurrentProperty = Node->GetFirstProperty();
+	FString NodeName = UTF8_TO_TCHAR(Node->GetName());
+	static const FString MetadataPrefix(FBX_METADATA_PREFIX);
+	while (CurrentProperty.IsValid())
+	{
+		if (CurrentProperty.GetFlag(FbxPropertyFlags::eUserDefined))
+		{
+			// Prefix the FBX metadata tag to make it distinguishable from other metadata
+			// so that it can be exportable through FBX export
+			FString MetadataTag = UTF8_TO_TCHAR(CurrentProperty.GetName());
+			if (!MetadataTag.StartsWith(NodeName))
+			{
+				// Append the node name in the tag since all the metadata will be flattened on the Object
+				MetadataTag = NodeName + TEXT(".") + MetadataTag;
+			}
+			MetadataTag = MetadataPrefix + MetadataTag;
+
+			FString MetadataValue = GetFbxPropertyStringValue(CurrentProperty);
+			Object->GetOutermost()->GetMetaData()->SetValue(Object, *MetadataTag, *MetadataValue);
+		}
+		CurrentProperty = Node->GetNextProperty(CurrentProperty);
+	}
+
+	int NumChildren = Node->GetChildCount();
+	for (int i = 0; i < NumChildren; ++i)
+	{
+		ImportNodeCustomProperties(Object, Node->GetChild(i));
+	}
+}
+
 } // namespace UnFbx
 
 #undef LOCTEXT_NAMESPACE
diff --git a/Engine/Source/Editor/UnrealEd/Private/Fbx/FbxSceneImportFactory.cpp b/Engine/Source/Editor/UnrealEd/Private/Fbx/FbxSceneImportFactory.cpp
--- a/Engine/Source/Editor/UnrealEd/Private/Fbx/FbxSceneImportFactory.cpp
+++ b/Engine/Source/Editor/UnrealEd/Private/Fbx/FbxSceneImportFactory.cpp
@@ -1946,10 +1946,15 @@ UObject* UFbxSceneImportFactory::ImportOneSkeletalMesh(void* VoidRootNodeToImpor
 				break;
 			RootNodeInfo->AttributeInfo->SetOriginalImportPath(PackageName);
 			FName SkeletalMeshFName = FName(*SkeletalMeshName);
+
+			TArray<FbxNode*> SkeletonNodeArray;
+			FbxImporter->FillFbxSkeletonArray(RootNodeToImport, SkeletonNodeArray);
+
 			//Import the skeletal mesh
 			UnFbx::FFbxImporter::FImportSkeletalMeshArgs ImportSkeletalMeshArgs;
 			ImportSkeletalMeshArgs.InParent = Pkg;
 			ImportSkeletalMeshArgs.NodeArray = bUseSkelMeshNodePivotArray ? SkelMeshNodePivotArray : SkelMeshNodeArray;
+			ImportSkeletalMeshArgs.BoneNodeArray = SkeletonNodeArray;
 			ImportSkeletalMeshArgs.Name = SkeletalMeshFName;
 			ImportSkeletalMeshArgs.Flags = Flags;
 			ImportSkeletalMeshArgs.TemplateImportData = SkeletalMeshImportData;
diff --git a/Engine/Source/Editor/UnrealEd/Private/Fbx/FbxSkeletalMeshExport.cpp b/Engine/Source/Editor/UnrealEd/Private/Fbx/FbxSkeletalMeshExport.cpp
--- a/Engine/Source/Editor/UnrealEd/Private/Fbx/FbxSkeletalMeshExport.cpp
+++ b/Engine/Source/Editor/UnrealEd/Private/Fbx/FbxSkeletalMeshExport.cpp
@@ -12,6 +12,7 @@
 
 #include "FbxExporter.h"
 #include "Exporters/FbxExportOption.h"
+#include "UObject/MetaData.h"
 
 DEFINE_LOG_CATEGORY_STATIC(LogFbxSkeletalMeshExport, Log, All);
 
@@ -469,6 +470,58 @@ void FFbxExporter::ExportSkeletalMeshComponent(USkeletalMeshComponent* SkelMeshC
 	}
 }
 
+void ExportObjectMetadataToBones(const UObject* ObjectToExport, const TArray<FbxNode*>& Nodes)
+{
+	if (!ObjectToExport || Nodes.Num() == 0)
+	{
+		return;
+	}
+
+	// Retrieve the metadata map without creating it
+	const TMap<FName, FString>* MetadataMap = UMetaData::GetMapForObject(ObjectToExport);
+	if (MetadataMap)
+	{
+		// Map the nodes to their names for fast access
+		TMap<FString, FbxNode*> NameToNode;
+		for (FbxNode* Node : Nodes)
+		{
+			NameToNode.Add(FString(Node->GetName()), Node);
+		}
+
+		static const FString MetadataPrefix(FBX_METADATA_PREFIX);
+		for (const auto& MetadataIt : *MetadataMap)
+		{
+			// Export object metadata tags that are prefixed as FBX custom user-defined properties
+			// Remove the prefix since it's for Unreal use only (and '.' is considered an invalid character for user property names in DCC like Maya)
+			FString TagAsString = MetadataIt.Key.ToString();
+			if (TagAsString.RemoveFromStart(MetadataPrefix))
+			{
+				// Extract the node name from the metadata tag
+				FString NodeName;
+				int32 CharPos = INDEX_NONE;
+				if (TagAsString.FindChar(TEXT('.'), CharPos))
+				{
+					NodeName = TagAsString.Left(CharPos);
+
+					// The remaining part is the actual metadata tag
+					TagAsString = TagAsString.RightChop(CharPos + 1); // exclude the period
+				}
+
+				// Try to attach the metadata to its associated node by name
+				FbxNode** Node = NameToNode.Find(NodeName);
+				if (Node)
+				{
+					FbxProperty Property = FbxProperty::Create(*Node, FbxStringDT, TCHAR_TO_UTF8(*TagAsString));
+					FbxString ValueString(TCHAR_TO_UTF8(*MetadataIt.Value));
+
+					Property.Set(ValueString);
+					Property.ModifyFlag(FbxPropertyFlags::eUserDefined, true);
+				}
+			}
+		}
+	}
+}
+
 /**
  * Add the given skeletal mesh to the Fbx scene in preparation for exporting.  Makes all new nodes a child of the given node
  */
@@ -518,10 +571,14 @@ FbxNode* FFbxExporter::ExportSkeletalMeshToFbx(const USkeletalMesh* SkelMesh, co
 			TmpNodeNoTransform->RemoveChild(SkeletonRootNode);
 			ActorRootNode->AddChild(SkeletonRootNode);
 		}
+
+		ExportObjectMetadataToBones(SkelMesh->Skeleton, BoneNodes);
+
 		if (MeshRootNode)
 		{
 			TmpNodeNoTransform->RemoveChild(MeshRootNode);
 			ActorRootNode->AddChild(MeshRootNode);
+			ExportObjectMetadata(SkelMesh, MeshRootNode);
 		}
 
 		Scene->GetRootNode()->RemoveChild(TmpNodeNoTransform);
diff --git a/Engine/Source/Editor/UnrealEd/Private/Fbx/FbxSkeletalMeshImport.cpp b/Engine/Source/Editor/UnrealEd/Private/Fbx/FbxSkeletalMeshImport.cpp
--- a/Engine/Source/Editor/UnrealEd/Private/Fbx/FbxSkeletalMeshImport.cpp
+++ b/Engine/Source/Editor/UnrealEd/Private/Fbx/FbxSkeletalMeshImport.cpp
@@ -1869,6 +1869,18 @@ USkeletalMesh* UnFbx::FFbxImporter::ImportSkeletalMesh(FImportSkeletalMeshArgs &
 		}
 	}
 
+	// Import mesh metadata to SkeletalMesh
+	for (FbxNode* MeshNode : ImportSkeletalMeshArgs.NodeArray)
+	{
+		ImportNodeCustomProperties(SkeletalMesh, MeshNode);
+	}
+
+	// Import bone metadata to Skeleton
+	for (FbxNode* SkeletonNode : ImportSkeletalMeshArgs.BoneNodeArray)
+	{
+		ImportNodeCustomProperties(SkeletalMesh->Skeleton, SkeletonNode);
+	}
+
 	// ComponentContexts will now go out of scope, causing components to be re-registered
 
 	if (RecreateExistingRenderStateContext)
@@ -2245,10 +2257,14 @@ USkeletalMesh* UnFbx::FFbxImporter::ReimportSkeletalMesh(USkeletalMesh* Mesh, UF
 			FSkeletalMeshImportData OutData;
 			if (LODIndex == 0)
 			{
+				TArray<FbxNode*> SkeletonNodeArray;
+				FillFbxSkeletonArray(Scene->GetRootNode(), SkeletonNodeArray);
+
 				ImportMeshLodData.AddZeroed();
 				UnFbx::FFbxImporter::FImportSkeletalMeshArgs ImportSkeletalMeshArgs;
 				ImportSkeletalMeshArgs.InParent = Mesh->GetOuter();
 				ImportSkeletalMeshArgs.NodeArray = SkelMeshNodeArray;
+				ImportSkeletalMeshArgs.BoneNodeArray = SkeletonNodeArray;
 				ImportSkeletalMeshArgs.Name = *Mesh->GetName();
 				ImportSkeletalMeshArgs.Flags = RF_Public | RF_Standalone;
 				ImportSkeletalMeshArgs.TemplateImportData = TemplateImportData;
diff --git a/Engine/Source/Editor/UnrealEd/Private/Fbx/FbxStaticMeshImport.cpp b/Engine/Source/Editor/UnrealEd/Private/Fbx/FbxStaticMeshImport.cpp
--- a/Engine/Source/Editor/UnrealEd/Private/Fbx/FbxStaticMeshImport.cpp
+++ b/Engine/Source/Editor/UnrealEd/Private/Fbx/FbxStaticMeshImport.cpp
@@ -37,7 +37,6 @@
 #include "MeshAttributes.h"
 #include "IMeshBuilderModule.h"
 #include "Settings/EditorExperimentalSettings.h"
-#include "UObject/MetaData.h"
 
 #define LOCTEXT_NAMESPACE "FbxStaticMeshImport"
 
@@ -1331,103 +1330,6 @@ void UnFbx::FFbxImporter::VerifyGeometry(UStaticMesh* StaticMesh)
 	}
 }
 
-FString GetFbxPropertyStringValue(const FbxProperty& Property)
-{
-	FString ValueStr(TEXT("Unsupported type"));
-
-	FbxDataType DataType = Property.GetPropertyDataType();
-	switch (DataType.GetType())
-	{
-	case eFbxBool:
-		{
-			FbxBool BoolValue = Property.Get<FbxBool>();
-			ValueStr = LexToString(BoolValue);
-		}
-		break;
-	case eFbxInt:
-		{
-			FbxInt IntValue = Property.Get<FbxInt>();
-			ValueStr = LexToString(IntValue);
-		}
-		break;
-	case eFbxEnum:
-		{
-			FbxEnum EnumValue = Property.Get<FbxEnum>();
-			ValueStr = LexToString(EnumValue);
-		}
-		break;
-	case eFbxFloat:
-		{
-			FbxFloat FloatValue = Property.Get<FbxFloat>();
-			ValueStr = LexToString(FloatValue);
-		}
-		break;
-	case eFbxDouble:
-		{
-			FbxDouble DoubleValue = Property.Get<FbxDouble>();
-			ValueStr = LexToString(DoubleValue);
-		}
-		break;
-	case eFbxDouble2:
-		{
-			FbxDouble2 Vec = Property.Get<FbxDouble2>();
-			ValueStr = FString::Printf(TEXT("(%f, %f, %f, %f)"), Vec[0], Vec[1]);
-		}
-		break;
-	case eFbxDouble3:
-		{
-			FbxDouble3 Vec = Property.Get<FbxDouble3>();
-			ValueStr = FString::Printf(TEXT("(%f, %f, %f)"), Vec[0], Vec[1], Vec[2]);
-		}
-		break;
-	case eFbxDouble4:
-		{
-			FbxDouble4 Vec = Property.Get<FbxDouble4>();
-			ValueStr = FString::Printf(TEXT("(%f, %f, %f, %f)"), Vec[0], Vec[1], Vec[2], Vec[3]);
-		}
-		break;
-	case eFbxString:
-		{
-			FbxString StringValue = Property.Get<FbxString>();
-			ValueStr = UTF8_TO_TCHAR(StringValue.Buffer());
-	}
-		break;
-	default:
-		break;
-	}
-	return ValueStr;
-}
-
-void ImportNodeCustomProperties(UObject* Object, FbxNode* Node)
-{
-	if (Object && Node)
-	{
-		// Import all custom user-defined FBX properties from the FBX node to the object metadata
-		FbxProperty CurrentProperty = Node->GetFirstProperty();
-		static const FString MetadataPrefix(FBX_METADATA_PREFIX);
-		while (CurrentProperty.IsValid())
-		{
-			if (CurrentProperty.GetFlag(FbxPropertyFlags::eUserDefined))
-			{
-				// Prefix the FBX metadata tag to make it distinguishable from other metadata
-				// so that it can be exportable through FBX export
-				FString MetadataTag = UTF8_TO_TCHAR(CurrentProperty.GetName());
-				MetadataTag = MetadataPrefix + MetadataTag;
-
-				FString MetadataValue = GetFbxPropertyStringValue(CurrentProperty);
-				Object->GetOutermost()->GetMetaData()->SetValue(Object, *MetadataTag, *MetadataValue);
-			}
-			CurrentProperty = Node->GetNextProperty(CurrentProperty);
-		}
-
-		int NumChildren = Node->GetChildCount();
-		for (int i = 0; i < NumChildren; ++i)
-		{
-			ImportNodeCustomProperties(Object, Node->GetChild(i));
-		}
-	}
-}
-
 UStaticMesh* UnFbx::FFbxImporter::ImportStaticMeshAsSingle(UObject* InParent, TArray<FbxNode*>& MeshNodeArray, const FName InName, EObjectFlags Flags, UFbxStaticMeshImportData* TemplateImportData, UStaticMesh* InStaticMesh, int LODIndex, void *ExistMeshDataPtr)
 {
 	struct ExistingStaticMeshData* ExistMeshData = (struct ExistingStaticMeshData*)ExistMeshDataPtr;
diff --git a/Engine/Source/Editor/UnrealEd/Private/FbxExporter.h b/Engine/Source/Editor/UnrealEd/Private/FbxExporter.h
--- a/Engine/Source/Editor/UnrealEd/Private/FbxExporter.h
+++ b/Engine/Source/Editor/UnrealEd/Private/FbxExporter.h
@@ -468,6 +468,9 @@ private:
 
 	/** Returns currently active FBX export options. Automation or UI diaog based options. */
 	UFbxExportOption* GetExportOptions();
+
+	/** Exports all the object's FBX metadata to the FBX node */
+	void ExportObjectMetadata(const UObject* ObjectToExport, FbxNode* Node);
 };
 
 
diff --git a/Engine/Source/Editor/UnrealEd/Public/FbxImporter.h b/Engine/Source/Editor/UnrealEd/Public/FbxImporter.h
--- a/Engine/Source/Editor/UnrealEd/Public/FbxImporter.h
+++ b/Engine/Source/Editor/UnrealEd/Public/FbxImporter.h
@@ -817,6 +817,7 @@ public:
 
 		UObject* InParent;
 		TArray<FbxNode*> NodeArray;
+		TArray<FbxNode*> BoneNodeArray;
 		FName Name;
 		EObjectFlags Flags;
 		UFbxSkeletalMeshImportData* TemplateImportData;
@@ -974,7 +975,15 @@ public:
 	* @param outMeshArray return Fbx meshes
 	*/
 	UNREALED_API void FillFbxMeshArray(FbxNode* Node, TArray<FbxNode*>& outMeshArray, UnFbx::FFbxImporter* FFbxImporter);
-	
+
+	/**
+	* Get all Fbx Skeleton nodes
+	*
+	* @param Node Root node to find skeleton nodes
+	* @param outNodeArray return skeleton nodes
+	*/
+	UNREALED_API void FillFbxSkeletonArray(FbxNode* Node, TArray<FbxNode*>& OutNodeArray);
+
 	/**
 	* Get all Fbx mesh objects not under a LOD group and all LOD group node
 	*
@@ -1692,6 +1701,9 @@ protected:
 	 */
 	bool RetrievePoseFromBindPose(const TArray<FbxNode*>& NodeArray, FbxArray<FbxPose*> & PoseArray) const;
 
+	/** Import the user-defined properties on the node as FBX metadata on the object */
+	void ImportNodeCustomProperties(UObject* Object, FbxNode* Node);
+
 public:
 	/** Import and set up animation related data from mesh **/
 	void SetupAnimationDataFromMesh(USkeletalMesh * SkeletalMesh, UObject* InParent, TArray<FbxNode*>& NodeArray, UFbxAnimSequenceImportData* ImportData, const FString& Filename);
