Fix additional cases where GetBlueprint could get called without a valid blueprint, thus checking

diff --git a/Engine/Source/Editor/UnrealEd/Private/Kismet2/EnumEditorUtils.cpp b/Engine/Source/Editor/UnrealEd/Private/Kismet2/EnumEditorUtils.cpp
--- a/Engine/Source/Editor/UnrealEd/Private/Kismet2/EnumEditorUtils.cpp
+++ b/Engine/Source/Editor/UnrealEd/Private/Kismet2/EnumEditorUtils.cpp
@@ -359,8 +359,9 @@ void FEnumEditorUtils::BroadcastChanges(const UUserDefinedEnum* Enum, const TArr
 				UK2Node* Node = Cast<UK2Node>(Pin->GetOuter());
 				if (FNodeValidatorHelper::IsValid(Node))
 				{
-					if (UBlueprint* Blueprint = Node->GetBlueprint())
+					if (Node->HasValidBlueprint())
 					{
+						UBlueprint* Blueprint = Node->GetBlueprint();
 						if (INDEX_NONE == Enum->GetIndexByNameString(Pin->DefaultValue))
 						{
 							Pin->Modify();
@@ -413,7 +414,10 @@ void FEnumEditorUtils::BroadcastChanges(const UUserDefinedEnum* Enum, const TArr
 							{
 								VarNode->ReconstructNode();
 
-								BlueprintsToRefresh.Add(VarNode->GetBlueprint());
+								if (VarNode->HasValidBlueprint())
+								{
+									BlueprintsToRefresh.Add(VarNode->GetBlueprint());
+								}
 							}
 						}
 
