Adding condition to guard against crash when pasting comments

diff --git a/Engine/Source/Editor/UnrealEd/Private/MaterialGraph.cpp b/Engine/Source/Editor/UnrealEd/Private/MaterialGraph.cpp
--- a/Engine/Source/Editor/UnrealEd/Private/MaterialGraph.cpp
+++ b/Engine/Source/Editor/UnrealEd/Private/MaterialGraph.cpp
@@ -300,23 +300,25 @@ void UMaterialGraph::LinkMaterialExpressionsFromGraph() const
 			else if (UMaterialGraphNode_Comment* CommentNode = Cast<UMaterialGraphNode_Comment>(Nodes[NodeIndex]))
 			{
 				UMaterialExpressionComment* Comment = CommentNode->MaterialExpressionComment;
-
-				if (Comment->MaterialExpressionEditorX != CommentNode->NodePosX
-					|| Comment->MaterialExpressionEditorY != CommentNode->NodePosY
-					|| Comment->Text != CommentNode->NodeComment
-					|| Comment->SizeX != CommentNode->NodeWidth
-					|| Comment->SizeY != CommentNode->NodeHeight
-					|| Comment->CommentColor != CommentNode->CommentColor)
+				if (Comment)
 				{
-					Comment->Modify();
-
-					// Update positions and comments
-					Comment->MaterialExpressionEditorX = CommentNode->NodePosX;
-					Comment->MaterialExpressionEditorY = CommentNode->NodePosY;
-					Comment->Text = CommentNode->NodeComment;
-					Comment->SizeX = CommentNode->NodeWidth;
-					Comment->SizeY = CommentNode->NodeHeight;
-					Comment->CommentColor = CommentNode->CommentColor;
+					if (Comment->MaterialExpressionEditorX != CommentNode->NodePosX
+						|| Comment->MaterialExpressionEditorY != CommentNode->NodePosY
+						|| Comment->Text != CommentNode->NodeComment
+						|| Comment->SizeX != CommentNode->NodeWidth
+						|| Comment->SizeY != CommentNode->NodeHeight
+						|| Comment->CommentColor != CommentNode->CommentColor)
+					{
+						Comment->Modify();
+
+						// Update positions and comments
+						Comment->MaterialExpressionEditorX = CommentNode->NodePosX;
+						Comment->MaterialExpressionEditorY = CommentNode->NodePosY;
+						Comment->Text = CommentNode->NodeComment;
+						Comment->SizeX = CommentNode->NodeWidth;
+						Comment->SizeY = CommentNode->NodeHeight;
+						Comment->CommentColor = CommentNode->CommentColor;
+					}
 				}
 			}
 		}
