Merging CL 4331753 from Dev-VR - Allowing output log to capture warning and error verbosity from commandline output it receives (like launch on Lumin)

diff --git a/Engine/Source/Editor/UnrealEd/Private/PlayLevel.cpp b/Engine/Source/Editor/UnrealEd/Private/PlayLevel.cpp
--- a/Engine/Source/Editor/UnrealEd/Private/PlayLevel.cpp
+++ b/Engine/Source/Editor/UnrealEd/Private/PlayLevel.cpp
@@ -1470,9 +1470,20 @@ void UEditorEngine::PlayStandaloneLocalPc(FString MapNameOverride, FIntPoint* Wi
 	FEditorDelegates::BeginStandaloneLocalPlay.Broadcast(ProcessID);
 }
 
-static void HandleOutputReceived(const FString& InMessage)
+static void HandleOutputReceived  (const FString& InMessage)
 {
-	UE_LOG(LogPlayLevel, Log, TEXT("%s"), *InMessage);
+	if (InMessage.Contains(TEXT("Error:")))
+	{
+		UE_LOG(LogPlayLevel, Error, TEXT("%s"), *InMessage);
+	}
+	else if (InMessage.Contains(TEXT("Warning:")))
+	{
+		UE_LOG(LogPlayLevel, Warning, TEXT("%s"), *InMessage);
+	}
+	else 
+	{
+		UE_LOG(LogPlayLevel, Log, TEXT("%s"), *InMessage);
+	}
 }
 
 static void HandleCancelButtonClicked(ILauncherWorkerPtr LauncherWorker)
