Fixed crash when deleting a package without an SSC state

diff --git a/Engine/Source/Editor/UnrealEd/Private/ObjectTools.cpp b/Engine/Source/Editor/UnrealEd/Private/ObjectTools.cpp
--- a/Engine/Source/Editor/UnrealEd/Private/ObjectTools.cpp
+++ b/Engine/Source/Editor/UnrealEd/Private/ObjectTools.cpp
@@ -1780,10 +1780,12 @@ namespace ObjectTools
 			const FString& PackageFilename = PackageFilesToDelete[PackageFileIdx];
 			if ( ISourceControlModule::Get().IsEnabled() )
 			{
-				const FSourceControlStateRef SourceControlState = PackageSCCStates[PackageFileIdx];
-				const bool bInDepot = SourceControlState->IsSourceControlled();
+				const FSourceControlStatePtr SourceControlState = PackageSCCStates.IsValidIndex(PackageFileIdx) ? PackageSCCStates[PackageFileIdx] : FSourceControlStatePtr();
+				const bool bInDepot = SourceControlState.IsValid() && SourceControlState->IsSourceControlled();
 				if ( bInDepot )
 				{
+					check(SourceControlState.IsValid());
+
 					// The file is managed by source control. Open it for delete.
 					FString FullPackageFilename = FPaths::ConvertRelativePathToFull(PackageFilename);
 
