Ensure we GC at various points that could have left lingering references to items being content hot-reloaded

diff --git a/Engine/Source/Editor/UnrealEd/Private/Toolkits/AssetEditorManager.cpp b/Engine/Source/Editor/UnrealEd/Private/Toolkits/AssetEditorManager.cpp
--- a/Engine/Source/Editor/UnrealEd/Private/Toolkits/AssetEditorManager.cpp
+++ b/Engine/Source/Editor/UnrealEd/Private/Toolkits/AssetEditorManager.cpp
@@ -163,7 +163,7 @@ TArray<IAssetEditorInstance*> FAssetEditorManager::FindEditorsForAsset(UObject*
 }
 
 
-void FAssetEditorManager::CloseAllEditorsForAsset(UObject* Asset)
+int32 FAssetEditorManager::CloseAllEditorsForAsset(UObject* Asset)
 {
 	TArray<IAssetEditorInstance*> EditorInstances = FindEditorsForAsset(Asset);
 
@@ -173,6 +173,8 @@ void FAssetEditorManager::CloseAllEditorsForAsset(UObject* Asset)
 	}
 
 	AssetEditorRequestCloseEvent.Broadcast(Asset, EAssetEditorCloseReason::CloseAllEditorsForAsset);
+
+	return EditorInstances.Num();
 }
 
 void FAssetEditorManager::RemoveAssetFromAllEditors(UObject* Asset)
@@ -787,9 +789,17 @@ void FAssetEditorManager::HandlePackageReloaded(const EPackageReloadPhase InPack
 			}
 		}
 
+		int32 NumAssetEditorsClosed = 0;
 		for (UObject* OldAsset : OldAssets)
 		{
-			CloseAllEditorsForAsset(OldAsset);
+			NumAssetEditorsClosed += CloseAllEditorsForAsset(OldAsset);
+		}
+
+		if (NumAssetEditorsClosed > 0)
+		{
+			// Closing asset editors might have have left objects pending GC that still reference the asset we're about to reload
+			// Run a GC now to ensure those are cleaned up before the fix-up phase happens
+			CollectGarbage(GARBAGE_COLLECTION_KEEPFLAGS);
 		}
 	}
 
diff --git a/Engine/Source/Editor/UnrealEd/Public/Toolkits/AssetEditorManager.h b/Engine/Source/Editor/UnrealEd/Public/Toolkits/AssetEditorManager.h
--- a/Engine/Source/Editor/UnrealEd/Public/Toolkits/AssetEditorManager.h
+++ b/Engine/Source/Editor/UnrealEd/Public/Toolkits/AssetEditorManager.h
@@ -82,8 +82,8 @@ public:
 	/** Returns all editors currently opened for the specified asset */
 	TArray<IAssetEditorInstance*> FindEditorsForAsset(UObject* Asset);
 
-	/** Close all active editors for the supplied asset */
-	void CloseAllEditorsForAsset(UObject* Asset);
+	/** Close all active editors for the supplied asset and return the number of asset editors that were closed */
+	int32 CloseAllEditorsForAsset(UObject* Asset);
 
 	/** Close any editor which is not this one */
 	void CloseOtherEditors(UObject* Asset, IAssetEditorInstance* OnlyEditor);
diff --git a/Engine/Source/Runtime/CoreUObject/Private/UObject/PackageReload.cpp b/Engine/Source/Runtime/CoreUObject/Private/UObject/PackageReload.cpp
--- a/Engine/Source/Runtime/CoreUObject/Private/UObject/PackageReload.cpp
+++ b/Engine/Source/Runtime/CoreUObject/Private/UObject/PackageReload.cpp
@@ -445,6 +445,12 @@ void ReloadPackages(const TArrayView<FReloadPackageData>& InPackagesToReload, TA
 			PreparingPackagesForReloadSlowTask.EnterProgressFrame(1.0f);
 			ExistingPackages.Emplace(PackageReloadInternal::ValidateAndPreparePackageForReload(PackageToReloadData.PackageToReload));
 		}
+
+		if (ExistingPackages.Num() > 0)
+		{
+			// Run a GC before we start to clean-up any lingering objects that may reference things we're about to reload
+			CollectGarbage(GARBAGE_COLLECTION_KEEPFLAGS);
+		}
 	}
 
 	// Rename the existing packages, load the new packages, then fix-up any references
