Close blueprint editor when replacing a blueprint, avoiding crash when using the merge tool to reload a blueprint

diff --git a/Engine/Source/Editor/UnrealEd/Private/Kismet2/Kismet2.cpp b/Engine/Source/Editor/UnrealEd/Private/Kismet2/Kismet2.cpp
--- a/Engine/Source/Editor/UnrealEd/Private/Kismet2/Kismet2.cpp
+++ b/Engine/Source/Editor/UnrealEd/Private/Kismet2/Kismet2.cpp
@@ -720,6 +720,9 @@ UBlueprint* FKismetEditorUtilities::ReplaceBlueprint(UBlueprint* Target, UBluepr
 
 	UPackage* BlueprintPackage = Target->GetOutermost();
 	check(BlueprintPackage != GetTransientPackage());
+	
+	// Need to close the editor now, it won't work once it's been garbage collected during the reload
+	FAssetEditorManager::Get().CloseAllEditorsForAsset(Target);
 
 	FBlueprintUnloader Unloader(Target);
 	Unloader.UnloadBlueprint(/*bResetPackage =*/false);
