Reduce CPU usage for Cook-on-the-fly server by not busy-waiting for assets to cook.

Added an event which is signaled when a network request has been received. The main COTF loop blocks on this event when the cook request queue is empty. This reduces CPU usage while keeping the request latency low.

graeme.thornton, martin.ridgers

diff --git a/Engine/Source/Editor/UnrealEd/Classes/CookOnTheSide/CookOnTheFlyServer.h b/Engine/Source/Editor/UnrealEd/Classes/CookOnTheSide/CookOnTheFlyServer.h
--- a/Engine/Source/Editor/UnrealEd/Classes/CookOnTheSide/CookOnTheFlyServer.h
+++ b/Engine/Source/Editor/UnrealEd/Classes/CookOnTheSide/CookOnTheFlyServer.h
@@ -409,6 +409,7 @@ public:
 	bool HasRecompileShaderRequests() const;
 
 	bool HasCookRequests() const;
+	void WaitForRequests(int TimeoutMs);
 
 	uint32 NumConnections() const;
 
diff --git a/Engine/Source/Editor/UnrealEd/Private/Commandlets/CookCommandlet.cpp b/Engine/Source/Editor/UnrealEd/Private/Commandlets/CookCommandlet.cpp
--- a/Engine/Source/Editor/UnrealEd/Private/Commandlets/CookCommandlet.cpp
+++ b/Engine/Source/Editor/UnrealEd/Private/Commandlets/CookCommandlet.cpp
@@ -515,6 +515,8 @@ bool UCookCommandlet::CookOnTheFly( FGuid InstanceId, int32 Timeout, bool bForce
 			}
 
 			CookOnTheFlyGCController.ConditionallyCollectGarbage(CookOnTheFlyServer);
+
+			CookOnTheFlyServer->WaitForRequests(100 /* timeoutMs */);
 		}
 	}
 
diff --git a/Engine/Source/Editor/UnrealEd/Private/CookOnTheFlyServer.cpp b/Engine/Source/Editor/UnrealEd/Private/CookOnTheFlyServer.cpp
--- a/Engine/Source/Editor/UnrealEd/Private/CookOnTheFlyServer.cpp
+++ b/Engine/Source/Editor/UnrealEd/Private/CookOnTheFlyServer.cpp
@@ -1265,6 +1265,8 @@ struct FPackageTracker : public FUObjectArray::FUObjectCreateListener, public FU
 	// List of requested files
 	CookRequestQueue		CookRequests;
 
+	FEvent*					CookRequestEvent = nullptr;
+
 	FThreadSafeUnsolicitedPackagesList			UnsolicitedCookedPackages;
 	FThreadSafeQueue<struct FRecompileRequest*> RecompileRequests;
 
@@ -1599,6 +1601,8 @@ bool UCookOnTheFlyServer::StartNetworkFileServer( const bool BindAnyPort )
 	}
 #endif
 
+	PackageTracker->CookRequestEvent = FPlatformProcess::GetSynchEventFromPool();
+
 	// loop while waiting for requests
 	GIsRequestingExit = false;
 	return true;
@@ -2010,6 +2014,14 @@ FString UCookOnTheFlyServer::GetContentDirecctoryForDLC() const
 
 COREUOBJECT_API extern bool GOutputCookingWarnings;
 
+void UCookOnTheFlyServer::WaitForRequests(int TimeoutMs)
+{
+	if (PackageTracker->CookRequestEvent)
+	{
+		PackageTracker->CookRequestEvent->Wait(TimeoutMs, true);
+	}
+}
+
 bool UCookOnTheFlyServer::HasCookRequests() const 
 { 
 	return PackageTracker->CookRequests.HasItems(); 
@@ -7090,12 +7102,11 @@ void UCookOnTheFlyServer::GetCookOnTheFlyUnsolicitedFiles(const FName& PlatformN
 	UPackage::WaitForAsyncFileWrites();
 }
 
-void UCookOnTheFlyServer::HandleNetworkFileServerFileRequest( const FString& Filename, const FString &PlatformName, TArray<FString>& UnsolicitedFiles )
+void UCookOnTheFlyServer::HandleNetworkFileServerFileRequest(const FString& Filename, const FString& PlatformName, TArray<FString>& UnsolicitedFiles)
 {
 	check(IsCookOnTheFlyMode());
 
-	bool bIsCookable = FPackageName::IsPackageExtension(*FPaths::GetExtension(Filename, true));
-
+	const bool bIsCookable = FPackageName::IsPackageExtension(*FPaths::GetExtension(Filename, true));
 
 	FName PlatformFname = FName(*PlatformName);
 
@@ -7122,7 +7133,9 @@ void UCookOnTheFlyServer::HandleNetworkFileServerFileRequest( const FString& Fil
 	UE_LOG(LogCook, Display, TEXT("Requesting file from cooker %s"), *StandardFileName);
 
 	PackageTracker->CookRequests.EnqueueUnique(FileRequest, true);
-
+	
+	if (PackageTracker->CookRequestEvent)
+		PackageTracker->CookRequestEvent->Trigger();
 
 #if PROFILE_NETWORK
 	bool bFoundNetworkEventWait = true;
@@ -7139,7 +7152,6 @@ void UCookOnTheFlyServer::HandleNetworkFileServerFileRequest( const FString& Fil
 			break;
 		}
 	}
-	
 
 	// wait for tick entry here
 	TimeTillRequestStarted += FPlatformTime::Seconds() - StartTime;
