Fixing localization warnings

diff --git a/Engine/Source/Editor/UnrealEd/Private/Kismet2/DebuggerCommands.cpp b/Engine/Source/Editor/UnrealEd/Private/Kismet2/DebuggerCommands.cpp
--- a/Engine/Source/Editor/UnrealEd/Private/Kismet2/DebuggerCommands.cpp
+++ b/Engine/Source/Editor/UnrealEd/Private/Kismet2/DebuggerCommands.cpp
@@ -1091,12 +1091,9 @@ TSharedRef< SWidget > FPlayWorldCommands::GenerateLaunchMenuContent( TSharedRef<
 							FString AggregateDevicedName(FString::Printf(TEXT("  %s"), *DeviceProxy->GetName())); //align with the other menu entries
 							FSlateIcon AggregateDeviceIcon(FEditorStyle::GetStyleSetName(), VanillaPlatform.PlatformInfo->GetIconStyleName(PlatformInfo::EPlatformIconSize::Normal));
 
-							FString SubMenuKey(FString::Printf(TEXT("%s_SubMenu"), *DeviceProxy->GetName()));
-							FString SubMenuToolTip(FString::Printf(TEXT("%s_SubMenuToolTip"), *DeviceProxy->GetName()));
-
 							MenuBuilder.AddSubMenu(
-								FInternationalization::ForUseOnlyByLocMacroAndGraphNodeTextLiterals_CreateText(*AggregateDevicedName, TEXT(LOCTEXT_NAMESPACE), *SubMenuKey),
-								FInternationalization::ForUseOnlyByLocMacroAndGraphNodeTextLiterals_CreateText(*AggregateDevicedName, TEXT(LOCTEXT_NAMESPACE), *SubMenuToolTip),
+								FText::FromString(AggregateDevicedName),
+								FText::FromString(AggregateDevicedName),
 								FNewMenuDelegate::CreateStatic(&MakeAllDevicesSubMenu, VanillaPlatform.PlatformInfo, DeviceProxy),
 								false, AggregateDeviceIcon, true
 							);
