Fixing localization warnings

Added support for gathering from the newly added UI_COMMAND_EXT macro

diff --git a/Engine/Source/Editor/UnrealEd/Classes/Commandlets/GatherTextFromSourceCommandlet.h b/Engine/Source/Editor/UnrealEd/Classes/Commandlets/GatherTextFromSourceCommandlet.h
--- a/Engine/Source/Editor/UnrealEd/Classes/Commandlets/GatherTextFromSourceCommandlet.h
+++ b/Engine/Source/Editor/UnrealEd/Classes/Commandlets/GatherTextFromSourceCommandlet.h
@@ -248,10 +248,23 @@ private:
 		FString Name;
 	};
 
-	class FCommandMacroDescriptor : public FMacroDescriptor
+	class FUICommandMacroDescriptor : public FMacroDescriptor
 	{
 	public:
-		FCommandMacroDescriptor() : FMacroDescriptor(TEXT("UI_COMMAND")) {}
+		FUICommandMacroDescriptor() : FMacroDescriptor(TEXT("UI_COMMAND")) {}
+
+		virtual void TryParse(const FString& Text, FSourceFileParseContext& Context) const override;
+
+	protected:
+		FUICommandMacroDescriptor(FString InName) : FMacroDescriptor(MoveTemp(InName)) {}
+
+		void TryParseArgs(const FString& Text, FSourceFileParseContext& Context, const TArray<FString>& Arguments, const int32 ArgIndexOffset) const;
+	};
+
+	class FUICommandExtMacroDescriptor : public FUICommandMacroDescriptor
+	{
+	public:
+		FUICommandExtMacroDescriptor() : FUICommandMacroDescriptor(TEXT("UI_COMMAND_EXT")) {}
 
 		virtual void TryParse(const FString& Text, FSourceFileParseContext& Context) const override;
 	};
diff --git a/Engine/Source/Editor/UnrealEd/Private/Commandlets/GatherTextFromSourceCommandlet.cpp b/Engine/Source/Editor/UnrealEd/Private/Commandlets/GatherTextFromSourceCommandlet.cpp
--- a/Engine/Source/Editor/UnrealEd/Private/Commandlets/GatherTextFromSourceCommandlet.cpp
+++ b/Engine/Source/Editor/UnrealEd/Private/Commandlets/GatherTextFromSourceCommandlet.cpp
@@ -234,7 +234,9 @@ int32 UGatherTextFromSourceCommandlet::Main( const FString& Params )
 
 	Parsables.Add(new FEndIfDescriptor());
 
-	Parsables.Add(new FCommandMacroDescriptor());
+	Parsables.Add(new FUICommandMacroDescriptor());
+
+	Parsables.Add(new FUICommandExtMacroDescriptor());
 
 	// New Localization System with Namespace as literal argument.
 	Parsables.Add(new FStringMacroDescriptor( FString(TEXT("NSLOCTEXT")),
@@ -854,6 +856,12 @@ bool UGatherTextFromSourceCommandlet::ParseSourceText(const FString& Text, const
 			{
 				FString& Token = ParsableTokens[ParIdx];
 
+				if (Token.Len() == ParsableMatchCounters[ParIdx])
+				{
+					// already seen this entire token and are looking for longer matches - skip it
+					continue;
+				}
+
 				if (*Cursor == Token[ParsableMatchCounters[ParIdx]])
 				{
 					// Char at cursor matches the next char in the parsable's identifying token
@@ -1411,70 +1419,104 @@ bool UGatherTextFromSourceCommandlet::FMacroDescriptor::PrepareArgument(FString&
 	return Error ? false : true;
 }
 
-void UGatherTextFromSourceCommandlet::FCommandMacroDescriptor::TryParse(const FString& Text, FSourceFileParseContext& Context) const
+void UGatherTextFromSourceCommandlet::FUICommandMacroDescriptor::TryParseArgs(const FString& Text, FSourceFileParseContext& Context, const TArray<FString>& Arguments, const int32 ArgIndexOffset) const
+{
+	FString Identifier = Arguments[ArgIndexOffset];
+	Identifier.TrimStartInline();
+
+	// Identifier may optionally be in quotes, as it's sometimes a string literal (in UE_COMMAND_EXT), and sometimes stringified by the macro (in UI_COMMAND)
+	// Because this is optional, we don't care if this processing fails
+	bool HasQuotes = false;
+	PrepareArgument(Identifier, true, FString(), HasQuotes);
+
+	FString SourceLocation = FSourceLocation(Context.Filename, Context.LineNumber).ToString();
+	if (Identifier.IsEmpty())
+	{
+		//The command doesn't have an identifier so we can't gather it
+		UE_LOG(LogGatherTextFromSourceCommandlet, Warning, TEXT("%s macro has an empty identifier and cannot be gathered. %s"), *GetToken(), *SourceLocation);
+		return;
+	}
+
+	FString SourceText = Arguments[ArgIndexOffset + 1];
+	SourceText.TrimStartInline();
+
+	static const FString UICommandRootNamespace = TEXT("UICommands");
+	FString Namespace = Context.WithinNamespaceDefine && !Context.Namespace.IsEmpty() ? FString::Printf(TEXT("%s.%s"), *UICommandRootNamespace, *Context.Namespace) : UICommandRootNamespace;
+
+	// parse DefaultLangString argument - this arg will be in quotes without TEXT macro
+	FString MacroDesc = FString::Printf(TEXT("\"FriendlyName\" argument in %s macro %s(%d):%s"), *GetToken(), *Context.Filename, Context.LineNumber, *Context.LineText);
+	if (PrepareArgument(SourceText, true, MacroDesc, HasQuotes))
+	{
+		if (HasQuotes && !Identifier.IsEmpty() && !SourceText.IsEmpty())
+		{
+			// First create the command entry
+			FManifestContext CommandContext;
+			CommandContext.Key = Identifier;
+			CommandContext.SourceLocation = SourceLocation;
+
+			Context.AddManifestText(GetToken(), Namespace, SourceText, CommandContext);
+
+			// parse DefaultLangTooltipString argument - this arg will be in quotes without TEXT macro
+			FString TooltipSourceText = Arguments[ArgIndexOffset + 2];
+			TooltipSourceText.TrimStartInline();
+			MacroDesc = FString::Printf(TEXT("\"InDescription\" argument in %s macro %s(%d):%s"), *GetToken(), *Context.Filename, Context.LineNumber, *Context.LineText);
+			if (PrepareArgument(TooltipSourceText, true, MacroDesc, HasQuotes))
+			{
+				if (HasQuotes && !TooltipSourceText.IsEmpty())
+				{
+					// Create the tooltip entry
+					FManifestContext CommandTooltipContext;
+					CommandTooltipContext.Key = Identifier + TEXT("_ToolTip");
+					CommandTooltipContext.SourceLocation = SourceLocation;
+
+					Context.AddManifestText(GetToken(), Namespace, TooltipSourceText, CommandTooltipContext);
+				}
+			}
+		}
+	}
+}
+
+void UGatherTextFromSourceCommandlet::FUICommandMacroDescriptor::TryParse(const FString& Text, FSourceFileParseContext& Context) const
 {
 	// Attempt to parse something of the format
-	// UI_COMMAND(LocKey, DefaultLangString, DefaultLangTooltipString, <IgnoredParam>, <IgnoredParam>)
+	// UI_COMMAND(LocKey, DefaultLangString, DefaultLangTooltipString, ...)
 
 	if (!Context.ExcludedRegion && !Context.WithinBlockComment && !Context.WithinLineComment && !Context.WithinStringLiteral)
 	{
 		TArray<FString> Arguments;
 		if (ParseArgsFromMacro(StripCommentsFromToken(Text, Context), Arguments, Context))
 		{
-			// 5 or 6 arguments depending on whether or not there are one or two bindings
-			if (Arguments.Num() != 5 && Arguments.Num() != 6)
+			// Need at least 3 arguments
+			if (Arguments.Num() < 3)
 			{
-				UE_LOG(LogGatherTextFromSourceCommandlet, Warning, TEXT("Too many arguments in command %s macro in %s(%d):%s"), *GetToken(), *Context.Filename, Context.LineNumber, *FLocTextHelper::SanitizeLogOutput(Context.LineText));
+				UE_LOG(LogGatherTextFromSourceCommandlet, Warning, TEXT("Too few arguments in %s macro in %s(%d):%s"), *GetToken(), *Context.Filename, Context.LineNumber, *FLocTextHelper::SanitizeLogOutput(Context.LineText));
 			}
 			else
 			{
-				Arguments[0].TrimStartInline();
-				FString Identifier = Arguments[0];
-				const FString UICommandRootNamespace = TEXT("UICommands");
-				FString Namespace = Context.WithinNamespaceDefine && !Context.Namespace.IsEmpty() ? FString::Printf( TEXT("%s.%s"), *UICommandRootNamespace, *Context.Namespace) : UICommandRootNamespace;
-				FString SourceLocation = FSourceLocation(Context.Filename, Context.LineNumber).ToString();
-				Arguments[1].TrimStartInline();
-				FString SourceText = Arguments[1];
-
-				if ( Identifier.IsEmpty() )
-				{
-					//The command doesn't have an identifier so we can't gather it
-					UE_LOG(LogGatherTextFromSourceCommandlet, Warning, TEXT("UICOMMAND macro has an empty identifier and cannot be gathered. %s"), *SourceLocation );
-					return;
-				}
+				TryParseArgs(Text, Context, Arguments, 0);
+			}
+		}
+	}
+}
 
-				// parse DefaultLangString argument - this arg will be in quotes without TEXT macro
-				bool HasQuotes;
-				FString MacroDesc = FString::Printf(TEXT("\"FriendlyName\" argument in %s macro %s(%d):%s"), *GetToken(), *Context.Filename, Context.LineNumber, *Context.LineText);
-				if ( PrepareArgument(SourceText, true, MacroDesc, HasQuotes) )
-				{
-					if ( HasQuotes && !Identifier.IsEmpty() && !SourceText.IsEmpty() )
-					{
-						// First create the command entry
-						FManifestContext CommandContext;
-						CommandContext.Key = Identifier;
-						CommandContext.SourceLocation = SourceLocation;
-
-						Context.AddManifestText( GetToken(), Namespace, SourceText, CommandContext );
-
-						// parse DefaultLangTooltipString argument - this arg will be in quotes without TEXT macro
-						Arguments[2].TrimStartInline();
-						FString TooltipSourceText = Arguments[2];
-						MacroDesc = FString::Printf(TEXT("\"InDescription\" argument in %s macro %s(%d):%s"), *GetToken(), *Context.Filename, Context.LineNumber, *Context.LineText);
-						if (PrepareArgument(TooltipSourceText, true, MacroDesc, HasQuotes))
-						{
-							if (HasQuotes && !TooltipSourceText.IsEmpty())
-							{
-								// Create the tooltip entry
-								FManifestContext CommandTooltipContext;
-								CommandTooltipContext.Key = Identifier + TEXT("_ToolTip");
-								CommandTooltipContext.SourceLocation = SourceLocation;
+void UGatherTextFromSourceCommandlet::FUICommandExtMacroDescriptor::TryParse(const FString& Text, FSourceFileParseContext& Context) const
+{
+	// Attempt to parse something of the format
+	// UI_COMMAND_EXT(<IgnoredParam>, <IgnoredParam>, LocKey, DefaultLangString, DefaultLangTooltipString, ...)
 
-								Context.AddManifestText( GetToken(), Namespace, TooltipSourceText, CommandTooltipContext );
-							}
-						}
-					}
-				}
+	if (!Context.ExcludedRegion && !Context.WithinBlockComment && !Context.WithinLineComment && !Context.WithinStringLiteral)
+	{
+		TArray<FString> Arguments;
+		if (ParseArgsFromMacro(StripCommentsFromToken(Text, Context), Arguments, Context))
+		{
+			// Need at least 5 arguments
+			if (Arguments.Num() < 5)
+			{
+				UE_LOG(LogGatherTextFromSourceCommandlet, Warning, TEXT("Too few arguments in %s macro in %s(%d):%s"), *GetToken(), *Context.Filename, Context.LineNumber, *FLocTextHelper::SanitizeLogOutput(Context.LineText));
+			}
+			else
+			{
+				TryParseArgs(Text, Context, Arguments, 2);
 			}
 		}
 	}
