Fixing a Cook On The Fly server crash when an asset that has been requested by the client is missing from the server

diff --git a/Engine/Source/Editor/UnrealEd/Private/CookOnTheFlyServer.cpp b/Engine/Source/Editor/UnrealEd/Private/CookOnTheFlyServer.cpp
--- a/Engine/Source/Editor/UnrealEd/Private/CookOnTheFlyServer.cpp
+++ b/Engine/Source/Editor/UnrealEd/Private/CookOnTheFlyServer.cpp
@@ -4017,6 +4017,10 @@ void UCookOnTheFlyServer::SaveCookedPackage(UPackage* Package, uint32 SaveFlags,
 
 		Package->SetPackageFlagsTo(OriginalPackageFlags);
 	}
+	else
+	{
+		SavePackageResults.Add(FSavePackageResultStruct(ESavePackageResult::MissingFile));
+	}
 
 	check(bIsSavingPackage == true);
 	bIsSavingPackage = false;
diff --git a/Engine/Source/Runtime/CoreUObject/Public/UObject/Package.h b/Engine/Source/Runtime/CoreUObject/Public/UObject/Package.h
--- a/Engine/Source/Runtime/CoreUObject/Public/UObject/Package.h
+++ b/Engine/Source/Runtime/CoreUObject/Public/UObject/Package.h
@@ -38,6 +38,8 @@ enum class ESavePackageResult
 	GenerateStub,
 	/** [When cooking] When performing package diff, the package generated in memory was different to the one that existed on disk */
 	DifferentContent,
+	/** [When cooking] The file requested (when cooking on the fly) did not exist on disk */
+	MissingFile
 };
 
 /**
