Remove debugging hook accidentally added in 4235135

diff --git a/Engine/Source/Editor/UnrealEd/Private/Kismet2/KismetDebugUtilities.cpp b/Engine/Source/Editor/UnrealEd/Private/Kismet2/KismetDebugUtilities.cpp
--- a/Engine/Source/Editor/UnrealEd/Private/Kismet2/KismetDebugUtilities.cpp
+++ b/Engine/Source/Editor/UnrealEd/Private/Kismet2/KismetDebugUtilities.cpp
@@ -61,9 +61,6 @@ public:
 
 	~FKismetDebugUtilitiesData()
 	{
-		volatile int32 a = 0;
-		++a;
-		(void)a;
 	}
 
 	void Reset()
