Remove unused dtor
#rnx

diff --git a/Engine/Source/Editor/UnrealEd/Private/Kismet2/KismetDebugUtilities.cpp b/Engine/Source/Editor/UnrealEd/Private/Kismet2/KismetDebugUtilities.cpp
--- a/Engine/Source/Editor/UnrealEd/Private/Kismet2/KismetDebugUtilities.cpp
+++ b/Engine/Source/Editor/UnrealEd/Private/Kismet2/KismetDebugUtilities.cpp
@@ -59,10 +59,6 @@ public:
 	{
 	}
 
-	~FKismetDebugUtilitiesData()
-	{
-	}
-
 	void Reset()
 	{
 		TargetGraphNodes.Empty();
