Fixed crash when undoing after deleting a function graph that has a math expression node

diff --git a/Engine/Source/Editor/UnrealEd/Private/Kismet2/BlueprintEditorUtils.cpp b/Engine/Source/Editor/UnrealEd/Private/Kismet2/BlueprintEditorUtils.cpp
--- a/Engine/Source/Editor/UnrealEd/Private/Kismet2/BlueprintEditorUtils.cpp
+++ b/Engine/Source/Editor/UnrealEd/Private/Kismet2/BlueprintEditorUtils.cpp
@@ -2697,8 +2697,6 @@ void FBlueprintEditorUtils::RemoveGraph(UBlueprint* Blueprint, class UEdGraph* G
 		GraphToRemove->SetFlags(RF_Transient);
 	}
 
-	GraphToRemove->MarkPendingKill();
-
 	if (Flags & EGraphRemoveFlags::Recompile )
 	{
 		FBlueprintEditorUtils::MarkBlueprintAsStructurallyModified(Blueprint);
