Fixed crash after hitting stop when in-stack-debugging
Clear GPlayInEditorID when stopping PIE at a blueprint exception/breakpoint

diff --git a/Engine/Source/Editor/UnrealEd/Private/Kismet2/KismetDebugUtilities.cpp b/Engine/Source/Editor/UnrealEd/Private/Kismet2/KismetDebugUtilities.cpp
--- a/Engine/Source/Editor/UnrealEd/Private/Kismet2/KismetDebugUtilities.cpp
+++ b/Engine/Source/Editor/UnrealEd/Private/Kismet2/KismetDebugUtilities.cpp
@@ -702,6 +702,7 @@ void FKismetDebugUtilities::AttemptToBreakExecution(UBlueprint* BlueprintObj, co
 	// Now enter within-the-frame debugging mode
 	if (bShouldInStackDebug)
 	{
+		TGuardValue<int32> GuardDisablePIE(GPlayInEditorID, INDEX_NONE);
 		const TArray<const FFrame*>& ScriptStack = FBlueprintExceptionTracker::Get().ScriptStack;
 		Data.LastExceptionMessage = Info.GetDescription();
 		FKismetEditorUtilities::BringKismetToFocusAttentionOnObject(NodeStoppedAt);
