Guard against null objects in asset editors

diff --git a/Engine/Source/Editor/UnrealEd/Private/Toolkits/AssetEditorToolkit.cpp b/Engine/Source/Editor/UnrealEd/Private/Toolkits/AssetEditorToolkit.cpp
--- a/Engine/Source/Editor/UnrealEd/Private/Toolkits/AssetEditorToolkit.cpp
+++ b/Engine/Source/Editor/UnrealEd/Private/Toolkits/AssetEditorToolkit.cpp
@@ -1105,10 +1105,19 @@ void FAssetEditorToolkit::RemoveAllToolbarWidgets()
 
 void FAssetEditorToolkit::FGCEditingObjects::AddReferencedObjects(FReferenceCollector& Collector)
 {
+	// Remove null objects as a safe guard against assets being forcibly GC'd
+	TArray<UObject*> EditingObjectsCopy = OwnerToolkit.EditingObjects;
 	Collector.AddReferencedObjects(OwnerToolkit.EditingObjects);
-}
-
 
+	OwnerToolkit.EditingObjects.Reset();
+	for (UObject* Object : EditingObjectsCopy)
+	{
+		if (Object != nullptr)
+		{
+			OwnerToolkit.EditingObjects.Add(Object);
+		}
+	}
+}
 
 
 TSharedPtr<FExtender> FExtensibilityManager::GetAllExtenders()
