Expose editor's component type cache to reference collectors

diff --git a/Engine/Source/Editor/UnrealEd/Private/ComponentTypeRegistry.cpp b/Engine/Source/Editor/UnrealEd/Private/ComponentTypeRegistry.cpp
--- a/Engine/Source/Editor/UnrealEd/Private/ComponentTypeRegistry.cpp
+++ b/Engine/Source/Editor/UnrealEd/Private/ComponentTypeRegistry.cpp
@@ -28,6 +28,7 @@
 
 struct FComponentTypeRegistryData
 	: public FTickableEditorObject
+	, public FGCObject
 {
 	FComponentTypeRegistryData();
 
@@ -40,7 +41,10 @@ struct FComponentTypeRegistryData
 	virtual void Tick(float) override;
 	virtual ETickableTickType GetTickableTickType() const override { return ETickableTickType::Always; }
 	virtual TStatId GetStatId() const override { RETURN_QUICK_DECLARE_CYCLE_STAT(FTypeDatabaseUpdater, STATGROUP_Tickables); }
-
+	
+	/** Implementation of FGCObject */
+	virtual void AddReferencedObjects(FReferenceCollector& Collector) override;
+	
 	// Request a refresh of the components list next frame
 	void Invalidate()
 	{
@@ -427,6 +431,19 @@ void FComponentTypeRegistryData::Tick(float)
 	}
 }
 
+void FComponentTypeRegistryData::AddReferencedObjects(FReferenceCollector& Collector)
+{
+	for(FComponentClassComboEntryPtr& ComboEntry : ComponentClassList)
+	{
+		ComboEntry->AddReferencedObjects(Collector);
+	}
+
+	for(FComponentTypeEntry& TypeEntry : ComponentTypeList)
+	{
+		Collector.AddReferencedObject(TypeEntry.ComponentClass);
+	}
+}
+
 //////////////////////////////////////////////////////////////////////////
 // FComponentTypeRegistry
 
diff --git a/Engine/Source/Editor/UnrealEd/Private/SComponentClassCombo.cpp b/Engine/Source/Editor/UnrealEd/Private/SComponentClassCombo.cpp
--- a/Engine/Source/Editor/UnrealEd/Private/SComponentClassCombo.cpp
+++ b/Engine/Source/Editor/UnrealEd/Private/SComponentClassCombo.cpp
@@ -26,6 +26,22 @@ FString FComponentClassComboEntry::GetClassName() const
 	return ComponentClass != nullptr ? ComponentClass->GetDisplayNameText().ToString() : ComponentName;
 }
 
+void FComponentClassComboEntry::AddReferencedObjects(FReferenceCollector& Collector)
+{
+	UClass* RawClass = ComponentClass;
+	Collector.AddReferencedObject(RawClass);
+	if(RawClass && RawClass->IsChildOf(UActorComponent::StaticClass()))
+	{
+		ComponentClass = RawClass;
+	}
+	else
+	{
+		ComponentClass = nullptr;
+	}
+
+	Collector.AddReferencedObject(IconClass);
+}
+
 void SComponentClassCombo::Construct(const FArguments& InArgs)
 {
 	PrevSelectedIndex = INDEX_NONE;
diff --git a/Engine/Source/Editor/UnrealEd/Public/SComponentClassCombo.h b/Engine/Source/Editor/UnrealEd/Public/SComponentClassCombo.h
--- a/Engine/Source/Editor/UnrealEd/Public/SComponentClassCombo.h
+++ b/Engine/Source/Editor/UnrealEd/Public/SComponentClassCombo.h
@@ -148,6 +148,8 @@ public:
 	FName GetIconOverrideBrushName() const { return CustomizationArgs.IconOverrideBrushName; }
 
 	int32 GetSortPriority() const { return CustomizationArgs.SortPriority; }
+
+	void AddReferencedObjects(FReferenceCollector& Collector);
 private:
 	TSubclassOf<UActorComponent> ComponentClass;
 	const UClass* IconClass;
