Remove legacy/backcompat Blueprint code that triggered a copy of the skeleton CDO to the generated CDO based on a serialized flag that's no longer being updated at compile time.

Notes:
- This is the extension of the 4.21-specific fix in 4467591 that we deferred until 4.22 to allow a larger testing window on the more impactful changes before release.

diff --git a/Engine/Source/Editor/Kismet/Private/BlueprintCompilationManager.cpp b/Engine/Source/Editor/Kismet/Private/BlueprintCompilationManager.cpp
--- a/Engine/Source/Editor/Kismet/Private/BlueprintCompilationManager.cpp
+++ b/Engine/Source/Editor/Kismet/Private/BlueprintCompilationManager.cpp
@@ -485,7 +485,7 @@ void FBlueprintCompilationManagerImpl::FlushCompilationQueueImpl(TArray<UObject*
 			if(bSkipCompile)
 			{
 				CurrentlyCompilingBPs.Add(FCompilerData(QueuedBP, ECompilationManagerJobType::SkeletonOnly, QueuedJob.ClientResultsLog, QueuedJob.CompileOptions, false));
-				if (QueuedBP->IsGeneratedClassAuthoritative() && (QueuedBP->GeneratedClass != nullptr))
+				if (QueuedBP->GeneratedClass != nullptr)
 				{
 					// set bIsRegeneratingOnLoad so that we don't reset loaders:
 					QueuedBP->bIsRegeneratingOnLoad = true;
diff --git a/Engine/Source/Editor/KismetCompiler/Private/KismetCompiler.cpp b/Engine/Source/Editor/KismetCompiler/Private/KismetCompiler.cpp
--- a/Engine/Source/Editor/KismetCompiler/Private/KismetCompiler.cpp
+++ b/Engine/Source/Editor/KismetCompiler/Private/KismetCompiler.cpp
@@ -4149,15 +4149,6 @@ void FKismetCompilerContext::CompileFunctions(EInternalCompilerFlags InternalFla
 		// Copy over the CDO properties if we're not already regenerating on load.  In that case, the copy will be done after compile on load is complete
 		FBlueprintEditorUtils::PropagateParentBlueprintDefaults(NewClass);
 
-		if (Blueprint->HasAnyFlags(RF_BeingRegenerated))
-		{
-			if (CompileOptions.CompileType == EKismetCompileType::Full)
-			{
-				check(Blueprint->PRIVATE_InnermostPreviousCDO == NULL);
-				Blueprint->PRIVATE_InnermostPreviousCDO = OldCDO;
-			}
-		}
-
 		if(bPropagateValuesToCDO)
 		{
 			if( !Blueprint->HasAnyFlags(RF_BeingRegenerated) )
@@ -4188,17 +4179,6 @@ void FKismetCompilerContext::CompileFunctions(EInternalCompilerFlags InternalFla
 					// Don't perform generated class validation since we didn't do any value propagation.
 					bSkipGeneratedClassValidation = true;
 				}
-
-				// >>> Backwards Compatibility: Propagate data from the skel CDO to the gen CDO if we haven't already done so for this blueprint
-				if( !bIsSkeletonOnly && !Blueprint->IsGeneratedClassAuthoritative() )
-				{
-					UEditorEngine::FCopyPropertiesForUnrelatedObjectsParams CopyDetails;
-					CopyDetails.bAggressiveDefaultSubobjectReplacement = false;
-					CopyDetails.bDoDelta = false;
-					UEditorEngine::CopyPropertiesForUnrelatedObjects(Blueprint->SkeletonGeneratedClass->GetDefaultObject(), NewCDO, CopyDetails);
-					Blueprint->SetLegacyGeneratedClassIsAuthoritative();
-				}
-				// <<< End Backwards Compatibility
 			}
 
 			PropagateValuesToCDO(NewCDO, OldCDO);
diff --git a/Engine/Source/Editor/UnrealEd/Private/Kismet2/BlueprintEditorUtils.cpp b/Engine/Source/Editor/UnrealEd/Private/Kismet2/BlueprintEditorUtils.cpp
--- a/Engine/Source/Editor/UnrealEd/Private/Kismet2/BlueprintEditorUtils.cpp
+++ b/Engine/Source/Editor/UnrealEd/Private/Kismet2/BlueprintEditorUtils.cpp
@@ -1420,20 +1420,9 @@ UClass* FBlueprintEditorUtils::RegenerateBlueprintClass(UBlueprint* Blueprint, U
 		}
 		else
 		{
-			if (Blueprint->IsGeneratedClassAuthoritative() && (Blueprint->GeneratedClass != nullptr))
+			if (Blueprint->GeneratedClass != nullptr)
 			{
 				RemoveStaleFunctions(Cast<UBlueprintGeneratedClass>(Blueprint->GeneratedClass), Blueprint);
-
-				check(PreviousCDO != nullptr);
-				check(Blueprint->SkeletonGeneratedClass != nullptr);
-
-				// We now know we're a data-only blueprint on the outer pass (generate class is valid), where generated class is authoritative
-				// If the PreviousCDO is to the skeleton, then it will corrupt data when copied over the AuthoriativeClass later on in this function
-				if (PreviousCDO == Blueprint->SkeletonGeneratedClass->GetDefaultObject())
-				{
-					check(Blueprint->PRIVATE_InnermostPreviousCDO == nullptr);
-					Blueprint->PRIVATE_InnermostPreviousCDO = Blueprint->GeneratedClass->GetDefaultObject();
-				}
 			}
 
 			// No actual compilation work to be done, but try to conform the class and fix up anything that might need to be updated if the native base class has changed in any way
@@ -1470,81 +1459,12 @@ UClass* FBlueprintEditorUtils::RegenerateBlueprintClass(UBlueprint* Blueprint, U
 
 		if (!FKismetEditorUtilities::IsClassABlueprintSkeleton(ClassToRegenerate))
 		{
-			if (Blueprint->bRecompileOnLoad)
-			{
-				// Verify that we had a skeleton generated class if we had a previous CDO, to make sure we have something to copy into
-				check((Blueprint->BlueprintType == BPTYPE_MacroLibrary) || Blueprint->SkeletonGeneratedClass);
-
-				const bool bPreviousMatchesGenerated = (PreviousCDO == Blueprint->GeneratedClass->GetDefaultObject());
-
-				if (Blueprint->BlueprintType != BPTYPE_MacroLibrary)
-				{
-					UObject* CDOThatKickedOffCOL = PreviousCDO;
-					if (Blueprint->IsGeneratedClassAuthoritative() && !bPreviousMatchesGenerated && Blueprint->PRIVATE_InnermostPreviousCDO)
-					{
-						PreviousCDO = Blueprint->PRIVATE_InnermostPreviousCDO;
-					}
-				}
-
-				// If this is the top of the compile-on-load stack for this object, copy the old CDO properties to the newly created one unless they are the same
-				UClass* AuthoritativeClass = (Blueprint->IsGeneratedClassAuthoritative() ? Blueprint->GeneratedClass : Blueprint->SkeletonGeneratedClass);
-				if (AuthoritativeClass != nullptr && PreviousCDO != AuthoritativeClass->GetDefaultObject())
-				{
-					TGuardValue<bool> GuardTemplateNameFlag(GCompilingBlueprint, true);
-
-					// Make sure the previous CDO has been fully loaded before we use it
-					FBlueprintEditorUtils::PreloadMembers(PreviousCDO);
-
-					// Copy over the properties from the old CDO to the new
-					PropagateParentBlueprintDefaults(AuthoritativeClass);
-					UObject* NewCDO = AuthoritativeClass->GetDefaultObject();
-					{
-						FSaveActorFlagsHelper SaveActorFlags(AuthoritativeClass);
-						UEditorEngine::FCopyPropertiesForUnrelatedObjectsParams CopyDetails;
-						CopyDetails.bAggressiveDefaultSubobjectReplacement = true;
-						CopyDetails.bDoDelta = false;
-						CopyDetails.bCopyDeprecatedProperties = true;
-						CopyDetails.bSkipCompilerGeneratedDefaults = true;
-						UEditorEngine::CopyPropertiesForUnrelatedObjects(PreviousCDO, NewCDO, CopyDetails);
-					}
-
-					if (bRegenerated)
-					{
-						PatchCDOSubobjectsIntoExport(PreviousCDO, NewCDO);
-						// We purposefully do not call post load here, it happens later on in the normal flow
-					}
-
-					// Update the custom property list used in post construction logic to include native class properties for which the regenerated Blueprint CDO now differs from the native CDO.
-					if (UBlueprintGeneratedClass* BPGClass = Cast<UBlueprintGeneratedClass>(AuthoritativeClass))
-					{
-						BPGClass->UpdateCustomPropertyListForPostConstruction();
-					}
-				}
-
-				Blueprint->PRIVATE_InnermostPreviousCDO = nullptr;
-			}
-			else
+			if (!Blueprint->bRecompileOnLoad)
 			{
 				// If we didn't recompile, we still need to propagate flags, and instance components
 				FKismetEditorUtilities::ConformBlueprintFlagsAndComponents(Blueprint);
 			}
 
-			// If this is the top of the compile-on-load stack for this object, copy the old CDO properties to the newly created one
-			if (!Blueprint->IsGeneratedClassAuthoritative() && Blueprint->GeneratedClass != nullptr)
-			{
-				TGuardValue<bool> GuardTemplateNameFlag(GCompilingBlueprint, true);
-
-				UObject* SkeletonCDO = Blueprint->SkeletonGeneratedClass->GetDefaultObject();
-				UObject* GeneratedCDO = Blueprint->GeneratedClass->GetDefaultObject();
-
-				UEditorEngine::FCopyPropertiesForUnrelatedObjectsParams CopyDetails;
-				CopyDetails.bAggressiveDefaultSubobjectReplacement = false;
-				CopyDetails.bDoDelta = false;
-				UEditorEngine::CopyPropertiesForUnrelatedObjects(SkeletonCDO, GeneratedCDO, CopyDetails);
-
-				Blueprint->SetLegacyGeneratedClassIsAuthoritative();
-			}
-
 			// Now that the CDO is valid, update the OwnedComponents, in case we've added or removed native components
 			if (AActor* MyActor = Cast<AActor>(Blueprint->GeneratedClass->GetDefaultObject()))
 			{
diff --git a/Engine/Source/Runtime/Core/Public/UObject/FrameworkObjectVersion.h b/Engine/Source/Runtime/Core/Public/UObject/FrameworkObjectVersion.h
--- a/Engine/Source/Runtime/Core/Public/UObject/FrameworkObjectVersion.h
+++ b/Engine/Source/Runtime/Core/Public/UObject/FrameworkObjectVersion.h
@@ -119,6 +119,9 @@ struct CORE_API FFrameworkObjectVersion
 		// Custom event and non-native interface event implementations add 'const' to reference parameters
 		EditableEventsUseConstRefParameters,
 
+		// No longer serialize the legacy flag that indicates this state, as it is now implied since we don't serialize the skeleton CDO
+		BlueprintGeneratedClassIsAlwaysAuthoritative,
+
 		// -----<new versions can be added above this line>-------------------------------------------------
 		VersionPlusOne,
 		LatestVersion = VersionPlusOne - 1
diff --git a/Engine/Source/Runtime/Engine/Classes/Engine/Blueprint.h b/Engine/Source/Runtime/Engine/Classes/Engine/Blueprint.h
--- a/Engine/Source/Runtime/Engine/Classes/Engine/Blueprint.h
+++ b/Engine/Source/Runtime/Engine/Classes/Engine/Blueprint.h
@@ -424,9 +424,6 @@ class ENGINE_API UBlueprint : public UBlueprintCore
 	UPROPERTY()
 	TSubclassOf<class UObject> ParentClass;
 
-	UPROPERTY(transient)
-	UObject* PRIVATE_InnermostPreviousCDO;
-
 	/** When the class generated by this blueprint is loaded, it will be recompiled the first time.  After that initial recompile, subsequent loads will skip the regeneration step */
 	UPROPERTY(transient)
 	uint32 bHasBeenRegenerated:1;
diff --git a/Engine/Source/Runtime/Engine/Classes/Engine/BlueprintCore.h b/Engine/Source/Runtime/Engine/Classes/Engine/BlueprintCore.h
--- a/Engine/Source/Runtime/Engine/Classes/Engine/BlueprintCore.h
+++ b/Engine/Source/Runtime/Engine/Classes/Engine/BlueprintCore.h
@@ -29,25 +29,18 @@ class ENGINE_API UBlueprintCore
 	bool bLegacyNeedToPurgeSkelRefs;
 
 private:
-
-	/** BackCompat: Whether or not this blueprint's authoritative CDO data has been migrated from the SkeletonGeneratedClass CDO to the GeneratedClass CDO */
-	UPROPERTY()
-	bool bLegacyGeneratedClassIsAuthoritative;
-
 	/** Blueprint Guid */
 	UPROPERTY()
 	FGuid BlueprintGuid;
 
 public:
+	DEPRECATED(4.22, "The minimum UE4 object version implies the Blueprint generated class is always authoritative. It is no longer necessary to explicitly set it.")
+	void SetLegacyGeneratedClassIsAuthoritative() {}
 
-	void SetLegacyGeneratedClassIsAuthoritative()
-	{
-		bLegacyGeneratedClassIsAuthoritative = true;
-	}
-
+	DEPRECATED(4.22, "The minimum UE4 object version implies the Blueprint generated class is always authoritative. It is no longer necessary to explicitly check it.")
 	bool IsGeneratedClassAuthoritative()
 	{
-		return bLegacyGeneratedClassIsAuthoritative;
+		return true;
 	}
 
 	virtual void Serialize( FArchive& Ar ) override;
diff --git a/Engine/Source/Runtime/Engine/Private/Blueprint.cpp b/Engine/Source/Runtime/Engine/Private/Blueprint.cpp
--- a/Engine/Source/Runtime/Engine/Private/Blueprint.cpp
+++ b/Engine/Source/Runtime/Engine/Private/Blueprint.cpp
@@ -4,6 +4,7 @@
 #include "Misc/CoreMisc.h"
 #include "Misc/ConfigCacheIni.h"
 #include "UObject/BlueprintsObjectVersion.h"
+#include "UObject/FrameworkObjectVersion.h"
 #include "UObject/UObjectHash.h"
 #include "Serialization/PropertyLocalizationDataGathering.h"
 #include "UObject/UnrealType.h"
@@ -285,7 +286,6 @@ UBlueprintCore::UBlueprintCore(const FObjectInitializer& ObjectInitializer)
 	{ static const FAutoRegisterLocalizationDataGatheringCallback AutomaticRegistrationOfLocalizationGatherer(UBlueprintCore::StaticClass(), &GatherBlueprintForLocalization); }
 #endif
 
-	bLegacyGeneratedClassIsAuthoritative = false;
 	bLegacyNeedToPurgeSkelRefs = true;
 }
 
@@ -295,9 +295,15 @@ void UBlueprintCore::Serialize(FArchive& Ar)
 
 #if WITH_EDITOR
 	Ar.UsingCustomVersion(FBlueprintsObjectVersion::GUID);
-#endif
+	Ar.UsingCustomVersion(FFrameworkObjectVersion::GUID);
 
-	Ar << bLegacyGeneratedClassIsAuthoritative;	
+	if (Ar.IsLoading() && Ar.CustomVer(FFrameworkObjectVersion::GUID) < FFrameworkObjectVersion::BlueprintGeneratedClassIsAlwaysAuthoritative)
+	{
+		// No longer in use.
+		bool bLegacyGeneratedClassIsAuthoritative;
+		Ar << bLegacyGeneratedClassIsAuthoritative;
+	}
+#endif
 
 	if ((Ar.UE4Ver() < VER_UE4_BLUEPRINT_SKEL_CLASS_TRANSIENT_AGAIN)
 		&& (Ar.UE4Ver() != VER_UE4_BLUEPRINT_SKEL_TEMPORARY_TRANSIENT))
