Sequencer: Fix fbx export node name so that it uses the sequencer display name

diff --git a/Engine/Source/Editor/UnrealEd/Private/Fbx/FbxMainExport.cpp b/Engine/Source/Editor/UnrealEd/Private/Fbx/FbxMainExport.cpp
--- a/Engine/Source/Editor/UnrealEd/Private/Fbx/FbxMainExport.cpp
+++ b/Engine/Source/Editor/UnrealEd/Private/Fbx/FbxMainExport.cpp
@@ -1464,7 +1464,7 @@ FString FFbxExporter::FLevelSequenceNodeNameAdapter::GetActorNodeName(const AAct
 		{
 			if (RuntimeObject.Get() == Actor)
 			{
-				NodeName = MovieSceneBinding.GetName();
+				NodeName = MovieScene->GetObjectDisplayName(MovieSceneBinding.GetObjectGuid()).ToString();
 			}
 		}
 	}
