Better fix for a Cook On The Fly Server crash - take the number of platforms into account when a file is missing from the server

diff --git a/Engine/Source/Editor/UnrealEd/Private/CookOnTheFlyServer.cpp b/Engine/Source/Editor/UnrealEd/Private/CookOnTheFlyServer.cpp
--- a/Engine/Source/Editor/UnrealEd/Private/CookOnTheFlyServer.cpp
+++ b/Engine/Source/Editor/UnrealEd/Private/CookOnTheFlyServer.cpp
@@ -3883,7 +3883,7 @@ void UCookOnTheFlyServer::SaveCookedPackage(UPackage* Package, uint32 SaveFlags,
 			}
 		}
 		
-		for (int32 PlatformIndex = 0; PlatformIndex < Platforms.Num(); ++PlatformIndex )
+		for (int32 PlatformIndex = 0; PlatformIndex < Platforms.Num(); ++PlatformIndex)
 		{
 			SavePackageResults.Add(FSavePackageResultStruct(ESavePackageResult::Success));
 			ITargetPlatform* Target = Platforms[PlatformIndex];
@@ -4019,7 +4019,10 @@ void UCookOnTheFlyServer::SaveCookedPackage(UPackage* Package, uint32 SaveFlags,
 	}
 	else
 	{
-		SavePackageResults.Add(FSavePackageResultStruct(ESavePackageResult::MissingFile));
+		for (int32 PlatformIndex = 0; PlatformIndex < TargetPlatformNames.Num(); ++PlatformIndex)
+		{
+			SavePackageResults.Add(FSavePackageResultStruct(ESavePackageResult::MissingFile));
+		}
 	}
 
 	check(bIsSavingPackage == true);
