Scripted FBX import can now autodetect type

diff --git a/Engine/Source/Editor/UnrealEd/Private/Fbx/FbxFactory.cpp b/Engine/Source/Editor/UnrealEd/Private/Fbx/FbxFactory.cpp
--- a/Engine/Source/Editor/UnrealEd/Private/Fbx/FbxFactory.cpp
+++ b/Engine/Source/Editor/UnrealEd/Private/Fbx/FbxFactory.cpp
@@ -270,6 +270,11 @@ UObject* UFbxFactory::FactoryCreateFile
 	UFbxImportUI* OverrideImportUI = AssetImportTask ? Cast<UFbxImportUI>(AssetImportTask->Options) : nullptr;
 	if (OverrideImportUI)
 	{
+		if (AssetImportTask->bAutomated && OverrideImportUI->bAutomatedImportShouldDetectType)
+		{
+			OverrideImportUI->MeshTypeToImport = ImportUI->MeshTypeToImport;
+			OverrideImportUI->OriginalImportType = ImportUI->OriginalImportType;
+		}
 		ImportUI = OverrideImportUI;
 	}
 	//We are not re-importing
