Fix for default asset names

diff --git a/Engine/Source/Editor/UnrealEd/Classes/Factories/MaterialFunctionInstanceFactory.h b/Engine/Source/Editor/UnrealEd/Classes/Factories/MaterialFunctionInstanceFactory.h
--- a/Engine/Source/Editor/UnrealEd/Classes/Factories/MaterialFunctionInstanceFactory.h
+++ b/Engine/Source/Editor/UnrealEd/Classes/Factories/MaterialFunctionInstanceFactory.h
@@ -27,6 +27,10 @@ class UMaterialFunctionMaterialLayerInstanceFactory : public UMaterialFunctionIn
 
 	//~ Begin UFactory Interface
 	virtual UObject* FactoryCreateNew(UClass* Class, UObject* InParent, FName Name, EObjectFlags Flags, UObject* Context, FFeedbackContext* Warn) override;
+	FString GetDefaultNewAssetName() const override
+	{
+		return FString(TEXT("NewMaterialLayerInstance"));
+	}
 	//~ Begin UFactory Interface	
 };
 
@@ -38,5 +42,9 @@ class UMaterialFunctionMaterialLayerBlendInstanceFactory : public UMaterialFunct
 
 	//~ Begin UFactory Interface
 	virtual UObject* FactoryCreateNew(UClass* Class, UObject* InParent, FName Name, EObjectFlags Flags, UObject* Context, FFeedbackContext* Warn) override;
+	FString GetDefaultNewAssetName() const override
+	{
+		return FString(TEXT("NewMaterialLayerBlendInstance"));
+	}
 	//~ Begin UFactory Interface	
 };
\ No newline at end of file
diff --git a/Engine/Source/Editor/UnrealEd/Classes/Factories/MaterialFunctionMaterialLayerBlendFactory.h b/Engine/Source/Editor/UnrealEd/Classes/Factories/MaterialFunctionMaterialLayerBlendFactory.h
--- a/Engine/Source/Editor/UnrealEd/Classes/Factories/MaterialFunctionMaterialLayerBlendFactory.h
+++ b/Engine/Source/Editor/UnrealEd/Classes/Factories/MaterialFunctionMaterialLayerBlendFactory.h
@@ -19,6 +19,10 @@ class UMaterialFunctionMaterialLayerBlendFactory : public UFactory
 	//~ Begin UFactory Interface
 	virtual bool CanCreateNew() const override;
 	virtual UObject* FactoryCreateNew(UClass* Class,UObject* InParent,FName Name,EObjectFlags Flags,UObject* Context,FFeedbackContext* Warn) override;
+	FString GetDefaultNewAssetName() const override
+	{
+		return FString(TEXT("NewMaterialLayerBlend"));
+	}
 	//~ Begin UFactory Interface	
 };
 
diff --git a/Engine/Source/Editor/UnrealEd/Classes/Factories/MaterialFunctionMaterialLayerFactory.h b/Engine/Source/Editor/UnrealEd/Classes/Factories/MaterialFunctionMaterialLayerFactory.h
--- a/Engine/Source/Editor/UnrealEd/Classes/Factories/MaterialFunctionMaterialLayerFactory.h
+++ b/Engine/Source/Editor/UnrealEd/Classes/Factories/MaterialFunctionMaterialLayerFactory.h
@@ -19,6 +19,10 @@ class UMaterialFunctionMaterialLayerFactory : public UFactory
 	//~ Begin UFactory Interface
 	virtual bool CanCreateNew() const override;
 	virtual UObject* FactoryCreateNew(UClass* Class,UObject* InParent,FName Name,EObjectFlags Flags,UObject* Context,FFeedbackContext* Warn) override;
+	FString GetDefaultNewAssetName() const override
+	{
+		return FString(TEXT("NewMaterialLayer"));
+	}
 	//~ Begin UFactory Interface	
 };
 
