Fix for Reset Layout not resetting subeditors

diff --git a/Engine/Source/Editor/UnrealEd/Public/Toolkits/AssetEditorToolkit.h b/Engine/Source/Editor/UnrealEd/Public/Toolkits/AssetEditorToolkit.h
--- a/Engine/Source/Editor/UnrealEd/Public/Toolkits/AssetEditorToolkit.h
+++ b/Engine/Source/Editor/UnrealEd/Public/Toolkits/AssetEditorToolkit.h
@@ -13,6 +13,7 @@
 #include "Framework/Docking/LayoutService.h"
 #include "Toolkits/IToolkitHost.h"
 #include "Toolkits/BaseToolkit.h"
+#include "UnrealEdMisc.h"
 
 class FAssetEditorModeManager;
 class FMenuBuilder;
@@ -317,7 +318,10 @@ private:
 	// Callback for persisting the Asset Editor's layout.
 	void HandleTabManagerPersistLayout( const TSharedRef<FTabManager::FLayout>& LayoutToSave )
 	{
-		FLayoutSaveRestore::SaveToConfig(GEditorLayoutIni, LayoutToSave);
+		if (FUnrealEdMisc::Get().IsSavingLayoutOnClosedAllowed())
+		{
+			FLayoutSaveRestore::SaveToConfig(GEditorLayoutIni, LayoutToSave);
+		}
 	}
 
 private:
