Copy delegates onto new version of table when reimporting curve and data tables.

diff --git a/Engine/Source/Editor/UnrealEd/Private/Factories/CSVImportFactory.cpp b/Engine/Source/Editor/UnrealEd/Private/Factories/CSVImportFactory.cpp
--- a/Engine/Source/Editor/UnrealEd/Private/Factories/CSVImportFactory.cpp
+++ b/Engine/Source/Editor/UnrealEd/Private/Factories/CSVImportFactory.cpp
@@ -211,8 +211,11 @@ UObject* UCSVImportFactory::FactoryCreateText(UClass* InClass, UObject* InParent
 
 			// If there is an existing table, need to call this to free data memory before recreating object
 			bool bStripFromClientBuilds = false;
+			UDataTable::FOnDataTableChanged OldOnDataTableChanged;
 			if (ExistingTable != nullptr)
 			{
+				OldOnDataTableChanged = MoveTemp(ExistingTable->OnDataTableChanged());
+				ExistingTable->OnDataTableChanged().Clear();
 				DataTableClass = ExistingTable->GetClass();
 				ExistingTable->EmptyTable();
 				bStripFromClientBuilds = ExistingTable->bStripFromClientBuilds;
@@ -226,6 +229,10 @@ UObject* UCSVImportFactory::FactoryCreateText(UClass* InClass, UObject* InParent
 			// Go ahead and create table from string
 			Problems = DoImportDataTable(NewTable, String);
 
+			// hook delegates back up and inform listeners of changes
+			NewTable->OnDataTableChanged() = MoveTemp(OldOnDataTableChanged);
+			NewTable->OnDataTableChanged().Broadcast();
+
 			// Print out
 			UE_LOG(LogCSVImportFactory, Log, TEXT("Imported DataTable '%s' - %d Problems"), *InName.ToString(), Problems.Num());
 			NewAsset = NewTable;
@@ -235,8 +242,11 @@ UObject* UCSVImportFactory::FactoryCreateText(UClass* InClass, UObject* InParent
 			UClass* CurveTableClass = UCurveTable::StaticClass();
 
 			// If there is an existing table, need to call this to free data memory before recreating object
+			UCurveTable::FOnCurveTableChanged OldOnCurveTableChanged;
 			if (ExistingCurveTable != nullptr)
 			{
+				OldOnCurveTableChanged = MoveTemp(ExistingCurveTable->OnCurveTableChanged());
+				ExistingCurveTable->OnCurveTableChanged().Clear();
 				CurveTableClass = ExistingCurveTable->GetClass();
 				ExistingCurveTable->EmptyTable();
 			}
@@ -248,6 +258,10 @@ UObject* UCSVImportFactory::FactoryCreateText(UClass* InClass, UObject* InParent
 			// Go ahead and create table from string
 			Problems = DoImportCurveTable(NewTable, String, ImportCurveInterpMode);
 
+			// hook delegates back up and inform listeners of changes
+			NewTable->OnCurveTableChanged() = MoveTemp(OldOnCurveTableChanged);
+			NewTable->OnCurveTableChanged().Broadcast();
+
 			// Print out
 			UE_LOG(LogCSVImportFactory, Log, TEXT("Imported CurveTable '%s' - %d Problems"), *InName.ToString(), Problems.Num());
 			NewAsset = NewTable;
