Autosave no longer saves map build data

diff --git a/Engine/Source/Editor/UnrealEd/Private/FileHelpers.cpp b/Engine/Source/Editor/UnrealEd/Private/FileHelpers.cpp
--- a/Engine/Source/Editor/UnrealEd/Private/FileHelpers.cpp
+++ b/Engine/Source/Editor/UnrealEd/Private/FileHelpers.cpp
@@ -2619,15 +2619,34 @@ EAutosaveContentPackagesResult::Type FEditorFileUtils::AutosaveContentPackagesEx
 		// If the package is dirty and is not the transient package, we'd like to autosave it
 		if ( CurPackage && ( CurPackage != TransientPackage ) && CurPackage->IsDirty() && (bForceIfNotInList || DirtyPackagesForAutoSave.Contains(CurPackage)) )
 		{
-			UWorld* MapWorld = UWorld::FindWorldInPackage(CurPackage);
+			bool bSkipPackage = false;
+			TArray<UObject*> ObjectsInPackage;
+			GetObjectsWithOuter(CurPackage, ObjectsInPackage, false);
+			for (auto ObjIt = ObjectsInPackage.CreateConstIterator(); ObjIt; ++ObjIt)
+			{
+				// Also, make sure this is not a map package
+				if (Cast<UWorld>(*ObjIt))
+				{
+					bSkipPackage = true;
+					break;
+				}
+				else if (Cast<UMapBuildDataRegistry>(*ObjIt))
+				{
+					// Do not auto save generated map build data packages
+					bSkipPackage = true;
+					break;
+				}
+			}
 
-			// Also, make sure this is not a map package
-			const bool bIsMapPackage = MapWorld != NULL;
+			if (bSkipPackage)
+			{
+				continue;
+			}
 
 			// Ignore packages with long, invalid names. This culls out packages with paths in read-only roots such as /Temp.
 			const bool bInvalidLongPackageName = !FPackageName::IsShortPackageName(CurPackage->GetFName()) && !FPackageName::IsValidLongPackageName(CurPackage->GetName(), /*bIncludeReadOnlyRoots=*/false);
 				
-			if ( !bIsMapPackage && !bInvalidLongPackageName )
+			if ( !bInvalidLongPackageName )
 			{
 				PackagesToSave.Add(CurPackage);
 			}
