Copy of cl: 4335009 into 4.21
Fix crash when unattended

diff --git a/Engine/Source/Editor/UnrealEd/Private/Factories/EditorFactories.cpp b/Engine/Source/Editor/UnrealEd/Private/Factories/EditorFactories.cpp
--- a/Engine/Source/Editor/UnrealEd/Private/Factories/EditorFactories.cpp
+++ b/Engine/Source/Editor/UnrealEd/Private/Factories/EditorFactories.cpp
@@ -5299,7 +5299,8 @@ EReimportResult::Type UReimportFbxStaticMeshFactory::Reimport( UObject* Obj )
 	{
 		ImportUI = NewObject<UFbxImportUI>(this, NAME_None, RF_Public);
 	}
-	const bool IsUnattended = GIsAutomationTesting || FApp::IsUnattended();
+	//Prevent any UI for automation, unattended and commandlet
+	const bool IsUnattended = GIsAutomationTesting || FApp::IsUnattended() || IsRunningCommandlet();
 	const bool ShowImportDialogAtReimport = GetDefault<UEditorPerProjectUserSettings>()->bShowImportDialogAtReimport && !IsUnattended;
 
 	if (ImportData == nullptr)
@@ -5568,7 +5569,8 @@ EReimportResult::Type UReimportFbxSkeletalMeshFactory::Reimport( UObject* Obj )
 	}
 
 	bool bSuccess = false;
-	const bool IsUnattended = GIsAutomationTesting || FApp::IsUnattended();
+	//Prevent any UI for automation, unattended and commandlet
+	const bool IsUnattended = GIsAutomationTesting || FApp::IsUnattended() || IsRunningCommandlet();
 	const bool ShowImportDialogAtReimport = GetDefault<UEditorPerProjectUserSettings>()->bShowImportDialogAtReimport && !IsUnattended;
 
 	if (ImportData == nullptr)
