CIS fix (shadow variable).

#rnx

diff --git a/Engine/Source/Editor/UnrealEd/Private/Kismet2/BlueprintEditorUtils.cpp b/Engine/Source/Editor/UnrealEd/Private/Kismet2/BlueprintEditorUtils.cpp
--- a/Engine/Source/Editor/UnrealEd/Private/Kismet2/BlueprintEditorUtils.cpp
+++ b/Engine/Source/Editor/UnrealEd/Private/Kismet2/BlueprintEditorUtils.cpp
@@ -1970,12 +1970,12 @@ void FBlueprintEditorUtils::PostDuplicateBlueprint(UBlueprint* Blueprint, bool b
 			{
 				// Skip CDO validation in this case as we will not have yet propagated values to the new CDO. Also skip
 				// Blueprint search data updates, as that will be handled by an OnAssetAdded() delegate in the FiB manager.
-				const EBlueprintCompileOptions CompileOptions =
+				const EBlueprintCompileOptions BPCompileOptions =
 					EBlueprintCompileOptions::SkipDefaultObjectValidation |
 					EBlueprintCompileOptions::SkipFiBSearchMetaUpdate;
 
 				FBlueprintCompilationManager::CompileSynchronously(
-					FBPCompileRequest(Blueprint, CompileOptions, nullptr)
+					FBPCompileRequest(Blueprint, BPCompileOptions, nullptr)
 				);
 			}
 			else
