Replace overly aggressive ensures in bookmark code.

diff --git a/Engine/Source/Editor/UnrealEd/Private/Bookmarks/BookmarkTypeToolsImpl.cpp b/Engine/Source/Editor/UnrealEd/Private/Bookmarks/BookmarkTypeToolsImpl.cpp
--- a/Engine/Source/Editor/UnrealEd/Private/Bookmarks/BookmarkTypeToolsImpl.cpp
+++ b/Engine/Source/Editor/UnrealEd/Private/Bookmarks/BookmarkTypeToolsImpl.cpp
@@ -1,4 +1,4 @@
-﻿// Copyright 1998-2018 Epic Games, Inc. All Rights Reserved.
+// Copyright 1998-2018 Epic Games, Inc. All Rights Reserved.
 
 #include "Bookmarks/IBookmarkTypeTools.h"
 #include "Bookmarks/IBookmarkTypeActions.h"
@@ -8,6 +8,10 @@
 #include "EditorViewportClient.h"
 #include "Engine/BookmarkBase.h"
 
+#include "Logging/LogMacros.h"
+
+DEFINE_LOG_CATEGORY_STATIC(LogEditorBookmarks, Warning, Warning);
+
 class FBookmarkTypeToolsImpl : public IBookmarkTypeTools
 {
 private:
@@ -23,7 +27,7 @@ private:
 
 	static FORCEINLINE const AWorldSettings* GetWorldSettings(const FEditorViewportClient* InViewportClient)
 	{
-		if (ensure(InViewportClient != nullptr))
+		if (InViewportClient != nullptr)
 		{
 			if (const UWorld* World = InViewportClient->GetWorld())
 			{
@@ -55,7 +59,7 @@ private:
 			}
 		}
 		
-		ensureMsgf(false, TEXT("FBookmarkTypeToolsImpl::GetBookmarkTypeActions - Unable to get appropriate BookmarkTypeActions for Class %s"), *GetPathNameSafe(Class));
+		UE_LOG(LogEditorBookmarks, Warning, TEXT("FBookmarkTypeToolsImpl::GetBookmarkTypeActions - Unable to get appropriate BookmarkTypeActions for Class %s"), *GetPathNameSafe(Class));
 		return nullptr;
 	}
 
@@ -106,11 +110,14 @@ public:
 			const TSharedPtr<IBookmarkTypeActions> Actions = GetBookmarkTypeActions(*WorldSettings);
 			if (Actions.IsValid())
 			{
-				UBookmarkBase* Bookmark = WorldSettings->GetOrAddBookmark(InIndex, true);
-				if (ensureMsgf(Bookmark != nullptr, TEXT("FBookmarkTypeToolsImpl::CreateOrSetBookmark - Failed to create bookmark at Index %d"), InIndex))
+				if (UBookmarkBase* Bookmark = WorldSettings->GetOrAddBookmark(InIndex, true))
 				{
 					Actions->InitFromViewport(Bookmark, *InViewportClient);
 				}
+				else
+				{
+					UE_LOG(LogEditorBookmarks, Error, TEXT("FBookmarkTypeToolsImpl::CreateOrSetBookmark - Failed to create bookmark at Index %d"), InIndex);
+				}
 			}
 		}
 	}
@@ -131,10 +138,9 @@ public:
 
 	/**
 	 * Jumps to a bookmark from the list.
-	 * This assumes the bookmark index is valid. Use CheckBookmark before calling this.
 	 * 
 	 * @param InIndex			Index of the bookmark to set.
-	 * @param InSettings		Settings to used when jumpting to the bookmark.
+	 * @param InSettings		Settings to used when jumping to the bookmark.
 	 * @param InViewportClient	Level editor viewport client used to reference the world which owns the bookmark.
 	 */
 	virtual void JumpToBookmark(const uint32 InIndex, const TSharedPtr<struct FBookmarkBaseJumpToSettings> InSettings, FEditorViewportClient* InViewportClient) const override
@@ -145,13 +151,20 @@ public:
 			if (Actions.IsValid())
 			{
 				const TArray<UBookmarkBase*> Bookmarks = WorldSettings->GetBookmarks();
-				if (ensureMsgf(Bookmarks.IsValidIndex(InIndex), TEXT("FBookmarkTypeToolsImpl::JumpToBookmark - Invalid bookmark index %d"), InIndex))
+				if (Bookmarks.IsValidIndex(InIndex))
 				{
-					UBookmarkBase* Bookmark = Bookmarks[InIndex];
-					if (ensureMsgf(Bookmark != nullptr, TEXT("FBookmarkTypeToolsImpl::JumpToBookmark - Null Bookmark at index %d"), InIndex))
+					if (UBookmarkBase* Bookmark = Bookmarks[InIndex])
 					{
 						Actions->JumpToBookmark(Bookmark, InSettings, *InViewportClient);
 					}
+					else
+					{
+						UE_LOG(LogEditorBookmarks, Warning, TEXT("FBookmarkTypeToolsImpl::JumpToBookmark - Null Bookmark at index %d"), InIndex);
+					}
+				}
+				else
+				{
+					UE_LOG(LogEditorBookmarks, Warning, TEXT("FBookmarkTypeToolsImpl::JumpToBookmark - Invalid bookmark index %d"), InIndex);
 				}
 			}
 		}
@@ -229,61 +242,62 @@ public:
 
 		if (AWorldSettings* WorldSettings = GetWorldSettings(InViewportClient))
 		{
-			// Sanity check to make sure we are handling the appropriate WorldSettings object.
-			if (ensureMsgf(WorldSettings == InWorldSettings,
-				TEXT("FBookmarkTypeToolsImpl::UpgradeBookmarks - Viewport client does not correspond to correct world (Viewport WorldSettings = %s Expected WorldSettings = %s)"),
-				*GetPathNameSafe(WorldSettings), *GetPathNameSafe(InWorldSettings)))
+			if (WorldSettings != InWorldSettings)
 			{
-				// Make sure the new bookmark class and type actions are valid.
-				if (UClass* NewBookmarkClass = InWorldSettings->GetDefaultBookmarkClass().Get())
+				UE_LOG(LogEditorBookmarks, Warning, TEXT("FBookmarkTypeToolsImpl::UpgradeBookmarks - Viewport client does not correspond to correct world (Viewport WorldSettings = %s Expected WorldSettings = %s)"),
+					*GetPathNameSafe(WorldSettings), *GetPathNameSafe(InWorldSettings));
+				return;
+			}
+
+			// Make sure the new bookmark class and type actions are valid.
+			if (UClass* NewBookmarkClass = InWorldSettings->GetDefaultBookmarkClass().Get())
+			{
+				const TSharedPtr<IBookmarkTypeActions> NewBookmarkActions = GetBookmarkTypeActions(*InWorldSettings);
+				if (NewBookmarkActions.IsValid())
 				{
-					const TSharedPtr<IBookmarkTypeActions> NewBookmarkActions = GetBookmarkTypeActions(*InWorldSettings);
-					if (NewBookmarkActions.IsValid())
+					UClass* OldBookmarkClass = nullptr;
+					TSharedPtr<IBookmarkTypeActions> OldBookmarkActions = nullptr;
+
+					// Cache off our current viewport state.
+					FViewportCameraTransform ViewportTransform = InViewportClient->GetViewTransform();
+					EViewModeIndex ViewMode = InViewportClient->GetViewMode();
+
+					const int32 NumBookmarks = InWorldSettings->GetMaxNumberOfBookmarks();
+					const TArray<UBookmarkBase*>& AvailableBookmarks = InWorldSettings->GetBookmarks();
+
+					// Here, we'll go through each existing bookmark, jump to it, then create a new bookmark
+					// from that restored state.
+					// This approach isn't going to be perfect due to potential incompatibility.
+					// Alternatives may be to allow intermixing of Bookmark types, or completely obliterating bookmarks,
+					// neither of which seem like great alternatives.
+					for (int32 i = 0; i < NumBookmarks; ++i)
 					{
-						UClass* OldBookmarkClass = nullptr;
-						TSharedPtr<IBookmarkTypeActions> OldBookmarkActions = nullptr;
-
-						// Cache off our current viewport state.
-						FViewportCameraTransform ViewportTransform = InViewportClient->GetViewTransform();
-						EViewModeIndex ViewMode = InViewportClient->GetViewMode();
-
-						const int32 NumBookmarks = InWorldSettings->GetMaxNumberOfBookmarks();
-						const TArray<UBookmarkBase*>& AvailableBookmarks = InWorldSettings->GetBookmarks();
-
-						// Here, we'll go through each existing bookmark, jump to it, then create a new bookmark
-						// from that restored state.
-						// This approach isn't going to be perfect due to potential incompatibility.
-						// Alternatives may be to allow intermixing of Bookmark types, or completely obliterating bookmarks,
-						// neither of which seem like great alternatives.
-						for (int32 i = 0; i < NumBookmarks; ++i)
+						if (UBookmarkBase* OldBookmark = AvailableBookmarks[i])
 						{
-							if (UBookmarkBase* OldBookmark = AvailableBookmarks[i])
+							if (OldBookmarkClass == nullptr ||
+								OldBookmarkClass != OldBookmark->GetClass())
 							{
-								if (OldBookmarkClass == nullptr ||
-									OldBookmarkClass != OldBookmark->GetClass())
-								{
-									OldBookmarkClass = OldBookmark->GetClass();
-									OldBookmarkActions = GetBookmarkTypeActions(OldBookmarkClass);
-								}
-
-								InWorldSettings->ClearBookmark(i);
-
-								if (OldBookmarkActions.IsValid())
-								{
-									OldBookmarkActions->JumpToBookmark(OldBookmark, nullptr, *InViewportClient);
-									UBookmarkBase* NewBookmark = InWorldSettings->GetOrAddBookmark(i, false);
-									NewBookmarkActions->InitFromViewport(NewBookmark, *InViewportClient);
-								}
+								OldBookmarkClass = OldBookmark->GetClass();
+								OldBookmarkActions = GetBookmarkTypeActions(OldBookmarkClass);
 							}
-						}
 
-						// Restore our viewport state.
-						InViewportClient->SetViewLocation(ViewportTransform.GetLocation());
-						InViewportClient->SetLookAtLocation(ViewportTransform.GetLookAt());
-						InViewportClient->SetOrthoZoom(ViewportTransform.GetOrthoZoom());
-						InViewportClient->SetViewRotation(ViewportTransform.GetRotation());
-						InViewportClient->SetViewMode(ViewMode);
+							InWorldSettings->ClearBookmark(i);
+
+							if (OldBookmarkActions.IsValid())
+							{
+								OldBookmarkActions->JumpToBookmark(OldBookmark, nullptr, *InViewportClient);
+								UBookmarkBase* NewBookmark = InWorldSettings->GetOrAddBookmark(i, false);
+								NewBookmarkActions->InitFromViewport(NewBookmark, *InViewportClient);
+							}
+						}
 					}
+
+					// Restore our viewport state.
+					InViewportClient->SetViewLocation(ViewportTransform.GetLocation());
+					InViewportClient->SetLookAtLocation(ViewportTransform.GetLookAt());
+					InViewportClient->SetOrthoZoom(ViewportTransform.GetOrthoZoom());
+					InViewportClient->SetViewRotation(ViewportTransform.GetRotation());
+					InViewportClient->SetViewMode(ViewMode);
 				}
 			}
 		}
