Fix for crash if the file path to reimport from is not set

diff --git a/Engine/Source/Editor/UnrealEd/Private/Factories/CSVImportFactory.cpp b/Engine/Source/Editor/UnrealEd/Private/Factories/CSVImportFactory.cpp
--- a/Engine/Source/Editor/UnrealEd/Private/Factories/CSVImportFactory.cpp
+++ b/Engine/Source/Editor/UnrealEd/Private/Factories/CSVImportFactory.cpp
@@ -409,7 +409,14 @@ bool UReimportDataTableFactory::CanReimport( UObject* Obj, TArray<FString>& OutF
 	if(DataTable)
 	{
 		DataTable->AssetImportData->ExtractFilenames(OutFilenames);
-		return true;
+		if (OutFilenames.Num() > 0)
+		{
+			return true;
+		}
+		else
+		{
+			return false;
+		}
 	}
 	return false;
 }
