Fix compile error fallout from CL# 4493121
#rnx

diff --git a/Engine/Source/Editor/UnrealEd/Private/Kismet2/KismetReinstanceUtilities.cpp b/Engine/Source/Editor/UnrealEd/Private/Kismet2/KismetReinstanceUtilities.cpp
--- a/Engine/Source/Editor/UnrealEd/Private/Kismet2/KismetReinstanceUtilities.cpp
+++ b/Engine/Source/Editor/UnrealEd/Private/Kismet2/KismetReinstanceUtilities.cpp
@@ -1476,7 +1476,7 @@ void FBlueprintCompileReinstancer::ReplaceInstancesOfClassEx(const FReplaceInsta
 {
 	TMap<UClass*, UClass*> OldToNewClassMap;
 	OldToNewClassMap.Add(Parameters.OldClass, Parameters.NewClass);
-	ReplaceInstancesOfClass_Inner(OldToNewClassMap, Parameters.OriginalCDO, Parameters.ObjectsThatShouldUseOldStuff, Parameters.bClassObjectReplaced, Parameters.bPreserveRootComponent, Parameters.InstancesThatShouldUseOldClass);
+	ReplaceInstancesOfClass_Inner(OldToNewClassMap, Parameters.OriginalCDO, Parameters.ObjectsThatShouldUseOldStuff, Parameters.bClassObjectReplaced, Parameters.bPreserveRootComponent, /*bArchetypesAreUpToDate=*/false, Parameters.InstancesThatShouldUseOldClass);
 }
 
 void FBlueprintCompileReinstancer::BatchReplaceInstancesOfClass(TMap<UClass*, UClass*>& InOldToNewClassMap, bool bArchetypesAreUpToDate)
