Fix re-import skeletalmesh crash

diff --git a/Engine/Source/Editor/UnrealEd/Private/Factories/SkeletalMeshImport.cpp b/Engine/Source/Editor/UnrealEd/Private/Factories/SkeletalMeshImport.cpp
--- a/Engine/Source/Editor/UnrealEd/Private/Factories/SkeletalMeshImport.cpp
+++ b/Engine/Source/Editor/UnrealEd/Private/Factories/SkeletalMeshImport.cpp
@@ -516,7 +516,6 @@ ExistingSkelMeshData* SaveExistingSkelMeshData(USkeletalMesh* ExistingSkelMesh,
 void TryRegenerateLODs(ExistingSkelMeshData* MeshData, USkeletalMesh* SkeletalMesh)
 {
 	int32 TotalLOD = MeshData->ExistingLODModels.Num();
-
 	// see if mesh reduction util is available
 	IMeshReductionManagerModule& Module = FModuleManager::Get().LoadModuleChecked<IMeshReductionManagerModule>("MeshReductionInterface");
 	static bool bAutoMeshReductionAvailable = Module.GetSkeletalMeshReductionInterface() != NULL;
@@ -526,19 +525,29 @@ void TryRegenerateLODs(ExistingSkelMeshData* MeshData, USkeletalMesh* SkeletalMe
 		GWarn->BeginSlowTask(LOCTEXT("RegenLODs", "Generating new LODs"), true);
 		FSkeletalMeshUpdateContext UpdateContext;
 		UpdateContext.SkeletalMesh = SkeletalMesh;
-
+		TArray<bool> Dependencies;
+		Dependencies.AddZeroed(TotalLOD+1);
+		Dependencies[0] = true;
 		for (int32 Index = 0; Index < TotalLOD; ++Index)
 		{
 			int32 LODIndex = Index + 1;
-			FSkeletalMeshLODInfo& LODInfo = MeshData->ExistingLODInfo[Index];
-			// reset material maps, it won't work anyway. 
-			LODInfo.LODMaterialMap.Empty();
-			// add LOD info back
-			SkeletalMesh->AddLODInfo(LODInfo); 
-			// force it to regen
-			FLODUtilities::SimplifySkeletalMeshLOD(UpdateContext, LODIndex, false);
+			if (LODIndex >= SkeletalMesh->GetLODInfoArray().Num())
+			{
+				FSkeletalMeshLODInfo& ExistLODInfo = MeshData->ExistingLODInfo[Index];
+				// reset material maps, it won't work anyway. 
+				ExistLODInfo.LODMaterialMap.Empty();
+				// add LOD info back
+				SkeletalMesh->AddLODInfo(ExistLODInfo);
+				check(LODIndex < SkeletalMesh->GetLODInfoArray().Num());
+			}
+			const FSkeletalMeshLODInfo* LODInfo = SkeletalMesh->GetLODInfo(LODIndex);
+			if (LODInfo && LODInfo->bHasBeenSimplified && Dependencies[LODInfo->ReductionSettings.BaseLOD])
+			{
+				Dependencies[LODIndex] = true;
+				// force it to regenerate
+				FLODUtilities::SimplifySkeletalMeshLOD(UpdateContext, LODIndex, false);
+			}
 		}
-
 		GWarn->EndSlowTask();
 	}
 	else
