Fix typo in fbx import dialog option, for Use T0 As Ref Pose

diff --git a/Engine/Source/Editor/UnrealEd/Classes/Factories/FbxSkeletalMeshImportData.h b/Engine/Source/Editor/UnrealEd/Classes/Factories/FbxSkeletalMeshImportData.h
--- a/Engine/Source/Editor/UnrealEd/Classes/Factories/FbxSkeletalMeshImportData.h
+++ b/Engine/Source/Editor/UnrealEd/Classes/Factories/FbxSkeletalMeshImportData.h
@@ -38,7 +38,7 @@ public:
 	uint32 bUpdateSkeletonReferencePose:1;
 
 	/** Enable this option to use frame 0 as reference pose */
-	UPROPERTY(EditAnywhere, AdvancedDisplay, config, Category= Mesh, meta=(ImportType="SkeletalMesh|RigOnly"))
+	UPROPERTY(EditAnywhere, AdvancedDisplay, config, Category= Mesh, meta=(ImportType="SkeletalMesh|RigOnly", DisplayName="Use T0 As Ref Pose"))
 	uint32 bUseT0AsRefPose:1;
 
 	/** If checked, triangles with non-matching smoothing groups will be physically split. */
