Fix crash with re-import and there is some NavCollision data or User data store in the static mesh

diff --git a/Engine/Source/Editor/UnrealEd/Private/Factories/EditorFactories.cpp b/Engine/Source/Editor/UnrealEd/Private/Factories/EditorFactories.cpp
--- a/Engine/Source/Editor/UnrealEd/Private/Factories/EditorFactories.cpp
+++ b/Engine/Source/Editor/UnrealEd/Private/Factories/EditorFactories.cpp
@@ -5382,15 +5382,22 @@ EReimportResult::Type UReimportFbxStaticMeshFactory::Reimport( UObject* Obj )
 		if ( FFbxImporter->ImportFromFile( *Filename, FPaths::GetExtension( Filename ), true ) )
 		{
 			FFbxImporter->ApplyTransformSettingsToFbxNode(FFbxImporter->Scene->GetRootNode(), ImportData);
+
+			// preserve the user data by doing a copy
 			const TArray<UAssetUserData*>* UserData = Mesh->GetAssetUserDataArray();
-			TArray<UAssetUserData*> UserDataCopy;
+			TMap<UAssetUserData*, bool> UserDataCopy;
 			if (UserData)
 			{
 				for (int32 Idx = 0; Idx < UserData->Num(); Idx++)
 				{
 					if ((*UserData)[Idx] != nullptr)
 					{
-						UserDataCopy.Add((UAssetUserData*)StaticDuplicateObject((*UserData)[Idx], GetTransientPackage()));
+						bool bAddDupToRoot = !((*UserData)[Idx]->IsRooted());
+						if (bAddDupToRoot)
+						{
+							(*UserData)[Idx]->AddToRoot();
+						}
+						UserDataCopy.Add((UAssetUserData*)StaticDuplicateObject((*UserData)[Idx], GetTransientPackage()), bAddDupToRoot);
 					}
 				}
 			}
@@ -5400,6 +5407,13 @@ EReimportResult::Type UReimportFbxStaticMeshFactory::Reimport( UObject* Obj )
 				(UNavCollisionBase*)StaticDuplicateObject(Mesh->NavCollision, GetTransientPackage()) :
 				nullptr;
 
+			bool bAddedNavCollisionDupToRoot = false;
+			if (NavCollision && !NavCollision->IsRooted())
+			{
+				bAddedNavCollisionDupToRoot = true;
+				NavCollision->AddToRoot();
+			}
+
 			// preserve extended bound settings
 			const FVector PositiveBoundsExtension = Mesh->PositiveBoundsExtension;
 			const FVector NegativeBoundsExtension = Mesh->NegativeBoundsExtension;
@@ -5409,14 +5423,25 @@ EReimportResult::Type UReimportFbxStaticMeshFactory::Reimport( UObject* Obj )
 				UE_LOG(LogEditorFactories, Log, TEXT("-- imported successfully") );
 
 				// Copy user data to newly created mesh
-				for (int32 Idx = 0; Idx < UserDataCopy.Num(); Idx++)
+				for (auto Kvp : UserDataCopy)
 				{
-					UserDataCopy[Idx]->Rename(nullptr, Mesh, REN_DontCreateRedirectors | REN_DoNotDirty);
-					Mesh->AddAssetUserData(UserDataCopy[Idx]);
+					UAssetUserData* UserDataObject = Kvp.Key;
+					if (Kvp.Value)
+					{
+						//if the duplicated temporary UObject was add to root, we must remove it from the root
+						UserDataObject->RemoveFromRoot();
+					}
+					UserDataObject->Rename(nullptr, Mesh, REN_DontCreateRedirectors | REN_DoNotDirty);
+					Mesh->AddAssetUserData(UserDataObject);
 				}
 
 				if (NavCollision)
 				{
+					if (bAddedNavCollisionDupToRoot)
+					{
+						//if the duplicated temporary UObject was add to root, we must remove it from the root
+						NavCollision->RemoveFromRoot();
+					}
 					Mesh->NavCollision = NavCollision;
 					NavCollision->Rename(NULL, Mesh, REN_DontCreateRedirectors | REN_DoNotDirty);
 				}
