Saving HLODs from a new level now always saves the level first

diff --git a/Engine/Source/Developer/HierarchicalLODUtilities/Private/HierarchicalLODUtilities.cpp b/Engine/Source/Developer/HierarchicalLODUtilities/Private/HierarchicalLODUtilities.cpp
--- a/Engine/Source/Developer/HierarchicalLODUtilities/Private/HierarchicalLODUtilities.cpp
+++ b/Engine/Source/Developer/HierarchicalLODUtilities/Private/HierarchicalLODUtilities.cpp
@@ -99,6 +99,20 @@ static FString GetHLODPackageName(const ULevel* InLevel, const uint32 InHLODLeve
 	return FString::Printf(TEXT("%s/HLOD/%s"), *PathName, *InOutHLODProxyName);
 }
 
+static FString GetHLODProxyName(const FString& InLevelPackageName, const uint32 InHLODLevelIndex)
+{
+	const FString BaseName = FPackageName::GetShortName(InLevelPackageName);	
+	return FString::Printf(TEXT("%s_%i_HLOD"), *BaseName, InHLODLevelIndex);
+}
+
+static FString GetHLODPackageName(const FString& InLevelPackageName, const uint32 InHLODLevelIndex, FString& InOutHLODProxyName)
+{
+	const FString PathName = FPackageName::GetLongPackagePath(InLevelPackageName);
+	const FString BaseName = FPackageName::GetShortName(InLevelPackageName);	
+	InOutHLODProxyName = GetHLODProxyName(InLevelPackageName, InHLODLevelIndex);
+	return FString::Printf(TEXT("%s/HLOD/%s"), *PathName, *InOutHLODProxyName);
+}
+
 void FHierarchicalLODUtilities::CleanStandaloneAssetsInPackage(UPackage* InPackage)
 {
 	TArray<UObject*> Objects;
@@ -205,6 +219,12 @@ UPackage* FHierarchicalLODUtilities::CreateOrRetrieveLevelHLODPackage(const ULev
 	return HLODPackage;
 }
 
+FString FHierarchicalLODUtilities::GetLevelHLODProxyName(const FString& InLevelPackageName, const uint32 InHLODLevelIndex)
+{
+	FString HLODProxyName;
+	FString HLODPackageName = GetHLODPackageName(InLevelPackageName, InHLODLevelIndex, HLODProxyName);
+	return HLODPackageName + TEXT(".") + HLODProxyName;
+}
 
 bool FHierarchicalLODUtilities::BuildStaticMeshForLODActor(ALODActor* LODActor, UPackage* AssetsOuter, const FHierarchicalSimplification& LODSetup, UMaterialInterface* InBaseMaterial)
 {
diff --git a/Engine/Source/Developer/HierarchicalLODUtilities/Public/HierarchicalLODUtilities.h b/Engine/Source/Developer/HierarchicalLODUtilities/Public/HierarchicalLODUtilities.h
--- a/Engine/Source/Developer/HierarchicalLODUtilities/Public/HierarchicalLODUtilities.h
+++ b/Engine/Source/Developer/HierarchicalLODUtilities/Public/HierarchicalLODUtilities.h
@@ -33,6 +33,7 @@ public:
 	virtual UPackage* RetrieveLevelHLODPackage(const ULevel* InLevel, const uint32 HLODLevelIndex) override;
 	virtual UHLODProxy* CreateOrRetrieveLevelHLODProxy(const ULevel* InLevel, const uint32 HLODLevelIndex) override;
 	virtual UHLODProxy* RetrieveLevelHLODProxy(const ULevel* InLevel, const uint32 HLODLevelIndex) override;
+	virtual FString GetLevelHLODProxyName(const FString& InLevelPackageName, const uint32 HLODLevelIndex) override;
 	DEPRECATED(4.20, "Use different signature containing BaseMaterial parameter")
 	virtual bool BuildStaticMeshForLODActor(ALODActor* LODActor, UPackage* AssetsOuter, const FHierarchicalSimplification& LODSetup) override;
 	DEPRECATED(4.20, "Use different signature using UHLODProxy")
diff --git a/Engine/Source/Developer/HierarchicalLODUtilities/Public/IHierarchicalLODUtilities.h b/Engine/Source/Developer/HierarchicalLODUtilities/Public/IHierarchicalLODUtilities.h
--- a/Engine/Source/Developer/HierarchicalLODUtilities/Public/IHierarchicalLODUtilities.h
+++ b/Engine/Source/Developer/HierarchicalLODUtilities/Public/IHierarchicalLODUtilities.h
@@ -75,6 +75,9 @@ public:
 	/** Retrieves the HLOD proxy that is created for the given level */
 	virtual UHLODProxy* RetrieveLevelHLODProxy(const ULevel* InLevel, const uint32 HLODLevelIndex) = 0;
 
+	/** Get the name of a HLOD proxy given the name of a level's package */
+	virtual FString GetLevelHLODProxyName(const FString& InLevelPackageName, const uint32 HLODLevelIndex) = 0;
+
 	/**
 	* Builds a virtual mesh object for the given LODACtor
 	*
diff --git a/Engine/Source/Editor/UnrealEd/Private/HierarchicalLOD.cpp b/Engine/Source/Editor/UnrealEd/Private/HierarchicalLOD.cpp
--- a/Engine/Source/Editor/UnrealEd/Private/HierarchicalLOD.cpp
+++ b/Engine/Source/Editor/UnrealEd/Private/HierarchicalLOD.cpp
@@ -731,54 +731,105 @@ void FHierarchicalLODBuilder::BuildMeshesForLODActors(bool bForceAll)
 
 void FHierarchicalLODBuilder::SaveMeshesForActors()
 {
-	TArray<UPackage*> PackagesToSave;
+	TArray<UPackage*> LevelPackagesToSave;
+	TArray<FString> OldLevelPackageNames;
 
+	bool bUnsavedLevel = false;
 	const TArray<ULevel*>& Levels = World->GetLevels();
 	for (const ULevel* Level : Levels)
 	{
-		const TArray<FHierarchicalSimplification>& BuildLODLevelSettings = Level->GetWorldSettings()->GetHierarchicalLODSetup();
-		UMaterialInterface* BaseMaterial = Level->GetWorldSettings()->GetHierarchicalLODBaseMaterial();
-		TArray<TArray<ALODActor*>> LODLevelActors;
-		LODLevelActors.AddDefaulted(BuildLODLevelSettings.Num());
+		// Levels might also need a resave, or levels might not have been saved yet
+		LevelPackagesToSave.Add(Level->GetOutermost());
+		OldLevelPackageNames.Add(Level->GetOutermost()->GetName());
+		bUnsavedLevel |= Level->GetOutermost()->GetName().StartsWith("/Temp/");
+	}
+
+	bool bSuccess = true;
+
+	// Save levels first if they are in the /Temp/ mount point
+	if(bUnsavedLevel)
+	{
+		bSuccess = UEditorLoadingAndSavingUtils::SavePackagesWithDialog(LevelPackagesToSave, true);
+	}
+	
+	if(bSuccess)
+	{
+		check(LevelPackagesToSave.Num() == OldLevelPackageNames.Num() && LevelPackagesToSave.Num() == Levels.Num());
+
+		TArray<UPackage*> HLODPackagesToSave;
 
-		if (Level->Actors.Num() > 0)
+		for(int32 PackageIndex = 0; PackageIndex < LevelPackagesToSave.Num(); ++PackageIndex)
 		{
-			FHierarchicalLODUtilitiesModule& Module = FModuleManager::LoadModuleChecked<FHierarchicalLODUtilitiesModule>("HierarchicalLODUtilities");
-			IHierarchicalLODUtilities* Utilities = Module.GetUtilities();
+			const ULevel* Level = Levels[PackageIndex];
 
-			// Retrieve LOD actors from the level
-			uint32 NumLODActors = 0;
-			for (int32 ActorId = 0; ActorId < Level->Actors.Num(); ++ActorId)
+			HLODPackagesToSave.Add(Level->GetOutermost());
+
+			const TArray<FHierarchicalSimplification>& BuildLODLevelSettings = Level->GetWorldSettings()->GetHierarchicalLODSetup();
+			UMaterialInterface* BaseMaterial = Level->GetWorldSettings()->GetHierarchicalLODBaseMaterial();
+			TArray<TArray<ALODActor*>> LODLevelActors;
+			LODLevelActors.AddDefaulted(BuildLODLevelSettings.Num());
+
+			if (Level->Actors.Num() > 0)
 			{
-				AActor* Actor = Level->Actors[ActorId];
-				if (Actor && Actor->IsA<ALODActor>())
+				FHierarchicalLODUtilitiesModule& Module = FModuleManager::LoadModuleChecked<FHierarchicalLODUtilitiesModule>("HierarchicalLODUtilities");
+				IHierarchicalLODUtilities* Utilities = Module.GetUtilities();
+
+				// Retrieve LOD actors from the level
+				uint32 NumLODActors = 0;
+				for (int32 ActorId = 0; ActorId < Level->Actors.Num(); ++ActorId)
 				{
-					ALODActor* LODActor = CastChecked<ALODActor>(Actor);
+					AActor* Actor = Level->Actors[ActorId];
+					if (Actor && Actor->IsA<ALODActor>())
+					{
+						ALODActor* LODActor = CastChecked<ALODActor>(Actor);
 						
-					LODLevelActors[LODActor->LODLevel - 1].Add(LODActor);
-					NumLODActors++;
-				}
-			}		
+						LODLevelActors[LODActor->LODLevel - 1].Add(LODActor);
+						NumLODActors++;
+					}
+				}		
 
-			if (NumLODActors)
-			{
-				const int32 NumLODLevels = LODLevelActors.Num();
-				for (int32 LODIndex = 0; LODIndex < NumLODLevels; ++LODIndex)
+				if (NumLODActors)
 				{
-					UPackage* AssetsOuter = Utilities->RetrieveLevelHLODPackage(Level, LODIndex);
-					if(AssetsOuter)
+					const int32 NumLODLevels = LODLevelActors.Num();
+					for (int32 LODIndex = 0; LODIndex < NumLODLevels; ++LODIndex)
 					{
-						PackagesToSave.Add(AssetsOuter);
+						UPackage* AssetsOuter = Utilities->RetrieveLevelHLODPackage(Level, LODIndex);
+						if(AssetsOuter)
+						{
+							HLODPackagesToSave.Add(AssetsOuter);
+						}
+						// If we couldn't find the HLOD package, the level may have been renamed while saving above, 
+						// so we need to relocate our old HLOD package before saving it.
+						else if(bUnsavedLevel && LevelPackagesToSave[PackageIndex]->GetName() != OldLevelPackageNames[PackageIndex])
+						{
+							FString OldHLODProxyName = Utilities->GetLevelHLODProxyName(OldLevelPackageNames[PackageIndex], LODIndex);
+							UHLODProxy* OldHLODPProxy = FindObject<UHLODProxy>(nullptr, *OldHLODProxyName);
+							if(OldHLODPProxy)
+							{
+								FString NewHLODProxyName = Utilities->GetLevelHLODProxyName(LevelPackagesToSave[PackageIndex]->GetName(), LODIndex);
+								FString OldHLODPackageName = FPackageName::ObjectPathToPackageName(OldHLODProxyName);
+								FString NewHLODPackageName = FPackageName::ObjectPathToPackageName(NewHLODProxyName);
+								UPackage* OldHLODPackage = FindObject<UPackage>(nullptr, *OldHLODPackageName);
+								if(OldHLODPackage)
+								{
+									OldHLODPProxy->Rename(*FPackageName::ObjectPathToObjectName(NewHLODProxyName), OldHLODPackage, REN_NonTransactional | REN_DontCreateRedirectors);
+									OldHLODPackage->Rename(*NewHLODPackageName, nullptr, REN_NonTransactional | REN_DontCreateRedirectors);
+									
+									HLODPackagesToSave.Add(OldHLODPackage);
+
+									// Mark the level package as dirty as we have changed export locations, and without a resave we will not pick up
+									// HLOD packages when reloaded.
+									Level->GetOutermost()->MarkPackageDirty();
+								}
+							}
+						}
 					}
 				}
 			}
 		}
 
-		// Levels might also need a resave
-		PackagesToSave.Add(Level->GetOutermost());
+		UEditorLoadingAndSavingUtils::SavePackagesWithDialog(HLODPackagesToSave, true);
 	}
-
-	UEditorLoadingAndSavingUtils::SavePackagesWithDialog(PackagesToSave, true);
 }
 
 bool FHierarchicalLODBuilder::NeedsBuild(bool bInForce) const
