Fix crash when opening a newly built HLOD proxy/mesh

diff --git a/Engine/Source/Editor/UnrealEd/Private/Toolkits/AssetEditorManager.cpp b/Engine/Source/Editor/UnrealEd/Private/Toolkits/AssetEditorManager.cpp
--- a/Engine/Source/Editor/UnrealEd/Private/Toolkits/AssetEditorManager.cpp
+++ b/Engine/Source/Editor/UnrealEd/Private/Toolkits/AssetEditorManager.cpp
@@ -403,7 +403,7 @@ bool FAssetEditorManager::OpenEditorForAsset(UObject* Asset, const EToolkitMode:
 		FString AssetPath = Asset->GetOuter()->GetPathName();
 		FContentBrowserModule& CBModule = FModuleManager::LoadModuleChecked<FContentBrowserModule>(TEXT("ContentBrowser"));
 		FMainMRUFavoritesList* RecentlyOpenedAssets = CBModule.GetRecentlyOpenedAssets();
-		if (RecentlyOpenedAssets)
+		if (RecentlyOpenedAssets && FPackageName::IsValidLongPackageName(AssetPath))
 		{
 			RecentlyOpenedAssets->AddMRUItem(AssetPath);
 		}
