Fixes for UE-64484 (GC during CompileAllBluePrints commandlet) & UE-64622 (wrong equality check for TCHAR;s)

diff --git a/Engine/Source/Editor/UnrealEd/Private/Commandlets/CompileAllBlueprintsCommandlet.cpp b/Engine/Source/Editor/UnrealEd/Private/Commandlets/CompileAllBlueprintsCommandlet.cpp
--- a/Engine/Source/Editor/UnrealEd/Private/Commandlets/CompileAllBlueprintsCommandlet.cpp
+++ b/Engine/Source/Editor/UnrealEd/Private/Commandlets/CompileAllBlueprintsCommandlet.cpp
@@ -8,6 +8,7 @@
 #include "ISourceControlModule.h"
 #include "Misc/FileHelper.h"
 #include "Misc/Paths.h"
+#include "Engine/Engine.h"
 
 DEFINE_LOG_CATEGORY_STATIC(LogCompileAllBlueprintsCommandlet, Log, All);
 
@@ -131,6 +132,8 @@ void UCompileAllBlueprintsCommandlet::ParseWhitelist(const FString& WhitelistFil
 
 void UCompileAllBlueprintsCommandlet::BuildBlueprints()
 {
+	double LastGCTime = FPlatformTime::Seconds();
+
 	for (FAssetData const& Asset : BlueprintAssetList)
 	{
 		if (ShouldBuildAsset(Asset))
@@ -150,6 +153,14 @@ void UCompileAllBlueprintsCommandlet::BuildBlueprints()
 			{
 				CompileBlueprint(LoadedBlueprint);
 			}
+
+			const double TimeNow = FPlatformTime::Seconds();
+
+			if (TimeNow - LastGCTime >= 10.0)
+			{
+				GEngine->TrimMemory();
+				LastGCTime = TimeNow;
+			}
 		}
 	}
 }
diff --git a/Engine/Source/Runtime/Core/Public/Misc/AutomationTest.h b/Engine/Source/Runtime/Core/Public/Misc/AutomationTest.h
--- a/Engine/Source/Runtime/Core/Public/Misc/AutomationTest.h
+++ b/Engine/Source/Runtime/Core/Public/Misc/AutomationTest.h
@@ -1255,7 +1255,7 @@ public:
 
 	void TestEqual(const FString& Description, const TCHAR* A, const TCHAR* B)
 	{
-		if (A != B)
+		if (FCString::Strcmp(A,B) != 0)
 		{
 			AddError(FString::Printf(TEXT("%s: The two values are not equal.\n  Actual: '%s'\nExpected: '%s'"), *Description, A, B), 1);
 		}
@@ -1263,7 +1263,7 @@ public:
 
 	void TestEqual(const FString& Description, const TCHAR* A, const FString& B)
 	{
-		if (A != B)
+		if (FCString::Strcmp(A, *B) != 0)
 		{
 			AddError(FString::Printf(TEXT("%s: The two values are not equal.\n  Actual: '%s'\nExpected: '%s'"), *Description, A, *B), 1);
 		}
