PR #5217: Update variable name if type is set to/from boolean (Contributed by [REDACTED_PERSON])

diff --git a/Engine/Source/Editor/UnrealEd/Private/Kismet2/BlueprintEditorUtils.cpp b/Engine/Source/Editor/UnrealEd/Private/Kismet2/BlueprintEditorUtils.cpp
--- a/Engine/Source/Editor/UnrealEd/Private/Kismet2/BlueprintEditorUtils.cpp
+++ b/Engine/Source/Editor/UnrealEd/Private/Kismet2/BlueprintEditorUtils.cpp
@@ -4726,6 +4726,13 @@ void FBlueprintEditorUtils::ChangeMemberVariableType(UBlueprint* Blueprint, cons
 
 				if(bChangeVariableType)
 				{
+					const bool bBecameBoolean = Variable.VarType.PinCategory != UEdGraphSchema_K2::PC_Boolean && NewPinType.PinCategory == UEdGraphSchema_K2::PC_Boolean;
+					const bool bBecameNotBoolean = Variable.VarType.PinCategory == UEdGraphSchema_K2::PC_Boolean && NewPinType.PinCategory != UEdGraphSchema_K2::PC_Boolean;
+					if (bBecameBoolean || bBecameNotBoolean)
+					{
+						Variable.FriendlyName = FName::NameToDisplayString(Variable.VarName.ToString(), bBecameBoolean);
+					}
+
 					Variable.VarType = NewPinType;
 
 					if(Variable.VarType.IsSet() || Variable.VarType.IsMap())
