PR #5156: UE-36747: Don't call SpawnActor when class is abstract or deprecated . (Contributed by [REDACTED_PERSON])

diff --git a/Engine/Source/Editor/UnrealEd/Private/ThumbnailHelpers.cpp b/Engine/Source/Editor/UnrealEd/Private/ThumbnailHelpers.cpp
--- a/Engine/Source/Editor/UnrealEd/Private/ThumbnailHelpers.cpp
+++ b/Engine/Source/Editor/UnrealEd/Private/ThumbnailHelpers.cpp
@@ -1040,7 +1040,7 @@ void FClassActorThumbnailScene::SpawnPreviewActor(UClass* InClass)
 		PreviewActor->Destroy();
 		PreviewActor = nullptr;
 	}
-	if (InClass)
+	if (InClass && !InClass->HasAnyClassFlags(CLASS_Deprecated | CLASS_Abstract))
 	{
 		// Create preview actor
 		FActorSpawnParameters SpawnInfo;
