Automation - Fix promotion utilities to use the main frame module to safely get the main editor window when needed.

diff --git a/Engine/Source/Editor/UnrealEd/Private/Tests/AutomationEditorPromotionCommon.cpp b/Engine/Source/Editor/UnrealEd/Private/Tests/AutomationEditorPromotionCommon.cpp
--- a/Engine/Source/Editor/UnrealEd/Private/Tests/AutomationEditorPromotionCommon.cpp
+++ b/Engine/Source/Editor/UnrealEd/Private/Tests/AutomationEditorPromotionCommon.cpp
@@ -19,6 +19,7 @@
 #include "Engine/Texture.h"
 #include "LevelEditor.h"
 #include "ScopedTransaction.h"
+#include "Interfaces/IMainFrameModule.h"
 
 #define LOCTEXT_NAMESPACE "EditorPromotionTestCommon"
 
@@ -157,10 +158,17 @@ FInputChord FEditorPromotionTestUtilities::GetOrSetUICommand(const FString& Cont
 */
 void FEditorPromotionTestUtilities::SendCommandToCurrentEditor(const FInputChord& InChord, const FName& WidgetTypeToFocus)
 {
-	//Focus the asset Editor / Graph 
-	TSharedRef<SWindow> EditorWindow = FSlateApplication::Get().GetActiveTopLevelWindow().ToSharedRef();
-	FSlateApplication::Get().ProcessWindowActivatedEvent(FWindowActivateEvent(FWindowActivateEvent::EA_Activate, EditorWindow));
-	TSharedPtr<SWidget> FocusWidget = FindFirstWidgetByClass(EditorWindow, WidgetTypeToFocus);
+	IMainFrameModule& MainFrame = FModuleManager::LoadModuleChecked<IMainFrameModule>("MainFrame");
+	TSharedPtr<SWindow>	ParentWindow = MainFrame.GetParentWindow();
+	if (ParentWindow.IsValid() == false)
+	{
+		UE_LOG(LogEditorPromotionTests, Error, TEXT("Failed to find main editor window to send commands to."));
+		return;
+	}
+
+	FSlateApplication::Get().ProcessWindowActivatedEvent(FWindowActivateEvent(FWindowActivateEvent::EA_Activate, ParentWindow.ToSharedRef()));
+
+	TSharedPtr<SWidget> FocusWidget = FindFirstWidgetByClass(ParentWindow.ToSharedRef(), WidgetTypeToFocus);
 	if (FocusWidget.IsValid())
 	{
 		FSlateApplication::Get().SetKeyboardFocus(FocusWidget.ToSharedRef(), EFocusCause::SetDirectly);
@@ -254,15 +262,15 @@ void FEditorPromotionTestUtilities::TakeScreenshot(const FString& ScreenshotName
 	else
 	{
 		//Find the main editor window
-		TArray<TSharedRef<SWindow> > AllWindows;
-		FSlateApplication::Get().GetAllVisibleWindowsOrdered(AllWindows);
-		if (AllWindows.Num() == 0)
+		IMainFrameModule& MainFrame = FModuleManager::LoadModuleChecked<IMainFrameModule>("MainFrame");
+		TSharedPtr<SWindow>	ParentWindow = MainFrame.GetParentWindow();
+		if (ParentWindow.IsValid() == false)
 		{
 			UE_LOG(LogEditorPromotionTests, Error, TEXT("ERROR: Could not find the main editor window."));
 			return;
 		}
 
-		Window = AllWindows[0];
+		Window = ParentWindow;
 	}
 
 	if (Window.IsValid())
