Fix crash when importing with geo or skinning only

diff --git a/Engine/Source/Editor/UnrealEd/Private/Fbx/FbxFactory.cpp b/Engine/Source/Editor/UnrealEd/Private/Fbx/FbxFactory.cpp
--- a/Engine/Source/Editor/UnrealEd/Private/Fbx/FbxFactory.cpp
+++ b/Engine/Source/Editor/UnrealEd/Private/Fbx/FbxFactory.cpp
@@ -614,7 +614,7 @@ UObject* UFbxFactory::FactoryCreateFile
 
 								if ( NewMesh )
 								{
-									if (ImportUI->bImportAnimations)
+									if (ImportOptions->bImportAnimations)
 									{
 										// We need to remove all scaling from the root node before we set up animation data.
 										// Othewise some of the global transform calculations will be incorrect.
@@ -676,7 +676,7 @@ UObject* UFbxFactory::FactoryCreateFile
 							}
 						
 							// import morph target
-							if (CreatedObject && ImportUI->SkeletalMeshImportData->bImportMorphTargets && ImportedSuccessfulLodIndex != INDEX_NONE)
+							if (CreatedObject && ImportOptions->bImportMorph && ImportedSuccessfulLodIndex != INDEX_NONE)
 							{
 								// Disable material importing when importing morph targets
 								uint32 bImportMaterials = ImportOptions->bImportMaterials;
diff --git a/Engine/Source/Editor/UnrealEd/Private/Fbx/FbxSkeletalMeshImport.cpp b/Engine/Source/Editor/UnrealEd/Private/Fbx/FbxSkeletalMeshImport.cpp
--- a/Engine/Source/Editor/UnrealEd/Private/Fbx/FbxSkeletalMeshImport.cpp
+++ b/Engine/Source/Editor/UnrealEd/Private/Fbx/FbxSkeletalMeshImport.cpp
@@ -1098,7 +1098,7 @@ bool UnFbx::FFbxImporter::ImportBone(TArray<FbxNode*>& NodeArray, FSkeletalMeshI
 	return true;
 }
 
-bool UnFbx::FFbxImporter::FillSkeletalMeshImportData(TArray<FbxNode*>& NodeArray, UFbxSkeletalMeshImportData* TemplateImportData, TArray<FbxShape*> *FbxShapeArray, FSkeletalMeshImportData* OutData, TArray<FName> &LastImportedMaterialNames)
+bool UnFbx::FFbxImporter::FillSkeletalMeshImportData(TArray<FbxNode*>& NodeArray, UFbxSkeletalMeshImportData* TemplateImportData, TArray<FbxShape*> *FbxShapeArray, FSkeletalMeshImportData* OutData, TArray<FName> &LastImportedMaterialNames, const bool bIsReimport)
 {
 	if (NodeArray.Num() == 0)
 	{
@@ -1131,7 +1131,7 @@ bool UnFbx::FFbxImporter::FillSkeletalMeshImportData(TArray<FbxNode*>& NodeArray
 	SkelMeshImportDataPtr->bUseT0AsRefPose = ImportOptions->bUseT0AsRefPose;
 	// Note: importing morph data causes additional passes through this function, so disable the warning dialogs
 	// from popping up again on each additional pass.
-	if (!ImportOptions->bImportAsSkeletalGeometry) //Do not import bone if we import only the geometry
+	if (!bIsReimport || !ImportOptions->bImportAsSkeletalGeometry) //Do not import bone if we import only the geometry and we are reimporting
 	{
 		if (!ImportBone(NodeArray, *SkelMeshImportDataPtr, TemplateImportData, SortedLinkArray, SkelMeshImportDataPtr->bDiffPose, (FbxShapeArray != nullptr), SkelMeshImportDataPtr->bUseT0AsRefPose, Node))
 		{
@@ -1515,7 +1515,7 @@ USkeletalMesh* UnFbx::FFbxImporter::ImportSkeletalMesh(FImportSkeletalMeshArgs &
 	// We must do a maximum of fail test before backing up the data since the backup is destructive on the existing skeletal mesh.
 	// See the comment later when we call the following function (SaveExistingSkelMeshData)
 
-	if (FillSkeletalMeshImportData(ImportSkeletalMeshArgs.NodeArray, ImportSkeletalMeshArgs.TemplateImportData, ImportSkeletalMeshArgs.FbxShapeArray, SkelMeshImportDataPtr, LastImportedMaterialNames) == false)
+	if (FillSkeletalMeshImportData(ImportSkeletalMeshArgs.NodeArray, ImportSkeletalMeshArgs.TemplateImportData, ImportSkeletalMeshArgs.FbxShapeArray, SkelMeshImportDataPtr, LastImportedMaterialNames, ExistingSkelMesh != nullptr) == false)
 	{
 		AddTokenizedErrorMessage(FTokenizedMessage::Create(EMessageSeverity::Warning, LOCTEXT("FbxSkeletaLMeshimport_FillupImportData", "Get Import Data has failed.")), FFbxErrors::SkeletalMesh_FillImportDataFailed);
 		if (SkeletalMesh)
@@ -1753,7 +1753,7 @@ USkeletalMesh* UnFbx::FFbxImporter::ImportSkeletalMesh(FImportSkeletalMeshArgs &
 		}
 
 		// merge bones to the selected skeleton
-		if (!ImportOptions->bImportAsSkeletalSkinning && !Skeleton->MergeAllBonesToBoneTree( SkeletalMesh ) )
+		if ((!ImportOptions->bImportAsSkeletalSkinning || ExistingSkelMesh == nullptr) && !Skeleton->MergeAllBonesToBoneTree( SkeletalMesh ) )
 		{
 			// We should only show the skeleton save toast once, not as many times as we have nodes to import
 			bool bToastSaveMessage = false;
diff --git a/Engine/Source/Editor/UnrealEd/Public/FbxImporter.h b/Engine/Source/Editor/UnrealEd/Public/FbxImporter.h
--- a/Engine/Source/Editor/UnrealEd/Public/FbxImporter.h
+++ b/Engine/Source/Editor/UnrealEd/Public/FbxImporter.h
@@ -1373,7 +1373,7 @@ public:
 	*
 	* @returns bool*	true if import successfully.
 	*/
-	bool FillSkeletalMeshImportData(TArray<FbxNode*>& NodeArray, UFbxSkeletalMeshImportData* TemplateImportData, TArray<FbxShape*> *FbxShapeArray, FSkeletalMeshImportData* OutData, TArray<FName> &LastImportedMaterialNames);
+	bool FillSkeletalMeshImportData(TArray<FbxNode*>& NodeArray, UFbxSkeletalMeshImportData* TemplateImportData, TArray<FbxShape*> *FbxShapeArray, FSkeletalMeshImportData* OutData, TArray<FName> &LastImportedMaterialNames, const bool bIsReimport);
 
 protected:
 
