Reinitialize skel mesh animations after reinstancing is complete

diff --git a/Engine/Source/Editor/UnrealEd/Private/Kismet2/KismetReinstanceUtilities.cpp b/Engine/Source/Editor/UnrealEd/Private/Kismet2/KismetReinstanceUtilities.cpp
--- a/Engine/Source/Editor/UnrealEd/Private/Kismet2/KismetReinstanceUtilities.cpp
+++ b/Engine/Source/Editor/UnrealEd/Private/Kismet2/KismetReinstanceUtilities.cpp
@@ -1671,18 +1671,6 @@ static void ReplaceObjectHelper(UObject*& OldObject, UClass* OldClass, UObject*&
 	UEditorEngine::CopyPropertiesForUnrelatedObjects(OldObject, NewUObject);
 	InstancedPropertyUtils::FArchiveInsertInstancedSubObjects InstancedSubObjSpawner(NewUObject, InstancedPropertyMap);
 
-	if (UAnimInstance* AnimTree = Cast<UAnimInstance>(NewUObject))
-	{
-		// Initialising the anim instance isn't enough to correctly set up the skeletal mesh again in a
-		// paused world, need to initialise the skeletal mesh component that contains the anim instance.
-		if (USkeletalMeshComponent* SkelComponent = Cast<USkeletalMeshComponent>(AnimTree->GetOuter()))
-		{
-			SkelComponent->InitAnim(true);
-			// compile change ignores motion vector, so ignore this. 
-			SkelComponent->ClearMotionVector();
-		}
-	}
-
 	UWorld* RegisteredWorld = nullptr;
 	bool bWasRegistered = false;
 	if (bIsComponent)
@@ -2119,6 +2107,27 @@ void FBlueprintCompileReinstancer::ReplaceInstancesOfClass_Inner(TMap<UClass*, U
 	}
 
 	FReplaceReferenceHelper::FindAndReplaceReferences(SourceObjects, ObjectsThatShouldUseOldStuff, ObjectsReplaced, OldToNewInstanceMap, ReinstancedObjectsWeakReferenceMap);
+	
+	for (UObject* Obj : ObjectsReplaced)
+	{
+		UObject** NewObject = OldToNewInstanceMap.Find(Obj);
+		if (NewObject && *NewObject)
+		{
+			if (UAnimInstance* AnimTree = Cast<UAnimInstance>(*NewObject))
+			{
+				// Initialising the anim instance isn't enough to correctly set up the skeletal mesh again in a
+				// paused world, need to initialise the skeletal mesh component that contains the anim instance.
+				if (USkeletalMeshComponent* SkelComponent = Cast<USkeletalMeshComponent>(AnimTree->GetOuter()))
+				{
+					SkelComponent->ClearAnimScriptInstance();
+					SkelComponent->InitAnim(true);
+					// compile change ignores motion vector, so ignore this. 
+					SkelComponent->ClearMotionVector();
+				}
+			}
+		}
+	}
+
 
 	if(SelectedActors)
 	{
