Reapply 4.20 hotfix that did not get carried over to 4.21
#rnx

diff --git a/Engine/Source/Editor/UnrealEd/Private/Kismet2/KismetDebugUtilities.cpp b/Engine/Source/Editor/UnrealEd/Private/Kismet2/KismetDebugUtilities.cpp
--- a/Engine/Source/Editor/UnrealEd/Private/Kismet2/KismetDebugUtilities.cpp
+++ b/Engine/Source/Editor/UnrealEd/Private/Kismet2/KismetDebugUtilities.cpp
@@ -1273,7 +1273,7 @@ FKismetDebugUtilities::EWatchTextResult FKismetDebugUtilities::FindDebuggingData
 						{
 							void* PropertyValue = SelfPinProperty->ContainerPtrToValuePtr<void>(SelfPinData);
 							UObject* TempActiveObject = SelfPinPropertyBase->GetObjectPropertyValue(PropertyValue);
-							if (TempActiveObject)
+							if (TempActiveObject && TempActiveObject != ActiveObject)
 							{
 								return FindDebuggingData(Blueprint, TempActiveObject, WatchPin, OutProperty, OutData, OutDelta, OutParent);
 							}
