(4.21.2) Fix crash if a parent blueprint doesn't have a simple construction script
#rnx

diff --git a/Engine/Source/Editor/UnrealEd/Private/Kismet2/Kismet2.cpp b/Engine/Source/Editor/UnrealEd/Private/Kismet2/Kismet2.cpp
--- a/Engine/Source/Editor/UnrealEd/Private/Kismet2/Kismet2.cpp
+++ b/Engine/Source/Editor/UnrealEd/Private/Kismet2/Kismet2.cpp
@@ -1486,10 +1486,13 @@ void FKismetEditorUtilities::AddComponentsToBlueprint(UBlueprint* Blueprint, TAr
 				{
 					for (UBlueprint* ParentBlueprint : ParentBPStack)
 					{
-						ParentSCSNode = ParentBlueprint->SimpleConstructionScript->FindSCSNode(SceneComponent->GetAttachParent()->GetFName());
-						if (ParentSCSNode)
+						if (ParentBlueprint->SimpleConstructionScript)
 						{
-							break;
+							ParentSCSNode = ParentBlueprint->SimpleConstructionScript->FindSCSNode(SceneComponent->GetAttachParent()->GetFName());
+							if (ParentSCSNode)
+							{
+								break;
+							}
 						}
 					}
 				}
