Sequencer: Import/export focus distance

diff --git a/Engine/Source/Editor/MovieSceneTools/Private/MovieSceneToolHelpers.cpp b/Engine/Source/Editor/MovieSceneTools/Private/MovieSceneToolHelpers.cpp
--- a/Engine/Source/Editor/MovieSceneTools/Private/MovieSceneToolHelpers.cpp
+++ b/Engine/Source/Editor/MovieSceneTools/Private/MovieSceneToolHelpers.cpp
@@ -1209,6 +1209,7 @@ void CopyCameraProperties(FbxCamera* CameraNode, ACineCameraActor* CameraActor)
 	CineCameraComponent->SetFieldOfView(FieldOfView);
 	CineCameraComponent->FilmbackSettings.SensorWidth = FUnitConversion::Convert(ApertureWidth, EUnit::Inches, EUnit::Millimeters);
 	CineCameraComponent->FilmbackSettings.SensorHeight = FUnitConversion::Convert(ApertureHeight, EUnit::Inches, EUnit::Millimeters);
+	CineCameraComponent->FocusSettings.ManualFocusDistance = CameraNode->FocusDistance;
 	if (FocalLength < CineCameraComponent->LensSettings.MinFocalLength)
 	{
 		CineCameraComponent->LensSettings.MinFocalLength = FocalLength;
diff --git a/Engine/Source/Editor/UnrealEd/Private/Fbx/FbxMainExport.cpp b/Engine/Source/Editor/UnrealEd/Private/Fbx/FbxMainExport.cpp
--- a/Engine/Source/Editor/UnrealEd/Private/Fbx/FbxMainExport.cpp
+++ b/Engine/Source/Editor/UnrealEd/Private/Fbx/FbxMainExport.cpp
@@ -588,6 +588,8 @@ void FFbxExporter::FillFbxCameraAttribute(FbxNode* ParentNode, FbxCamera* Camera
 	float ApertureHeightInInches = 0.612f; // 0.612f is a magic number from Maya that represents the ApertureHeight
 	float ApertureWidthInInches = CameraComponent->AspectRatio * ApertureHeightInInches;
 	float FocalLength = Camera->ComputeFocalLength(CameraComponent->FieldOfView);
+	
+	TOptional<float> FocusDistance;
 
 	if (CameraComponent->IsA(UCineCameraComponent::StaticClass()))
 	{
@@ -597,6 +599,7 @@ void FFbxExporter::FillFbxCameraAttribute(FbxNode* ParentNode, FbxCamera* Camera
 			ApertureWidthInInches = FUnitConversion::Convert(CineCameraComponent->FilmbackSettings.SensorWidth, EUnit::Millimeters, EUnit::Inches);
 			ApertureHeightInInches = FUnitConversion::Convert(CineCameraComponent->FilmbackSettings.SensorHeight, EUnit::Millimeters, EUnit::Inches);
 			FocalLength = CineCameraComponent->CurrentFocalLength;
+			FocusDistance = CineCameraComponent->FocusSettings.ManualFocusDistance;
 		}
 	}
 
@@ -609,6 +612,11 @@ void FFbxExporter::FillFbxCameraAttribute(FbxNode* ParentNode, FbxCamera* Camera
 	Camera->SetApertureMode(FbxCamera::eFocalLength);
 	Camera->FocalLength.Set(FocalLength);
 
+	if (FocusDistance.IsSet())
+	{
+		Camera->FocusDistance.Set(FocusDistance.GetValue());
+	}
+
 	// Add one user property for recording the AspectRatio animation
 	CreateAnimatableUserProperty(ParentNode, CameraComponent->AspectRatio, "UE_AspectRatio", "UE_Matinee_Camera_AspectRatio");
 
