Renderpass updates.

Lightmass landscape render. Self-contained.

diff --git a/Engine/Source/Editor/UnrealEd/Private/Lightmass/LightmassLandscapeRender.cpp b/Engine/Source/Editor/UnrealEd/Private/Lightmass/LightmassLandscapeRender.cpp
--- a/Engine/Source/Editor/UnrealEd/Private/Lightmass/LightmassLandscapeRender.cpp
+++ b/Engine/Source/Editor/UnrealEd/Private/Lightmass/LightmassLandscapeRender.cpp
@@ -123,22 +123,26 @@ void RenderLandscapeMaterialForLightmass(const FLandscapeStaticLightingMesh* Lan
 			{
 				// Set the RHI render target.
 				RHICmdList.TransitionResource(EResourceTransitionAccess::EWritable, RenderTarget->GetRenderTargetTexture());
-				::SetRenderTarget(RHICmdList, RenderTarget->GetRenderTargetTexture(), FTextureRHIRef());
 
-				const FIntRect RTViewRect = FIntRect(0, 0, RenderTarget->GetRenderTargetTexture()->GetSizeX(), RenderTarget->GetRenderTargetTexture()->GetSizeY());
+				FRHIRenderPassInfo RPInfo(RenderTarget->GetRenderTargetTexture(), ERenderTargetActions::Load_Store);
+				RHICmdList.BeginRenderPass(RPInfo, TEXT("CanvasFlushSetup"));
+				{
+					const FIntRect RTViewRect = FIntRect(0, 0, RenderTarget->GetRenderTargetTexture()->GetSizeX(), RenderTarget->GetRenderTargetTexture()->GetSizeY());
 
-				// set viewport to RT size
-				RHICmdList.SetViewport(RTViewRect.Min.X, RTViewRect.Min.Y, 0.0f, RTViewRect.Max.X, RTViewRect.Max.Y, 1.0f);
+					// set viewport to RT size
+					RHICmdList.SetViewport(RTViewRect.Min.X, RTViewRect.Min.Y, 0.0f, RTViewRect.Max.X, RTViewRect.Max.Y, 1.0f);
 
-				FSceneView View(ViewInitOptions);
+					FSceneView View(ViewInitOptions);
 
-				FDrawingPolicyRenderState DrawRenderState(View);
+					FDrawingPolicyRenderState DrawRenderState(View);
 
-				// disable depth test & writes
-				DrawRenderState.SetDepthStencilState(TStaticDepthStencilState<false, CF_Always>::GetRHI());
+					// disable depth test & writes
+					DrawRenderState.SetDepthStencilState(TStaticDepthStencilState<false, CF_Always>::GetRHI());
 
-				//SCOPED_DRAW_EVENT(RHICmdList, RenderLandscapeMaterialToTexture);
-				GetRendererModule().DrawTileMesh(RHICmdList, DrawRenderState, View, Mesh, false, FHitProxyId());
+					//SCOPED_DRAW_EVENT(RHICmdList, RenderLandscapeMaterialToTexture);
+					GetRendererModule().DrawTileMesh(RHICmdList, DrawRenderState, View, Mesh, false, FHitProxyId());
+				}
+				RHICmdList.EndRenderPass();
 			}
 		});
 	FlushRenderingCommands();
