(4.21.2) Make a number of FBlueprintEditorUtils functions resilient to having null passed in to them
#rnx

diff --git a/Engine/Source/Editor/UnrealEd/Private/Kismet2/BlueprintEditorUtils.cpp b/Engine/Source/Editor/UnrealEd/Private/Kismet2/BlueprintEditorUtils.cpp
--- a/Engine/Source/Editor/UnrealEd/Private/Kismet2/BlueprintEditorUtils.cpp
+++ b/Engine/Source/Editor/UnrealEd/Private/Kismet2/BlueprintEditorUtils.cpp
@@ -3131,7 +3131,7 @@ bool FBlueprintEditorUtils::IsBlueprintConst(const UBlueprint* Blueprint)
 {
 	// Macros aren't marked as const because they can modify variables when instanced into a non const class
 	// and will be caught at compile time if they're modifying variables on a const class.
-	return Blueprint->BlueprintType == BPTYPE_Const;
+	return Blueprint && Blueprint->BlueprintType == BPTYPE_Const;
 }
 
 bool FBlueprintEditorUtils::IsBlutility(const UBlueprint* Blueprint)
@@ -3147,7 +3147,7 @@ bool FBlueprintEditorUtils::IsBlutility(const UBlueprint* Blueprint)
 
 bool FBlueprintEditorUtils::IsActorBased(const UBlueprint* Blueprint)
 {
-	return Blueprint->ParentClass && Blueprint->ParentClass->IsChildOf(AActor::StaticClass());
+	return Blueprint && Blueprint->ParentClass && Blueprint->ParentClass->IsChildOf(AActor::StaticClass());
 }
 
 bool FBlueprintEditorUtils::IsDelegateSignatureGraph(const UEdGraph* Graph)
@@ -3173,22 +3173,22 @@ bool FBlueprintEditorUtils::IsMathExpressionGraph(const UEdGraph* InGraph)
 
 bool FBlueprintEditorUtils::IsInterfaceBlueprint(const UBlueprint* Blueprint)
 {
-	return (Blueprint->BlueprintType == BPTYPE_Interface);
+	return (Blueprint && Blueprint->BlueprintType == BPTYPE_Interface);
 }
 
 bool FBlueprintEditorUtils::IsLevelScriptBlueprint(const UBlueprint* Blueprint)
 {
-	return (Blueprint->BlueprintType == BPTYPE_LevelScript);
+	return (Blueprint && Blueprint->BlueprintType == BPTYPE_LevelScript);
 }
 
 bool FBlueprintEditorUtils::IsAnonymousBlueprintClass(const UClass* Class)
 {
-	return (Class->GetOutermost()->ContainsMap());
+	return (Class && Class->GetOutermost()->ContainsMap());
 }
 
 ULevel* FBlueprintEditorUtils::GetLevelFromBlueprint(const UBlueprint* Blueprint)
 {
-	return Cast<ULevel>(Blueprint->GetOuter());
+	return (Blueprint ? Cast<ULevel>(Blueprint->GetOuter()) : nullptr);
 }
 
 bool FBlueprintEditorUtils::SupportsConstructionScript(const UBlueprint* Blueprint)
