Mirror 4562879
Assign UberGraphFunctionKey on dupicated class

diff --git a/Engine/Source/Editor/UnrealEd/Private/Kismet2/KismetReinstanceUtilities.cpp b/Engine/Source/Editor/UnrealEd/Private/Kismet2/KismetReinstanceUtilities.cpp
--- a/Engine/Source/Editor/UnrealEd/Private/Kismet2/KismetReinstanceUtilities.cpp
+++ b/Engine/Source/Editor/UnrealEd/Private/Kismet2/KismetReinstanceUtilities.cpp
@@ -1533,6 +1533,13 @@ UClass* FBlueprintCompileReinstancer::MoveCDOToNewClass(UClass* OwnerClass, cons
 		BPGDuplicatedClass->OverridenArchetypeForCDO = BPClassToReinstance->OverridenArchetypeForCDO;
 	}
 
+#if VALIDATE_UBER_GRAPH_PERSISTENT_FRAME
+	if (BPGDuplicatedClass && BPClassToReinstance)
+	{
+		BPGDuplicatedClass->UberGraphFunctionKey = BPClassToReinstance->UberGraphFunctionKey;
+	}
+#endif
+
 	UFunction* DuplicatedClassUberGraphFunction = BPGDuplicatedClass ? BPGDuplicatedClass->UberGraphFunction : nullptr;
 	if (DuplicatedClassUberGraphFunction)
 	{
