Fbx Importer, prevent duplicated logs when we reach the maximum LOD.

diff --git a/Engine/Source/Editor/UnrealEd/Private/Fbx/FbxFactory.cpp b/Engine/Source/Editor/UnrealEd/Private/Fbx/FbxFactory.cpp
--- a/Engine/Source/Editor/UnrealEd/Private/Fbx/FbxFactory.cpp
+++ b/Engine/Source/Editor/UnrealEd/Private/Fbx/FbxFactory.cpp
@@ -410,8 +410,10 @@ UObject* UFbxFactory::FactoryCreateFile
 								{
 									if (GroupLodIndex >= MAX_STATIC_MESH_LODS)
 									{
-										FbxImporter->AddTokenizedErrorMessage(FTokenizedMessage::Create(EMessageSeverity::Warning, FText::Format(LOCTEXT("ImporterLimits_MaximumStaticMeshLODReach", "Reach the maximum LOD number({0}) for a staticmesh."), FText::AsNumber(MAX_STATIC_MESH_LODS))), FFbxErrors::Generic_Mesh_TooManyLODs);
-										continue;
+										FbxImporter->AddTokenizedErrorMessage(FTokenizedMessage::Create(EMessageSeverity::Warning, FText::Format(
+											LOCTEXT("ImporterLimits_MaximumStaticMeshLODReach", "Reached the maximum number of LODs for a Static Mesh({0}) - discarding {1} LOD meshes."), FText::AsNumber(MAX_STATIC_MESH_LODS), FText::AsNumber(LODGroup->GetChildCount() - MAX_STATIC_MESH_LODS))
+										), FFbxErrors::Generic_Mesh_TooManyLODs);
+										break;
 									}
 									TArray<FbxNode*> AllNodeInLod;
 									FbxImporter->FindAllLODGroupNode(AllNodeInLod, LODGroup, GroupLodIndex);
@@ -774,12 +776,15 @@ UObject* UFbxFactory::RecursiveImportNode(void* VoidFbxImporter, void* VoidNode,
 		if (CreatedObject && bImportMeshLODs)
 		{
 			// import LOD meshes
+			
 			for (int32 LODIndex = 1; LODIndex < Node->GetChildCount(); LODIndex++)
 			{
 				if (LODIndex >= MAX_STATIC_MESH_LODS)
 				{
-					FbxImporter->AddTokenizedErrorMessage(FTokenizedMessage::Create(EMessageSeverity::Warning, FText::Format(LOCTEXT("ImporterLimits_MaximumStaticMeshLODReach", "Reach the maximum LOD number({0}) for a staticmesh."), FText::AsNumber(MAX_STATIC_MESH_LODS))), FFbxErrors::Generic_Mesh_TooManyLODs);
-					continue;
+					FbxImporter->AddTokenizedErrorMessage(FTokenizedMessage::Create(EMessageSeverity::Warning, FText::Format(
+						LOCTEXT("ImporterLimits_MaximumStaticMeshLODReach", "Reached the maximum number of LODs for a Static Mesh({0}) - discarding {1} LOD meshes."), FText::AsNumber(MAX_STATIC_MESH_LODS), FText::AsNumber(Node->GetChildCount() - MAX_STATIC_MESH_LODS))
+					), FFbxErrors::Generic_Mesh_TooManyLODs);
+					break;
 				}
 				AllNodeInLod.Empty();
 				FbxImporter->FindAllLODGroupNode(AllNodeInLod, Node, LODIndex);
diff --git a/Engine/Source/Editor/UnrealEd/Private/Fbx/FbxSceneImportFactory.cpp b/Engine/Source/Editor/UnrealEd/Private/Fbx/FbxSceneImportFactory.cpp
--- a/Engine/Source/Editor/UnrealEd/Private/Fbx/FbxSceneImportFactory.cpp
+++ b/Engine/Source/Editor/UnrealEd/Private/Fbx/FbxSceneImportFactory.cpp
@@ -2173,8 +2173,10 @@ UObject* UFbxSceneImportFactory::RecursiveImportNode(void* VoidFbxImporter, void
 				{
 					if (LODIndex >= MAX_STATIC_MESH_LODS)
 					{
-						FFbxImporter->AddTokenizedErrorMessage(FTokenizedMessage::Create(EMessageSeverity::Warning, FText::Format(LOCTEXT("ImporterLimits_MaximumStaticMeshLODReach", "Reach the maximum LOD number({0}) for a staticmesh."), FText::AsNumber(MAX_STATIC_MESH_LODS))), FFbxErrors::Generic_Mesh_TooManyLODs);
-						continue;
+						FFbxImporter->AddTokenizedErrorMessage(FTokenizedMessage::Create(EMessageSeverity::Warning, FText::Format(
+							LOCTEXT("ImporterLimits_MaximumStaticMeshLODReach", "Reached the maximum number of LODs for a Static Mesh({0}) - discarding {1} LOD meshes."), FText::AsNumber(MAX_STATIC_MESH_LODS), FText::AsNumber(Node->GetChildCount() - MAX_STATIC_MESH_LODS))
+						), FFbxErrors::Generic_Mesh_TooManyLODs);
+						break;
 					}
 					AllNodeInLod.Empty();
 					FFbxImporter->FindAllLODGroupNode(AllNodeInLod, Node, LODIndex);
diff --git a/Engine/Source/Editor/UnrealEd/Private/Fbx/FbxStaticMeshImport.cpp b/Engine/Source/Editor/UnrealEd/Private/Fbx/FbxStaticMeshImport.cpp
--- a/Engine/Source/Editor/UnrealEd/Private/Fbx/FbxStaticMeshImport.cpp
+++ b/Engine/Source/Editor/UnrealEd/Private/Fbx/FbxStaticMeshImport.cpp
@@ -1091,8 +1091,10 @@ UStaticMesh* UnFbx::FFbxImporter::ReimportStaticMesh(UStaticMesh* Mesh, UFbxStat
 					{
 						if (GroupLodIndex >= MAX_STATIC_MESH_LODS)
 						{
-							AddTokenizedErrorMessage(FTokenizedMessage::Create(EMessageSeverity::Warning, FText::Format(LOCTEXT("ImporterLimits_MaximumStaticMeshLODReach", "Reach the maximum LOD number({0}) for a staticmesh."), FText::AsNumber(MAX_STATIC_MESH_LODS))), FFbxErrors::Generic_Mesh_TooManyLODs);
-							continue;
+							AddTokenizedErrorMessage(FTokenizedMessage::Create(EMessageSeverity::Warning, FText::Format(
+								LOCTEXT("ImporterLimits_MaximumStaticMeshLODReach", "Reached the maximum number of LODs for a Static Mesh({0}) - discarding {1} LOD meshes."), FText::AsNumber(MAX_STATIC_MESH_LODS), FText::AsNumber(LODGroup->GetChildCount() - MAX_STATIC_MESH_LODS))
+							), FFbxErrors::Generic_Mesh_TooManyLODs);
+							break;
 						}
 						TArray<FbxNode*> AllNodeInLod;
 						FindAllLODGroupNode(AllNodeInLod, LODGroup, GroupLodIndex);
