Fix LOD capitalisation writing in the import options

diff --git a/Engine/Source/Editor/UnrealEd/Classes/Factories/FbxImportUI.h b/Engine/Source/Editor/UnrealEd/Classes/Factories/FbxImportUI.h
--- a/Engine/Source/Editor/UnrealEd/Classes/Factories/FbxImportUI.h
+++ b/Engine/Source/Editor/UnrealEd/Classes/Factories/FbxImportUI.h
@@ -72,39 +72,39 @@ public:
 	class UPhysicsAsset* PhysicsAsset;
 
 	/** If checked, the static mesh auto compute LOD distance will be turn on. If unchecked user will be able to specify custom LOD distance for every LOD. */
-	UPROPERTY(EditAnywhere, BlueprintReadWrite, config, Category = LodSettings, meta = (ImportType = "StaticMesh"))
+	UPROPERTY(EditAnywhere, BlueprintReadWrite, config, Category = LODSettings, meta = (ImportType = "StaticMesh", DisplayName = "Auto Compute LOD Distances"))
 	uint32 bAutoComputeLodDistances : 1;
 	/** Specify the LOD distance for LOD 0*/
-	UPROPERTY(EditAnywhere, BlueprintReadWrite, config, Category = LodSettings, meta = (ImportType = "StaticMesh", UIMin = "0.0"))
+	UPROPERTY(EditAnywhere, BlueprintReadWrite, config, Category = LODSettings, meta = (ImportType = "StaticMesh", UIMin = "0.0", DisplayName = "LOD Distance 0"))
 	float LodDistance0;
 	/** Specify the LOD distance for LOD 1*/
-	UPROPERTY(EditAnywhere, BlueprintReadWrite, config, Category = LodSettings, meta = (ImportType = "StaticMesh", UIMin = "0.0"))
+	UPROPERTY(EditAnywhere, BlueprintReadWrite, config, Category = LODSettings, meta = (ImportType = "StaticMesh", UIMin = "0.0", DisplayName = "LOD Distance 1"))
 	float LodDistance1;
 	/** Specify the LOD distance for LOD 2*/
-	UPROPERTY(EditAnywhere, BlueprintReadWrite, config, Category = LodSettings, meta = (ImportType = "StaticMesh", UIMin = "0.0"))
+	UPROPERTY(EditAnywhere, BlueprintReadWrite, config, Category = LODSettings, meta = (ImportType = "StaticMesh", UIMin = "0.0", DisplayName = "LOD Distance 2"))
 	float LodDistance2;
 	/** Specify the LOD distance for LOD 3*/
-	UPROPERTY(EditAnywhere, BlueprintReadWrite, config, Category = LodSettings, meta = (ImportType = "StaticMesh", UIMin = "0.0"))
+	UPROPERTY(EditAnywhere, BlueprintReadWrite, config, Category = LODSettings, meta = (ImportType = "StaticMesh", UIMin = "0.0", DisplayName = "LOD Distance 3"))
 	float LodDistance3;
 	/** Specify the LOD distance for LOD 4*/
-	UPROPERTY(EditAnywhere, BlueprintReadWrite, config, Category = LodSettings, meta = (ImportType = "StaticMesh", UIMin = "0.0"))
+	UPROPERTY(EditAnywhere, BlueprintReadWrite, config, Category = LODSettings, meta = (ImportType = "StaticMesh", UIMin = "0.0", DisplayName = "LOD Distance 4"))
 	float LodDistance4;
 	/** Specify the LOD distance for LOD 5*/
-	UPROPERTY(EditAnywhere, BlueprintReadWrite, config, Category = LodSettings, meta = (ImportType = "StaticMesh", UIMin = "0.0"))
+	UPROPERTY(EditAnywhere, BlueprintReadWrite, config, Category = LODSettings, meta = (ImportType = "StaticMesh", UIMin = "0.0", DisplayName = "LOD Distance 5"))
 	float LodDistance5;
 	/** Specify the LOD distance for LOD 6*/
-	UPROPERTY(EditAnywhere, BlueprintReadWrite, config, Category = LodSettings, meta = (ImportType = "StaticMesh", UIMin = "0.0"))
+	UPROPERTY(EditAnywhere, BlueprintReadWrite, config, Category = LODSettings, meta = (ImportType = "StaticMesh", UIMin = "0.0", DisplayName = "LOD Distance 6"))
 	float LodDistance6;
 	/** Specify the LOD distance for LOD 7*/
-	UPROPERTY(EditAnywhere, BlueprintReadWrite, config, Category = LodSettings, meta = (ImportType = "StaticMesh", UIMin = "0.0"))
+	UPROPERTY(EditAnywhere, BlueprintReadWrite, config, Category = LODSettings, meta = (ImportType = "StaticMesh", UIMin = "0.0", DisplayName = "LOD Distance 7"))
 	float LodDistance7;
 
 	/** Set the minimum LOD number. A value of 0 disable the option. */
-	UPROPERTY(EditAnywhere, BlueprintReadWrite, config, AdvancedDisplay, Category = LodSettings, meta = (ImportType = "StaticMesh", UIMin = "0"))
+	UPROPERTY(EditAnywhere, BlueprintReadWrite, config, AdvancedDisplay, Category = LODSettings, meta = (ImportType = "StaticMesh", UIMin = "0", DisplayName = "Minimum LOD Number"))
 	int32 MinimumLodNumber;
 
 	/** Set the number of LODs. A value of 0 disable the option. */
-	UPROPERTY(EditAnywhere, BlueprintReadWrite, config, AdvancedDisplay, Category = LodSettings, meta = (ImportType = "StaticMesh", UIMin = "0"))
+	UPROPERTY(EditAnywhere, BlueprintReadWrite, config, AdvancedDisplay, Category = LODSettings, meta = (ImportType = "StaticMesh", UIMin = "0", DisplayName = "LOD Number"))
 	int32 LodNumber;
 
 	/** True to import animations from the FBX File */
