Fix the socket import to remove the extra conversion

diff --git a/Engine/Source/Editor/UnrealEd/Private/Fbx/FbxMainImport.cpp b/Engine/Source/Editor/UnrealEd/Private/Fbx/FbxMainImport.cpp
--- a/Engine/Source/Editor/UnrealEd/Private/Fbx/FbxMainImport.cpp
+++ b/Engine/Source/Editor/UnrealEd/Private/Fbx/FbxMainImport.cpp
@@ -1175,6 +1175,12 @@ void FFbxImporter::ConvertScene()
 	FileAxisSystem = Scene->GetGlobalSettings().GetAxisSystem();
 	FileUnitSystem = Scene->GetGlobalSettings().GetSystemUnit();
 
+	FbxAMatrix AxisConversionMatrix;
+	AxisConversionMatrix.SetIdentity();
+
+	FbxAMatrix JointOrientationMatrix;
+	JointOrientationMatrix.SetIdentity();
+
 	if (GetImportOptions()->bConvertScene)
 	{
 		// we use -Y as forward axis here when we import. This is odd considering our forward axis is technically +X
@@ -1195,19 +1201,30 @@ void FFbxImporter::ConvertScene()
 
 		FbxAxisSystem SourceSetup = Scene->GetGlobalSettings().GetAxisSystem();
 
+
 		if (SourceSetup != UnrealImportAxis)
 		{
 			FbxRootNodeUtility::RemoveAllFbxRoots(Scene);
 			UnrealImportAxis.ConvertScene(Scene);
-			FbxAMatrix JointOrientationMatrix;
-			JointOrientationMatrix.SetIdentity();
+			
+			FbxAMatrix SourceMatrix;
+			SourceSetup.GetMatrix(SourceMatrix);
+			FbxAMatrix UE4Matrix;
+			UnrealImportAxis.GetMatrix(UE4Matrix);
+			AxisConversionMatrix = SourceMatrix.Inverse() * UE4Matrix;
+			
+			
 			if (GetImportOptions()->bForceFrontXAxis)
 			{
 				JointOrientationMatrix.SetR(FbxVector4(-90.0, -90.0, 0.0));
 			}
-			FFbxDataConverter::SetJointPostConversionMatrix(JointOrientationMatrix);
 		}
 	}
+
+	FFbxDataConverter::SetJointPostConversionMatrix(JointOrientationMatrix);
+
+	FFbxDataConverter::SetAxisConversionMatrix(AxisConversionMatrix);
+
 	// Convert the scene's units to what is used in this program, if needed.
 	// The base unit used in both FBX and Unreal is centimeters.  So unless the units 
 	// are already in centimeters (ie: scalefactor 1.0) then it needs to be converted
diff --git a/Engine/Source/Editor/UnrealEd/Private/Fbx/FbxStaticMeshImport.cpp b/Engine/Source/Editor/UnrealEd/Private/Fbx/FbxStaticMeshImport.cpp
--- a/Engine/Source/Editor/UnrealEd/Private/Fbx/FbxStaticMeshImport.cpp
+++ b/Engine/Source/Editor/UnrealEd/Private/Fbx/FbxStaticMeshImport.cpp
@@ -2236,71 +2236,86 @@ void UnFbx::FFbxImporter::ImportStaticMeshLocalSockets(UStaticMesh* StaticMesh,
 	check(MeshNodeArray.Num());
 	FbxNode *MeshRootNode = MeshNodeArray[0];
 	const FbxAMatrix &MeshTotalMatrix = ComputeTotalMatrix(MeshRootNode);
+	TArray<FbxSocketNode> AllSocketNodes;
 	for (FbxNode* RootNode : MeshNodeArray)
 	{
 		// Find all nodes that are sockets
 		TArray<FbxSocketNode> SocketNodes;
 		FindMeshSockets(RootNode, SocketNodes);
-
-		// Create a UStaticMeshSocket for each fbx socket
 		for (int32 SocketIndex = 0; SocketIndex < SocketNodes.Num(); ++SocketIndex)
 		{
-			FbxSocketNode& SocketNode = SocketNodes[SocketIndex];
-
-			UStaticMeshSocket* Socket = StaticMesh->FindSocket(SocketNode.SocketName);
-			if (!Socket)
+			bool bFoundNewSocket = true;
+			for (int32 AllSocketIndex = 0; AllSocketIndex < AllSocketNodes.Num(); ++AllSocketIndex)
 			{
-				// If the socket didn't exist create a new one now
-				Socket = NewObject<UStaticMeshSocket>(StaticMesh);
-				Socket->bSocketCreatedAtImport = true;
-				check(Socket);
-
-				Socket->SocketName = SocketNode.SocketName;
-				StaticMesh->Sockets.Add(Socket);
+				if (AllSocketNodes[AllSocketIndex].SocketName == SocketNodes[SocketIndex].SocketName)
+				{
+					bFoundNewSocket = false;
+					break;
+				}
 			}
-
-			if (Socket)
+			if (bFoundNewSocket)
 			{
-				const FbxAMatrix& SocketMatrix = Scene->GetAnimationEvaluator()->GetNodeLocalTransform(SocketNode.Node);
-				FbxAMatrix FinalSocketMatrix = MeshTotalMatrix * SocketMatrix;
-				FTransform SocketTransform;
-				SocketTransform.SetTranslation(Converter.ConvertPos(FinalSocketMatrix.GetT()));
-				SocketTransform.SetRotation(Converter.ConvertRotToQuat(FinalSocketMatrix.GetQ()));
-				SocketTransform.SetScale3D(Converter.ConvertScale(FinalSocketMatrix.GetS()));
-
-				Socket->RelativeLocation = SocketTransform.GetLocation();
-				Socket->RelativeRotation = SocketTransform.GetRotation().Rotator();
-				Socket->RelativeScale = SocketTransform.GetScale3D();
+				AllSocketNodes.Add(SocketNodes[SocketIndex]);
 			}
 		}
-		// Delete mesh sockets that were removed from the import data
-		if (StaticMesh->Sockets.Num() != SocketNodes.Num())
+	}
+		
+	// Create a UStaticMeshSocket for each fbx socket
+	for (int32 SocketIndex = 0; SocketIndex < AllSocketNodes.Num(); ++SocketIndex)
+	{
+		FbxSocketNode& SocketNode = AllSocketNodes[SocketIndex];
+
+		UStaticMeshSocket* Socket = StaticMesh->FindSocket(SocketNode.SocketName);
+		if (!Socket)
 		{
-			for (int32 MeshSocketIx = 0; MeshSocketIx < StaticMesh->Sockets.Num(); ++MeshSocketIx)
-			{
-				bool Found = false;
-				UStaticMeshSocket* MeshSocket = StaticMesh->Sockets[MeshSocketIx];
-				//Do not remove socket that was not generated at import
-				if (!MeshSocket->bSocketCreatedAtImport)
-				{
-					continue;
-				}
+			// If the socket didn't exist create a new one now
+			Socket = NewObject<UStaticMeshSocket>(StaticMesh);
+			Socket->bSocketCreatedAtImport = true;
+			check(Socket);
 
-				for (int32 FbxSocketIx = 0; FbxSocketIx < SocketNodes.Num(); FbxSocketIx++)
-				{
-					if (SocketNodes[FbxSocketIx].SocketName == MeshSocket->SocketName)
-					{
-						Found = true;
-						break;
-					}
-				}
-				if (!Found)
-				{
-					StaticMesh->Sockets.RemoveAt(MeshSocketIx);
-					MeshSocketIx--;
-				}
+			Socket->SocketName = SocketNode.SocketName;
+			StaticMesh->Sockets.Add(Socket);
+		}
+
+		if (Socket)
+		{
+			const FbxAMatrix& SocketMatrix = Scene->GetAnimationEvaluator()->GetNodeLocalTransform(SocketNode.Node);
+			//Remove the axis conversion for the socket since its attach to a mesh containing this conversion.
+			FbxAMatrix FinalSocketMatrix = (MeshTotalMatrix * SocketMatrix) * FFbxDataConverter::GetAxisConversionMatrixInv();
+			FTransform SocketTransform;
+			SocketTransform.SetTranslation(Converter.ConvertPos(FinalSocketMatrix.GetT()));
+			SocketTransform.SetRotation(Converter.ConvertRotToQuat(FinalSocketMatrix.GetQ()));
+			SocketTransform.SetScale3D(Converter.ConvertScale(FinalSocketMatrix.GetS()));
+
+			Socket->RelativeLocation = SocketTransform.GetLocation();
+			Socket->RelativeRotation = SocketTransform.GetRotation().Rotator();
+			Socket->RelativeScale = SocketTransform.GetScale3D();
+		}
+	}
+	// Delete mesh sockets that were removed from the import data
+	for (int32 MeshSocketIx = 0; MeshSocketIx < StaticMesh->Sockets.Num(); ++MeshSocketIx)
+	{
+		bool Found = false;
+		UStaticMeshSocket* MeshSocket = StaticMesh->Sockets[MeshSocketIx];
+		//Do not remove socket that was not generated at import
+		if (!MeshSocket->bSocketCreatedAtImport)
+		{
+			continue;
+		}
+
+		for (int32 FbxSocketIx = 0; FbxSocketIx < AllSocketNodes.Num(); FbxSocketIx++)
+		{
+			if (AllSocketNodes[FbxSocketIx].SocketName == MeshSocket->SocketName)
+			{
+				Found = true;
+				break;
 			}
 		}
+		if (!Found)
+		{
+			StaticMesh->Sockets.RemoveAt(MeshSocketIx);
+			MeshSocketIx--;
+		}
 	}
 }
 
@@ -2326,8 +2341,8 @@ void UnFbx::FFbxImporter::ImportStaticMeshGlobalSockets( UStaticMesh* StaticMesh
 
 			Socket->SocketName = SocketNode.SocketName;
 			StaticMesh->Sockets.Add(Socket);
-
-			const FbxAMatrix& SocketMatrix = Scene->GetAnimationEvaluator()->GetNodeGlobalTransform(SocketNode.Node);
+			//Remove the axis conversion for the socket since its attach to a mesh containing this conversion.
+			const FbxAMatrix& SocketMatrix = Scene->GetAnimationEvaluator()->GetNodeGlobalTransform(SocketNode.Node) * FFbxDataConverter::GetAxisConversionMatrixInv();
 			FTransform SocketTransform;
 			SocketTransform.SetTranslation(Converter.ConvertPos(SocketMatrix.GetT()));
 			SocketTransform.SetRotation(Converter.ConvertRotToQuat(SocketMatrix.GetQ()));
@@ -2340,33 +2355,29 @@ void UnFbx::FFbxImporter::ImportStaticMeshGlobalSockets( UStaticMesh* StaticMesh
 			Socket->bSocketCreatedAtImport = true;
 		}
 	}
-	// Delete mesh sockets that were removed from the import data
-	if (StaticMesh->Sockets.Num() != SocketNodes.Num())
+	for (int32 MeshSocketIx = 0; MeshSocketIx < StaticMesh->Sockets.Num(); ++MeshSocketIx)
 	{
-		for (int32 MeshSocketIx = 0; MeshSocketIx < StaticMesh->Sockets.Num(); ++MeshSocketIx)
+		bool Found = false;
+		UStaticMeshSocket* MeshSocket = StaticMesh->Sockets[MeshSocketIx];
+		//Do not remove socket that was not generated at import
+		if (!MeshSocket->bSocketCreatedAtImport)
 		{
-			bool Found = false;
-			UStaticMeshSocket* MeshSocket = StaticMesh->Sockets[MeshSocketIx];
-			//Do not remove socket that was not generated at import
-			if (!MeshSocket->bSocketCreatedAtImport)
-			{
-				continue;
-			}
+			continue;
+		}
 
-			for (int32 FbxSocketIx = 0; FbxSocketIx < SocketNodes.Num(); FbxSocketIx++)
-			{
-				if (SocketNodes[FbxSocketIx].SocketName == MeshSocket->SocketName)
-				{
-					Found = true;
-					break;
-				}
-			}
-			if (!Found)
+		for (int32 FbxSocketIx = 0; FbxSocketIx < SocketNodes.Num(); FbxSocketIx++)
+		{
+			if (SocketNodes[FbxSocketIx].SocketName == MeshSocket->SocketName)
 			{
-				StaticMesh->Sockets.RemoveAt(MeshSocketIx);
-				MeshSocketIx--;
+				Found = true;
+				break;
 			}
 		}
+		if (!Found)
+		{
+			StaticMesh->Sockets.RemoveAt(MeshSocketIx);
+			MeshSocketIx--;
+		}
 	}
 }
 
diff --git a/Engine/Source/Editor/UnrealEd/Private/Fbx/FbxUtilsImport.cpp b/Engine/Source/Editor/UnrealEd/Private/Fbx/FbxUtilsImport.cpp
--- a/Engine/Source/Editor/UnrealEd/Private/Fbx/FbxUtilsImport.cpp
+++ b/Engine/Source/Editor/UnrealEd/Private/Fbx/FbxUtilsImport.cpp
@@ -7,6 +7,8 @@ using namespace UnFbx;
 
 
 FbxAMatrix FFbxDataConverter::JointPostConversionMatrix;
+FbxAMatrix FFbxDataConverter::AxisConversionMatrix;
+FbxAMatrix FFbxDataConverter::AxisConversionMatrixInv;
 
 FVector FFbxDataConverter::ConvertPos(FbxVector4 Vector)
 {
diff --git a/Engine/Source/Editor/UnrealEd/Public/FbxImporter.h b/Engine/Source/Editor/UnrealEd/Public/FbxImporter.h
--- a/Engine/Source/Editor/UnrealEd/Public/FbxImporter.h
+++ b/Engine/Source/Editor/UnrealEd/Public/FbxImporter.h
@@ -461,6 +461,10 @@ public:
 	static void SetJointPostConversionMatrix(FbxAMatrix ConversionMatrix) { JointPostConversionMatrix = ConversionMatrix; }
 	static const FbxAMatrix &GetJointPostConversionMatrix() { return JointPostConversionMatrix; }
 
+	static void SetAxisConversionMatrix(FbxAMatrix ConversionMatrix) { AxisConversionMatrix = ConversionMatrix; AxisConversionMatrixInv = ConversionMatrix.Inverse(); }
+	static const FbxAMatrix &GetAxisConversionMatrix() { return AxisConversionMatrix; }
+	static const FbxAMatrix &GetAxisConversionMatrixInv() { return AxisConversionMatrixInv; }
+
 	static FVector ConvertPos(FbxVector4 Vector);
 	static FVector ConvertDir(FbxVector4 Vector);
 	static FRotator ConvertEuler(FbxDouble3 Euler);
@@ -498,6 +502,8 @@ public:
 
 private:
 	static FbxAMatrix JointPostConversionMatrix;
+	static FbxAMatrix AxisConversionMatrix;
+	static FbxAMatrix AxisConversionMatrixInv;
 };
 
 FBXImportOptions* GetImportOptions( class FFbxImporter* FbxImporter, UFbxImportUI* ImportUI, bool bShowOptionDialog, bool bIsAutomated, const FString& FullPath, bool& OutOperationCanceled, bool& OutImportAll, bool bIsObjFormat, const FString& InFilename, bool bForceImportType = false, EFBXImportType ImportType = FBXIT_StaticMesh, UObject* ReimportObject = nullptr);
