Clamp fbx import dialog "CustomSampleRate" value between 0 and 1920

diff --git a/Engine/Source/Editor/UnrealEd/Classes/Factories/FbxAnimSequenceImportData.h b/Engine/Source/Editor/UnrealEd/Classes/Factories/FbxAnimSequenceImportData.h
--- a/Engine/Source/Editor/UnrealEd/Classes/Factories/FbxAnimSequenceImportData.h
+++ b/Engine/Source/Editor/UnrealEd/Classes/Factories/FbxAnimSequenceImportData.h
@@ -63,7 +63,7 @@ class UNREALED_API UFbxAnimSequenceImportData : public UFbxAssetImportData
 	bool bUseDefaultSampleRate;
 
 	/** Use this option to specify a sample rate for the imported animation, a value of 0 use the best matching samplerate. */
-	UPROPERTY(EditAnywhere, AdvancedDisplay, config, Category = ImportSettings, meta = (EditCondition = "!bUseDefaultSampleRate", ToolTip = "Sample fbx animation data at the specified sample rate, 0 find automaticaly the best sample rate"))
+	UPROPERTY(EditAnywhere, AdvancedDisplay, config, Category = ImportSettings, meta = (EditCondition = "!bUseDefaultSampleRate", ToolTip = "Sample fbx animation data at the specified sample rate, 0 find automaticaly the best sample rate", ClampMin = 0, UIMin = 0, ClampMax = 48000, UIMax = 60))
 	int32 CustomSampleRate;
 
 	/** Name of source animation that was imported, used to reimport correct animation from the FBX file */
