PR #5063: UE-30316: Re-set bHiddenEdTemporary when recreating BSP during brush . (Contributed by [REDACTED_PERSON])

diff --git a/Engine/Source/Editor/UnrealEd/Private/BSPOps.cpp b/Engine/Source/Editor/UnrealEd/Private/BSPOps.cpp
--- a/Engine/Source/Editor/UnrealEd/Private/BSPOps.cpp
+++ b/Engine/Source/Editor/UnrealEd/Private/BSPOps.cpp
@@ -1127,6 +1127,10 @@ int32	FBSPOps::bspAddNode( UModel* Model, int32 iParent, ENodePlace NodePlace, u
 
 		Surf->Actor	 		= EdPoly->Actor;
 		Surf->iBrushPoly	= EdPoly->iBrushPoly;
+		
+		Surf->bHiddenEdTemporary = EdPoly->Actor->IsTemporarilyHiddenInEditor();
+		Surf->bHiddenEdLevel = EdPoly->Actor->bHiddenEdLevel;
+		Surf->bHiddenEdLayer = EdPoly->Actor->bHiddenEdLayer;
 
 		Surf->Plane			= FPlane(EdPoly->Vertices[0],EdPoly->Normal);
 	}
