Fix for crash is the EdPoly doesn't have a brush actor.

diff --git a/Engine/Source/Editor/UnrealEd/Private/BSPOps.cpp b/Engine/Source/Editor/UnrealEd/Private/BSPOps.cpp
--- a/Engine/Source/Editor/UnrealEd/Private/BSPOps.cpp
+++ b/Engine/Source/Editor/UnrealEd/Private/BSPOps.cpp
@@ -1128,9 +1128,12 @@ int32	FBSPOps::bspAddNode( UModel* Model, int32 iParent, ENodePlace NodePlace, u
 		Surf->Actor	 		= EdPoly->Actor;
 		Surf->iBrushPoly	= EdPoly->iBrushPoly;
 		
-		Surf->bHiddenEdTemporary = EdPoly->Actor->IsTemporarilyHiddenInEditor();
-		Surf->bHiddenEdLevel = EdPoly->Actor->bHiddenEdLevel;
-		Surf->bHiddenEdLayer = EdPoly->Actor->bHiddenEdLayer;
+		if (EdPoly->Actor)
+		{
+			Surf->bHiddenEdTemporary = EdPoly->Actor->IsTemporarilyHiddenInEditor();
+			Surf->bHiddenEdLevel = EdPoly->Actor->bHiddenEdLevel;
+			Surf->bHiddenEdLayer = EdPoly->Actor->bHiddenEdLayer;
+		}
 
 		Surf->Plane			= FPlane(EdPoly->Vertices[0],EdPoly->Normal);
 	}
