Downgraded Mesh triangulation warning to display

diff --git a/Engine/Source/Editor/UnrealEd/Private/Fbx/FbxStaticMeshImport.cpp b/Engine/Source/Editor/UnrealEd/Private/Fbx/FbxStaticMeshImport.cpp
--- a/Engine/Source/Editor/UnrealEd/Private/Fbx/FbxStaticMeshImport.cpp
+++ b/Engine/Source/Editor/UnrealEd/Private/Fbx/FbxStaticMeshImport.cpp
@@ -374,7 +374,7 @@ bool UnFbx::FFbxImporter::BuildStaticMeshFromGeometry(FbxNode* Node, UStaticMesh
 	if (!Mesh->IsTriangleMesh())
 	{
 		if(!GIsAutomationTesting)
-		UE_LOG(LogFbx, Warning, TEXT("Triangulating static mesh %s"), UTF8_TO_TCHAR(Node->GetName()));
+		UE_LOG(LogFbx, Display, TEXT("Triangulating static mesh %s"), UTF8_TO_TCHAR(Node->GetName()));
 
 		const bool bReplace = true;
 		FbxNodeAttribute* ConvertedNode = GeometryConverter->Triangulate(Mesh, bReplace);
