Fbx Importer: call post edit change on extra staticmesh when importing without combine flag. Posteditchange was call only on the first imported mesh.

diff --git a/Engine/Source/Editor/UnrealEd/Private/Fbx/FbxFactory.cpp b/Engine/Source/Editor/UnrealEd/Private/Fbx/FbxFactory.cpp
--- a/Engine/Source/Editor/UnrealEd/Private/Fbx/FbxFactory.cpp
+++ b/Engine/Source/Editor/UnrealEd/Private/Fbx/FbxFactory.cpp
@@ -483,6 +483,8 @@ UObject* UFbxFactory::FactoryCreateFile
 							{
 								FAssetRegistryModule::AssetCreated(Asset);
 								Asset->MarkPackageDirty();
+								//Make sure the build is up to date with the latest section info map
+								Asset->PostEditChange();
 							}
 						}
 
