Fix some tooltip and display name around LODs options.

diff --git a/Engine/Source/Editor/Persona/Private/PersonaMeshDetails.cpp b/Engine/Source/Editor/Persona/Private/PersonaMeshDetails.cpp
--- a/Engine/Source/Editor/Persona/Private/PersonaMeshDetails.cpp
+++ b/Engine/Source/Editor/Persona/Private/PersonaMeshDetails.cpp
@@ -798,14 +798,14 @@ void FPersonaMeshDetails::AddLODLevelCategories(IDetailLayoutBuilder& DetailLayo
 			SNew(STextBlock)
 			.Font(IDetailLayoutBuilder::GetDetailFont())
 			.Text(this, &FPersonaMeshDetails::GetLODCustomModeNameContent, (int32)INDEX_NONE)
-			.ToolTipText(LOCTEXT("LODCustomModeFirstRowTooltip", "Custom Mode allow editing multiple LOD in same time."))
+			.ToolTipText(LOCTEXT("LODCustomModeFirstRowTooltip", "Custom Mode shows multiple LOD's properties at the same time for easier editing."))
 		]
 		.ValueContent()
 		[
 			SNew(SCheckBox)
 			.IsChecked(this, &FPersonaMeshDetails::IsLODCustomModeCheck, (int32)INDEX_NONE)
 			.OnCheckStateChanged(this, &FPersonaMeshDetails::SetLODCustomModeCheck, (int32)INDEX_NONE)
-			.ToolTipText(LOCTEXT("LODCustomModeFirstRowTooltip", "Custom Mode allow editing multiple LOD in same time."))
+			.ToolTipText(LOCTEXT("LODCustomModeFirstRowTooltip", "Custom Mode shows multiple LOD's properties at the same time for easier editing."))
 		];
 
 		LodCategories.Empty(SkelMeshLODCount);
@@ -2620,7 +2620,7 @@ TSharedRef<SWidget> FPersonaMeshDetails::OnGenerateLodMenuForLodPicker()
 
 	FText AutoLodText = FText::FromString((TEXT("Auto LOD")));
 	FUIAction AutoLodAction(FExecuteAction::CreateSP(this, &FPersonaMeshDetails::SetCurrentLOD, 0));
-	MenuBuilder.AddMenuEntry(AutoLodText, LOCTEXT("OnGenerateLodMenuForSectionList_Auto_ToolTip", "LOD0 is edit when selecting Auto LOD"), FSlateIcon(), AutoLodAction);
+	MenuBuilder.AddMenuEntry(AutoLodText, LOCTEXT("OnGenerateLodMenuForSectionList_Auto_ToolTip", "With Auto LOD selected, LOD0's properties are visible for editing."), FSlateIcon(), AutoLodAction);
 	// Add a menu item for each texture.  Clicking on the texture will display it in the content browser
 	for (int32 AllLodIndex = 0; AllLodIndex < SkelMeshLODCount; ++AllLodIndex)
 	{
diff --git a/Engine/Source/Editor/StaticMeshEditor/Private/StaticMeshEditorTools.cpp b/Engine/Source/Editor/StaticMeshEditor/Private/StaticMeshEditorTools.cpp
--- a/Engine/Source/Editor/StaticMeshEditor/Private/StaticMeshEditorTools.cpp
+++ b/Engine/Source/Editor/StaticMeshEditor/Private/StaticMeshEditorTools.cpp
@@ -2964,14 +2964,14 @@ void FLevelOfDetailSettingsLayout::AddLODLevelCategories( IDetailLayoutBuilder&
 			SNew(STextBlock)
 			.Font(IDetailLayoutBuilder::GetDetailFont())
 			.Text(this, &FLevelOfDetailSettingsLayout::GetLODCustomModeNameContent, (int32)INDEX_NONE)
-			.ToolTipText(LOCTEXT("LODCustomModeFirstRowTooltip", "Custom Mode allow editing multiple LOD in same time."))
+			.ToolTipText(LOCTEXT("LODCustomModeFirstRowTooltip", "Custom Mode shows multiple LOD's properties at the same time for easier editing."))
 		]
 		.ValueContent()
 		[
 			SNew(SCheckBox)
 			.IsChecked(this, &FLevelOfDetailSettingsLayout::IsLODCustomModeCheck, (int32)INDEX_NONE)
 			.OnCheckStateChanged(this, &FLevelOfDetailSettingsLayout::SetLODCustomModeCheck, (int32)INDEX_NONE)
-			.ToolTipText(LOCTEXT("LODCustomModeFirstRowTooltip", "Custom Mode allow editing multiple LOD in same time."))
+			.ToolTipText(LOCTEXT("LODCustomModeFirstRowTooltip", "Custom Mode shows multiple LOD's properties at the same time for easier editing."))
 		];
 		// Create information panel for each LOD level.
 		for(int32 LODIndex = 0; LODIndex < StaticMeshLODCount; ++LODIndex)
@@ -3890,7 +3890,7 @@ TSharedRef<SWidget> FLevelOfDetailSettingsLayout::OnGenerateLodMenuForLodPicker(
 
 	FText AutoLodText = FText::FromString((TEXT("LOD Auto")));
 	FUIAction AutoLodAction(FExecuteAction::CreateSP(this, &FLevelOfDetailSettingsLayout::OnSelectedLODChanged, 0));
-	MenuBuilder.AddMenuEntry(AutoLodText, LOCTEXT("OnGenerateLodMenuForLodPicker_Auto_ToolTip", "LOD0 is edit when selecting Auto LOD"), FSlateIcon(), AutoLodAction);
+	MenuBuilder.AddMenuEntry(AutoLodText, LOCTEXT("OnGenerateLodMenuForLodPicker_Auto_ToolTip", "With Auto LOD selected, LOD0's properties are visible for editing."), FSlateIcon(), AutoLodAction);
 	// Add a menu item for each texture.  Clicking on the texture will display it in the content browser
 	for (int32 AllLodIndex = 0; AllLodIndex < StaticMeshLODCount; ++AllLodIndex)
 	{
diff --git a/Engine/Source/Editor/UnrealEd/Classes/Factories/FbxImportUI.h b/Engine/Source/Editor/UnrealEd/Classes/Factories/FbxImportUI.h
--- a/Engine/Source/Editor/UnrealEd/Classes/Factories/FbxImportUI.h
+++ b/Engine/Source/Editor/UnrealEd/Classes/Factories/FbxImportUI.h
@@ -71,40 +71,40 @@ public:
 	UPROPERTY(EditAnywhere, BlueprintReadWrite, AdvancedDisplay, Category=Mesh, meta=(ImportType="SkeletalMesh", editcondition="!bCreatePhysicsAsset"))
 	class UPhysicsAsset* PhysicsAsset;
 
-	/** If checked, the static mesh auto compute LOD distance will be turn on. If unchecked user will be able to specify custom LOD distance for every LOD. */
-	UPROPERTY(EditAnywhere, BlueprintReadWrite, config, Category = LODSettings, meta = (ImportType = "StaticMesh", DisplayName = "Auto Compute LOD Distances"))
+	/** If checked, the editor will automatically compute screen size values for the static mesh's LODs. If unchecked, the user can enter custom screen size values for each LOD. */
+	UPROPERTY(EditAnywhere, BlueprintReadWrite, config, Category = LODSettings, meta = (ImportType = "StaticMesh", DisplayName = "Auto Compute LOD Screen Size"))
 	uint32 bAutoComputeLodDistances : 1;
-	/** Specify the LOD distance for LOD 0*/
-	UPROPERTY(EditAnywhere, BlueprintReadWrite, config, Category = LODSettings, meta = (ImportType = "StaticMesh", UIMin = "0.0", DisplayName = "LOD Distance 0"))
+	/** Set a screen size value for LOD 0*/
+	UPROPERTY(EditAnywhere, BlueprintReadWrite, config, Category = LODSettings, meta = (ImportType = "StaticMesh", UIMin = "0.0", DisplayName = "LOD 0 Screen Size"))
 	float LodDistance0;
-	/** Specify the LOD distance for LOD 1*/
-	UPROPERTY(EditAnywhere, BlueprintReadWrite, config, Category = LODSettings, meta = (ImportType = "StaticMesh", UIMin = "0.0", DisplayName = "LOD Distance 1"))
+	/** Set a screen size value for LOD 1*/
+	UPROPERTY(EditAnywhere, BlueprintReadWrite, config, Category = LODSettings, meta = (ImportType = "StaticMesh", UIMin = "0.0", DisplayName = "LOD 1 Screen Size"))
 	float LodDistance1;
-	/** Specify the LOD distance for LOD 2*/
-	UPROPERTY(EditAnywhere, BlueprintReadWrite, config, Category = LODSettings, meta = (ImportType = "StaticMesh", UIMin = "0.0", DisplayName = "LOD Distance 2"))
+	/** Set a screen size value for LOD 2*/
+	UPROPERTY(EditAnywhere, BlueprintReadWrite, config, Category = LODSettings, meta = (ImportType = "StaticMesh", UIMin = "0.0", DisplayName = "LOD 2 Screen Size"))
 	float LodDistance2;
-	/** Specify the LOD distance for LOD 3*/
-	UPROPERTY(EditAnywhere, BlueprintReadWrite, config, Category = LODSettings, meta = (ImportType = "StaticMesh", UIMin = "0.0", DisplayName = "LOD Distance 3"))
+	/** Set a screen size value for LOD 3*/
+	UPROPERTY(EditAnywhere, BlueprintReadWrite, config, Category = LODSettings, meta = (ImportType = "StaticMesh", UIMin = "0.0", DisplayName = "LOD 3 Screen Size"))
 	float LodDistance3;
-	/** Specify the LOD distance for LOD 4*/
-	UPROPERTY(EditAnywhere, BlueprintReadWrite, config, Category = LODSettings, meta = (ImportType = "StaticMesh", UIMin = "0.0", DisplayName = "LOD Distance 4"))
+	/** Set a screen size value for LOD 4*/
+	UPROPERTY(EditAnywhere, BlueprintReadWrite, config, Category = LODSettings, meta = (ImportType = "StaticMesh", UIMin = "0.0", DisplayName = "LOD 4 Screen Size"))
 	float LodDistance4;
-	/** Specify the LOD distance for LOD 5*/
-	UPROPERTY(EditAnywhere, BlueprintReadWrite, config, Category = LODSettings, meta = (ImportType = "StaticMesh", UIMin = "0.0", DisplayName = "LOD Distance 5"))
+	/** Set a screen size value for LOD 5*/
+	UPROPERTY(EditAnywhere, BlueprintReadWrite, config, Category = LODSettings, meta = (ImportType = "StaticMesh", UIMin = "0.0", DisplayName = "LOD 5 Screen Size"))
 	float LodDistance5;
-	/** Specify the LOD distance for LOD 6*/
-	UPROPERTY(EditAnywhere, BlueprintReadWrite, config, Category = LODSettings, meta = (ImportType = "StaticMesh", UIMin = "0.0", DisplayName = "LOD Distance 6"))
+	/** Set a screen size value for LOD 6*/
+	UPROPERTY(EditAnywhere, BlueprintReadWrite, config, Category = LODSettings, meta = (ImportType = "StaticMesh", UIMin = "0.0", DisplayName = "LOD 6 Screen Size"))
 	float LodDistance6;
-	/** Specify the LOD distance for LOD 7*/
-	UPROPERTY(EditAnywhere, BlueprintReadWrite, config, Category = LODSettings, meta = (ImportType = "StaticMesh", UIMin = "0.0", DisplayName = "LOD Distance 7"))
+	/** Set a screen size value for LOD 7*/
+	UPROPERTY(EditAnywhere, BlueprintReadWrite, config, Category = LODSettings, meta = (ImportType = "StaticMesh", UIMin = "0.0", DisplayName = "LOD 7 Screen Size"))
 	float LodDistance7;
 
-	/** Set the minimum LOD number. A value of 0 disable the option. */
-	UPROPERTY(EditAnywhere, BlueprintReadWrite, config, AdvancedDisplay, Category = LODSettings, meta = (ImportType = "StaticMesh", UIMin = "0", DisplayName = "Minimum LOD Number"))
+	/** Set the minimum LOD used for rendering. Setting the value to 0 will use the default value of LOD0. */
+	UPROPERTY(EditAnywhere, BlueprintReadWrite, config, AdvancedDisplay, Category = LODSettings, meta = (ImportType = "StaticMesh", UIMin = "0", DisplayName = "Minimum LOD"))
 	int32 MinimumLodNumber;
 
-	/** Set the number of LODs. A value of 0 disable the option. */
-	UPROPERTY(EditAnywhere, BlueprintReadWrite, config, AdvancedDisplay, Category = LODSettings, meta = (ImportType = "StaticMesh", UIMin = "0", DisplayName = "LOD Number"))
+	/** Set the number of LODs for the editor to import. Setting the value to 0 imports the number of LODs found in the file (up to the maximum). */
+	UPROPERTY(EditAnywhere, BlueprintReadWrite, config, AdvancedDisplay, Category = LODSettings, meta = (ImportType = "StaticMesh", UIMin = "0", DisplayName = "Number of LODs"))
 	int32 LodNumber;
 
 	/** True to import animations from the FBX File */
diff --git a/Engine/Source/Runtime/Engine/Classes/Engine/SkeletalMesh.h b/Engine/Source/Runtime/Engine/Classes/Engine/SkeletalMesh.h
--- a/Engine/Source/Runtime/Engine/Classes/Engine/SkeletalMesh.h
+++ b/Engine/Source/Runtime/Engine/Classes/Engine/SkeletalMesh.h
@@ -500,7 +500,7 @@ private:
 
 public:
 	/** Minimum LOD to render. Can be overridden per component as well as set here for all mesh instances here */
-	UPROPERTY(EditAnywhere, Category = LODSettings)
+	UPROPERTY(EditAnywhere, Category = LODSettings, meta = (DisplayName = "Minimum LOD"))
 	FPerPlatformInt MinLod;
 
 #if WITH_EDITORONLY_DATA
diff --git a/Engine/Source/Runtime/Engine/Classes/Engine/SkeletalMeshLODSettings.h b/Engine/Source/Runtime/Engine/Classes/Engine/SkeletalMeshLODSettings.h
--- a/Engine/Source/Runtime/Engine/Classes/Engine/SkeletalMeshLODSettings.h
+++ b/Engine/Source/Runtime/Engine/Classes/Engine/SkeletalMeshLODSettings.h
@@ -102,7 +102,7 @@ class USkeletalMeshLODSettings : public UDataAsset
 protected:
 
 	/** Minimum LOD to render. Can be overridden per component as well as set here for all mesh instances here */
-	UPROPERTY(globalconfig, EditAnywhere, Category=LODGroups)
+	UPROPERTY(globalconfig, EditAnywhere, Category=LODGroups, meta = (DisplayName = "Minimum LOD"))
 	FPerPlatformInt MinLod;
 
 	UPROPERTY(globalconfig, EditAnywhere, Category=LODGroups)
