PR #5022: Change Logging level in SourceCodeNavigation (Contributed by [REDACTED_PERSON])

diff --git a/Engine/Source/Editor/UnrealEd/Private/SourceCodeNavigation.cpp b/Engine/Source/Editor/UnrealEd/Private/SourceCodeNavigation.cpp
--- a/Engine/Source/Editor/UnrealEd/Private/SourceCodeNavigation.cpp
+++ b/Engine/Source/Editor/UnrealEd/Private/SourceCodeNavigation.cpp
@@ -498,7 +498,7 @@ void FSourceCodeNavigationImpl::NavigateToFunctionSource( const FString& Functio
 				SourceLineNumber = FileAndLineInfo.LineNumber;
 			}
 
-			UE_LOG(LogSelectionDetails, Warning, TEXT( "NavigateToFunctionSource:  Found symbols for [%s] - File [%s], Line [%i], Column [%i]" ),
+			UE_LOG(LogSelectionDetails, Verbose, TEXT( "NavigateToFunctionSource:  Found symbols for [%s] - File [%s], Line [%i], Column [%i]" ),
 				*FunctionSymbolName,
 				*SourceFileName,
 				(uint32)FileAndLineInfo.LineNumber,
