Fbx importer now support animation layer when computing the time span, the merge layer is done before the conversion

diff --git a/Engine/Source/Editor/UnrealEd/Private/Fbx/FbxCurveDataImport.cpp b/Engine/Source/Editor/UnrealEd/Private/Fbx/FbxCurveDataImport.cpp
--- a/Engine/Source/Editor/UnrealEd/Private/Fbx/FbxCurveDataImport.cpp
+++ b/Engine/Source/Editor/UnrealEd/Private/Fbx/FbxCurveDataImport.cpp
@@ -749,8 +749,6 @@ namespace UnFbx {
 		if (!AnimStack)
 			return;
 
-		MergeAllLayerAnimation(AnimStack, FbxTime::GetFrameRate(Scene->GetGlobalSettings().GetTimeMode()));
-
 		FbxAnimLayer* AnimLayer = AnimStack->GetMember<FbxAnimLayer>(0);
 		if (AnimLayer == NULL)
 			return;
diff --git a/Engine/Source/Editor/UnrealEd/Private/Fbx/FbxMainImport.cpp b/Engine/Source/Editor/UnrealEd/Private/Fbx/FbxMainImport.cpp
--- a/Engine/Source/Editor/UnrealEd/Private/Fbx/FbxMainImport.cpp
+++ b/Engine/Source/Editor/UnrealEd/Private/Fbx/FbxMainImport.cpp
@@ -143,7 +143,7 @@ FBXImportOptions* GetImportOptions( UnFbx::FFbxImporter* FbxImporter, UFbxImport
 		{
 			FbxAnimStack* CurAnimStack = FbxImporter->Scene->GetSrcObject<FbxAnimStack>(AnimStackIndex);
 			FbxTimeSpan AnimatedInterval(FBXSDK_TIME_INFINITE, FBXSDK_TIME_MINUS_INFINITE);
-			FbxImporter->Scene->GetRootNode()->GetAnimationInterval(AnimatedInterval, CurAnimStack);
+			FbxImporter->GetAnimationIntervalMultiLayer(FbxImporter->Scene->GetRootNode(), CurAnimStack, AnimatedInterval);
 			// find the most range that covers by both method, that'll be used for clamping
 			AnimTimeSpan.SetStart(FMath::Min<FbxTime>(AnimTimeSpan.GetStart(), AnimatedInterval.GetStart()));
 			AnimTimeSpan.SetStop(FMath::Max<FbxTime>(AnimTimeSpan.GetStop(), AnimatedInterval.GetStop()));
@@ -1171,6 +1171,16 @@ bool FFbxImporter::ImportFile(FString Filename, bool bPreventMaterialNameClash /
 
 void FFbxImporter::ConvertScene()
 {
+	//Merge the anim stack before the conversion since the above 0 layer will not be converted
+	int32 AnimStackCount = Scene->GetSrcObjectCount<FbxAnimStack>();
+	//Merge the animation stack layer before converting the scene
+	for (int32 AnimStackIndex = 0; AnimStackIndex < AnimStackCount; AnimStackIndex++)
+	{
+		FbxAnimStack* CurAnimStack = Scene->GetSrcObject<FbxAnimStack>(AnimStackIndex);
+		int32 ResampleRate = GetGlobalAnimStackSampleRate(CurAnimStack);
+		MergeAllLayerAnimation(CurAnimStack, ResampleRate);
+	}
+
 	//Set the original file information
 	FileAxisSystem = Scene->GetGlobalSettings().GetAxisSystem();
 	FileUnitSystem = Scene->GetGlobalSettings().GetSystemUnit();
diff --git a/Engine/Source/Editor/UnrealEd/Private/Fbx/FbxMatineeImport.cpp b/Engine/Source/Editor/UnrealEd/Private/Fbx/FbxMatineeImport.cpp
--- a/Engine/Source/Editor/UnrealEd/Private/Fbx/FbxMatineeImport.cpp
+++ b/Engine/Source/Editor/UnrealEd/Private/Fbx/FbxMatineeImport.cpp
@@ -163,8 +163,6 @@ bool FFbxImporter::ImportMatineeSequence(AMatineeActor* InMatineeActor)
 	FbxAnimStack* AnimStack = Scene->GetMember<FbxAnimStack>(0);
 	if (!AnimStack) return false;
 		
-	MergeAllLayerAnimation(AnimStack, FbxTime::GetFrameRate(Scene->GetGlobalSettings().GetTimeMode()));
-
 	FbxAnimLayer* AnimLayer = AnimStack->GetMember<FbxAnimLayer>(0);
 	if (AnimLayer == NULL) return false;
 
@@ -562,8 +560,6 @@ float FFbxImporter::ImportMatineeActor(FbxNode* Node, UInterpGroupInst* MatineeG
 	FbxAnimStack* AnimStack = Scene->GetMember<FbxAnimStack>(0);
 	if (!AnimStack) return -1.0f;
 
-	MergeAllLayerAnimation(AnimStack, FbxTime::GetFrameRate(Scene->GetGlobalSettings().GetTimeMode()));
-
 	FbxAnimLayer* AnimLayer = AnimStack->GetMember<FbxAnimLayer>(0);
 	if (AnimLayer == NULL) return -1.0f;
 	
diff --git a/Engine/Source/Editor/UnrealEd/Private/FbxAnimUtils.cpp b/Engine/Source/Editor/UnrealEd/Private/FbxAnimUtils.cpp
--- a/Engine/Source/Editor/UnrealEd/Private/FbxAnimUtils.cpp
+++ b/Engine/Source/Editor/UnrealEd/Private/FbxAnimUtils.cpp
@@ -115,8 +115,6 @@ namespace FbxAnimUtils
 				FbxAnimStack* AnimStack = FbxImporter->Scene->GetMember<FbxAnimStack>(0);
 				if (AnimStack != nullptr)
 				{
-					FbxImporter->MergeAllLayerAnimation(AnimStack, FbxTime::GetFrameRate(FbxImporter->Scene->GetGlobalSettings().GetTimeMode()));
-
 					FbxTimeSpan AnimTimeSpan = FbxImporter->GetAnimationTimeSpan(FbxImporter->Scene->GetRootNode(), AnimStack);
 
 					// Grab the start time, as we might have a preroll
diff --git a/Engine/Source/Editor/UnrealEd/Private/SkeletalMeshEdit.cpp b/Engine/Source/Editor/UnrealEd/Private/SkeletalMeshEdit.cpp
--- a/Engine/Source/Editor/UnrealEd/Private/SkeletalMeshEdit.cpp
+++ b/Engine/Source/Editor/UnrealEd/Private/SkeletalMeshEdit.cpp
@@ -34,6 +34,10 @@
 
 #define LOCTEXT_NAMESPACE "SkeletalMeshEdit"
 
+//The max reference rate is use to cap the maximum rate we support.
+//It must be base on DEFAULT_SAMPLERATE*2ExpX where X is a integer with range [1 to 6] because we use KINDA_SMALL_NUMBER(0.0001) we do not want to pass 1920Hz 1/1920 = 0.0005
+#define MaxReferenceRate 1920.0f
+
 UAnimSequence * UEditorEngine::ImportFbxAnimation( USkeleton* Skeleton, UObject* Outer, UFbxAnimSequenceImportData* TemplateImportData, const TCHAR* InFilename, const TCHAR* AnimName, bool bImportMorphTracks )
 {
 	check(Skeleton);
@@ -518,6 +522,25 @@ FbxTimeSpan UnFbx::FFbxImporter::GetAnimationTimeSpan(FbxNode* RootNode, FbxAnim
 
 	return AnimTimeSpan;
 }
+
+void UnFbx::FFbxImporter::GetAnimationIntervalMultiLayer(FbxNode* RootNode, FbxAnimStack* AnimStack, FbxTimeSpan& AnimTimeSpan)
+{
+	int NumAnimLayers = AnimStack != nullptr ? AnimStack->GetMemberCount() : 0;
+	for (int AnimLayerIndex = 0; AnimLayerIndex < NumAnimLayers; ++AnimLayerIndex)
+	{
+		FbxTimeSpan LayerAnimTimeSpan(FBXSDK_TIME_INFINITE, FBXSDK_TIME_MINUS_INFINITE);
+		RootNode->GetAnimationInterval(LayerAnimTimeSpan, AnimStack, AnimLayerIndex);
+		if (LayerAnimTimeSpan.GetStart() < AnimTimeSpan.GetStart())
+		{
+			AnimTimeSpan.SetStart(LayerAnimTimeSpan.GetStart());
+		}
+		if (LayerAnimTimeSpan.GetStop() > AnimTimeSpan.GetStop())
+		{
+			AnimTimeSpan.SetStop(LayerAnimTimeSpan.GetStop());
+		}
+	}
+}
+
 /**
 * Add to the animation set, the animations contained within the FBX document, for the given skeleton
 */
@@ -632,8 +655,7 @@ UAnimSequence * UnFbx::FFbxImporter::ImportAnimations(USkeleton* Skeleton, UObje
 
 //Get the smallest sample rate(integer) representing the DeltaTime(time between 0.0f and 1.0f).
 //@DeltaTime: the time to find the rate between 0.0f and 1.0f
-//@MaxReferenceRate: the maximum rate we can find
-int32 GetTimeSampleRate(const float DeltaTime, const float MaxReferenceRate)
+int32 GetTimeSampleRate(const float DeltaTime)
 {
 	float OriginalSampleRateDivider = 1.0f / DeltaTime;
 	float SampleRateDivider = OriginalSampleRateDivider;
@@ -654,7 +676,7 @@ int32 GetTimeSampleRate(const float DeltaTime, const float MaxReferenceRate)
 	return FMath::Min(FPlatformMath::RoundToInt(SampleRateDivider), FPlatformMath::RoundToInt(MaxReferenceRate));
 }
 
-int32 GetAnimationCurveRate(FbxAnimCurve* CurrentCurve, float MaxReferenceRate)
+int32 GetAnimationCurveRate(FbxAnimCurve* CurrentCurve)
 {
 	if (CurrentCurve == nullptr)
 		return 0;
@@ -702,7 +724,7 @@ int32 GetAnimationCurveRate(FbxAnimCurve* CurrentCurve, float MaxReferenceRate)
 				int32 DeltaKey = FPlatformMath::RoundToInt(Delta*KeyMultiplier);
 				if (!FMath::IsNearlyZero(Delta, KINDA_SMALL_NUMBER) && !DeltaComputed.Contains(DeltaKey))
 				{
-					int32 ComputeSampleRate = GetTimeSampleRate(Delta, MaxReferenceRate);
+					int32 ComputeSampleRate = GetTimeSampleRate(Delta);
 					DeltaComputed.Add(DeltaKey);
 					//Use the least common multiplier with the new delta entry
 					int32 LeastCommonMultiplier = FMath::Min(FMath::LeastCommonMultiplier(SampleRate, ComputeSampleRate), FPlatformMath::RoundToInt(MaxReferenceRate));
@@ -717,87 +739,62 @@ int32 GetAnimationCurveRate(FbxAnimCurve* CurrentCurve, float MaxReferenceRate)
 	return 0;
 }
 
-int32 UnFbx::FFbxImporter::GetMaxSampleRate(TArray<FbxNode*>& SortedLinks, TArray<FbxNode*>& NodeArray)
+void GetNodeSampleRate(FbxNode* Node, FbxAnimLayer* AnimLayer, TArray<int32>& NodeAnimSampleRates, bool bCurve, bool bBlendCurve)
 {
-	//The max reference rate is use to cap the maximum rate we support.
-	//It must be base on DEFAULT_SAMPLERATE*2ExpX where X is a integer with range [1 to 6] because we use KINDA_SMALL_NUMBER(0.0001) we do not want to pass 1920Hz 1/1920 = 0.0005
-	float MaxReferenceRate = 1920.0f;
-	int32 MaxStackResampleRate = 0;
-	TArray<int32> CurveAnimSampleRates;
-	const FBXImportOptions* ImportOption = GetImportOptions();
-	int32 AnimStackCount = Scene->GetSrcObjectCount<FbxAnimStack>();
-	for( int32 AnimStackIndex = 0; AnimStackIndex < AnimStackCount; AnimStackIndex++)
+	if (bCurve)
 	{
-		FbxAnimStack* CurAnimStack = Scene->GetSrcObject<FbxAnimStack>(AnimStackIndex);
+		const int32 MaxElement = 9;
+		FbxAnimCurve* Curves[MaxElement];
 
-		FbxTimeSpan AnimStackTimeSpan = GetAnimationTimeSpan(SortedLinks[0], CurAnimStack);
+		Curves[0] = Node->LclTranslation.GetCurve(AnimLayer, FBXSDK_CURVENODE_COMPONENT_X, false);
+		Curves[1] = Node->LclTranslation.GetCurve(AnimLayer, FBXSDK_CURVENODE_COMPONENT_Y, false);
+		Curves[2] = Node->LclTranslation.GetCurve(AnimLayer, FBXSDK_CURVENODE_COMPONENT_Z, false);
+		Curves[3] = Node->LclRotation.GetCurve(AnimLayer, FBXSDK_CURVENODE_COMPONENT_X, false);
+		Curves[4] = Node->LclRotation.GetCurve(AnimLayer, FBXSDK_CURVENODE_COMPONENT_Y, false);
+		Curves[5] = Node->LclRotation.GetCurve(AnimLayer, FBXSDK_CURVENODE_COMPONENT_Z, false);
+		Curves[6] = Node->LclScaling.GetCurve(AnimLayer, FBXSDK_CURVENODE_COMPONENT_X, false);
+		Curves[7] = Node->LclScaling.GetCurve(AnimLayer, FBXSDK_CURVENODE_COMPONENT_Y, false);
+		Curves[8] = Node->LclScaling.GetCurve(AnimLayer, FBXSDK_CURVENODE_COMPONENT_Z, false);
 
-		double AnimStackStart = AnimStackTimeSpan.GetStart().GetSecondDouble();
-		double AnimStackStop = AnimStackTimeSpan.GetStop().GetSecondDouble();
 
-		FbxAnimLayer* AnimLayer = (FbxAnimLayer*)CurAnimStack->GetMember(0);
-		for(int32 LinkIndex = 0; LinkIndex < SortedLinks.Num(); ++LinkIndex)
+		for (int32 CurveIndex = 0; CurveIndex < MaxElement; ++CurveIndex)
 		{
-			FbxNode* CurrentLink = SortedLinks[LinkIndex];
-
-			const int32 MaxElement = 9;
-			FbxAnimCurve* Curves[MaxElement];
-
-			Curves[0] = CurrentLink->LclTranslation.GetCurve(AnimLayer, FBXSDK_CURVENODE_COMPONENT_X, false);
-			Curves[1] = CurrentLink->LclTranslation.GetCurve(AnimLayer, FBXSDK_CURVENODE_COMPONENT_Y, false);
-			Curves[2] = CurrentLink->LclTranslation.GetCurve(AnimLayer, FBXSDK_CURVENODE_COMPONENT_Z, false);
-			Curves[3] = CurrentLink->LclRotation.GetCurve(AnimLayer, FBXSDK_CURVENODE_COMPONENT_X, false);
-			Curves[4] = CurrentLink->LclRotation.GetCurve(AnimLayer, FBXSDK_CURVENODE_COMPONENT_Y, false);
-			Curves[5] = CurrentLink->LclRotation.GetCurve(AnimLayer, FBXSDK_CURVENODE_COMPONENT_Z, false);
-			Curves[6] = CurrentLink->LclScaling.GetCurve(AnimLayer, FBXSDK_CURVENODE_COMPONENT_X, false);
-			Curves[7] = CurrentLink->LclScaling.GetCurve(AnimLayer, FBXSDK_CURVENODE_COMPONENT_Y, false);
-			Curves[8] = CurrentLink->LclScaling.GetCurve(AnimLayer, FBXSDK_CURVENODE_COMPONENT_Z, false);
-
-			
-			for(int32 CurveIndex = 0; CurveIndex < MaxElement; ++CurveIndex)
+			FbxAnimCurve* CurrentCurve = Curves[CurveIndex];
+			if (CurrentCurve)
 			{
-				FbxAnimCurve* CurrentCurve = Curves[CurveIndex];
-				if(CurrentCurve)
+				int32 CurveAnimRate = GetAnimationCurveRate(CurrentCurve);
+				if (CurveAnimRate != 0)
 				{
-					int32 CurveAnimRate = GetAnimationCurveRate(CurrentCurve, MaxReferenceRate);
-					if (CurveAnimRate != 0)
-					{
-						CurveAnimSampleRates.AddUnique(CurveAnimRate);
-					}
+					NodeAnimSampleRates.AddUnique(CurveAnimRate);
 				}
 			}
 		}
+	}
 
-		// it doens't matter whether you choose to import morphtarget or not
-		// blendshape are always imported. Import morphtarget is only used for morphtarget for mesh
+	if (bBlendCurve)
+	{
+		FbxGeometry* Geometry = (FbxGeometry*)Node->GetNodeAttribute();
+		if (Geometry)
 		{
-			for (int32 NodeIndex = 0; NodeIndex < NodeArray.Num(); NodeIndex++)
+			int32 BlendShapeDeformerCount = Geometry->GetDeformerCount(FbxDeformer::eBlendShape);
+			for (int32 BlendShapeIndex = 0; BlendShapeIndex < BlendShapeDeformerCount; ++BlendShapeIndex)
 			{
-				// consider blendshape animation curve
-				FbxGeometry* Geometry = (FbxGeometry*)NodeArray[NodeIndex]->GetNodeAttribute();
-				if (Geometry)
+				FbxBlendShape* BlendShape = (FbxBlendShape*)Geometry->GetDeformer(BlendShapeIndex, FbxDeformer::eBlendShape);
+
+				int32 BlendShapeChannelCount = BlendShape->GetBlendShapeChannelCount();
+				for (int32 ChannelIndex = 0; ChannelIndex < BlendShapeChannelCount; ++ChannelIndex)
 				{
-					int32 BlendShapeDeformerCount = Geometry->GetDeformerCount(FbxDeformer::eBlendShape);
-					for (int32 BlendShapeIndex = 0; BlendShapeIndex < BlendShapeDeformerCount; ++BlendShapeIndex)
-					{
-						FbxBlendShape* BlendShape = (FbxBlendShape*)Geometry->GetDeformer(BlendShapeIndex, FbxDeformer::eBlendShape);
+					FbxBlendShapeChannel* Channel = BlendShape->GetBlendShapeChannel(ChannelIndex);
 
-						int32 BlendShapeChannelCount = BlendShape->GetBlendShapeChannelCount();
-						for (int32 ChannelIndex = 0; ChannelIndex < BlendShapeChannelCount; ++ChannelIndex)
+					if (Channel)
+					{
+						FbxAnimCurve* CurrentCurve = Geometry->GetShapeChannel(BlendShapeIndex, ChannelIndex, AnimLayer);
+						if (CurrentCurve)
 						{
-							FbxBlendShapeChannel* Channel = BlendShape->GetBlendShapeChannel(ChannelIndex);
-
-							if (Channel)
+							int32 CurveAnimRate = GetAnimationCurveRate(CurrentCurve);
+							if (CurveAnimRate != 0)
 							{
-								FbxAnimCurve* CurrentCurve = Geometry->GetShapeChannel(BlendShapeIndex, ChannelIndex, AnimLayer);
-								if (CurrentCurve)
-								{
-									int32 CurveAnimRate = GetAnimationCurveRate(CurrentCurve, MaxReferenceRate);
-									if (CurveAnimRate != 0)
-									{
-										CurveAnimSampleRates.AddUnique(CurveAnimRate);
-									}
-								}
+								NodeAnimSampleRates.AddUnique(CurveAnimRate);
 							}
 						}
 					}
@@ -805,6 +802,97 @@ int32 UnFbx::FFbxImporter::GetMaxSampleRate(TArray<FbxNode*>& SortedLinks, TArra
 			}
 		}
 	}
+}
+
+int32 UnFbx::FFbxImporter::GetGlobalAnimStackSampleRate(FbxAnimStack* CurAnimStack)
+{
+	int32 ResampleRate = DEFAULT_SAMPLERATE;
+	if (ImportOptions->bResample)
+	{
+		TArray<int32> CurveAnimSampleRates;
+		int32 MaxStackResampleRate = 0;
+		int32 AnimStackLayerCount = CurAnimStack->GetMemberCount();
+		for (int32 LayerIndex = 0; LayerIndex < AnimStackLayerCount; ++LayerIndex)
+		{
+			FbxAnimLayer* AnimLayer = (FbxAnimLayer*)CurAnimStack->GetMember(LayerIndex);
+			for (int32 NodeIndex = 0; NodeIndex < Scene->GetNodeCount(); ++NodeIndex)
+			{
+				FbxNode* Node = Scene->GetNode(NodeIndex);
+				//Get both the transform properties curve and the blend shape animation sample rate
+				GetNodeSampleRate(Node, AnimLayer, CurveAnimSampleRates, true, true);
+			}
+		}
+
+		MaxStackResampleRate = CurveAnimSampleRates.Num() > 0 ? 1 : MaxStackResampleRate;
+		//Find the lowest sample rate that will pass by all the keys from all curves
+		for (int32 CurveSampleRate : CurveAnimSampleRates)
+		{
+			if (CurveSampleRate >= MaxReferenceRate && MaxStackResampleRate < CurveSampleRate)
+			{
+				MaxStackResampleRate = CurveSampleRate;
+			}
+			else if (MaxStackResampleRate < MaxReferenceRate)
+			{
+				int32 LeastCommonMultiplier = FMath::LeastCommonMultiplier(MaxStackResampleRate, CurveSampleRate);
+				MaxStackResampleRate = LeastCommonMultiplier != 0 ? LeastCommonMultiplier : FMath::Max3(FPlatformMath::RoundToInt(DEFAULT_SAMPLERATE), MaxStackResampleRate, CurveSampleRate);
+				if (MaxStackResampleRate >= MaxReferenceRate)
+				{
+					MaxStackResampleRate = MaxReferenceRate;
+				}
+			}
+		}
+
+		// Make sure we're not hitting 0 for samplerate
+		if (MaxStackResampleRate != 0)
+		{
+			//Make sure the resample rate is positive
+			if (!ensure(MaxStackResampleRate >= 0))
+			{
+				MaxStackResampleRate *= -1;
+			}
+			ResampleRate = MaxStackResampleRate;
+		}
+	}
+	return ResampleRate;
+}
+
+int32 UnFbx::FFbxImporter::GetMaxSampleRate(TArray<FbxNode*>& SortedLinks, TArray<FbxNode*>& NodeArray)
+{
+	
+	int32 MaxStackResampleRate = 0;
+	TArray<int32> CurveAnimSampleRates;
+	const FBXImportOptions* ImportOption = GetImportOptions();
+	int32 AnimStackCount = Scene->GetSrcObjectCount<FbxAnimStack>();
+	for( int32 AnimStackIndex = 0; AnimStackIndex < AnimStackCount; AnimStackIndex++)
+	{
+		FbxAnimStack* CurAnimStack = Scene->GetSrcObject<FbxAnimStack>(AnimStackIndex);
+
+		FbxTimeSpan AnimStackTimeSpan = GetAnimationTimeSpan(SortedLinks[0], CurAnimStack);
+
+		double AnimStackStart = AnimStackTimeSpan.GetStart().GetSecondDouble();
+		double AnimStackStop = AnimStackTimeSpan.GetStop().GetSecondDouble();
+
+		int32 AnimStackLayerCount = CurAnimStack->GetMemberCount();
+		for (int32 LayerIndex = 0; LayerIndex < AnimStackLayerCount; ++LayerIndex)
+		{
+			FbxAnimLayer* AnimLayer = (FbxAnimLayer*)CurAnimStack->GetMember(LayerIndex);
+			for (int32 LinkIndex = 0; LinkIndex < SortedLinks.Num(); ++LinkIndex)
+			{
+				FbxNode* CurrentLink = SortedLinks[LinkIndex];
+				GetNodeSampleRate(CurrentLink, AnimLayer, CurveAnimSampleRates, true, false);
+			}
+
+			// it doens't matter whether you choose to import morphtarget or not
+			// blendshape are always imported. Import morphtarget is only used for morphtarget for mesh
+			{
+				for (int32 NodeIndex = 0; NodeIndex < NodeArray.Num(); NodeIndex++)
+				{
+					// consider blendshape animation curve
+					GetNodeSampleRate(NodeArray[NodeIndex], AnimLayer, CurveAnimSampleRates, false, true);
+				}
+			}
+		}
+	}
 
 	MaxStackResampleRate = CurveAnimSampleRates.Num() > 0 ? 1 : MaxStackResampleRate;
 	//Find the lowest sample rate that will pass by all the keys from all curves
@@ -846,9 +934,6 @@ bool UnFbx::FFbxImporter::ValidateAnimStack(TArray<FbxNode*>& SortedLinks, TArra
 
 	UE_LOG(LogFbx, Log, TEXT("Parsing AnimStack %s"),UTF8_TO_TCHAR(CurAnimStack->GetName()));
 
-	// There are a FBX unroll filter bug, so don't bake animation layer at all
-	MergeAllLayerAnimation(CurAnimStack, ResampleRate);
-
 	bool bValidAnimStack = true;
 
 	AnimTimeSpan = GetAnimationTimeSpan(SortedLinks[0], CurAnimStack);
diff --git a/Engine/Source/Editor/UnrealEd/Public/FbxImporter.h b/Engine/Source/Editor/UnrealEd/Public/FbxImporter.h
--- a/Engine/Source/Editor/UnrealEd/Public/FbxImporter.h
+++ b/Engine/Source/Editor/UnrealEd/Public/FbxImporter.h
@@ -843,6 +843,12 @@ public:
 	 */
 	FbxTimeSpan GetAnimationTimeSpan(FbxNode* RootNode, FbxAnimStack* AnimStack);
 
+	/**
+	* When we get exported time we call GetanimationInterval from fbx sdk and it return the layer 0 by default
+	* This function return the sum of all layer instead of just the layer 0.
+	*/
+	void GetAnimationIntervalMultiLayer(FbxNode* RootNode, FbxAnimStack* AnimStack, FbxTimeSpan& AnimTimeSpan);
+
 	/**
 	 * Import one animation from CurAnimStack
 	 *
@@ -856,6 +862,12 @@ public:
 	 * @param AnimTimeSpan	AnimTimeSpan
 	 */
 	bool ImportAnimation(USkeleton* Skeleton, UAnimSequence* DestSeq, const FString& FileName, TArray<FbxNode*>& SortedLinks, TArray<FbxNode*>& NodeArray, FbxAnimStack* CurAnimStack, const int32 ResampleRate, const FbxTimeSpan AnimTimeSpan);
+	/**
+	* Calculate the global Sample Rate for all the nodes in the FbxAnimStack pass in parameter
+	*
+	* @param FbxAnimStack	The anim stack we want to know the best sample rate
+	*/
+	int32 GetGlobalAnimStackSampleRate(FbxAnimStack* CurAnimStack);
 	/**
 	 * Calculate Max Sample Rate - separate out of the original ImportAnimations
 	 *
