Adding documentation for analytic event Edtior.Usage.FBX

diff --git a/Engine/Source/Editor/UnrealEd/Private/Fbx/FbxMainImport.cpp b/Engine/Source/Editor/UnrealEd/Private/Fbx/FbxMainImport.cpp
--- a/Engine/Source/Editor/UnrealEd/Private/Fbx/FbxMainImport.cpp
+++ b/Engine/Source/Editor/UnrealEd/Private/Fbx/FbxMainImport.cpp
@@ -1385,7 +1385,23 @@ bool FFbxImporter::ImportFromFile(const FString& Filename, const FString& Type,
 				//Convert the scene
 				ConvertScene();
 
-				// do analytics on getting Fbx data
+				// Run Analytics for FBX Import data
+				/**
+				  * @EventName Editor.Usage.FBX
+				  * @Trigger Fires when the user clicks OK in the FBX Import Dialog
+				  * @Type Editor
+				  * @EventParam OriginalVendor string Returns the name of the vendor that manufactures the original application used to generate the imported FBX
+				  * @EventParam OriginalAppName string Returns the name of the original application used to generate the imported FBX
+				  * @EventParam OriginalAppVersion string Returns the revision of the original application used to generate the imported FBX
+				  * @EventParam LastSavedVendor string Returns the name of the vendor that manufactures the last application used to modify the imported FBX
+				  * @EventParam LastSavedAppName string Returns the name of the last application used to modify the imported FBX
+				  * @EventParam LastSavedAppVersion string Returns the revision of the last application used to modify the imported FBX
+				  * @EventParam FBXFileVersion string Returns the FBX SDK used to generate the imported FBX
+				  * @EventParam FilenameHash string Returns the filename of the meshes imported from the FBX 
+				  * @EventParam ImportType string Returns the mesh data type (Static, Skeletal) being imported from the FBX
+				  * 
+				  * @Owner Alexis.Matte
+				*/
 				FbxDocumentInfo* DocInfo = Scene->GetSceneInfo();
 				if (DocInfo)
 				{
