Curve allows for changing the times of multiple keys if they're on different curves

diff --git a/Engine/Source/Editor/UnrealEd/Private/SCurveEditor.cpp b/Engine/Source/Editor/UnrealEd/Private/SCurveEditor.cpp
--- a/Engine/Source/Editor/UnrealEd/Private/SCurveEditor.cpp
+++ b/Engine/Source/Editor/UnrealEd/Private/SCurveEditor.cpp
@@ -33,6 +33,8 @@
 #include "Framework/Commands/GenericCommands.h"
 #include "Widgets/Input/SNumericEntryBox.h"
 #include "Widgets/Input/STextEntryPopup.h"
+#include "Widgets/Notifications/SNotificationList.h"
+#include "Framework/Notifications/NotificationManager.h"
 
 #define LOCTEXT_NAMESPACE "SCurveEditor"
 
@@ -733,7 +735,19 @@ EVisibility SCurveEditor::GetFrameEditVisibility() const
 
 bool SCurveEditor::GetInputEditEnabled() const
 {
-	return (SelectedKeys.Num() == 1);
+	bool bKeysOnSameCurves = false;
+	for (int32 SelectedIndex = 0; SelectedIndex < SelectedKeys.Num() - 1; SelectedIndex++)
+	{
+		for (int32 CompareIndex = SelectedIndex + 1; CompareIndex < SelectedKeys.Num(); CompareIndex++)
+		{
+			if (SelectedKeys[SelectedIndex].Curve == SelectedKeys[CompareIndex].Curve)
+			{
+				bKeysOnSameCurves = true;
+				break;
+			}
+		}
+	}
+	return (SelectedKeys.Num() == 1) || !bKeysOnSameCurves;
 }
 
 int32 SCurveEditor::OnPaint(const FPaintArgs& Args, const FGeometry& AllottedGeometry, const FSlateRect& MyCullingRect, FSlateWindowElementList& OutDrawElements, int32 LayerId, const FWidgetStyle& InWidgetStyle, bool bParentEnabled ) const
@@ -1655,12 +1669,12 @@ void SCurveEditor::TryStartDrag(const FGeometry& InMyGeometry, const FPointerEve
 				DragState = EDragState::DragKey;
 				DraggedKeyHandle = HitKey.KeyHandle;
 				PreDragKeyLocations.Empty();
-				for (auto selectedKey : SelectedKeys)
+				for (auto SelectedKey : SelectedKeys)
 				{
-					PreDragKeyLocations.Add(selectedKey.KeyHandle, FVector2D
+					PreDragKeyLocations.Add(SelectedKey.KeyHandle, FVector2D
 					(
-						selectedKey.Curve->GetKeyTime(selectedKey.KeyHandle),
-						selectedKey.Curve->GetKeyValue(selectedKey.KeyHandle)
+						SelectedKey.Curve->GetKeyTime(SelectedKey.KeyHandle),
+						SelectedKey.Curve->GetKeyValue(SelectedKey.KeyHandle)
 					));
 				}
 			}
@@ -2053,18 +2067,9 @@ void SCurveEditor::OnTimeComitted(float NewTime, ETextCommit::Type CommitType)
 	// Don't digest the number if we just clicked away from the pop-up
 	if ( !bIsUsingSlider && ((CommitType == ETextCommit::OnEnter) || ( CommitType == ETextCommit::OnUserMovedFocus )) )
 	{
-		if ( SelectedKeys.Num() >= 1 )
+		for(FSelectedCurveKey Key : SelectedKeys)
 		{
-			auto Key = SelectedKeys[0];
-			if ( IsValidCurve(Key.Curve) )
-			{
-				const FScopedTransaction Transaction(LOCTEXT("CurveEditor_NewTime", "New Time Entered"));
-				CurveOwner->ModifyOwner();
-				Key.Curve->SetKeyTime(Key.KeyHandle, NewTime);
-				TArray<FRichCurveEditInfo> ChangedCurveEditInfos;
-				ChangedCurveEditInfos.Add(GetViewModelForCurve(Key.Curve)->CurveInfo);
-				CurveOwner->OnCurveChanged(ChangedCurveEditInfos);
-			}
+			UpdateCurveTimeSingleKey(Key, NewTime);
 		}
 
 		FSlateApplication::Get().DismissAllMenus();
@@ -2075,18 +2080,9 @@ void SCurveEditor::OnTimeChanged(float NewTime)
 {
 	if ( bIsUsingSlider )
 	{
-		if ( SelectedKeys.Num() >= 1 )
+		for (FSelectedCurveKey Key : SelectedKeys)
 		{
-			auto Key = SelectedKeys[0];
-			if ( IsValidCurve(Key.Curve) )
-			{
-				const FScopedTransaction Transaction( LOCTEXT( "CurveEditor_NewTime", "New Time Entered" ) );
-				CurveOwner->ModifyOwner();
-				Key.Curve->SetKeyTime(Key.KeyHandle, NewTime);
-				TArray<FRichCurveEditInfo> ChangedCurveEditInfos;
-				ChangedCurveEditInfos.Add(GetViewModelForCurve(Key.Curve)->CurveInfo);
-				CurveOwner->OnCurveChanged(ChangedCurveEditInfos);
-			}
+			UpdateCurveTimeSingleKey(Key, NewTime);
 		}
 	}
 }
@@ -2111,18 +2107,9 @@ void SCurveEditor::OnTimeInFramesComitted(int32 NewFrame, ETextCommit::Type Comm
 	// Don't digest the number if we just clicked away from the pop-up
 	if ( !bIsUsingSlider && ((CommitType == ETextCommit::OnEnter) || ( CommitType == ETextCommit::OnUserMovedFocus )) )
 	{
-		if ( SelectedKeys.Num() >= 1 )
+		for (FSelectedCurveKey Key : SelectedKeys)
 		{
-			auto Key = SelectedKeys[0];
-			if ( IsValidCurve(Key.Curve) )
-			{
-				const FScopedTransaction Transaction(LOCTEXT("CurveEditor_NewFrame", "New Frame Entered"));
-				CurveOwner->ModifyOwner();
-				Key.Curve->SetKeyTime(Key.KeyHandle, FrameToTime(NewFrame));
-				TArray<FRichCurveEditInfo> ChangedCurveEditInfos;
-				ChangedCurveEditInfos.Add(GetViewModelForCurve(Key.Curve)->CurveInfo);
-				CurveOwner->OnCurveChanged(ChangedCurveEditInfos);
-			}
+			UpdateCurveTimeSingleKey(Key, NewFrame);
 		}
 
 		FSlateApplication::Get().DismissAllMenus();
@@ -2133,22 +2120,51 @@ void SCurveEditor::OnTimeInFramesChanged(int32 NewFrame)
 {
 	if ( bIsUsingSlider )
 	{
-		if ( SelectedKeys.Num() >= 1 )
+		for (FSelectedCurveKey Key : SelectedKeys)
 		{
-			auto Key = SelectedKeys[0];
-			if ( IsValidCurve(Key.Curve) )
-			{
-				const FScopedTransaction Transaction( LOCTEXT( "CurveEditor_NewFrame", "New Frame Entered" ) );
-				CurveOwner->ModifyOwner();
-				Key.Curve->SetKeyTime(Key.KeyHandle, FrameToTime(NewFrame));
-				TArray<FRichCurveEditInfo> ChangedCurveEditInfos;
-				ChangedCurveEditInfos.Add(GetViewModelForCurve(Key.Curve)->CurveInfo);
-				CurveOwner->OnCurveChanged(ChangedCurveEditInfos);
-			}
+			UpdateCurveTimeSingleKey(Key, NewFrame);
+		}
+	}
+}
+
+void SCurveEditor::UpdateCurveTimeSingleKey(FSelectedCurveKey Key, float NewTime, bool bSetFromFrame)
+{
+	// If the curve is valid and doesn't already have a key at that time
+	if (IsValidCurve(Key.Curve))
+	{
+		if (!Key.Curve->KeyExistsAtTime(NewTime))
+		{
+			FText TransactionText = bSetFromFrame ? LOCTEXT("CurveEditor_NewFrame", "New Frame Entered") : LOCTEXT("CurveEditor_NewTime", "New Time Entered");
+			const FScopedTransaction Transaction(TransactionText);
+			CurveOwner->ModifyOwner();
+			Key.Curve->SetKeyTime(Key.KeyHandle, NewTime);
+			TArray<FRichCurveEditInfo> ChangedCurveEditInfos;
+			ChangedCurveEditInfos.Add(GetViewModelForCurve(Key.Curve)->CurveInfo);
+			CurveOwner->OnCurveChanged(ChangedCurveEditInfos);
+		}
+		// if the existing key is not this key
+		else if (Key.Curve->FindKey(NewTime) != Key.KeyHandle)
+		{
+			LogAndToastCurveTimeWarning(Key.Curve);
 		}
 	}
 }
 
+void SCurveEditor::UpdateCurveTimeSingleKey(FSelectedCurveKey Key, int32 NewFrame)
+{
+	UpdateCurveTimeSingleKey(Key, FrameToTime(NewFrame), true);
+}
+
+void SCurveEditor::LogAndToastCurveTimeWarning(FRichCurve* Curve)
+{
+	FText Error = FText::Format(LOCTEXT("KeyTimeCollision","A key on {0} could not be moved because there was already a key at that time."), FText::FromName(GetViewModelForCurve(Curve)->CurveInfo.CurveName));
+	FNotificationInfo Info(Error);
+	Info.ExpireDuration = 5.0f;
+	FSlateNotificationManager::Get().AddNotification(Info);
+
+	UE_LOG(LogCurveEditor, Warning, TEXT("%s"), *Error.ToString());
+}
+
 TOptional<float> SCurveEditor::OnGetValue() const
 {
 	TOptional<float> Value;
diff --git a/Engine/Source/Editor/UnrealEd/Public/SCurveEditor.h b/Engine/Source/Editor/UnrealEd/Public/SCurveEditor.h
--- a/Engine/Source/Editor/UnrealEd/Public/SCurveEditor.h
+++ b/Engine/Source/Editor/UnrealEd/Public/SCurveEditor.h
@@ -312,6 +312,12 @@ private:
 			return (Curve == Other.Curve) && (KeyHandle == Other.KeyHandle);
 		}
 
+		/** Are the curve and the key different ?*/
+		bool			operator != (const FSelectedCurveKey& Other) const
+		{
+			return (Curve != Other.Curve) && (KeyHandle != Other.KeyHandle);
+		}
+
 		FRichCurve*		Curve;
 		FKeyHandle		KeyHandle;
 	};
@@ -433,6 +439,10 @@ private:
 	void OnTimeInFramesComitted(int32 NewValue, ETextCommit::Type CommitType);
 	void OnTimeInFramesChanged(int32 NewValue);
 
+	void UpdateCurveTimeSingleKey(FSelectedCurveKey Key, float NewTime, bool bSetFromFrame = false);
+	void UpdateCurveTimeSingleKey(FSelectedCurveKey Key, int32 NewFrame);
+	void LogAndToastCurveTimeWarning(FRichCurve* Curve);
+
 	TOptional<float> OnGetValue() const;
 	void OnValueComitted(float NewValue, ETextCommit::Type CommitType);
 	void OnValueChanged(float NewValue);
diff --git a/Engine/Source/Runtime/Engine/Classes/Curves/RichCurve.h b/Engine/Source/Runtime/Engine/Classes/Curves/RichCurve.h
--- a/Engine/Source/Runtime/Engine/Classes/Curves/RichCurve.h
+++ b/Engine/Source/Runtime/Engine/Classes/Curves/RichCurve.h
@@ -273,9 +273,12 @@ public:
 	/** Get the time for the Key with the specified index. */
 	float GetKeyTime(FKeyHandle KeyHandle) const;
 
-	/** Finds a key a the specified time */
+	/** Finds a key a the specified time - if it doesn't exist, adds it */
 	FKeyHandle FindKey(float KeyTime, float KeyTimeTolerance = KINDA_SMALL_NUMBER) const;
 
+	/** True if a key exists already, false otherwise */
+	bool KeyExistsAtTime(float KeyTime, float KeyTimeTolerance = KINDA_SMALL_NUMBER) const;
+
 	/** Set the value of the specified key */
 	void SetKeyValue(FKeyHandle KeyHandle, float NewValue, bool bAutoSetTangents=true);
 
@@ -348,6 +351,7 @@ public:
 
 private:
 	void RemoveRedundantKeysInternal(float Tolerance, int32 InStartKeepKey, int32 InEndKeepKey);
+	int32 FindKeyInternal(float KeyTime, float KeyTimeTolerance) const;
 
 public:
 
diff --git a/Engine/Source/Runtime/Engine/Private/Curves/RichCurve.cpp b/Engine/Source/Runtime/Engine/Private/Curves/RichCurve.cpp
--- a/Engine/Source/Runtime/Engine/Private/Curves/RichCurve.cpp
+++ b/Engine/Source/Runtime/Engine/Private/Curves/RichCurve.cpp
@@ -389,34 +389,49 @@ float FRichCurve::GetKeyTime(FKeyHandle KeyHandle) const
 
 
 FKeyHandle FRichCurve::FindKey(float KeyTime, float KeyTimeTolerance) const
+{
+	int32 KeyPosition = FindKeyInternal(KeyTime, KeyTimeTolerance);
+	if (KeyPosition >= 0)
+	{
+		return GetKeyHandle(KeyPosition);
+	}
+
+	return FKeyHandle();
+}
+
+bool FRichCurve::KeyExistsAtTime(float KeyTime, float KeyTimeTolerance) const
+{
+	return FindKeyInternal(KeyTime, KeyTimeTolerance) >= 0;
+}
+
+int32 FRichCurve::FindKeyInternal(float KeyTime, float KeyTimeTolerance) const
 {
 	int32 Start = 0;
-	int32 End = Keys.Num()-1;
+	int32 End = Keys.Num() - 1;
 
 	// Binary search since the keys are in sorted order
 	while (Start <= End)
 	{
-		int32 TestPos = Start + (End-Start) / 2;
+		int32 TestPos = Start + (End - Start) / 2;
 		float TestKeyTime = Keys[TestPos].Time;
 
 		if (FMath::IsNearlyEqual(TestKeyTime, KeyTime, KeyTimeTolerance))
 		{
-			return GetKeyHandle(TestPos);
+			return TestPos;
 		}
-		else if(TestKeyTime < KeyTime)
+		else if (TestKeyTime < KeyTime)
 		{
-			Start = TestPos+1;
+			Start = TestPos + 1;
 		}
 		else
 		{
-			End = TestPos-1;
+			End = TestPos - 1;
 		}
 	}
 
-	return FKeyHandle();
+	return -1;
 }
 
-
 void FRichCurve::SetKeyValue(FKeyHandle KeyHandle, float NewValue, bool bAutoSetTangents)
 {
 	if (!IsKeyHandleValid(KeyHandle))
