Restricted preview mesh types
Hide the whole thumbnail category on materials because materials have their own preview mesh property.

diff --git a/Engine/Source/Editor/UnrealEd/Classes/ThumbnailRendering/SceneThumbnailInfoWithPrimitive.h b/Engine/Source/Editor/UnrealEd/Classes/ThumbnailRendering/SceneThumbnailInfoWithPrimitive.h
--- a/Engine/Source/Editor/UnrealEd/Classes/ThumbnailRendering/SceneThumbnailInfoWithPrimitive.h
+++ b/Engine/Source/Editor/UnrealEd/Classes/ThumbnailRendering/SceneThumbnailInfoWithPrimitive.h
@@ -25,7 +25,7 @@ class USceneThumbnailInfoWithPrimitive : public USceneThumbnailInfo
 	TEnumAsByte<EThumbnailPrimType> PrimitiveType;
 
 	/** The custom mesh used when the primitive type is TPT_None */
-	UPROPERTY(EditAnywhere, Category=Thumbnail)
+	UPROPERTY(EditAnywhere, Category=Thumbnail, meta = (AllowedClasses = "StaticMesh,SkeletalMesh", ExactClass = "true"))
 	FSoftObjectPath PreviewMesh;
 
 	UPROPERTY()
diff --git a/Engine/Source/Runtime/Engine/Classes/Materials/MaterialInterface.h b/Engine/Source/Runtime/Engine/Classes/Materials/MaterialInterface.h
--- a/Engine/Source/Runtime/Engine/Classes/Materials/MaterialInterface.h
+++ b/Engine/Source/Runtime/Engine/Classes/Materials/MaterialInterface.h
@@ -206,7 +206,7 @@ struct FMaterialTextureInfo
 	ENGINE_API bool IsValid(bool bCheckTextureIndex = false) const; 
 };
 
-UCLASS(abstract, BlueprintType,MinimalAPI)
+UCLASS(abstract, BlueprintType, MinimalAPI, HideCategories = (Thumbnail))
 class UMaterialInterface : public UObject, public IBlendableInterface, public IInterface_AssetUserData
 {
 	GENERATED_UCLASS_BODY()
