Deprecate FRawMesh - Miscellaneous Part 2

diff --git a/Engine/Source/Editor/UnrealEd/Private/EditorExporters.cpp b/Engine/Source/Editor/UnrealEd/Private/EditorExporters.cpp
--- a/Engine/Source/Editor/UnrealEd/Private/EditorExporters.cpp
+++ b/Engine/Source/Editor/UnrealEd/Private/EditorExporters.cpp
@@ -59,7 +59,8 @@
 #include "Editor.h"
 #include "MatineeExporter.h"
 #include "FbxExporter.h"
-#include "RawMesh.h"
+#include "MeshDescription.h"
+#include "MeshDescriptionOperations.h"
 #include "MaterialUtilities.h"
 #include "InstancedFoliageActor.h"
 #include "LandscapeProxy.h"
@@ -1718,11 +1719,17 @@ bool UStaticMeshExporterOBJ::ExportText(const FExportObjectInnerContext* Context
 	UStaticMesh* StaticMesh = CastChecked<UStaticMesh>( Object );
 
 	{
+		// Currently, we only export LOD 0 of the static mesh. In the future, we could potentially export all available LODs
+		const FStaticMeshLODResources& RenderData = StaticMesh->GetLODForExport(0);
+		uint32 Count = RenderData.GetNumTriangles();
+
+		bool bHasUVLightMap = RenderData.VertexBuffers.StaticMeshVertexBuffer.GetNumTexCoords() > 1;
+
 		// Create a new filename for the lightmap coordinate OBJ file (by adding "_UV1" to the end of the filename)
 		FString Filename = UExporter::CurrentFilename.Left( UExporter::CurrentFilename.Len() - 4 ) + "_UV1." + UExporter::CurrentFilename.Right( 3 );
 
 		// Open a second archive here so we can export lightmap coordinates at the same time we export the regular mesh
-		FArchive* UV1File = IFileManager::Get().CreateFileWriter( *Filename );
+		FArchive* UV1File = bHasUVLightMap ? IFileManager::Get().CreateFileWriter( *Filename ) : nullptr;
 
 		TArray<FVector> Verts;				// The verts in the mesh
 		TArray<FVector2D> UVs;				// UV coords from channel 0
@@ -1731,26 +1738,51 @@ bool UStaticMeshExporterOBJ::ExportText(const FExportObjectInnerContext* Context
 		TArray<uint32> SmoothingMasks;		// Complete collection of the smoothing groups from all triangles
 		TArray<uint32> UniqueSmoothingMasks;	// Collection of the unique smoothing groups (used when writing out the face info into the OBJ file so we can group by smoothing group)
 
-		UV1File->Logf( TEXT("# UnrealEd OBJ exporter\r\n") );
+		if (bHasUVLightMap)
+		{
+			UV1File->Logf(TEXT("# UnrealEd OBJ exporter\r\n"));
+		}
 		Ar.Log( TEXT("# UnrealEd OBJ exporter\r\n") );
 
-		// Currently, we only export LOD 0 of the static mesh. In the future, we could potentially export all available LODs
-		const FStaticMeshLODResources& RenderData = StaticMesh->GetLODForExport(0);
-		FRawMesh RawMesh;
-		StaticMesh->SourceModels[0].LoadRawMesh(RawMesh);
+		FMeshDescription* MeshDescription = StaticMesh->GetOriginalMeshDescription(0);
+		if (MeshDescription != nullptr)
+		{
+			uint32 TriangleCount = 0;
+			for (const FPolygonID PolygonID : MeshDescription->Polygons().GetElementIDs())
+			{
+				TriangleCount += MeshDescription->GetPolygonTriangles(PolygonID).Num();
+			}
+			if (Count == TriangleCount)
+			{
+				SmoothingMasks.AddZeroed(TriangleCount);
 
-		uint32 Count = RenderData.GetNumTriangles();
+				FMeshDescriptionOperations::ConvertHardEdgesToSmoothGroup(*MeshDescription, SmoothingMasks);
+				for (uint32 SmoothValue : SmoothingMasks)
+				{
+					UniqueSmoothingMasks.AddUnique(SmoothValue);
+				}
+			}
+		}
+
+		//This can happen in case the base LOD is reduce or the recompute normals has generate a different number of cases
+		//the count will be lower and in such a case the smooth group of the original mesh description cannot be used since it is out of sync with the render data.
+		if (SmoothingMasks.Num() != Count)
+		{
+			//Create one SmoothGroup with 0 values
+			SmoothingMasks.Empty(Count);
+			SmoothingMasks.AddZeroed(Count);
+			UniqueSmoothingMasks.Empty(1);
+			UniqueSmoothingMasks.Add(0);
+		}
 
 		// Collect all the data about the mesh
 		Verts.Reserve(3 * Count);
 		UVs.Reserve(3 * Count);
 		UVLMs.Reserve(3 * Count);
 		Normals.Reserve(3 * Count);
-		SmoothingMasks.Reserve(Count);
-		UniqueSmoothingMasks.Reserve(Count);
 
 		FIndexArrayView Indices = RenderData.IndexBuffer.GetArrayView();
-
+		
 		for( uint32 tri = 0 ; tri < Count ; tri++ )
 		{
 			uint32 Index1 = Indices[(tri * 3) + 0];
@@ -1772,51 +1804,60 @@ bool UStaticMeshExporterOBJ::ExportText(const FExportObjectInnerContext* Context
 			UVs.Add( RenderData.VertexBuffers.StaticMeshVertexBuffer.GetVertexUV(Index3, 0) );
 
 			// UVs from channel 1 (lightmap coords)
-			UVLMs.Add( RenderData.VertexBuffers.StaticMeshVertexBuffer.GetVertexUV(Index1, 1) );
-			UVLMs.Add( RenderData.VertexBuffers.StaticMeshVertexBuffer.GetVertexUV(Index2, 1) );
-			UVLMs.Add( RenderData.VertexBuffers.StaticMeshVertexBuffer.GetVertexUV(Index3, 1) );
+			if (bHasUVLightMap)
+			{
+				UVLMs.Add(RenderData.VertexBuffers.StaticMeshVertexBuffer.GetVertexUV(Index1, 1));
+				UVLMs.Add(RenderData.VertexBuffers.StaticMeshVertexBuffer.GetVertexUV(Index2, 1));
+				UVLMs.Add(RenderData.VertexBuffers.StaticMeshVertexBuffer.GetVertexUV(Index3, 1));
+			}
 
 			// Normals
 			Normals.Add( RenderData.VertexBuffers.StaticMeshVertexBuffer.VertexTangentZ(Index1) );
 			Normals.Add( RenderData.VertexBuffers.StaticMeshVertexBuffer.VertexTangentZ(Index2) );
 			Normals.Add( RenderData.VertexBuffers.StaticMeshVertexBuffer.VertexTangentZ(Index3) );
-
-			// Smoothing groups
-			SmoothingMasks.Add( RawMesh.FaceSmoothingMasks[tri] );
-
-			// Unique smoothing groups
-			UniqueSmoothingMasks.AddUnique( RawMesh.FaceSmoothingMasks[tri] );
 		}
 
 		// Write out the vertex data
-
-		UV1File->Logf( TEXT("\r\n") );
+		if (bHasUVLightMap)
+		{
+			UV1File->Logf(TEXT("\r\n"));
+		}
 		Ar.Log( TEXT("\r\n") );
 		for( int32 v = 0 ; v < Verts.Num() ; ++v )
 		{
 			// Transform to Lightwave's coordinate system
-			UV1File->Logf( TEXT("v %f %f %f\r\n"), Verts[v].X, Verts[v].Z, Verts[v].Y );
+			if (bHasUVLightMap)
+			{
+				UV1File->Logf(TEXT("v %f %f %f\r\n"), Verts[v].X, Verts[v].Z, Verts[v].Y);
+			}
 			Ar.Logf( TEXT("v %f %f %f\r\n"), Verts[v].X, Verts[v].Z, Verts[v].Y );
 		}
 
 		// Write out the UV data (the lightmap file differs here in that it writes from the UVLMs array instead of UVs)
-
-		UV1File->Logf( TEXT("\r\n") );
+		if (bHasUVLightMap)
+		{
+			UV1File->Logf(TEXT("\r\n"));
+		}
 		Ar.Log( TEXT("\r\n") );
 		for( int32 uv = 0 ; uv < UVs.Num() ; ++uv )
 		{
 			// Invert the y-coordinate (Lightwave has their bitmaps upside-down from us).
-			UV1File->Logf( TEXT("vt %f %f\r\n"), UVLMs[uv].X, 1.0f - UVLMs[uv].Y );
+			if (bHasUVLightMap)
+			{
+				UV1File->Logf(TEXT("vt %f %f\r\n"), UVLMs[uv].X, 1.0f - UVLMs[uv].Y);
+			}
 			Ar.Logf( TEXT("vt %f %f\r\n"), UVs[uv].X, 1.0f - UVs[uv].Y );
 		}
 
 		// Write object header
-
-		UV1File->Logf( TEXT("\r\n") );
+		if (bHasUVLightMap)
+		{
+			UV1File->Logf(TEXT("\r\n"));
+			UV1File->Logf(TEXT("g UnrealEdObject\r\n"));
+			UV1File->Logf(TEXT("\r\n"));
+		}
 		Ar.Log( TEXT("\r\n") );
-		UV1File->Logf( TEXT("g UnrealEdObject\r\n") );
 		Ar.Log( TEXT("g UnrealEdObject\r\n") );
-		UV1File->Logf( TEXT("\r\n") );
 		Ar.Log( TEXT("\r\n") );
 
 		// Write out the face windings, sectioned by unique smoothing groups
@@ -1825,7 +1866,10 @@ bool UStaticMeshExporterOBJ::ExportText(const FExportObjectInnerContext* Context
 
 		for( int32 sm = 0 ; sm < UniqueSmoothingMasks.Num() ; ++sm )
 		{
-			UV1File->Logf( TEXT("s %i\r\n"), SmoothingGroup );
+			if (bHasUVLightMap)
+			{
+				UV1File->Logf(TEXT("s %i\r\n"), SmoothingGroup);
+			}
 			Ar.Logf( TEXT("s %i\r\n"), SmoothingGroup );
 			SmoothingGroup++;
 
@@ -1834,22 +1878,29 @@ bool UStaticMeshExporterOBJ::ExportText(const FExportObjectInnerContext* Context
 				if( SmoothingMasks[tri] == UniqueSmoothingMasks[sm]  )
 				{
 					int idx = 1 + (tri * 3);
-
-					UV1File->Logf( TEXT("f %d/%d %d/%d %d/%d\r\n"), idx, idx, idx+1, idx+1, idx+2, idx+2 );
+					if (bHasUVLightMap)
+					{
+						UV1File->Logf(TEXT("f %d/%d %d/%d %d/%d\r\n"), idx, idx, idx + 1, idx + 1, idx + 2, idx + 2);
+					}
 					Ar.Logf( TEXT("f %d/%d %d/%d %d/%d\r\n"), idx, idx, idx+1, idx+1, idx+2, idx+2 );
 				}
 			}
 		}
 
 		// Write out footer
-
-		UV1File->Logf( TEXT("\r\n") );
+		if (bHasUVLightMap)
+		{
+			UV1File->Logf(TEXT("\r\n"));
+			UV1File->Logf(TEXT("g\r\n"));
+		}
 		Ar.Log( TEXT("\r\n") );
-		UV1File->Logf( TEXT("g\r\n") );
 		Ar.Log( TEXT("g\r\n") );
 
-		// Clean up and finish
-		delete UV1File;
+		if (bHasUVLightMap)
+		{
+			// Clean up and finish
+			delete UV1File;
+		}
 	}
 
 	// ------------------------------------------------------
diff --git a/Engine/Source/Editor/UnrealEd/Private/GeomFitUtils.cpp b/Engine/Source/Editor/UnrealEd/Private/GeomFitUtils.cpp
--- a/Engine/Source/Editor/UnrealEd/Private/GeomFitUtils.cpp
+++ b/Engine/Source/Editor/UnrealEd/Private/GeomFitUtils.cpp
@@ -190,18 +190,6 @@ int32 GenerateKDopAsSimpleCollision(UStaticMesh* StaticMesh, const TArray<FVecto
 
 /* ******************************** BOX ******************************** */
 
-static void CalcBoundingBox(const FRawMesh& RawMesh, FVector& Center, FVector& Extents, FVector& LimitVec)
-{
-	FBox Box(ForceInit);
-
-	for (uint32 i = 0; i < (uint32)RawMesh.VertexPositions.Num(); i++)
-	{
-		Box += RawMesh.VertexPositions[i] * LimitVec;
-	}
-
-	Box.GetCenterAndExtents(Center, Extents);
-}
-
 static void CalcBoundingBox(const FMeshDescription& MeshDescription, FVector& Center, FVector& Extents, FVector& LimitVec)
 {
 	TVertexAttributesConstRef<FVector> VertexPositions = MeshDescription.VertexAttributes().GetAttributesRef<FVector>(MeshAttribute::Vertex::Position);
@@ -270,127 +258,6 @@ int32 GenerateBoxAsSimpleCollision(UStaticMesh* StaticMesh)
 
 // This algorithm taken from Ritter, 1990
 // This one seems to do well with asymmetric input.
-static void CalcBoundingSphere(const FRawMesh& RawMesh, FSphere& sphere, FVector& LimitVec)
-{
-	if(RawMesh.VertexPositions.Num() == 0)
-		return;
-
-	FBox Box;
-
-	// First, find AABB, remembering furthest points in each dir.
-	Box.Min = RawMesh.VertexPositions[0] * LimitVec;
-	Box.Max = Box.Min;
-
-	FIntVector minIx, maxIx; // Extreme points.
-
-	minIx = FIntVector::ZeroValue;
-	maxIx = FIntVector::ZeroValue;
-
-	for(uint32 i=1; i<(uint32)RawMesh.VertexPositions.Num(); i++) 
-	{
-		FVector p = RawMesh.VertexPositions[i] * LimitVec;
-
-		// X //
-		if(p.X < Box.Min.X)
-		{
-			Box.Min.X = p.X;
-			minIx.X = i;
-		}
-		else if(p.X > Box.Max.X)
-		{
-			Box.Max.X = p.X;
-			maxIx.X = i;
-		}
-
-		// Y //
-		if(p.Y < Box.Min.Y)
-		{
-			Box.Min.Y = p.Y;
-			minIx.Y = i;
-		}
-		else if(p.Y > Box.Max.Y)
-		{
-			Box.Max.Y = p.Y;
-			maxIx.Y = i;
-		}
-
-		// Z //
-		if(p.Z < Box.Min.Z)
-		{
-			Box.Min.Z = p.Z;
-			minIx.Z = i;
-		}
-		else if(p.Z > Box.Max.Z)
-		{
-			Box.Max.Z = p.Z;
-			maxIx.Z = i;
-		}
-	}
-
-	const FVector Extremes[3]={ (RawMesh.VertexPositions[maxIx.X] - RawMesh.VertexPositions[minIx.X]) * LimitVec,
-								(RawMesh.VertexPositions[maxIx.Y] - RawMesh.VertexPositions[minIx.Y]) * LimitVec,
-								(RawMesh.VertexPositions[maxIx.Z] - RawMesh.VertexPositions[minIx.Z]) * LimitVec };
-
-	// Now find extreme points furthest apart, and initial center and radius of sphere.
-	float d2 = 0.f;
-	for(int32 i=0; i<3; i++)
-	{
-		const float tmpd2 = Extremes[i].SizeSquared();
-		if(tmpd2 > d2)
-		{
-			d2 = tmpd2;
-			sphere.Center = (RawMesh.VertexPositions[minIx(i)] + (0.5f * Extremes[i])) * LimitVec;
-			sphere.W = 0.f;
-		}
-	}
-
-	const FVector Extents = FVector(Extremes[0].X, Extremes[1].Y, Extremes[2].Z);
-
-	// radius and radius squared
-	float r = 0.5f * Extents.GetMax();
-	float r2 = FMath::Square(r);
-
-	// Now check each point lies within this sphere. If not - expand it a bit.
-	for(uint32 i=0; i<(uint32)RawMesh.VertexPositions.Num(); i++) 
-	{
-		const FVector cToP = (RawMesh.VertexPositions[i] * LimitVec) - sphere.Center;
-
-		const float pr2 = cToP.SizeSquared();
-
-		// If this point is outside our current bounding sphere's radius
-		if(pr2 > r2)
-		{
-			// ..expand radius just enough to include this point.
-			const float pr = FMath::Sqrt(pr2);
-			r = 0.5f * (r + pr);
-			r2 = FMath::Square(r);
-
-			sphere.Center += ((pr-r)/pr * cToP);
-		}
-	}
-
-	sphere.W = r;
-}
-
-// This is the one thats already used by unreal.
-// Seems to do better with more symmetric input...
-static void CalcBoundingSphere2(const FRawMesh& RawMesh, FSphere& sphere, FVector& LimitVec)
-{
-	FVector Center, Extents;
-	CalcBoundingBox(RawMesh, Center, Extents, LimitVec);
-
-	sphere.Center = Center;
-	sphere.W = 0.0f;
-
-	for( uint32 i=0; i<(uint32)RawMesh.VertexPositions.Num(); i++ )
-	{
-		float Dist = FVector::DistSquared(RawMesh.VertexPositions[i] * LimitVec, sphere.Center);
-		if( Dist > sphere.W )
-			sphere.W = Dist;
-	}
-	sphere.W = FMath::Sqrt(sphere.W);
-}
-
 static void CalcBoundingSphere(const FMeshDescription& MeshDescription, FSphere& sphere, FVector& LimitVec)
 {
 	if (MeshDescription.Vertices().Num() == 0)
@@ -664,90 +531,6 @@ static void CalcBoundingSphyl(const FMeshDescription& MeshDescription, FSphere&
 	length = hl * 2.0f;
 }
 
-static void CalcBoundingSphyl(const FRawMesh& RawMesh, FSphere& sphere, float& length, FRotator& rotation, FVector& LimitVec)
-{
-	if (RawMesh.VertexPositions.Num() == 0)
-		return;
-
-	FVector Center, Extents;
-	CalcBoundingBox(RawMesh, Center, Extents, LimitVec);
-
-	// @todo sphere.Center could perhaps be adjusted to best fit if model is non-symmetric on it's longest axis
-	sphere.Center = Center;
-
-	// Work out best axis aligned orientation (longest side)
-	float Extent = Extents.GetMax();
-	if (Extent == Extents.X)
-	{
-		rotation = FRotator(90.f, 0.f, 0.f);
-		Extents.X = 0.0f;
-	}
-	else if (Extent == Extents.Y)
-	{
-		rotation = FRotator(0.f, 0.f, 90.f);
-		Extents.Y = 0.0f;
-	}
-	else
-	{
-		rotation = FRotator(0.f, 0.f, 0.f);
-		Extents.Z = 0.0f;
-	}
-
-	// Cleared the largest axis above, remaining determines the radius
-	float r = Extents.GetMax();
-	float r2 = FMath::Square(r);
-	
-	// Now check each point lies within this the radius. If not - expand it a bit.
-	for (uint32 i = 0; i<(uint32)RawMesh.VertexPositions.Num(); i++)
-	{
-		FVector cToP = (RawMesh.VertexPositions[i] * LimitVec) - sphere.Center;
-		cToP = rotation.UnrotateVector(cToP);
-
-		const float pr2 = cToP.SizeSquared2D();	// Ignore Z here...
-
-		// If this point is outside our current bounding sphere's radius
-		if (pr2 > r2)
-		{
-			// ..expand radius just enough to include this point.
-			const float pr = FMath::Sqrt(pr2);
-			r = 0.5f * (r + pr);
-			r2 = FMath::Square(r);
-		}
-	}
-	
-	// The length is the longest side minus the radius.
-	float hl = FMath::Max(0.0f, Extent - r);
-
-	// Now check each point lies within the length. If not - expand it a bit.
-	for (uint32 i = 0; i<(uint32)RawMesh.VertexPositions.Num(); i++)
-	{
-		FVector cToP = (RawMesh.VertexPositions[i] * LimitVec) - sphere.Center;
-		cToP = rotation.UnrotateVector(cToP);
-
-		// If this point is outside our current bounding sphyl's length
-		if (FMath::Abs(cToP.Z) > hl)
-		{
-			const bool bFlip = (cToP.Z < 0.f ? true : false);
-			const FVector cOrigin(0.f, 0.f, (bFlip ? -hl : hl));
-
-			const float pr2 = (cOrigin - cToP).SizeSquared();
-
-			// If this point is outside our current bounding sphyl's radius
-			if (pr2 > r2)
-			{
-				FVector cPoint;
-				FMath::SphereDistToLine(cOrigin, r, cToP, (bFlip ? FVector(0.f, 0.f, 1.f) : FVector(0.f, 0.f, -1.f)), cPoint);
-
-				// Don't accept zero as a valid diff when we know it's outside the sphere (saves needless retest on further iterations of like points)
-				hl += FMath::Max(FMath::Abs(cToP.Z - cPoint.Z), 1.e-6f);
-			}
-		}
-	}
-
-	sphere.W = r;
-	length = hl * 2.0f;
-}
-
 // // //
 
 int32 GenerateSphylAsSimpleCollision(UStaticMesh* StaticMesh)
diff --git a/Engine/Source/Editor/UnrealEd/Private/StaticMeshEdit.cpp b/Engine/Source/Editor/UnrealEd/Private/StaticMeshEdit.cpp
--- a/Engine/Source/Editor/UnrealEd/Private/StaticMeshEdit.cpp
+++ b/Engine/Source/Editor/UnrealEd/Private/StaticMeshEdit.cpp
@@ -18,7 +18,6 @@
 #include "Editor.h"
 #include "StaticMeshResources.h"
 #include "BSPOps.h"
-#include "RawMesh.h"
 #include "PhysicsEngine/ConvexElem.h"
 #include "PhysicsEngine/BoxElem.h"
 #include "PhysicsEngine/SphereElem.h"
@@ -983,7 +982,6 @@ struct ExistingLODMeshData
 {
 	FMeshBuildSettings				ExistingBuildSettings;
 	FMeshReductionSettings			ExistingReductionSettings;
-	FRawMesh						ExistingRawMesh;
 	TUniquePtr<FMeshDescription>	ExistingMeshDescription;
 	TArray<FStaticMaterial>			ExistingMaterials;
 	FPerPlatformFloat				ExistingScreenSize;
@@ -1146,8 +1144,6 @@ ExistingStaticMeshData* SaveExistingStaticMeshData(UStaticMesh* ExistingMesh, Un
 			{
 				ExistingMeshDataPtr->ExistingLODData[i].ExistingMeshDescription = MakeUnique<FMeshDescription>(*MeshDescription);
 			}
-
-			ExistingMesh->SourceModels[i].RawMeshBulkData->LoadRawMesh(ExistingMeshDataPtr->ExistingLODData[i].ExistingRawMesh);
 		}
 
 		ExistingMeshDataPtr->ExistingSockets = ExistingMesh->Sockets;
@@ -1232,7 +1228,8 @@ void RestoreExistingMeshSettings(ExistingStaticMeshData* ExistingMesh, UStaticMe
 			{
 				NewMesh->AddSourceModel();
 			}
-			bool bSwapFromGeneratedToImported = !ExistingMesh->ExistingLODData[i].ExistingRawMesh.IsValid() && !NewMesh->SourceModels[i].RawMeshBulkData->IsEmpty();
+			FMeshDescription* LODMeshDescription = NewMesh->GetOriginalMeshDescription(i);
+			bool bSwapFromGeneratedToImported = !ExistingMesh->ExistingLODData[i].ExistingMeshDescription.IsValid() && (LODMeshDescription && LODMeshDescription->Polygons().Num() > 0);
 			bool bWasReduced = ExistingMesh->ExistingLODData[i].ExistingReductionSettings.PercentTriangles < 1.0f || ExistingMesh->ExistingLODData[i].ExistingReductionSettings.MaxDeviation > 0.0f;
 			if (!bSwapFromGeneratedToImported && bWasReduced)
 			{
@@ -1248,7 +1245,8 @@ void RestoreExistingMeshSettings(ExistingStaticMeshData* ExistingMesh, UStaticMe
 		//Just set the old configuration for the desired LODIndex
 		if(LODIndex >= 0 && LODIndex < CurrentNumLods && LODIndex < ExistingNumLods)
 		{
-			bool bSwapFromGeneratedToImported = !ExistingMesh->ExistingLODData[LODIndex].ExistingRawMesh.IsValid() && !NewMesh->SourceModels[LODIndex].RawMeshBulkData->IsEmpty();
+			FMeshDescription* LODMeshDescription = NewMesh->GetOriginalMeshDescription(LODIndex);
+			bool bSwapFromGeneratedToImported = !ExistingMesh->ExistingLODData[LODIndex].ExistingMeshDescription.IsValid() && (LODMeshDescription && LODMeshDescription->Polygons().Num() > 0);
 			bool bWasReduced = ExistingMesh->ExistingLODData[LODIndex].ExistingReductionSettings.PercentTriangles < 1.0f || ExistingMesh->ExistingLODData[LODIndex].ExistingReductionSettings.MaxDeviation > 0.0f;
 			if (!bSwapFromGeneratedToImported && bWasReduced)
 			{
@@ -1447,8 +1445,9 @@ void RestoreExistingMeshData(ExistingStaticMeshData* ExistingMeshDataPtr, UStati
 	for(int32 i=0; i<NumCommonLODs; i++)
 	{
 		NewMesh->SourceModels[i].BuildSettings = ExistingMeshDataPtr->ExistingLODData[i].ExistingBuildSettings;
+		FMeshDescription* LODMeshDescription = NewMesh->GetOriginalMeshDescription(i);
 		//Restore the reduction settings only if the existing data was a using reduction. Because we can set some value if we reimport from existing rawmesh to auto generated.
-		bool bSwapFromGeneratedToImported = !ExistingMeshDataPtr->ExistingLODData[i].ExistingRawMesh.IsValid() && !NewMesh->SourceModels[i].RawMeshBulkData->IsEmpty();
+		bool bSwapFromGeneratedToImported = !ExistingMeshDataPtr->ExistingLODData[i].ExistingMeshDescription.IsValid() && (LODMeshDescription && LODMeshDescription->Polygons().Num() > 0);
 		bool bWasReduced = ExistingMeshDataPtr->ExistingLODData[i].ExistingReductionSettings.PercentTriangles < 1.0f || ExistingMeshDataPtr->ExistingLODData[i].ExistingReductionSettings.MaxDeviation > 0.0f;
 		if ( !bSwapFromGeneratedToImported && bWasReduced)
 		{
@@ -1466,10 +1465,7 @@ void RestoreExistingMeshData(ExistingStaticMeshData* ExistingMeshDataPtr, UStati
 			FMeshDescription* MeshDescription = NewMesh->CreateOriginalMeshDescription(i);
 			*MeshDescription = MoveTemp(*ExistingMeshDataPtr->ExistingLODData[i].ExistingMeshDescription);
 			ExistingMeshDataPtr->ExistingLODData[i].ExistingMeshDescription.Reset();
-		}
-		if (ExistingMeshDataPtr->ExistingLODData[i].ExistingRawMesh.IsValidOrFixable())
-		{
-			SrcModel.RawMeshBulkData->SaveRawMesh(ExistingMeshDataPtr->ExistingLODData[i].ExistingRawMesh);
+			NewMesh->CommitOriginalMeshDescription(i);
 		}
 		SrcModel.BuildSettings = ExistingMeshDataPtr->ExistingLODData[i].ExistingBuildSettings;
 		SrcModel.ReductionSettings = ExistingMeshDataPtr->ExistingLODData[i].ExistingReductionSettings;
diff --git a/Engine/Source/Editor/UnrealEd/Public/FbxImporter.h b/Engine/Source/Editor/UnrealEd/Public/FbxImporter.h
--- a/Engine/Source/Editor/UnrealEd/Public/FbxImporter.h
+++ b/Engine/Source/Editor/UnrealEd/Public/FbxImporter.h
@@ -36,7 +36,6 @@ class USkeleton;
 class UStaticMesh;
 class UTexture;
 struct FExpressionInput;
-struct FRawMesh;
 struct FRichCurve;
 struct FStaticMaterial;
 struct FSkeletalMaterial;
