Fix a crash when importing a skeletal mesh with LODs

diff --git a/Engine/Source/Editor/UnrealEd/Private/Fbx/FbxMainImport.cpp b/Engine/Source/Editor/UnrealEd/Private/Fbx/FbxMainImport.cpp
--- a/Engine/Source/Editor/UnrealEd/Private/Fbx/FbxMainImport.cpp
+++ b/Engine/Source/Editor/UnrealEd/Private/Fbx/FbxMainImport.cpp
@@ -34,6 +34,7 @@
 #include "Editor/EditorPerProjectUserSettings.h"
 #include "Engine/SkeletalMesh.h"
 #include "Engine/StaticMesh.h"
+#include "IMeshReductionInterfaces.h"
 
 DEFINE_LOG_CATEGORY(LogFbx);
 
@@ -1822,6 +1823,21 @@ void FFbxImporter::ValidateAllMeshesAreReferenceByNodeAttribute()
 
 void FFbxImporter::ConvertLodPrefixToLodGroup()
 {
+	IMeshReductionModule& ReductionModule = FModuleManager::Get().LoadModuleChecked<IMeshReductionModule>("MeshReductionInterface");
+	IMeshReduction* SkeletalMeshReduction = ReductionModule.GetSkeletalMeshReductionInterface();
+	IMeshReduction* StaticMeshReduction = ReductionModule.GetStaticMeshReductionInterface();
+	bool bCanReduce = true;
+	bool bWarnUserNoReduction = false;
+	if (ImportOptions->ImportType == FBXIT_SkeletalMesh && !SkeletalMeshReduction)
+	{
+		bCanReduce = false;
+	}
+
+	if (ImportOptions->ImportType == FBXIT_StaticMesh && !StaticMeshReduction)
+	{
+		bCanReduce = false;
+	}
+
 	const FString LodPrefix = TEXT("LOD");
 	TMap<FString, TArray<uint64>> LodPrefixNodeMap;
 	TMap<uint64, FbxNode*> NodeMap;
@@ -1920,12 +1936,19 @@ void FFbxImporter::ConvertLodPrefixToLodGroup()
 		{
 			if (LodGroupNodes[CurrentLodIndex] == MAX_uint64)
 			{
-				FString FbxGeneratedNodeName = UTF8_TO_TCHAR(FirstNode->GetName());
-				FbxGeneratedNodeName = FbxGeneratedNodeName.RightChop(5);
-				FbxGeneratedNodeName += TEXT(GeneratedLODNameSuffix) + FString::FromInt(CurrentLodIndex);
-				//Generated LOD add dummy FbxNode to tell the import to add such a LOD
-				FbxNode* DummyGeneratedLODActorNode = FbxNode::Create(Scene, TCHAR_TO_UTF8(*FbxGeneratedNodeName));
-				ActorNode->AddChild(DummyGeneratedLODActorNode);
+				if (bCanReduce)
+				{
+					FString FbxGeneratedNodeName = UTF8_TO_TCHAR(FirstNode->GetName());
+					FbxGeneratedNodeName = FbxGeneratedNodeName.RightChop(5);
+					FbxGeneratedNodeName += TEXT(GeneratedLODNameSuffix) + FString::FromInt(CurrentLodIndex);
+					//Generated LOD add dummy FbxNode to tell the import to add such a LOD
+					FbxNode* DummyGeneratedLODActorNode = FbxNode::Create(Scene, TCHAR_TO_UTF8(*FbxGeneratedNodeName));
+					ActorNode->AddChild(DummyGeneratedLODActorNode);
+				}
+				else
+				{
+					bWarnUserNoReduction = true;
+				}
 				continue;
 			}
 			FbxNode* CurrentNode = NodeMap[LodGroupNodes[CurrentLodIndex]];
@@ -1941,6 +1964,22 @@ void FFbxImporter::ConvertLodPrefixToLodGroup()
 		check(ParentNode != nullptr);
 		ParentNode->AddChild(ActorNode);
 	}
+
+	if (bWarnUserNoReduction)
+	{
+		FText WarningMessage;
+		if (ImportOptions->ImportType == FBXIT_SkeletalMesh && !SkeletalMeshReduction)
+		{
+			WarningMessage = FText(LOCTEXT("FBX_ImportSkeletalMeshNoReductionModule", "No skeletal mesh reduction module available. Cannot add generated LOD between fbx node LOD prefix."));
+		}
+
+		if (ImportOptions->ImportType == FBXIT_StaticMesh && !StaticMeshReduction)
+		{
+			WarningMessage = FText(LOCTEXT("FBX_ImportStaticMeshNoReductionModule", "No static mesh reduction module available. Cannot add generated LOD between fbx node LOD prefix."));
+		}
+
+		AddTokenizedErrorMessage( FTokenizedMessage::Create( EMessageSeverity::Warning, WarningMessage ), FFbxErrors::Generic_Mesh_NoReductionModuleAvailable );
+	}
 }
 
 FbxNode *FFbxImporter::RecursiveGetFirstMeshNode(FbxNode* Node, FbxNode* NodeToFind)
diff --git a/Engine/Source/Runtime/Core/Private/Misc/FbxErrors.cpp b/Engine/Source/Runtime/Core/Private/Misc/FbxErrors.cpp
--- a/Engine/Source/Runtime/Core/Private/Misc/FbxErrors.cpp
+++ b/Engine/Source/Runtime/Core/Private/Misc/FbxErrors.cpp
@@ -53,6 +53,8 @@ FName FFbxErrors::Generic_Mesh_TooManyLODs(TEXT("Generic_Mesh_TooManyLODs"));
 
 FName FFbxErrors::Generic_Mesh_TangentsComputeError(TEXT("Generic_Mesh_TangentsComputeError"));
 
+FName FFbxErrors::Generic_Mesh_NoReductionModuleAvailable(TEXT("Generic_Mesh_NoReductionModuleAvailable"));
+
 /** Static Mesh **/
 FName FFbxErrors::StaticMesh_TooManyMaterials(TEXT("StaticMesh_TooManyMaterials"));
 
diff --git a/Engine/Source/Runtime/Core/Public/Misc/FbxErrors.h b/Engine/Source/Runtime/Core/Public/Misc/FbxErrors.h
--- a/Engine/Source/Runtime/Core/Public/Misc/FbxErrors.h
+++ b/Engine/Source/Runtime/Core/Public/Misc/FbxErrors.h
@@ -65,6 +65,8 @@ struct CORE_API FFbxErrors
 
 	static FName Generic_Mesh_TangentsComputeError;
 
+	static FName Generic_Mesh_NoReductionModuleAvailable;
+
 	/** Static Mesh **/
 	static FName StaticMesh_TooManyMaterials;
 
