Add fbx import options to analytic data

diff --git a/Engine/Source/Editor/UnrealEd/Private/Fbx/FbxMainImport.cpp b/Engine/Source/Editor/UnrealEd/Private/Fbx/FbxMainImport.cpp
--- a/Engine/Source/Editor/UnrealEd/Private/Fbx/FbxMainImport.cpp
+++ b/Engine/Source/Editor/UnrealEd/Private/Fbx/FbxMainImport.cpp
@@ -1399,7 +1399,63 @@ bool FFbxImporter::ImportFromFile(const FString& Filename, const FString& Type,
 				  * @EventParam LastSavedAppVersion string Returns the revision of the last application used to modify the imported FBX
 				  * @EventParam FBXFileVersion string Returns the FBX SDK used to generate the imported FBX
 				  * @EventParam FilenameHash string Returns the filename of the meshes imported from the FBX 
-				  * @EventParam ImportType string Returns the mesh data type (Static, Skeletal) being imported from the FBX
+				  * @EventParam ImportType string Returns the mesh data type (Static, Skeletal, Animation) being imported from the FBX
+				  * @EventParam ConvertScene boolean Returns whether the import fbx should be converted to unreal axis system
+				  * @EventParam ConvertSceneUnit boolean Returns whether the import fbx should converted the unit to unreal unit (cm)
+				  * @EventParam ForceFrontXAxis boolean Returns whether the import fbx should be converted to unreal axis system with front axis being X
+				  * @EventParam ImportMaterials boolean Returns whether the importer should create the missing materials
+				  * @EventParam ImportTextures boolean Returns whether the importer should create the missing textures
+				  * @EventParam InvertNormalMap boolean Returns whether the importer should inverse the incoming normal map textures
+				  * @EventParam RemoveNameSpace boolean Returns whether the importer should remove namespace on all nodes name
+				  * @EventParam UsedAsFullName boolean Returns whether the importer should use the filename to name the imported mesh
+				  * @EventParam ImportTranslation string Returns the translation vector apply on the import data
+				  * @EventParam ImportRotation string Returns the FRotator vector apply on the import data
+				  * @EventParam ImportUniformScale float Returns the uniform scale apply on the import data
+				  * @EventParam MaterialBasePath string Returns the path pointing on the base material use to import material instance
+				  * @EventParam MaterialSearchLocation string Returns the scope of the search for existing materials (if material not found it can create one depending on bImportMaterials value)
+				  * @EventParam AutoGenerateCollision boolean Returns whether the importer should create collision primitive
+				  * @EventParam CombineToSingle boolean Returns whether the importer should combine all mesh part together or import many meshes
+				  * @EventParam BakePivotInVertex boolean Returns whether the importer should bake the fbx mesh pivot into the vertex position
+				  * @EventParam TransformVertexToAbsolute boolean Returns whether the importer should bake the global fbx node transform into the vertex position
+				  * @EventParam ImportRigidMesh boolean Returns whether the importer should try to create a rigid mesh (static mesh import as skeletal mesh)
+				  * @EventParam NormalImportMethod string Return if the tangents or normal should be imported or compute
+				  * @EventParam NormalGenerationMethod string Return tangents generation method
+				  * @EventParam CreatePhysicsAsset boolean Returns whether the importer should create the physic asset
+				  * @EventParam ImportAnimations boolean Returns whether the importer should import also the animation
+				  * @EventParam ImportAsSkeletalGeometry boolean Returns whether the importer should import only the geometry
+				  * @EventParam ImportAsSkeletalSkinning boolean Returns whether the importer should import only the skinning
+				  * @EventParam ImportMeshesInBoneHierarchy boolean Returns whether the importer should import also the mesh found in the bone hierarchy
+				  * @EventParam ImportMorph boolean Returns whether the importer should import the morph targets
+				  * @EventParam ImportSkeletalMeshLODs boolean Returns whether the importer should import the LODs
+				  * @EventParam PreserveSmoothingGroups boolean Returns whether the importer should import the smoothing groups
+				  * @EventParam UpdateSkeletonReferencePose boolean Returns whether the importer should update the skeleton reference pose
+				  * @EventParam UseT0AsRefPose boolean Returns whether the importer should use the the animation 0 time has the reference pose
+				  * @EventParam ThresholdPosition float Returns the threshold delta to weld vertices
+				  * @EventParam ThresholdTangentNormal float Returns the threshold delta to weld tangents and normals
+				  * @EventParam ThresholdUV float Returns the threshold delta to weld UVs
+				  * @EventParam AutoComputeLodDistances boolean Returns whether the importer should set the auto compute LOD distance
+				  * @EventParam LodNumber integer Returns the LOD number we should have after the import
+				  * @EventParam BuildAdjacencyBuffer boolean Returns whether the importer should fill the adjacency buffer when building the static mesh
+				  * @EventParam BuildReversedIndexBuffer boolean Returns whether the importer should fill the reverse index buffer when building the static mesh
+				  * @EventParam GenerateLightmapUVs boolean Returns whether the importer should generate light map UVs
+				  * @EventParam ImportStaticMeshLODs boolean Returns whether the importer should import the LODs
+				  * @EventParam RemoveDegenerates boolean Returns whether the importer should remove the degenerated triangles when building the static mesh
+				  * @EventParam MinimumLodNumber integer Returns the minimum LOD use by the rendering
+				  * @EventParam StaticMeshLODGroup string Returns the LOD Group settings we use to build this imported static mesh
+				  * @EventParam VertexColorImportOption string Returns how the importer should import the vertex color
+				  * @EventParam VertexOverrideColor string Returns the color use if we need to override the vertex color
+				  * @EventParam AnimationLengthImportType string Returns how we choose the animation time span
+				  * @EventParam DeleteExistingMorphTargetCurves boolean Returns whether the importer should delete the existing morph target curves
+				  * @EventParam AnimationRange string Returns the range of animation the importer should sample if the time span is custom
+				  * @EventParam DoNotImportCurveWithZero boolean Returns whether the importer should import curves containing only zero value
+				  * @EventParam ImportBoneTracks boolean Returns whether the importer should import the bone tracks
+				  * @EventParam ImportCustomAttribute boolean Returns whether the importer should import the custom attribute curves
+				  * @EventParam PreserveLocalTransform boolean Returns whether the importer should preserve the local transform when importing the animation
+				  * @EventParam RemoveRedundantKeys boolean Returns whether the importer should remove all redundant key in an animation
+				  * @EventParam Resample boolean Returns whether the importer should re-sample the animation
+				  * @EventParam SetMaterialDriveParameterOnCustomAttribute boolean Returns whether the importer should hook all custom attribute curve to unreal material attribute
+				  * @EventParam SetMaterialDriveParameterOnCustomAttribute boolean Returns whether the importer should hook some custom attribute (having the suffix) curve to unreal material attribute
+				  * @EventParam ResampleRate float Returns the rate the exporter is suppose to re-sample any imported animations
 				  * 
 				  * @Owner Alexis.Matte
 				*/
@@ -1409,6 +1465,12 @@ bool FFbxImporter::ImportFromFile(const FString& Filename, const FString& Type,
 					if( FEngineAnalytics::IsAvailable() )
 					{
 						const static UEnum* FBXImportTypeEnum = FindObject<UEnum>(ANY_PACKAGE, TEXT("EFBXImportType"));
+						const static UEnum* FBXAnimationLengthImportTypeEnum = FindObject<UEnum>(ANY_PACKAGE, TEXT("EFBXAnimationLengthImportType"));
+						const static UEnum* MaterialSearchLocationEnum = FindObject<UEnum>(ANY_PACKAGE, TEXT("EMaterialSearchLocation"));
+						const static UEnum* FBXNormalGenerationMethodEnum = FindObject<UEnum>(ANY_PACKAGE, TEXT("EFBXNormalGenerationMethod"));
+						const static UEnum* FBXNormalImportMethodEnum = FindObject<UEnum>(ANY_PACKAGE, TEXT("EFBXNormalImportMethod"));
+						const static UEnum* VertexColorImportOptionEnum = FindObject<UEnum>(ANY_PACKAGE, TEXT("EVertexColorImportOption"));
+						
 						TArray<FAnalyticsEventAttribute> Attribs;
 
 						FString OriginalVendor(UTF8_TO_TCHAR(DocInfo->Original_ApplicationVendor.Get().Buffer()));
@@ -1432,7 +1494,103 @@ bool FFbxImporter::ImportFromFile(const FString& Filename, const FString& Type,
 						Attribs.Add(FAnalyticsEventAttribute(TEXT("FBX Version"), FbxFileVersion));
 						Attribs.Add(FAnalyticsEventAttribute(TEXT("Filename Hash"), FilenameHash));
 
-						Attribs.Add(FAnalyticsEventAttribute(TEXT("Import Type"), FBXImportTypeEnum->GetNameStringByValue(ImportOptions->ImportType)));
+						//////////////////////////////////////////////////////////////////////////
+						//FBX import options
+						Attribs.Add(FAnalyticsEventAttribute(TEXT("GenOpt ImportType"), FBXImportTypeEnum->GetNameStringByValue(ImportOptions->ImportType)));
+						Attribs.Add(FAnalyticsEventAttribute(TEXT("GenOpt ConvertScene"), ImportOptions->bConvertScene));
+						Attribs.Add(FAnalyticsEventAttribute(TEXT("GenOpt ConvertSceneUnit"), ImportOptions->bConvertSceneUnit));
+						Attribs.Add(FAnalyticsEventAttribute(TEXT("GenOpt ForceFrontXAxis"), ImportOptions->bForceFrontXAxis));
+						Attribs.Add(FAnalyticsEventAttribute(TEXT("GenOpt ImportMaterials"), ImportOptions->bImportMaterials));
+						Attribs.Add(FAnalyticsEventAttribute(TEXT("GenOpt ImportTextures"), ImportOptions->bImportTextures));
+						Attribs.Add(FAnalyticsEventAttribute(TEXT("GenOpt InvertNormalMap"), ImportOptions->bInvertNormalMap));
+						Attribs.Add(FAnalyticsEventAttribute(TEXT("GenOpt RemoveNameSpace"), ImportOptions->bRemoveNameSpace));
+						Attribs.Add(FAnalyticsEventAttribute(TEXT("GenOpt UsedAsFullName"), ImportOptions->bUsedAsFullName));
+						Attribs.Add(FAnalyticsEventAttribute(TEXT("GenOpt ImportTranslation"), ImportOptions->ImportTranslation.ToString()));
+						Attribs.Add(FAnalyticsEventAttribute(TEXT("GenOpt ImportRotation"), ImportOptions->ImportRotation.ToString()));
+						Attribs.Add(FAnalyticsEventAttribute(TEXT("GenOpt ImportUniformScale"), ImportOptions->ImportUniformScale));
+						Attribs.Add(FAnalyticsEventAttribute(TEXT("GenOpt MaterialBasePath"), ImportOptions->MaterialBasePath));
+						Attribs.Add(FAnalyticsEventAttribute(TEXT("GenOpt MaterialSearchLocation"), MaterialSearchLocationEnum->GetNameStringByValue((uint64)(ImportOptions->MaterialSearchLocation))));
+
+						//We cant capture a this member, so just assign the pointer here
+						FBXImportOptions* CaptureImportOptions = ImportOptions;
+						auto AddMeshAnalytic = [&Attribs, &CaptureImportOptions]()
+						{
+							Attribs.Add(FAnalyticsEventAttribute(TEXT("MeshOpt AutoGenerateCollision"), CaptureImportOptions->bAutoGenerateCollision));
+							Attribs.Add(FAnalyticsEventAttribute(TEXT("MeshOpt CombineToSingle"), CaptureImportOptions->bCombineToSingle));
+							Attribs.Add(FAnalyticsEventAttribute(TEXT("MeshOpt BakePivotInVertex"), CaptureImportOptions->bBakePivotInVertex));
+							Attribs.Add(FAnalyticsEventAttribute(TEXT("MeshOpt TransformVertexToAbsolute"), CaptureImportOptions->bTransformVertexToAbsolute));
+							Attribs.Add(FAnalyticsEventAttribute(TEXT("MeshOpt ImportRigidMesh"), CaptureImportOptions->bImportRigidMesh));
+							Attribs.Add(FAnalyticsEventAttribute(TEXT("MeshOpt NormalGenerationMethod"), FBXNormalGenerationMethodEnum->GetNameStringByValue(CaptureImportOptions->NormalGenerationMethod)));
+							Attribs.Add(FAnalyticsEventAttribute(TEXT("MeshOpt NormalImportMethod"), FBXNormalImportMethodEnum->GetNameStringByValue(CaptureImportOptions->NormalImportMethod)));
+						};
+
+						auto AddSKAnalytic = [&Attribs, &CaptureImportOptions]()
+						{
+							Attribs.Add(FAnalyticsEventAttribute(TEXT("SkeletalMeshOpt CreatePhysicsAsset"), CaptureImportOptions->bCreatePhysicsAsset));
+							Attribs.Add(FAnalyticsEventAttribute(TEXT("SkeletalMeshOpt ImportAnimations"), CaptureImportOptions->bImportAnimations));
+							Attribs.Add(FAnalyticsEventAttribute(TEXT("SkeletalMeshOpt ImportAsSkeletalGeometry"), CaptureImportOptions->bImportAsSkeletalGeometry));
+							Attribs.Add(FAnalyticsEventAttribute(TEXT("SkeletalMeshOpt ImportAsSkeletalSkinning"), CaptureImportOptions->bImportAsSkeletalSkinning));
+							Attribs.Add(FAnalyticsEventAttribute(TEXT("SkeletalMeshOpt ImportMeshesInBoneHierarchy"), CaptureImportOptions->bImportMeshesInBoneHierarchy));
+							Attribs.Add(FAnalyticsEventAttribute(TEXT("SkeletalMeshOpt ImportMorph"), CaptureImportOptions->bImportMorph));
+							Attribs.Add(FAnalyticsEventAttribute(TEXT("SkeletalMeshOpt ImportSkeletalMeshLODs"), CaptureImportOptions->bImportSkeletalMeshLODs));
+							Attribs.Add(FAnalyticsEventAttribute(TEXT("SkeletalMeshOpt PreserveSmoothingGroups"), CaptureImportOptions->bPreserveSmoothingGroups));
+							Attribs.Add(FAnalyticsEventAttribute(TEXT("SkeletalMeshOpt UpdateSkeletonReferencePose"), CaptureImportOptions->bUpdateSkeletonReferencePose));
+							Attribs.Add(FAnalyticsEventAttribute(TEXT("SkeletalMeshOpt UseT0AsRefPose"), CaptureImportOptions->bUseT0AsRefPose));
+							Attribs.Add(FAnalyticsEventAttribute(TEXT("SkeletalMeshOpt OverlappingThresholds.ThresholdPosition"), CaptureImportOptions->OverlappingThresholds.ThresholdPosition));
+							Attribs.Add(FAnalyticsEventAttribute(TEXT("SkeletalMeshOpt OverlappingThresholds.ThresholdTangentNormal"), CaptureImportOptions->OverlappingThresholds.ThresholdTangentNormal));
+							Attribs.Add(FAnalyticsEventAttribute(TEXT("SkeletalMeshOpt OverlappingThresholds.ThresholdUV"), CaptureImportOptions->OverlappingThresholds.ThresholdUV));
+						};
+
+						auto AddSMAnalytic = [&Attribs, &CaptureImportOptions]()
+						{
+							Attribs.Add(FAnalyticsEventAttribute(TEXT("StaticMeshOpt AutoComputeLodDistances"), CaptureImportOptions->bAutoComputeLodDistances));
+							Attribs.Add(FAnalyticsEventAttribute(TEXT("StaticMeshOpt LodNumber"), CaptureImportOptions->LodNumber));
+							Attribs.Add(FAnalyticsEventAttribute(TEXT("StaticMeshOpt BuildAdjacencyBuffer"), CaptureImportOptions->bBuildAdjacencyBuffer));
+							Attribs.Add(FAnalyticsEventAttribute(TEXT("StaticMeshOpt BuildReversedIndexBuffer"), CaptureImportOptions->bBuildReversedIndexBuffer));
+							Attribs.Add(FAnalyticsEventAttribute(TEXT("StaticMeshOpt GenerateLightmapUVs"), CaptureImportOptions->bGenerateLightmapUVs));
+							Attribs.Add(FAnalyticsEventAttribute(TEXT("StaticMeshOpt ImportStaticMeshLODs"), CaptureImportOptions->bImportStaticMeshLODs));
+							Attribs.Add(FAnalyticsEventAttribute(TEXT("StaticMeshOpt RemoveDegenerates"), CaptureImportOptions->bRemoveDegenerates));
+							Attribs.Add(FAnalyticsEventAttribute(TEXT("StaticMeshOpt MinimumLodNumber"), CaptureImportOptions->MinimumLodNumber));
+							Attribs.Add(FAnalyticsEventAttribute(TEXT("StaticMeshOpt StaticMeshLODGroup"), CaptureImportOptions->StaticMeshLODGroup));
+							Attribs.Add(FAnalyticsEventAttribute(TEXT("StaticMeshOpt VertexColorImportOption"), VertexColorImportOptionEnum->GetNameStringByValue(CaptureImportOptions->VertexColorImportOption)));
+							Attribs.Add(FAnalyticsEventAttribute(TEXT("StaticMeshOpt VertexOverrideColor"), CaptureImportOptions->VertexOverrideColor.ToString()));
+						};
+
+						auto AddAnimAnalytic = [&Attribs, &CaptureImportOptions]()
+						{
+							Attribs.Add(FAnalyticsEventAttribute(TEXT("AnimOpt AnimationLengthImportType"), FBXAnimationLengthImportTypeEnum->GetNameStringByValue(CaptureImportOptions->AnimationLengthImportType)));
+							Attribs.Add(FAnalyticsEventAttribute(TEXT("AnimOpt DeleteExistingMorphTargetCurves"), CaptureImportOptions->bDeleteExistingMorphTargetCurves));
+							Attribs.Add(FAnalyticsEventAttribute(TEXT("AnimOpt AnimationRange"), CaptureImportOptions->AnimationRange.ToString()));
+							Attribs.Add(FAnalyticsEventAttribute(TEXT("AnimOpt DoNotImportCurveWithZero"), CaptureImportOptions->bDoNotImportCurveWithZero));
+							Attribs.Add(FAnalyticsEventAttribute(TEXT("AnimOpt ImportBoneTracks"), CaptureImportOptions->bImportBoneTracks));
+							Attribs.Add(FAnalyticsEventAttribute(TEXT("AnimOpt ImportCustomAttribute"), CaptureImportOptions->bImportCustomAttribute));
+							Attribs.Add(FAnalyticsEventAttribute(TEXT("AnimOpt PreserveLocalTransform"), CaptureImportOptions->bPreserveLocalTransform));
+							Attribs.Add(FAnalyticsEventAttribute(TEXT("AnimOpt RemoveRedundantKeys"), CaptureImportOptions->bRemoveRedundantKeys));
+							Attribs.Add(FAnalyticsEventAttribute(TEXT("AnimOpt Resample"), CaptureImportOptions->bResample));
+							Attribs.Add(FAnalyticsEventAttribute(TEXT("AnimOpt SetMaterialDriveParameterOnCustomAttribute"), CaptureImportOptions->bSetMaterialDriveParameterOnCustomAttribute));
+							Attribs.Add(FAnalyticsEventAttribute(TEXT("AnimOpt MaterialCurveSuffixes"), CaptureImportOptions->MaterialCurveSuffixes));
+							Attribs.Add(FAnalyticsEventAttribute(TEXT("AnimOpt ResampleRate"), CaptureImportOptions->ResampleRate));
+							
+						};
+						
+						if (ImportOptions->ImportType == FBXIT_SkeletalMesh)
+						{
+							AddMeshAnalytic();
+							AddSKAnalytic();
+							if (ImportOptions->bImportAnimations)
+							{
+								AddAnimAnalytic();
+							}
+						}
+						else if (ImportOptions->ImportType == FBXIT_StaticMesh)
+						{
+							AddMeshAnalytic();
+							AddSMAnalytic();
+						}
+						else if (ImportOptions->ImportType == FBXIT_Animation)
+						{
+							AddAnimAnalytic();
+						}
 
 						FString EventString = FString::Printf(TEXT("Editor.Usage.FBX.Import"));
 						FEngineAnalytics::GetProvider().RecordEvent(EventString, Attribs);
