Fix the static mesh reset to default material behavior. Need to compare materialInstance with nullptr and also with MD_SURFACE Default material.

diff --git a/Engine/Source/Editor/UnrealEd/Private/Fbx/FbxReimportDialogs.cpp b/Engine/Source/Editor/UnrealEd/Private/Fbx/FbxReimportDialogs.cpp
--- a/Engine/Source/Editor/UnrealEd/Private/Fbx/FbxReimportDialogs.cpp
+++ b/Engine/Source/Editor/UnrealEd/Private/Fbx/FbxReimportDialogs.cpp
@@ -206,7 +206,7 @@ void ResetMaterialSlot(const TArray<TMaterialType>& CurrentMaterial, TArray<TMat
 	// the material slot name, artist will be able to remap the material instance correctly
 	for (int32 MaterialIndex = 0; MaterialIndex < ResultMaterial.Num(); ++MaterialIndex)
 	{
-		if (ResultMaterial[MaterialIndex].MaterialInterface == nullptr)
+		if (ResultMaterial[MaterialIndex].MaterialInterface == nullptr || ResultMaterial[MaterialIndex].MaterialInterface == UMaterial::GetDefaultMaterial(MD_Surface))
 		{
 			bool bFoundMatch = false;
 			for (int32 ExistMaterialIndex = 0; ExistMaterialIndex < CurrentMaterial.Num(); ++ExistMaterialIndex)
