Updated text in the notifications about old or missing Xcode to be more helpful

diff --git a/Engine/Source/Editor/UnrealEd/Private/EditorEngine.cpp b/Engine/Source/Editor/UnrealEd/Private/EditorEngine.cpp
--- a/Engine/Source/Editor/UnrealEd/Private/EditorEngine.cpp
+++ b/Engine/Source/Editor/UnrealEd/Private/EditorEngine.cpp
@@ -6834,8 +6834,8 @@ void UEditorEngine::UpdateAutoLoadProject()
 		const ECheckBoxState DontAskAgainCheckBoxState = Local::GetDontAskAgainCheckBoxState();
 		if (DontAskAgainCheckBoxState == ECheckBoxState::Unchecked)
 		{
-			const FText NoXcodeMessageText = LOCTEXT("XcodeNotInstalledWarningNotification", "Xcode is not installed on this Mac.\nMetal shader compilation will fall back to runtime compiled text shaders, which are slower.\nPlease install latest version of Xcode for best performance.");
-			const FText OldXcodeMessageText = LOCTEXT("OldXcodeVersionWarningNotification", "Xcode installed on this Mac is too old to be used for Metal shader compilation.\nFalling back to runtime compiled text shaders, which are slower.\nPlease update to latest version of Xcode for best performance.");
+			const FText NoXcodeMessageText = LOCTEXT("XcodeNotInstalledWarningNotification", "Xcode was not detected on this Mac.\nMetal shader compilation will fall back to runtime compiled text shaders, which are slower.\nPlease install latest version of Xcode for best performance\nand make sure it's set as default using xcode-select tool.");
+			const FText OldXcodeMessageText = LOCTEXT("OldXcodeVersionWarningNotification", "Xcode installed on this Mac is too old to be used for Metal shader compilation.\nFalling back to runtime compiled text shaders, which are slower.\nPlease update to latest version of Xcode for best performance\nand make sure it's set as default using xcode-select tool.");
 
 			FNotificationInfo Info(bIsXcodeInstalled ? OldXcodeMessageText : NoXcodeMessageText);
 			Info.bFireAndForget = false;
