PR #5267: UE-66108: Use human readable name for data table column names (Contributed by [REDACTED_PERSON])

diff --git a/Engine/Source/Editor/UnrealEd/Private/DataTableEditorUtils.cpp b/Engine/Source/Editor/UnrealEd/Private/DataTableEditorUtils.cpp
--- a/Engine/Source/Editor/UnrealEd/Private/DataTableEditorUtils.cpp
+++ b/Engine/Source/Editor/UnrealEd/Private/DataTableEditorUtils.cpp
@@ -268,7 +268,7 @@ void FDataTableEditorUtils::CacheDataTableForEditing(const UDataTable* DataTable
 	for (int32 Index = 0; Index < StructProps.Num(); ++Index)
 	{
 		const UProperty* Prop = StructProps[Index];
-		const FText PropertyDisplayName = FText::FromString(DataTableUtils::GetPropertyDisplayName(Prop, Prop->GetName()));
+		const FText PropertyDisplayName = FText::FromString(DataTableUtils::GetPropertyDisplayName(Prop, FName::NameToDisplayString(Prop->GetName(), Prop->IsA<UBoolProperty>())));
 
 		FDataTableEditorColumnHeaderDataPtr CachedColumnData;
 		
