Restore the skeletal mesh LOD section data after a reimport

diff --git a/Engine/Source/Editor/UnrealEd/Private/Fbx/FbxSkeletalMeshImport.cpp b/Engine/Source/Editor/UnrealEd/Private/Fbx/FbxSkeletalMeshImport.cpp
--- a/Engine/Source/Editor/UnrealEd/Private/Fbx/FbxSkeletalMeshImport.cpp
+++ b/Engine/Source/Editor/UnrealEd/Private/Fbx/FbxSkeletalMeshImport.cpp
@@ -3727,6 +3727,66 @@ bool UnFbx::FFbxImporter::ImportSkeletalMeshLOD(USkeletalMesh* InSkeletalMesh, U
 		}
 	}
 
+	//Restore the LOD section data in case this LOD was reimport and some material match
+	if (DestImportedResource->LODModels.IsValidIndex(DesiredLOD) && ImportedResource->LODModels.IsValidIndex(0))
+	{
+		const TArray<FSkelMeshSection>& ExistingSections = DestImportedResource->LODModels[DesiredLOD].Sections;
+		const FSkeletalMeshLODInfo& ExistingInfo = *(BaseSkeletalMesh->GetLODInfo(DesiredLOD));
+
+		TArray<FSkelMeshSection>& ImportedSections = ImportedResource->LODModels[0].Sections;
+		const FSkeletalMeshLODInfo& ImportedInfo = *(InSkeletalMesh->GetLODInfo(0));
+		
+		auto GetImportMaterialSlotName = [](const USkeletalMesh* SkelMesh, const FSkelMeshSection& Section, const FSkeletalMeshLODInfo& Info, int32& OutMaterialIndex)->FName
+		{
+			check(SkelMesh->Materials.Num() > 0);
+			OutMaterialIndex = Section.MaterialIndex;
+			if (Info.LODMaterialMap.IsValidIndex(Section.MaterialIndex))
+			{
+				OutMaterialIndex = Info.LODMaterialMap[Section.MaterialIndex];
+			}
+			FName ImportedMaterialSlotName = NAME_None;
+			if (SkelMesh->Materials.IsValidIndex(OutMaterialIndex))
+			{
+				ImportedMaterialSlotName = SkelMesh->Materials[OutMaterialIndex].ImportedMaterialSlotName;
+			}
+			else
+			{
+				ImportedMaterialSlotName = SkelMesh->Materials[0].ImportedMaterialSlotName;
+				OutMaterialIndex = 0;
+			}
+			return ImportedMaterialSlotName;
+		};
+
+		for (const FSkelMeshSection& ExistingSection : ExistingSections)
+		{
+			int32 ExistingMaterialIndex = 0;
+			FName ExistingImportedMaterialSlotName = GetImportMaterialSlotName(BaseSkeletalMesh, ExistingSection, ExistingInfo, ExistingMaterialIndex);
+
+			for (FSkelMeshSection& ImportedSection : ImportedSections)
+			{
+				int32 ImportedMaterialIndex = 0;
+				FName ImportedImportedMaterialSlotName = GetImportMaterialSlotName(InSkeletalMesh, ImportedSection, ImportedInfo, ImportedMaterialIndex);
+				if (ExistingImportedMaterialSlotName != NAME_None)
+				{
+					if (ImportedImportedMaterialSlotName == ExistingImportedMaterialSlotName)
+					{
+						//Set the value and exit
+						ImportedSection.bCastShadow = ExistingSection.bCastShadow;
+						ImportedSection.bRecomputeTangent = ExistingSection.bRecomputeTangent;
+						break;
+					}
+				}
+				else if(InSkeletalMesh->Materials[ImportedMaterialIndex] == BaseSkeletalMesh->Materials[ExistingMaterialIndex]) //Use material slot compare to match in case the name is none
+				{
+					//Set the value and exit
+					ImportedSection.bCastShadow = ExistingSection.bCastShadow;
+					ImportedSection.bRecomputeTangent = ExistingSection.bRecomputeTangent;
+					break;
+				}
+			}
+		}
+	}
+
 	if (bNeedToReregister)
 	{
 		// Shut down the skeletal mesh component that is previewing this mesh.
