Better logging and an early out for cases where 'None' is finding it's way to package searching.

diff --git a/Engine/Source/Editor/UnrealEd/Private/ObjectTools.cpp b/Engine/Source/Editor/UnrealEd/Private/ObjectTools.cpp
--- a/Engine/Source/Editor/UnrealEd/Private/ObjectTools.cpp
+++ b/Engine/Source/Editor/UnrealEd/Private/ObjectTools.cpp
@@ -4053,6 +4053,12 @@ namespace ThumbnailTools
 		FString PackageName = InPackageFileName;
 		if (FPackageName::TryConvertFilenameToLongPackageName(PackageName, PackageName))
 		{
+			if (PackageName == TEXT("None"))
+			{
+				UE_LOG(LogUObjectGlobals, Warning, TEXT("Attempted to FindCachedThumbnailInPackage named 'None' - PackageName: %s InPackageFileName: %s"), *PackageName, *InPackageFileName);
+				return nullptr;
+			}
+
 			// First check to see if the package is already in memory.  If it is, some or all of the thumbnails
 			// may already be loaded and ready.
 			UObject* PackageOuter = nullptr;
diff --git a/Engine/Source/Runtime/CoreUObject/Private/UObject/UObjectGlobals.cpp b/Engine/Source/Runtime/CoreUObject/Private/UObject/UObjectGlobals.cpp
--- a/Engine/Source/Runtime/CoreUObject/Private/UObject/UObjectGlobals.cpp
+++ b/Engine/Source/Runtime/CoreUObject/Private/UObject/UObjectGlobals.cpp
@@ -626,7 +626,7 @@ UPackage* FindPackage( UObject* InOuter, const TCHAR* PackageName )
 	}
 	else
 	{
-		UE_LOG(LogUObjectGlobals, Fatal, TEXT("%s"), TEXT("Attempted to create a package named 'None'") );
+		UE_LOG(LogUObjectGlobals, Fatal, TEXT("Attempted to find a package named 'None' - InName: %s"), PackageName);
 	}
 	return Result;
 }
