PR #5289: Viewport setting to invert right mouse dolly (Contributed by [REDACTED_PERSON])

diff --git a/Engine/Config/BaseEditorPerProjectUserSettings.ini b/Engine/Config/BaseEditorPerProjectUserSettings.ini
--- a/Engine/Config/BaseEditorPerProjectUserSettings.ini
+++ b/Engine/Config/BaseEditorPerProjectUserSettings.ini
@@ -309,6 +309,7 @@ CameraSpeedScalar=1.0f
 ;Mouse speed when using middle mouse scroll in viewport
 MouseScrollCameraSpeed=5
 MouseSensitivty=.2f
+bInvertRightMouseDollyYAxis=False
 ; Whether to use mouse position as direct widget position
 bUseAbsoluteTranslation=True
 bLevelStreamingVolumePrevis=False
diff --git a/Engine/Source/Editor/UnrealEd/Classes/Settings/LevelEditorViewportSettings.h b/Engine/Source/Editor/UnrealEd/Classes/Settings/LevelEditorViewportSettings.h
--- a/Engine/Source/Editor/UnrealEd/Classes/Settings/LevelEditorViewportSettings.h
+++ b/Engine/Source/Editor/UnrealEd/Classes/Settings/LevelEditorViewportSettings.h
@@ -270,6 +270,10 @@ class UNREALED_API ULevelEditorViewportSettings
 	UPROPERTY(EditAnywhere, config, Category=Controls)
 	bool bInvertMiddleMousePan;
 
+	/** Whether or not to invert the direction of right mouse dolly on the Y axis in orbit mode */
+	UPROPERTY(EditAnywhere, config, Category = Controls)
+	bool bInvertRightMouseDollyYAxis;
+
 	/** Whether to use mouse position as direct widget position. */
 	UPROPERTY(EditAnywhere, config, Category=Controls, AdvancedDisplay)
 	uint32 bUseAbsoluteTranslation:1;
diff --git a/Engine/Source/Editor/UnrealEd/Private/EditorViewportClient.cpp b/Engine/Source/Editor/UnrealEd/Private/EditorViewportClient.cpp
--- a/Engine/Source/Editor/UnrealEd/Private/EditorViewportClient.cpp
+++ b/Engine/Source/Editor/UnrealEd/Private/EditorViewportClient.cpp
@@ -2246,7 +2246,7 @@ void FEditorViewportClient::InputAxisForOrbit(FViewport* InViewport, const FVect
 	}
 	else if ( IsOrbitPanMode( InViewport ) )
 	{
-		bool bInvert = GetDefault<ULevelEditorViewportSettings>()->bInvertMiddleMousePan;
+		const bool bInvert = GetDefault<ULevelEditorViewportSettings>()->bInvertMiddleMousePan;
 
 		const float CameraSpeed = GetCameraSpeed();
 		Drag *= CameraSpeed;
@@ -2269,12 +2269,14 @@ void FEditorViewportClient::InputAxisForOrbit(FViewport* InViewport, const FVect
 	}
 	else if ( IsOrbitZoomMode( InViewport ) )
 	{
+		const bool bInvertY = GetDefault<ULevelEditorViewportSettings>()->bInvertRightMouseDollyYAxis;
+
 		FMatrix OrbitMatrix = ViewTransform.ComputeOrbitMatrix().InverseFast();
 
 		const float CameraSpeed = GetCameraSpeed();
 		Drag *= CameraSpeed;
 
-		FVector DeltaLocation = FVector(0, Drag.X+ -Drag.Y, 0);
+		FVector DeltaLocation = bInvertY ? FVector(0, Drag.X + Drag.Y, 0) : FVector(0, Drag.X+ -Drag.Y, 0);
 
 		FVector LookAt = ViewTransform.GetLookAt();
 
