Modified the Mac message box asking users to update to latest macOS to show once for every engine release

diff --git a/Engine/Source/Editor/UnrealEd/Private/EditorEngine.cpp b/Engine/Source/Editor/UnrealEd/Private/EditorEngine.cpp
--- a/Engine/Source/Editor/UnrealEd/Private/EditorEngine.cpp
+++ b/Engine/Source/Editor/UnrealEd/Private/EditorEngine.cpp
@@ -6758,13 +6758,14 @@ void UEditorEngine::UpdateAutoLoadProject()
 #if PLATFORM_MAC
 	if ( !GIsBuildMachine )
 	{
-		if(FPlatformMisc::MacOSXVersionCompare(10,13,5) < 0)
+		if(FPlatformMisc::MacOSXVersionCompare(10,14,1) < 0)
 		{
 			if(FSlateApplication::IsInitialized())
 			{
-				FSuppressableWarningDialog::FSetupInfo Info( LOCTEXT("UpdateMacOSX_Body","Please update to the latest version of macOS for best performance and stability."), LOCTEXT("UpdateMacOSX_Title","Update macOS"), TEXT("UpdateMacOSX"), GEditorSettingsIni );
+				FString SupressSettingName(FString(TEXT("UpdateMacOSX_")) + VERSION_STRINGIFY(ENGINE_MAJOR_VERSION) + TEXT("_") + VERSION_STRINGIFY(ENGINE_MINOR_VERSION) + TEXT("_") + VERSION_STRINGIFY(ENGINE_PATCH_VERSION));
+				FSuppressableWarningDialog::FSetupInfo Info( LOCTEXT("UpdateMacOSX_Body","Please update to the latest version of macOS for best performance and stability."), LOCTEXT("UpdateMacOSX_Title","Update macOS"), *SupressSettingName, GEditorSettingsIni );
 				Info.ConfirmText = LOCTEXT( "OK", "OK");
-				Info.bDefaultToSuppressInTheFuture = true;
+				Info.bDefaultToSuppressInTheFuture = false;
 				FSuppressableWarningDialog OSUpdateWarning( Info );
 				OSUpdateWarning.ShowModal();
 			}
