Add import vertex color options to skeletalmesh

diff --git a/Engine/Source/Editor/DetailCustomizations/Private/FbxImportUIDetails.cpp b/Engine/Source/Editor/DetailCustomizations/Private/FbxImportUIDetails.cpp
--- a/Engine/Source/Editor/DetailCustomizations/Private/FbxImportUIDetails.cpp
+++ b/Engine/Source/Editor/DetailCustomizations/Private/FbxImportUIDetails.cpp
@@ -411,6 +411,15 @@ void FFbxImportUIDetails::CustomizeDetails( IDetailLayoutBuilder& DetailBuilder
 
 						PropertyRow.IsEnabled(TAttribute<bool>(this, &FFbxImportUIDetails::GetVertexOverrideColorEnabledState));
 					}
+
+					if (Property->GetFName() == GET_MEMBER_NAME_CHECKED(UFbxSkeletalMeshImportData, VertexOverrideColor))
+					{
+						// Cache the VertexColorImportOption property
+						SkeletalMeshVertexColorImportOptionHandle = SkeletalMeshDataProp->GetChildHandle(GET_MEMBER_NAME_CHECKED(UFbxSkeletalMeshImportData, VertexColorImportOption));
+
+						PropertyRow.IsEnabled(TAttribute<bool>(this, &FFbxImportUIDetails::GetSkeletalMeshVertexOverrideColorEnabledState));
+					}
+					
 				}
 			}
 			//Add refresh callback
@@ -850,12 +859,22 @@ void FFbxImportUIDetails::OnLODGroupChanged(TSharedPtr<FString> NewValue, ESelec
 bool FFbxImportUIDetails::GetVertexOverrideColorEnabledState() const
 {
 	uint8 VertexColorImportOption;
-	check(VertexColorImportOptionHandle.IsValid())
+	check(VertexColorImportOptionHandle.IsValid());
 	ensure(VertexColorImportOptionHandle->GetValue(VertexColorImportOption) == FPropertyAccess::Success);
 
 	return (VertexColorImportOption == EVertexColorImportOption::Override);
 }
 
+
+bool FFbxImportUIDetails::GetSkeletalMeshVertexOverrideColorEnabledState() const
+{
+	uint8 VertexColorImportOption;
+	check(SkeletalMeshVertexColorImportOptionHandle.IsValid());
+	ensure(SkeletalMeshVertexColorImportOptionHandle->GetValue(VertexColorImportOption) == FPropertyAccess::Success);
+
+	return (VertexColorImportOption == EVertexColorImportOption::Override);
+}
+
 void FFbxImportUIDetails::CollectChildPropertiesRecursive(TSharedPtr<IPropertyHandle> Node, TArray<TSharedPtr<IPropertyHandle>>& OutProperties)
 {
 	uint32 NodeNumChildren = 0;
diff --git a/Engine/Source/Editor/DetailCustomizations/Private/FbxImportUIDetails.h b/Engine/Source/Editor/DetailCustomizations/Private/FbxImportUIDetails.h
--- a/Engine/Source/Editor/DetailCustomizations/Private/FbxImportUIDetails.h
+++ b/Engine/Source/Editor/DetailCustomizations/Private/FbxImportUIDetails.h
@@ -93,6 +93,7 @@ private:
 
 	/** Called to determine the visibility of the VertexOverrideColor property */
 	bool GetVertexOverrideColorEnabledState() const;
+	bool GetSkeletalMeshVertexOverrideColorEnabledState() const;
 
 	FReply ShowConflictDialog(ConflictDialogType DialogType);
 	bool ShowCompareResult();
@@ -106,6 +107,7 @@ private:
 
 	/** Cached VertexColorImportOption property handle */
 	TSharedPtr<IPropertyHandle> VertexColorImportOptionHandle;
+	TSharedPtr<IPropertyHandle> SkeletalMeshVertexColorImportOptionHandle;
 
 	TArray< TSharedPtr< FString > > BaseColorNames;
 	TArray< TSharedPtr< FString > > BaseTextureNames;
diff --git a/Engine/Source/Editor/MeshPaint/Private/MeshPaintHelpers.cpp b/Engine/Source/Editor/MeshPaint/Private/MeshPaintHelpers.cpp
--- a/Engine/Source/Editor/MeshPaint/Private/MeshPaintHelpers.cpp
+++ b/Engine/Source/Editor/MeshPaint/Private/MeshPaintHelpers.cpp
@@ -49,6 +49,8 @@
 #include "VREditorInteractor.h"
 #include "EditorWorldExtension.h"
 
+#include "Factories/FbxSkeletalMeshImportData.h"
+
 #include "Async/ParallelFor.h"
 #include "Rendering/SkeletalMeshModel.h"
 
@@ -1533,6 +1535,18 @@ void MeshPaintHelpers::ImportVertexColorsToSkeletalMesh(USkeletalMesh* SkeletalM
 		const FVector2D UV = LODModel.Sections[SectionIndex].SoftVertices[SectionVertexIndex].UVs[UVIndex];
 		LODModel.Sections[SectionIndex].SoftVertices[SectionVertexIndex].Color = PickVertexColorFromTextureData(MipData, UV, Texture, ColorMask);
 	}
+
+	//Make sure we change the import data so the re-import do not replace the new data
+	if (SkeletalMesh->AssetImportData)
+	{
+		UFbxSkeletalMeshImportData* ImportData = Cast<UFbxSkeletalMeshImportData>(SkeletalMesh->AssetImportData);
+		if (ImportData && ImportData->VertexColorImportOption != EVertexColorImportOption::Ignore)
+		{
+			ImportData->SetFlags(RF_Transactional);
+			ImportData->Modify();
+			ImportData->VertexColorImportOption = EVertexColorImportOption::Ignore;
+		}
+	}
 }
 
 FColor MeshPaintHelpers::PickVertexColorFromTextureData(const uint8* MipData, const FVector2D& UVCoordinate, const UTexture2D* Texture, const FColor ColorMask)
diff --git a/Engine/Source/Editor/MeshPaint/Private/MeshPaintSkeletalMeshAdapter.cpp b/Engine/Source/Editor/MeshPaint/Private/MeshPaintSkeletalMeshAdapter.cpp
--- a/Engine/Source/Editor/MeshPaint/Private/MeshPaintSkeletalMeshAdapter.cpp
+++ b/Engine/Source/Editor/MeshPaint/Private/MeshPaintSkeletalMeshAdapter.cpp
@@ -8,6 +8,7 @@
 #include "ComponentReregisterContext.h"
 #include "Rendering/SkeletalMeshRenderData.h"
 #include "Rendering/SkeletalMeshModel.h"
+#include "Factories/FbxSkeletalMeshImportData.h"
 
 //////////////////////////////////////////////////////////////////////////
 // FMeshPaintGeometryAdapterForSkeletalMeshes
@@ -338,6 +339,17 @@ void FMeshPaintGeometryAdapterForSkeletalMeshes::PreEdit()
 		ReferencedSkeletalMesh->VertexColorGuid = FGuid::NewGuid();
 		BeginInitResource(&LODData->StaticVertexBuffers.ColorVertexBuffer);
 	}
+	//Make sure we change the import data so the re-import do not replace the new data
+	if (ReferencedSkeletalMesh->AssetImportData)
+	{
+		UFbxSkeletalMeshImportData* ImportData = Cast<UFbxSkeletalMeshImportData>(ReferencedSkeletalMesh->AssetImportData);
+		if (ImportData && ImportData->VertexColorImportOption != EVertexColorImportOption::Ignore)
+		{
+			ImportData->SetFlags(RF_Transactional);
+			ImportData->Modify();
+			ImportData->VertexColorImportOption = EVertexColorImportOption::Ignore;
+		}
+	}
 }
 
 void FMeshPaintGeometryAdapterForSkeletalMeshes::PostEdit()
diff --git a/Engine/Source/Editor/UnrealEd/Classes/Factories/FbxMeshImportData.h b/Engine/Source/Editor/UnrealEd/Classes/Factories/FbxMeshImportData.h
--- a/Engine/Source/Editor/UnrealEd/Classes/Factories/FbxMeshImportData.h
+++ b/Engine/Source/Editor/UnrealEd/Classes/Factories/FbxMeshImportData.h
@@ -28,6 +28,20 @@ namespace EFBXNormalGenerationMethod
 	};
 }
 
+UENUM()
+namespace EVertexColorImportOption
+{
+	enum Type
+	{
+		/** Import the static mesh using the vertex colors from the FBX file. */
+		Replace,
+		/** Ignore vertex colors from the FBX file, and keep the existing mesh vertex colors. */
+		Ignore,
+		/** Override all vertex colors with the specified color. */
+		Override
+	};
+}
+
 /** Action to add nodes to the graph based on selected objects*/
 USTRUCT()
 struct FImportMeshLodSectionsData
diff --git a/Engine/Source/Editor/UnrealEd/Classes/Factories/FbxSkeletalMeshImportData.h b/Engine/Source/Editor/UnrealEd/Classes/Factories/FbxSkeletalMeshImportData.h
--- a/Engine/Source/Editor/UnrealEd/Classes/Factories/FbxSkeletalMeshImportData.h
+++ b/Engine/Source/Editor/UnrealEd/Classes/Factories/FbxSkeletalMeshImportData.h
@@ -33,6 +33,14 @@ public:
 	UPROPERTY(EditAnywhere, Transient, Category = Mesh, meta = (ImportType = "SkeletalMesh", DisplayName = "Import Content Type", OBJRestrict = "true"))
 	TEnumAsByte<enum EFBXImportContentType> ImportContentType;
 
+	/** Specify how vertex colors should be imported */
+	UPROPERTY(EditAnywhere, BlueprintReadWrite, AdvancedDisplay, config, Category = Mesh, meta = (OBJRestrict = "true", ImportType = "SkeletalMesh"))
+	TEnumAsByte<EVertexColorImportOption::Type> VertexColorImportOption;
+
+	/** Specify override color in the case that VertexColorImportOption is set to Override */
+	UPROPERTY(EditAnywhere, BlueprintReadWrite, AdvancedDisplay, config, Category = Mesh, meta = (OBJRestrict = "true", ImportType = "SkeletalMesh"))
+	FColor VertexOverrideColor;
+
 	/** Enable this option to update Skeleton (of the mesh)'s reference pose. Mesh's reference pose is always updated.  */
 	UPROPERTY(EditAnywhere, AdvancedDisplay, Category=Mesh, meta=(ImportType="SkeletalMesh|RigOnly", ToolTip="If enabled, update the Skeleton (of the mesh being imported)'s reference pose."))
 	uint32 bUpdateSkeletonReferencePose:1;
diff --git a/Engine/Source/Editor/UnrealEd/Classes/Factories/FbxStaticMeshImportData.h b/Engine/Source/Editor/UnrealEd/Classes/Factories/FbxStaticMeshImportData.h
--- a/Engine/Source/Editor/UnrealEd/Classes/Factories/FbxStaticMeshImportData.h
+++ b/Engine/Source/Editor/UnrealEd/Classes/Factories/FbxStaticMeshImportData.h
@@ -13,20 +13,6 @@
 
 class UStaticMesh;
 
-UENUM()
-namespace EVertexColorImportOption
-{
-	enum Type
-	{
-		/** Import the static mesh using the vertex colors from the FBX file. */
-		Replace,
-		/** Ignore vertex colors from the FBX file, and keep the existing mesh vertex colors. */
-		Ignore,
-		/** Override all vertex colors with the specified color. */
-		Override
-	};
-}
-
 UCLASS(config=EditorPerProjectUserSettings, AutoExpandCategories=(Options), MinimalAPI)
 class UFbxStaticMeshImportData : public UFbxMeshImportData
 {
diff --git a/Engine/Source/Editor/UnrealEd/Private/Fbx/FbxMainImport.cpp b/Engine/Source/Editor/UnrealEd/Private/Fbx/FbxMainImport.cpp
--- a/Engine/Source/Editor/UnrealEd/Private/Fbx/FbxMainImport.cpp
+++ b/Engine/Source/Editor/UnrealEd/Private/Fbx/FbxMainImport.cpp
@@ -396,6 +396,8 @@ void ApplyImportUIToImportOptions(UFbxImportUI* ImportUI, FBXImportOptions& InOu
 		InOutImportOptions.bConvertScene			= StaticMeshData->bConvertScene;
 		InOutImportOptions.bForceFrontXAxis			= StaticMeshData->bForceFrontXAxis;
 		InOutImportOptions.bConvertSceneUnit		= StaticMeshData->bConvertSceneUnit;
+		InOutImportOptions.VertexColorImportOption	= StaticMeshData->VertexColorImportOption;
+		InOutImportOptions.VertexOverrideColor		= StaticMeshData->VertexOverrideColor;
 	}
 	else if ( ImportUI->MeshTypeToImport == FBXIT_SkeletalMesh )
 	{
@@ -413,6 +415,8 @@ void ApplyImportUIToImportOptions(UFbxImportUI* ImportUI, FBXImportOptions& InOu
 		InOutImportOptions.bConvertScene				= SkeletalMeshData->bConvertScene;
 		InOutImportOptions.bForceFrontXAxis				= SkeletalMeshData->bForceFrontXAxis;
 		InOutImportOptions.bConvertSceneUnit			= SkeletalMeshData->bConvertSceneUnit;
+		InOutImportOptions.VertexColorImportOption		= SkeletalMeshData->VertexColorImportOption;
+		InOutImportOptions.VertexOverrideColor			= SkeletalMeshData->VertexOverrideColor;
 
 		if(ImportUI->bImportAnimations)
 		{
@@ -425,7 +429,6 @@ void ApplyImportUIToImportOptions(UFbxImportUI* ImportUI, FBXImportOptions& InOu
 			AnimData->bForceFrontXAxis				= SkeletalMeshData->bForceFrontXAxis;
 			AnimData->bConvertSceneUnit				= SkeletalMeshData->bConvertSceneUnit;
 		}
-
 	}
 	else
 	{
@@ -448,8 +451,7 @@ void ApplyImportUIToImportOptions(UFbxImportUI* ImportUI, FBXImportOptions& InOu
 	InOutImportOptions.OverlappingThresholds.ThresholdTangentNormal = ImportUI->SkeletalMeshImportData->ThresholdTangentNormal;
 	InOutImportOptions.OverlappingThresholds.ThresholdUV = ImportUI->SkeletalMeshImportData->ThresholdUV;
 	InOutImportOptions.bCombineToSingle = ImportUI->StaticMeshImportData->bCombineMeshes;
-	InOutImportOptions.VertexColorImportOption = ImportUI->StaticMeshImportData->VertexColorImportOption;
-	InOutImportOptions.VertexOverrideColor = ImportUI->StaticMeshImportData->VertexOverrideColor;
+	
 	InOutImportOptions.bRemoveDegenerates = ImportUI->StaticMeshImportData->bRemoveDegenerates;
 	InOutImportOptions.bBuildAdjacencyBuffer = ImportUI->StaticMeshImportData->bBuildAdjacencyBuffer;
 	InOutImportOptions.bBuildReversedIndexBuffer = ImportUI->StaticMeshImportData->bBuildReversedIndexBuffer;
diff --git a/Engine/Source/Editor/UnrealEd/Private/Fbx/FbxSkeletalMeshImport.cpp b/Engine/Source/Editor/UnrealEd/Private/Fbx/FbxSkeletalMeshImport.cpp
--- a/Engine/Source/Editor/UnrealEd/Private/Fbx/FbxSkeletalMeshImport.cpp
+++ b/Engine/Source/Editor/UnrealEd/Private/Fbx/FbxSkeletalMeshImport.cpp
@@ -3033,15 +3033,20 @@ bool UnFbx::FFbxImporter::FillSkelMeshImporterFromFbx( FSkeletalMeshImportData&
 	//
 	// get the first vertex color layer
 	//
+	
 	FbxLayerElementVertexColor* LayerElementVertexColor = BaseLayer->GetVertexColors();
 	FbxLayerElement::EReferenceMode VertexColorReferenceMode(FbxLayerElement::eDirect);
 	FbxLayerElement::EMappingMode VertexColorMappingMode(FbxLayerElement::eByControlPoint);
-	if (LayerElementVertexColor)
+	if (LayerElementVertexColor && ImportOptions->VertexColorImportOption == EVertexColorImportOption::Replace)
 	{
 		VertexColorReferenceMode = LayerElementVertexColor->GetReferenceMode();
 		VertexColorMappingMode = LayerElementVertexColor->GetMappingMode();
 		ImportData.bHasVertexColors = true;
 	}
+	else if(ImportOptions->VertexColorImportOption == EVertexColorImportOption::Override)
+	{
+		ImportData.bHasVertexColors = true;
+	}
 
 	//
 	// get the first normal layer
@@ -3270,7 +3275,7 @@ bool UnFbx::FFbxImporter::FillSkelMeshImporterFromFbx( FSkeletalMeshImportData&
 		}
 
 		// Read vertex colors if they exist.
-		if( LayerElementVertexColor )
+		if( LayerElementVertexColor && ImportOptions->VertexColorImportOption == EVertexColorImportOption::Replace)
 		{
 			switch(VertexColorMappingMode)
 			{
@@ -3310,6 +3315,14 @@ bool UnFbx::FFbxImporter::FillSkelMeshImporterFromFbx( FSkeletalMeshImportData&
 				break;
 			}
 		}
+		else if (ImportOptions->VertexColorImportOption == EVertexColorImportOption::Override)
+		{
+			for (int32 VertexIndex = 0; VertexIndex < 3; VertexIndex++)
+			{
+				int32 UnrealVertexIndex = OddNegativeScale ? 2 - VertexIndex : VertexIndex;
+				TmpWedges[UnrealVertexIndex].Color = ImportOptions->VertexOverrideColor;
+			}
+		}
 		
 		//
 		// basic wedges matching : 3 unique per face. TODO Can we do better ?
diff --git a/Engine/Source/Editor/UnrealEd/Private/Fbx/FbxSkeletalMeshImportData.cpp b/Engine/Source/Editor/UnrealEd/Private/Fbx/FbxSkeletalMeshImportData.cpp
--- a/Engine/Source/Editor/UnrealEd/Private/Fbx/FbxSkeletalMeshImportData.cpp
+++ b/Engine/Source/Editor/UnrealEd/Private/Fbx/FbxSkeletalMeshImportData.cpp
@@ -9,6 +9,7 @@ UFbxSkeletalMeshImportData::UFbxSkeletalMeshImportData(const FObjectInitializer&
 {
 	bTransformVertexToAbsolute = true;
 	bBakePivotInVertex = false;
+	VertexColorImportOption = EVertexColorImportOption::Replace;
 }
 
 UFbxSkeletalMeshImportData* UFbxSkeletalMeshImportData::GetImportDataForSkeletalMesh(USkeletalMesh* SkeletalMesh, UFbxSkeletalMeshImportData* TemplateForCreation)
