Set the base LOD to be the previous custom LOD

diff --git a/Engine/Source/Developer/MeshBuilder/Private/StaticMeshBuilder.cpp b/Engine/Source/Developer/MeshBuilder/Private/StaticMeshBuilder.cpp
--- a/Engine/Source/Developer/MeshBuilder/Private/StaticMeshBuilder.cpp
+++ b/Engine/Source/Developer/MeshBuilder/Private/StaticMeshBuilder.cpp
@@ -91,14 +91,10 @@ bool FStaticMeshBuilder::Build(FStaticMeshRenderData& StaticMeshRenderData, USta
 		const FStaticMeshSourceModel& SrcModel = StaticMesh->SourceModels[LodIndex];
 		FMeshReductionSettings ReductionSettings = LODGroup.GetSettings(SrcModel.ReductionSettings, LodIndex);
 
-		const int32 BaseReduceLodIndex = ReductionSettings.BaseLODModel;
+		//Make sure we do not reduce a non custom LOD by himself
+		const int32 BaseReduceLodIndex = FMath::Clamp<int32>(ReductionSettings.BaseLODModel, 0, OriginalMeshDescription == nullptr ? LodIndex - 1 : LodIndex);
 		// Use simplifier if a reduction in triangles or verts has been requested.
-
-		const bool bVertTermination = (!bIsThirdPartyReductiontool) && (ReductionSettings.TerminationCriterion != EStaticMeshReductionTerimationCriterion::Triangles);
-		const bool bTriTermination  = ReductionSettings.TerminationCriterion != EStaticMeshReductionTerimationCriterion::Vertices;
-		bool bUseReduction = ( bTriTermination && ReductionSettings.PercentTriangles < 1.0f )
-			              || ( bVertTermination && ReductionSettings.PercentVertices < 1.0f )
-			              || ( ReductionSettings.MaxDeviation > 0.0f );
+		bool bUseReduction = StaticMesh->IsReductionActive(LodIndex);
 
 		if (OriginalMeshDescription != nullptr)
 		{
@@ -150,25 +146,24 @@ bool FStaticMeshBuilder::Build(FStaticMeshRenderData& StaticMeshRenderData, USta
 		//Reduce LODs
 		if (bUseReduction)
 		{
-			const int32 BaseLodIndex = ReductionSettings.BaseLODModel;
 			float OverlappingThreshold = LODBuildSettings.bRemoveDegenerates ? THRESH_POINTS_ARE_SAME : 0.0f;
 			FOverlappingCorners OverlappingCorners;
-			FMeshDescriptionOperations::FindOverlappingCorners(OverlappingCorners, MeshDescriptions[BaseLodIndex], OverlappingThreshold);
+			FMeshDescriptionOperations::FindOverlappingCorners(OverlappingCorners, MeshDescriptions[BaseReduceLodIndex], OverlappingThreshold);
 
 			//Create a reduced mesh from the base LOD
 			UStaticMesh::RegisterMeshAttributes(MeshDescriptions[LodIndex]);
 			
 			int32 OldSectionInfoMapCount = StaticMesh->SectionInfoMap.GetSectionNumber(LodIndex);
 
-			if (LodIndex == BaseLodIndex)
+			if (LodIndex == BaseReduceLodIndex)
 			{
 				//When using LOD 0, we use a copy of the mesh description since reduce do not support inline reducing
-				FMeshDescription BaseMeshDescription = MeshDescriptions[BaseLodIndex];
+				FMeshDescription BaseMeshDescription = MeshDescriptions[BaseReduceLodIndex];
 				MeshDescriptionHelper.ReduceLOD(BaseMeshDescription, MeshDescriptions[LodIndex], ReductionSettings, OverlappingCorners, MaxDeviation);
 			}
 			else
 			{
-				MeshDescriptionHelper.ReduceLOD(MeshDescriptions[BaseLodIndex], MeshDescriptions[LodIndex], ReductionSettings, OverlappingCorners, MaxDeviation);
+				MeshDescriptionHelper.ReduceLOD(MeshDescriptions[BaseReduceLodIndex], MeshDescriptions[LodIndex], ReductionSettings, OverlappingCorners, MaxDeviation);
 			}
 			
 
@@ -209,7 +204,7 @@ bool FStaticMeshBuilder::Build(FStaticMeshRenderData& StaticMeshRenderData, USta
 				if (!bHasValidLODInfoMap)
 				{
 					bool bSectionInfoSet = false;
-					for (int32 BaseLodSectionIndex = 0; BaseLodSectionIndex < LODModelSectionInfoMap.GetSectionNumber(BaseLodIndex); ++BaseLodSectionIndex)
+					for (int32 BaseLodSectionIndex = 0; BaseLodSectionIndex < LODModelSectionInfoMap.GetSectionNumber(BaseReduceLodIndex); ++BaseLodSectionIndex)
 					{
 						FMeshSectionInfo SectionInfo = LODModelSectionInfoMap.Get(ReductionSettings.BaseLODModel, BaseLodSectionIndex);
 						if (BaseLodSectionIndex == UniqueMaterialIndex[SectionIndex])
diff --git a/Engine/Source/Developer/MeshMergeUtilities/Private/MeshMergeHelpers.cpp b/Engine/Source/Developer/MeshMergeUtilities/Private/MeshMergeHelpers.cpp
--- a/Engine/Source/Developer/MeshMergeUtilities/Private/MeshMergeHelpers.cpp
+++ b/Engine/Source/Developer/MeshMergeUtilities/Private/MeshMergeHelpers.cpp
@@ -395,7 +395,7 @@ void FMeshMergeHelpers::RetrieveMesh(const UStaticMesh* StaticMesh, int32 LODInd
 	const bool bImportedMesh = !StaticMeshModel.IsRawMeshEmpty() || (StaticMeshModel.OriginalMeshDescription != nullptr);
 	
 	// Check whether or not this mesh has been reduced in-engine
-	const bool bReducedMesh = (StaticMeshModel.ReductionSettings.PercentTriangles < 1.0f);
+	const bool bReducedMesh = StaticMesh->IsReductionActive(LODIndex);
 	// Trying to retrieve rawmesh from SourceStaticMeshModel was giving issues, which causes a mismatch			
 	const bool bRenderDataMismatch = (LODIndex > 0) || StaticMeshModel.BuildSettings.bGenerateLightmapUVs;
 
diff --git a/Engine/Source/Editor/StaticMeshEditor/Private/StaticMeshEditor.cpp b/Engine/Source/Editor/StaticMeshEditor/Private/StaticMeshEditor.cpp
--- a/Engine/Source/Editor/StaticMeshEditor/Private/StaticMeshEditor.cpp
+++ b/Engine/Source/Editor/StaticMeshEditor/Private/StaticMeshEditor.cpp
@@ -1375,7 +1375,9 @@ void FStaticMeshEditor::HandleReimportAllMesh()
 				{
 					continue;
 				}
-				bool bHasBeenSimplified = SourceModels[LodIndex].RawMeshBulkData->IsEmpty() || SourceModels[LodIndex].ReductionSettings.PercentTriangles < 1.0f || SourceModels[LodIndex].ReductionSettings.MaxDeviation > 0.0f;
+				
+				
+				bool bHasBeenSimplified = StaticMesh->GetOriginalMeshDescription(LodIndex) == nullptr || StaticMesh->IsReductionActive(LodIndex);
 				if (!bHasBeenSimplified)
 				{
 					FbxMeshUtils::ImportMeshLODDialog(StaticMesh, LodIndex);
diff --git a/Engine/Source/Editor/StaticMeshEditor/Private/StaticMeshEditorTools.cpp b/Engine/Source/Editor/StaticMeshEditor/Private/StaticMeshEditorTools.cpp
--- a/Engine/Source/Editor/StaticMeshEditor/Private/StaticMeshEditorTools.cpp
+++ b/Engine/Source/Editor/StaticMeshEditor/Private/StaticMeshEditorTools.cpp
@@ -1012,9 +1012,10 @@ void FMeshBuildSettingsLayout::OnDistanceFieldResolutionScaleCommitted(float New
 	OnDistanceFieldResolutionScaleChanged(NewValue);
 }
 
-FMeshReductionSettingsLayout::FMeshReductionSettingsLayout( TSharedRef<FLevelOfDetailSettingsLayout> InParentLODSettings, int32 InCurrentLODIndex)
+FMeshReductionSettingsLayout::FMeshReductionSettingsLayout( TSharedRef<FLevelOfDetailSettingsLayout> InParentLODSettings, int32 InCurrentLODIndex, bool InCanReduceMyself)
 	: ParentLODSettings( InParentLODSettings )
 	, CurrentLODIndex( InCurrentLODIndex )
+	, bCanReduceMyself(InCanReduceMyself)
 {
 
 	FillEnumOptions(ImportanceOptions, GetFeatureImportanceEnum());
@@ -1324,6 +1325,7 @@ void FMeshReductionSettingsLayout::GenerateChildContent( IDetailChildrenBuilder&
 
 	//Base LOD
 	{
+		int32 MaxBaseReduceIndex = bCanReduceMyself ? CurrentLODIndex : CurrentLODIndex - 1;
 		ChildrenBuilder.AddCustomRow( LOCTEXT("BaseLOD", "Base LOD") )
 			.NameContent()
 			.HAlign(HAlign_Left)
@@ -1338,9 +1340,9 @@ void FMeshReductionSettingsLayout::GenerateChildContent( IDetailChildrenBuilder&
 				SNew(SNumericEntryBox<int32>)
 				.AllowSpin(true)
 				.MinSliderValue(0)
-				.MaxSliderValue(CurrentLODIndex)
+				.MaxSliderValue(MaxBaseReduceIndex)
 				.MinValue(0)
-				.MaxValue(CurrentLODIndex)
+				.MaxValue(MaxBaseReduceIndex)
 				.Value(this, &FMeshReductionSettingsLayout::GetBaseLODIndex)
 				.OnValueChanged(this, &FMeshReductionSettingsLayout::SetBaseLODIndex)
 			];
@@ -3218,7 +3220,7 @@ void FLevelOfDetailSettingsLayout::AddLODLevelCategories( IDetailLayoutBuilder&
 
 			if (IsAutoMeshReductionAvailable())
 			{
-				ReductionSettingsWidgets[LODIndex] = MakeShareable( new FMeshReductionSettingsLayout( AsShared(), LODIndex) );
+				ReductionSettingsWidgets[LODIndex] = MakeShareable( new FMeshReductionSettingsLayout( AsShared(), LODIndex, StaticMesh->GetOriginalMeshDescription(LODIndex) != nullptr ));
 			}
 
 			if (LODIndex < StaticMesh->SourceModels.Num())
@@ -3258,7 +3260,7 @@ void FLevelOfDetailSettingsLayout::AddLODLevelCategories( IDetailLayoutBuilder&
 			CategoryName.AppendInt( LODIndex );
 
 			FText LODLevelString = FText::FromString(FString(TEXT("LOD ")) + FString::FromInt(LODIndex) );
-			bool bHasBeenSimplified = StaticMesh->SourceModels[LODIndex].RawMeshBulkData->IsEmpty() || StaticMesh->SourceModels[LODIndex].ReductionSettings.PercentTriangles < 1.0f || StaticMesh->SourceModels[LODIndex].ReductionSettings.MaxDeviation > 0.0f;
+			bool bHasBeenSimplified = StaticMesh->GetOriginalMeshDescription(LODIndex) == nullptr || StaticMesh->IsReductionActive(LODIndex);
 			FText GeneratedString = FText::FromString(bHasBeenSimplified ? TEXT("[generated]") : TEXT(""));
 
 			IDetailCategoryBuilder& LODCategory = DetailBuilder.EditCategory( *CategoryName, LODLevelString, ECategoryPriority::Important );
diff --git a/Engine/Source/Editor/StaticMeshEditor/Private/StaticMeshEditorTools.h b/Engine/Source/Editor/StaticMeshEditor/Private/StaticMeshEditorTools.h
--- a/Engine/Source/Editor/StaticMeshEditor/Private/StaticMeshEditorTools.h
+++ b/Engine/Source/Editor/StaticMeshEditor/Private/StaticMeshEditorTools.h
@@ -215,7 +215,7 @@ private:
 class FMeshReductionSettingsLayout : public IDetailCustomNodeBuilder, public TSharedFromThis<FMeshReductionSettingsLayout>
 {
 public:
-	FMeshReductionSettingsLayout( TSharedRef<FLevelOfDetailSettingsLayout> InParentLODSettings, int32 InCurrentLODIndex);
+	FMeshReductionSettingsLayout( TSharedRef<FLevelOfDetailSettingsLayout> InParentLODSettings, int32 InCurrentLODIndex, bool InCanReduceMyself);
 	virtual ~FMeshReductionSettingsLayout();
 
 	const FMeshReductionSettings& GetSettings() const;
@@ -289,6 +289,7 @@ private:
 	TWeakPtr<FLevelOfDetailSettingsLayout> ParentLODSettings;
 	FMeshReductionSettings ReductionSettings;
 	int32 CurrentLODIndex;
+	bool bCanReduceMyself;
 
 	// Used by simplygon
 	TArray<TSharedPtr<FString> > ImportanceOptions;
diff --git a/Engine/Source/Editor/UnrealEd/Private/Fbx/FbxFactory.cpp b/Engine/Source/Editor/UnrealEd/Private/Fbx/FbxFactory.cpp
--- a/Engine/Source/Editor/UnrealEd/Private/Fbx/FbxFactory.cpp
+++ b/Engine/Source/Editor/UnrealEd/Private/Fbx/FbxFactory.cpp
@@ -827,8 +827,9 @@ UObject* UFbxFactory::RecursiveImportNode(void* VoidFbxImporter, void* VoidNode,
 						{
 							NewStaticMesh->AddSourceModel();
 						}
-						if (LODIndex - 1 > 0 && (NewStaticMesh->SourceModels[LODIndex - 1].ReductionSettings.PercentTriangles < 1.0f || NewStaticMesh->SourceModels[LODIndex - 1].ReductionSettings.MaxDeviation > 0.0f))
+						if (LODIndex - 1 > 0 && NewStaticMesh->IsReductionActive(LODIndex-1))
 						{
+							//Do not add the LODGroup bias here, since the bias will be apply during the build
 							if (NewStaticMesh->SourceModels[LODIndex - 1].ReductionSettings.PercentTriangles < 1.0f)
 							{
 								NewStaticMesh->SourceModels[LODIndex].ReductionSettings.PercentTriangles = NewStaticMesh->SourceModels[LODIndex - 1].ReductionSettings.PercentTriangles * 0.5f;
diff --git a/Engine/Source/Editor/UnrealEd/Private/Fbx/FbxStaticMeshImport.cpp b/Engine/Source/Editor/UnrealEd/Private/Fbx/FbxStaticMeshImport.cpp
--- a/Engine/Source/Editor/UnrealEd/Private/Fbx/FbxStaticMeshImport.cpp
+++ b/Engine/Source/Editor/UnrealEd/Private/Fbx/FbxStaticMeshImport.cpp
@@ -1032,7 +1032,7 @@ void UnFbx::FFbxImporter::AddStaticMeshSourceModelGeneratedLOD(UStaticMesh* Stat
 	{
 		StaticMesh->AddSourceModel();
 	}
-	if (LODIndex - 1 > 0 && (StaticMesh->SourceModels[LODIndex - 1].ReductionSettings.PercentTriangles < 1.0f || StaticMesh->SourceModels[LODIndex - 1].ReductionSettings.MaxDeviation > 0.0f))
+	if (LODIndex - 1 > 0 && StaticMesh->IsReductionActive(LODIndex - 1))
 	{
 		if (StaticMesh->SourceModels[LODIndex - 1].ReductionSettings.PercentTriangles < 1.0f)
 		{
diff --git a/Engine/Source/Editor/UnrealEd/Private/StaticMeshEdit.cpp b/Engine/Source/Editor/UnrealEd/Private/StaticMeshEdit.cpp
--- a/Engine/Source/Editor/UnrealEd/Private/StaticMeshEdit.cpp
+++ b/Engine/Source/Editor/UnrealEd/Private/StaticMeshEdit.cpp
@@ -34,6 +34,10 @@
 #include "MeshDescription.h"
 #include "MeshAttributes.h"
 
+#include "Modules/ModuleManager.h"
+#include "IMeshReductionManagerModule.h"
+#include "IMeshReductionInterfaces.h"
+
 
 bool GBuildStaticMeshCollision = 1;
 
@@ -1196,6 +1200,23 @@ ExistingStaticMeshData* SaveExistingStaticMeshData(UStaticMesh* ExistingMesh, Un
 	return ExistingMeshDataPtr;
 }
 
+// Helper to find if some reduction settings are active
+bool IsReductionActive(const FMeshReductionSettings& ReductionSettings)
+{
+	bool bUseQuadricSimplier = true;
+	{
+		// Are we using our tool, or simplygon?  The tool is only changed during editor restarts
+		IMeshReduction* ReductionModule = FModuleManager::Get().LoadModuleChecked<IMeshReductionManagerModule>("MeshReductionInterface").GetStaticMeshReductionInterface();
+		FString VersionString = ReductionModule->GetVersionString();
+		TArray<FString> SplitVersionString;
+		VersionString.ParseIntoArray(SplitVersionString, TEXT("_"), true);
+		bUseQuadricSimplier = SplitVersionString[0].Equals("QuadricMeshReduction");
+	}
+	const bool bVertTermination = (bUseQuadricSimplier) && (ReductionSettings.TerminationCriterion != EStaticMeshReductionTerimationCriterion::Triangles) && (ReductionSettings.PercentVertices < 1.0f);
+	const bool bTriTermination = ReductionSettings.TerminationCriterion != EStaticMeshReductionTerimationCriterion::Vertices && (ReductionSettings.PercentTriangles < 1.0f);
+	return bTriTermination || bVertTermination || (ReductionSettings.MaxDeviation > 0.0f);
+}
+
 /* This function is call before building the mesh when we do a re-import*/
 void RestoreExistingMeshSettings(ExistingStaticMeshData* ExistingMesh, UStaticMesh* NewMesh, int32 LODIndex)
 {
@@ -1230,7 +1251,9 @@ void RestoreExistingMeshSettings(ExistingStaticMeshData* ExistingMesh, UStaticMe
 			}
 			FMeshDescription* LODMeshDescription = NewMesh->GetOriginalMeshDescription(i);
 			bool bSwapFromGeneratedToImported = !ExistingMesh->ExistingLODData[i].ExistingMeshDescription.IsValid() && (LODMeshDescription && LODMeshDescription->Polygons().Num() > 0);
-			bool bWasReduced = ExistingMesh->ExistingLODData[i].ExistingReductionSettings.PercentTriangles < 1.0f || ExistingMesh->ExistingLODData[i].ExistingReductionSettings.MaxDeviation > 0.0f;
+
+			bool bWasReduced = IsReductionActive(ExistingMesh->ExistingLODData[i].ExistingReductionSettings);
+
 			if (!bSwapFromGeneratedToImported && bWasReduced)
 			{
 				NewMesh->SourceModels[i].ReductionSettings = ExistingMesh->ExistingLODData[i].ExistingReductionSettings;
@@ -1247,7 +1270,7 @@ void RestoreExistingMeshSettings(ExistingStaticMeshData* ExistingMesh, UStaticMe
 		{
 			FMeshDescription* LODMeshDescription = NewMesh->GetOriginalMeshDescription(LODIndex);
 			bool bSwapFromGeneratedToImported = !ExistingMesh->ExistingLODData[LODIndex].ExistingMeshDescription.IsValid() && (LODMeshDescription && LODMeshDescription->Polygons().Num() > 0);
-			bool bWasReduced = ExistingMesh->ExistingLODData[LODIndex].ExistingReductionSettings.PercentTriangles < 1.0f || ExistingMesh->ExistingLODData[LODIndex].ExistingReductionSettings.MaxDeviation > 0.0f;
+			bool bWasReduced = IsReductionActive(ExistingMesh->ExistingLODData[LODIndex].ExistingReductionSettings);
 			if (!bSwapFromGeneratedToImported && bWasReduced)
 			{
 				NewMesh->SourceModels[LODIndex].ReductionSettings = ExistingMesh->ExistingLODData[LODIndex].ExistingReductionSettings;
@@ -1448,7 +1471,7 @@ void RestoreExistingMeshData(ExistingStaticMeshData* ExistingMeshDataPtr, UStati
 		FMeshDescription* LODMeshDescription = NewMesh->GetOriginalMeshDescription(i);
 		//Restore the reduction settings only if the existing data was a using reduction. Because we can set some value if we reimport from existing rawmesh to auto generated.
 		bool bSwapFromGeneratedToImported = !ExistingMeshDataPtr->ExistingLODData[i].ExistingMeshDescription.IsValid() && (LODMeshDescription && LODMeshDescription->Polygons().Num() > 0);
-		bool bWasReduced = ExistingMeshDataPtr->ExistingLODData[i].ExistingReductionSettings.PercentTriangles < 1.0f || ExistingMeshDataPtr->ExistingLODData[i].ExistingReductionSettings.MaxDeviation > 0.0f;
+		bool bWasReduced = IsReductionActive(ExistingMeshDataPtr->ExistingLODData[i].ExistingReductionSettings);
 		if ( !bSwapFromGeneratedToImported && bWasReduced)
 		{
 			NewMesh->SourceModels[i].ReductionSettings = ExistingMeshDataPtr->ExistingLODData[i].ExistingReductionSettings;
@@ -1490,7 +1513,7 @@ void RestoreExistingMeshData(ExistingStaticMeshData* ExistingMeshDataPtr, UStati
 			}
 
 			//When re-importing the asset, do not touch the LOD that was imported from file, the material array is keep intact so the section should still be valid.
-			bool NoRemapForThisLOD = LodLevel == INDEX_NONE && i != 0 && NewMesh->SourceModels[i].ReductionSettings.PercentTriangles >= 1.0f  && NewMesh->SourceModels[i].ReductionSettings.MaxDeviation <= 0.0f;
+			bool NoRemapForThisLOD = LodLevel == INDEX_NONE && i != 0 && !IsReductionActive(NewMesh->SourceModels[i].ReductionSettings);
 
 			FStaticMeshLODResources& LOD = NewMesh->RenderData->LODResources[i];
 			
diff --git a/Engine/Source/Runtime/Engine/Classes/Engine/StaticMesh.h b/Engine/Source/Runtime/Engine/Classes/Engine/StaticMesh.h
--- a/Engine/Source/Runtime/Engine/Classes/Engine/StaticMesh.h
+++ b/Engine/Source/Runtime/Engine/Classes/Engine/StaticMesh.h
@@ -1150,6 +1150,12 @@ public:
 	/** Get multicast delegate broadcast after mesh building */
 	FOnPostMeshBuild& OnPostMeshBuild() { return PostMeshBuild; }
 
+	/* Return true if the reduction settings are setup to reduce a LOD*/
+	ENGINE_API bool IsReductionActive(int32 LODIndex) const;
+
+	/* Get a copy of the reduction settings for a specified LOD index. */
+	ENGINE_API struct FMeshReductionSettings GetReductionSettings(int32 LODIndex) const;
+
 private:
 	/**
 	 * Converts legacy LODDistance in the source models to Display Factor
diff --git a/Engine/Source/Runtime/Engine/Private/StaticMesh.cpp b/Engine/Source/Runtime/Engine/Private/StaticMesh.cpp
--- a/Engine/Source/Runtime/Engine/Private/StaticMesh.cpp
+++ b/Engine/Source/Runtime/Engine/Private/StaticMesh.cpp
@@ -52,6 +52,9 @@
 #include "ScopedTransaction.h"
 #include "IMeshBuilderModule.h"
 #include "MeshDescriptionOperations.h"
+#include "IMeshReductionManagerModule.h"
+#include "IMeshReductionInterfaces.h"
+
 #endif // #if WITH_EDITOR
 
 #include "MeshDescription.h"
@@ -1312,6 +1315,34 @@ void UStaticMesh::GetLODGroupsDisplayNames(TArray<FText>& OutLODGroupsDisplayNam
 	RunningPlatform->GetStaticMeshLODSettings().GetLODGroupDisplayNames(OutLODGroupsDisplayNames);
 }
 
+bool UStaticMesh::IsReductionActive(int32 LODIndex) const
+{
+	FMeshReductionSettings ReductionSettings = GetReductionSettings(LODIndex);
+	
+	// Are we using our tool, or simplygon?  The tool is only changed during editor restarts
+	IMeshReduction* ReductionModule = FModuleManager::Get().LoadModuleChecked<IMeshReductionManagerModule>("MeshReductionInterface").GetStaticMeshReductionInterface();
+	FString VersionString = ReductionModule->GetVersionString();
+	TArray<FString> SplitVersionString;
+	VersionString.ParseIntoArray(SplitVersionString, TEXT("_"), true);
+	bool bUseQuadricSimplier = SplitVersionString[0].Equals("QuadricMeshReduction");
+
+	const bool bVertTermination = (bUseQuadricSimplier) && (ReductionSettings.TerminationCriterion != EStaticMeshReductionTerimationCriterion::Triangles) && (ReductionSettings.PercentVertices < 1.0f);
+	const bool bTriTermination = ReductionSettings.TerminationCriterion != EStaticMeshReductionTerimationCriterion::Vertices && (ReductionSettings.PercentTriangles < 1.0f);
+
+	return bTriTermination || bVertTermination || (ReductionSettings.MaxDeviation > 0.0f);
+}
+
+FMeshReductionSettings UStaticMesh::GetReductionSettings(int32 LODIndex) const
+{
+	//Retrieve the reduction settings, make sure we use the LODGroup if the Group is valid
+	ITargetPlatformManagerModule& TargetPlatformManager = GetTargetPlatformManagerRef();
+	ITargetPlatform* RunningPlatform = TargetPlatformManager.GetRunningTargetPlatform();
+	check(RunningPlatform);
+	const FStaticMeshLODSettings& LODSettings = RunningPlatform->GetStaticMeshLODSettings();
+	const FStaticMeshLODGroup& SMLODGroup = LODSettings.GetLODGroup(LODGroup);
+	const FStaticMeshSourceModel& SrcModel = SourceModels[LODIndex];
+	return SMLODGroup.GetSettings(SrcModel.ReductionSettings, LODIndex);
+}
 
 void UStaticMesh::PostDuplicate(bool bDuplicateForPIE)
 {
@@ -1336,6 +1367,7 @@ static void SerializeReductionSettingsForDDC(FArchive& Ar, FMeshReductionSetting
 	Ar << ReductionSettings.SilhouetteImportance;
 	Ar << ReductionSettings.TextureImportance;
 	Ar << ReductionSettings.ShadingImportance;
+	Ar << ReductionSettings.BaseLODModel;
 	FArchive_Serialize_BitfieldBool(Ar, ReductionSettings.bRecalculateNormals);
 }
 
@@ -2285,6 +2317,19 @@ void UStaticMesh::SetNumSourceModels(const int32 Num)
 	for (int32 Index = OldNum; Index < Num; ++Index)
 	{
 		SourceModels[Index].StaticMeshOwner = this;
+		int32 PreviousCustomLODIndex = 0;
+		//Find the previous custom LOD
+		for (int32 ReverseIndex = Index - 1; ReverseIndex > 0; ReverseIndex--)
+		{
+			const FStaticMeshSourceModel& StaticMeshModel = SourceModels[ReverseIndex];
+			//If the custom import LOD is reduce and is not using himself as the source, do not consider it
+			if (GetOriginalMeshDescription(ReverseIndex) != nullptr && !(IsReductionActive(ReverseIndex) && StaticMeshModel.ReductionSettings.BaseLODModel != ReverseIndex))
+			{
+				PreviousCustomLODIndex = ReverseIndex;
+				break;
+			}
+		}
+		SourceModels[Index].ReductionSettings.BaseLODModel = PreviousCustomLODIndex;
 	}
 }
 
