Make sure the staticmesh thumbnail always display the LOD 0. force the LOD in the temporary staticmesh component of the thumbnail scene.

diff --git a/Engine/Source/Editor/UnrealEd/Private/ThumbnailHelpers.cpp b/Engine/Source/Editor/UnrealEd/Private/ThumbnailHelpers.cpp
--- a/Engine/Source/Editor/UnrealEd/Private/ThumbnailHelpers.cpp
+++ b/Engine/Source/Editor/UnrealEd/Private/ThumbnailHelpers.cpp
@@ -497,6 +497,10 @@ void FStaticMeshThumbnailScene::SetStaticMesh(UStaticMesh* StaticMesh)
 		FTransform MeshTransform = FTransform::Identity;
 
 		PreviewActor->SetActorLocation(FVector(0,0,0), false);
+		
+		//Force LOD 0
+		PreviewActor->GetStaticMeshComponent()->ForcedLodModel = 1;
+
 		PreviewActor->GetStaticMeshComponent()->UpdateBounds();
 
 		// Center the mesh at the world origin then offset to put it on top of the plane
