New instances should not be in the set of objects that we avoid searching for references in when recreating instances after a class layout change. New instances that had references to old instances were sometimes not being updated

diff --git a/Engine/Source/Editor/UnrealEd/Private/Kismet2/KismetReinstanceUtilities.cpp b/Engine/Source/Editor/UnrealEd/Private/Kismet2/KismetReinstanceUtilities.cpp
--- a/Engine/Source/Editor/UnrealEd/Private/Kismet2/KismetReinstanceUtilities.cpp
+++ b/Engine/Source/Editor/UnrealEd/Private/Kismet2/KismetReinstanceUtilities.cpp
@@ -2056,12 +2056,7 @@ void FBlueprintCompileReinstancer::ReplaceInstancesOfClass_Inner(TMap<UClass*, U
 
 	// Now replace any pointers to the old archetypes/instances with pointers to the new one
 	TArray<UObject*> SourceObjects;
-	TArray<UObject*> DstObjects;
-
 	OldToNewInstanceMap.GenerateKeyArray(SourceObjects);
-	OldToNewInstanceMap.GenerateValueArray(DstObjects); // Also look for references in new spawned objects.
-
-	SourceObjects.Append(DstObjects);
 	
 	if (InOriginalCDO)
 	{
