During asset registry generation in the cooker, pass the list of cooked files to the AssetManager when doing PreSaveAssetRegistry()
Change TEncryptedAssetSet to be a map of group-name to a set of package names, rather than primary asset ids

diff --git a/Engine/Source/Editor/UnrealEd/Private/Commandlets/AssetRegistryGenerator.cpp b/Engine/Source/Editor/UnrealEd/Private/Commandlets/AssetRegistryGenerator.cpp
--- a/Engine/Source/Editor/UnrealEd/Private/Commandlets/AssetRegistryGenerator.cpp
+++ b/Engine/Source/Editor/UnrealEd/Private/Commandlets/AssetRegistryGenerator.cpp
@@ -431,45 +431,40 @@ void FAssetRegistryGenerator::InjectEncryptionData(FAssetRegistryState& TargetSt
 		for (TEncryptedAssetSet::ElementType EncryptedAssetSetElement : EncryptedAssetSet)
 		{
 			FName SetName = EncryptedAssetSetElement.Key;
-			TSet<FPrimaryAssetId>& EncryptedPrimaryAssets = EncryptedAssetSetElement.Value;
+			TSet<FName>& EncryptedRootAssets = EncryptedAssetSetElement.Value;
 
-			for (FPrimaryAssetId EncryptedPrimaryAsset : EncryptedPrimaryAssets)
+			for (FName EncryptedRootPackageName : EncryptedRootAssets)
 			{
-				FAssetData PrimaryAssetData;
-				if (AssetManager.GetPrimaryAssetData(EncryptedPrimaryAsset, PrimaryAssetData))
+				FAssetData* AssetData = const_cast<FAssetData*>(TargetState.GetAssetByObjectPath(EncryptedRootPackageName));
+
+				if (AssetData)
 				{
-					FName EncryptedPrimaryAssetPackageName = PrimaryAssetData.PackageName;
-					FAssetData* AssetData = const_cast<FAssetData*>(TargetState.GetAssetByObjectPath(PrimaryAssetData.ObjectPath));
+					FString GuidString;
 
-					if (AssetData)
+					if (AssetData->ChunkIDs.Num() > 1)
+					{
+						UE_LOG(LogAssetRegistryGenerator, Error, TEXT("Encrypted primary asset '%s' exists in two chunks. Only secondary assets should be shared between chunks."));
+					}
+					else if (AssetData->ChunkIDs.Num() == 1)
 					{
-						FString GuidString;
+						int32 ChunkID = AssetData->ChunkIDs[0];
+						FGuid Guid;
 
-						if (AssetData->ChunkIDs.Num() > 1)
+						if (GuidCache.Contains(ChunkID))
 						{
-							UE_LOG(LogAssetRegistryGenerator, Error, TEXT("Encrypted primary asset '%s' exists in two chunks. Only secondary assets should be shared between chunks."));
+							Guid = GuidCache[ChunkID];
 						}
-						else if (AssetData->ChunkIDs.Num() == 1)
+						else
 						{
-							int32 ChunkID = AssetData->ChunkIDs[0];
-							FGuid Guid;
-
-							if (GuidCache.Contains(ChunkID))
-							{
-								Guid = GuidCache[ChunkID];
-							}
-							else
-							{
-								Guid = GuidCache.Add(ChunkID, AssetManager.GetChunkEncryptionKeyGuid(ChunkID));
-							}
+							Guid = GuidCache.Add(ChunkID, AssetManager.GetChunkEncryptionKeyGuid(ChunkID));
+						}
 
-							if (Guid.IsValid())
-							{
-								FAssetDataTagMap TagsAndValues = AssetData->TagsAndValues.GetMap();
-								TagsAndValues.Add(UAssetManager::GetEncryptionKeyAssetTagName(), Guid.ToString());
-								FAssetData NewAssetData = FAssetData(AssetData->PackageName, AssetData->PackagePath, AssetData->AssetName, AssetData->AssetClass, TagsAndValues, AssetData->ChunkIDs, AssetData->PackageFlags);
-								TargetState.UpdateAssetData(AssetData, NewAssetData);
-							}
+						if (Guid.IsValid())
+						{
+							FAssetDataTagMap TagsAndValues = AssetData->TagsAndValues.GetMap();
+							TagsAndValues.Add(UAssetManager::GetEncryptionKeyAssetTagName(), Guid.ToString());
+							FAssetData NewAssetData = FAssetData(AssetData->PackageName, AssetData->PackagePath, AssetData->AssetName, AssetData->AssetClass, TagsAndValues, AssetData->ChunkIDs, AssetData->PackageFlags);
+							TargetState.UpdateAssetData(AssetData, NewAssetData);
 						}
 					}
 				}
@@ -803,11 +798,11 @@ void FAssetRegistryGenerator::BuildChunkManifest(const TSet<FName>& InCookedPack
 
 }
 
-void FAssetRegistryGenerator::PreSave()
+void FAssetRegistryGenerator::PreSave(const TSet<FName>& InCookedPackages)
 {
 	if (bUseAssetManager)
 	{
-		UAssetManager::Get().PreSaveAssetRegistry(TargetPlatform);
+		UAssetManager::Get().PreSaveAssetRegistry(TargetPlatform, InCookedPackages);
 	}
 }
 
diff --git a/Engine/Source/Editor/UnrealEd/Private/CookOnTheFlyServer.cpp b/Engine/Source/Editor/UnrealEd/Private/CookOnTheFlyServer.cpp
--- a/Engine/Source/Editor/UnrealEd/Private/CookOnTheFlyServer.cpp
+++ b/Engine/Source/Editor/UnrealEd/Private/CookOnTheFlyServer.cpp
@@ -6337,7 +6337,7 @@ void UCookOnTheFlyServer::CookByTheBookFinished()
 					IgnorePackageNames.Add(UncookedEditorOnlyPackage);
 				}
 				{
-					Generator.PreSave();
+					Generator.PreSave(CookedPackageNames);
 				}
 				{
 					SCOPE_TIMER(BuildChunkManifest);
diff --git a/Engine/Source/Editor/UnrealEd/Public/Commandlets/AssetRegistryGenerator.h b/Engine/Source/Editor/UnrealEd/Public/Commandlets/AssetRegistryGenerator.h
--- a/Engine/Source/Editor/UnrealEd/Public/Commandlets/AssetRegistryGenerator.h
+++ b/Engine/Source/Editor/UnrealEd/Public/Commandlets/AssetRegistryGenerator.h
@@ -71,7 +71,7 @@ public:
 	* PreSave
 	* Notify generator that we are about to save the registry and chunk manifests
 	*/
-	void PreSave();
+	void PreSave(const TSet<FName>& InCookedPackages);
 
 	/**
 	* PostSave
diff --git a/Engine/Source/Runtime/Engine/Classes/Engine/AssetManager.h b/Engine/Source/Runtime/Engine/Classes/Engine/AssetManager.h
--- a/Engine/Source/Runtime/Engine/Classes/Engine/AssetManager.h
+++ b/Engine/Source/Runtime/Engine/Classes/Engine/AssetManager.h
@@ -22,7 +22,7 @@ struct FPrimaryAssetRulesCustomOverride;
 DECLARE_DELEGATE_OneParam(FAssetManagerAcquireResourceDelegate, bool);
 
 /** Type that maps a named group to a set of primary assets */
-typedef TMap <FName, TSet<FPrimaryAssetId>> TEncryptedAssetSet;
+typedef TMap <FName, TSet<FName>> TEncryptedAssetSet;
 
 /** 
  * A singleton UObject that is responsible for loading and unloading PrimaryAssets, and maintaining game-specific asset references
@@ -479,7 +479,7 @@ public:
 	/** 
 	  * Called immediately before saving the asset registry during cooking
 	  */
-	virtual void PreSaveAssetRegistry(const class ITargetPlatform* TargetPlatform) {}
+	virtual void PreSaveAssetRegistry(const class ITargetPlatform* TargetPlatform, const TSet<FName>& InCookedPackages) {}
 
 	/** 
 	  * Called immediately after saving the asset registry during cooking
