Moving over:
'Lighting not building in sublevels using Resave Packages
#fix Array.Contains check was using the incorrect string to compare against
#misc Changed HLOD rebuild behaviour to not impact overall package saving/checking out when also rebuilding other data within the same ResaveCommandlet run

diff --git a/Engine/Source/Editor/UnrealEd/Private/Commandlets/ContentCommandlets.cpp b/Engine/Source/Editor/UnrealEd/Private/Commandlets/ContentCommandlets.cpp
--- a/Engine/Source/Editor/UnrealEd/Private/Commandlets/ContentCommandlets.cpp
+++ b/Engine/Source/Editor/UnrealEd/Private/Commandlets/ContentCommandlets.cpp
@@ -1349,7 +1349,8 @@ void UResavePackagesCommandlet::PerformAdditionalOperations(class UWorld* World,
 	ABrush::OnRebuildDone();
 
 	const bool bShouldBuildTextureStreamingForWorld = bShouldBuildTextureStreaming && !bShouldBuildTextureStreamingForAll;
-	if (bShouldBuildLighting || bShouldBuildTextureStreamingForWorld || bShouldBuildHLOD || bShouldBuildReflectionCaptures)
+	const bool bBuildingNonHLODData = (bShouldBuildLighting || bShouldBuildTextureStreamingForWorld || bShouldBuildReflectionCaptures);
+	if (bBuildingNonHLODData || bShouldBuildHLOD)
 	{
 		bool bShouldProceedWithRebuild = true;
 
@@ -1407,12 +1408,13 @@ void UResavePackagesCommandlet::PerformAdditionalOperations(class UWorld* World,
 		FString WorldPackageCheckedOutUser;
 		if (FPackageName::DoesPackageExist(World->GetOutermost()->GetName(), NULL, &WorldPackageName))
 		{
-			if(bShouldBuildHLOD)
+			// If we are only building HLODs check if level can be checked out, if so add to list of files that will be saved/checked-out after rebuilding the data
+			if(bShouldBuildHLOD && !bBuildingNonHLODData)
 			{
 				if (CanCheckoutFile(WorldPackageName, WorldPackageCheckedOutUser) || !bSkipCheckedOutFiles)
 				{
-					SublevelFilenames.Add(WorldPackageName);
-				}
+				SublevelFilenames.Add(WorldPackageName);
+			}
 				else 
 				{
 					bShouldProceedWithRebuild = false;
@@ -1444,8 +1446,8 @@ void UResavePackagesCommandlet::PerformAdditionalOperations(class UWorld* World,
 
 			for (ULevelStreaming* NextStreamingLevel : World->GetStreamingLevels())
 			{
-				// If we are building HLODs, we dont check out ahead of time
-				if(!bShouldBuildHLOD)
+				// If we are not building HLODs or are but also rebuilding lighting we check out the level file, otherwise we don't to try and ensure a minimal HLOD rebuild
+				if (!bShouldBuildHLOD || bBuildingNonHLODData)
 				{
 					CheckOutLevelFile(NextStreamingLevel->GetLoadedLevel());
 				}
@@ -1454,14 +1456,14 @@ void UResavePackagesCommandlet::PerformAdditionalOperations(class UWorld* World,
 				const FString StreamingLevelWorldAssetPackageName = NextStreamingLevel->GetWorldAssetPackageName();
 				if (FPackageName::DoesPackageExist(StreamingLevelWorldAssetPackageName, NULL, &StreamingLevelPackageFilename))
 				{
-					// If we are building HLODs, we dont check out ahead of time
-					if(bShouldBuildHLOD)
+					// If we are building HLODs only, we dont check out the files ahead of rebuilding the data
+					if(bShouldBuildHLOD && !bBuildingNonHLODData)
 					{
 						FString CurrentlyCheckedOutUser;
 						if (CanCheckoutFile(StreamingLevelPackageFilename, CurrentlyCheckedOutUser) || !bSkipCheckedOutFiles)
 						{
-							SublevelFilenames.Add(StreamingLevelPackageFilename);
-						}
+						SublevelFilenames.Add(StreamingLevelPackageFilename);
+					}
 						else 
 						{
 							UE_LOG(LogContentCommandlet, Warning, TEXT("[REPORT] Skipping %s as it is checked out by %s"), *StreamingLevelPackageFilename, *CurrentlyCheckedOutUser);
@@ -1476,6 +1478,7 @@ void UResavePackagesCommandlet::PerformAdditionalOperations(class UWorld* World,
 						}
 						else
 						{
+							UE_LOG(LogContentCommandlet, Error, TEXT("[REPORT] %s is currently already checked out, cannot continue resaving"), *StreamingLevelPackageFilename);
 							bShouldProceedWithRebuild = false;
 							break;
 						}
@@ -1688,11 +1691,11 @@ void UResavePackagesCommandlet::PerformAdditionalOperations(class UWorld* World,
 				{
 					FString StreamingLevelPackageFilename;
 					const FString StreamingLevelWorldAssetPackageName = NextStreamingLevel->GetWorldAssetPackageName();
-					if (FPackageName::DoesPackageExist(StreamingLevelWorldAssetPackageName, NULL, &StreamingLevelPackageFilename) && SublevelFilenames.Contains(StreamingLevelWorldAssetPackageName))
+					if (FPackageName::DoesPackageExist(StreamingLevelWorldAssetPackageName, NULL, &StreamingLevelPackageFilename) && SublevelFilenames.Contains(StreamingLevelPackageFilename))
 					{
 						UPackage* SubLevelPackage = NextStreamingLevel->GetLoadedLevel()->GetOutermost();
 						bool bSaveSubLevelPackage = true;
-						if(bShouldBuildHLOD)
+						if(bShouldBuildHLOD && !bBuildingNonHLODData)
 						{
 							// If we are building HLOD, only save packages that were dirtied
 							bSaveSubLevelPackage = SubLevelPackage->IsDirty();
@@ -1702,12 +1705,14 @@ void UResavePackagesCommandlet::PerformAdditionalOperations(class UWorld* World,
 						{
 							// When building HLODs we dont check out/modify maps unless dirty
 							bool bFileCheckedOut = true;
-							if(bShouldBuildHLOD)
+							if(bShouldBuildHLOD && !bBuildingNonHLODData)
 							{
 								bFileCheckedOut = CheckoutFile(StreamingLevelPackageFilename, true);
 							}
 
-							if (!bFileCheckedOut || !SavePackageHelper(SubLevelPackage, StreamingLevelPackageFilename))
+							// Try to save the level package 
+							const bool bSavePackageResult = SavePackageHelper(SubLevelPackage, StreamingLevelPackageFilename);
+							if (!bFileCheckedOut && !bSavePackageResult)
 							{
 								UE_LOG(LogContentCommandlet, Error, TEXT("[REPORT] Failed to save sub level: %s"), *StreamingLevelPackageFilename);
 							}
@@ -1726,8 +1731,8 @@ void UResavePackagesCommandlet::PerformAdditionalOperations(class UWorld* World,
 			}
 			else
 			{
-				UE_LOG(LogContentCommandlet, Error, TEXT("[REPORT] Failed to complete steps necessary to start a lightmass or texture streaming build of %s"), *World->GetName());
-			}
+			UE_LOG(LogContentCommandlet, Error, TEXT("[REPORT] Failed to complete steps necessary to start a lightmass or texture streaming build of %s"), *World->GetName());
+		}
 		}
 
 		if ((bShouldProceedWithRebuild == false)||(bSavePackage == false))
@@ -1751,9 +1756,9 @@ void UResavePackagesCommandlet::PerformAdditionalOperations(class UWorld* World,
 				FilesToSubmit.AddUnique(SublevelFilename);
 			}
 
-			if(bShouldBuildHLOD)
+			if(bShouldBuildHLOD && !bBuildingNonHLODData)
 			{
-				// Don't save outer package if it isn't dirty
+				// Don't save outer package if it isn't dirty when doing a HLOD rebuild only
 				bSavePackage = World->GetOutermost()->IsDirty();
 			}
 		}
