Fixed a bug in asset registry generation where objects were being incorrectly looked up by ObjectPath in the registry state, using a PackageName. This would stop metadata from being properly embedded in the asset registry data

diff --git a/Engine/Source/Editor/UnrealEd/Private/Commandlets/AssetRegistryGenerator.cpp b/Engine/Source/Editor/UnrealEd/Private/Commandlets/AssetRegistryGenerator.cpp
--- a/Engine/Source/Editor/UnrealEd/Private/Commandlets/AssetRegistryGenerator.cpp
+++ b/Engine/Source/Editor/UnrealEd/Private/Commandlets/AssetRegistryGenerator.cpp
@@ -435,36 +435,41 @@ void FAssetRegistryGenerator::InjectEncryptionData(FAssetRegistryState& TargetSt
 
 			for (FName EncryptedRootPackageName : EncryptedRootAssets)
 			{
-				FAssetData* AssetData = const_cast<FAssetData*>(TargetState.GetAssetByObjectPath(EncryptedRootPackageName));
+				const TArray<const FAssetData*>& PackageAssets = TargetState.GetAssetsByPackageName(EncryptedRootPackageName);
 
-				if (AssetData)
+				for (const FAssetData* PackageAsset : PackageAssets)
 				{
-					FString GuidString;
+					FAssetData* AssetData = const_cast<FAssetData*>(PackageAsset);
 
-					if (AssetData->ChunkIDs.Num() > 1)
+					if (AssetData)
 					{
-						UE_LOG(LogAssetRegistryGenerator, Error, TEXT("Encrypted primary asset '%s' exists in two chunks. Only secondary assets should be shared between chunks."));
-					}
-					else if (AssetData->ChunkIDs.Num() == 1)
-					{
-						int32 ChunkID = AssetData->ChunkIDs[0];
-						FGuid Guid;
+						FString GuidString;
 
-						if (GuidCache.Contains(ChunkID))
+						if (AssetData->ChunkIDs.Num() > 1)
 						{
-							Guid = GuidCache[ChunkID];
+							UE_LOG(LogAssetRegistryGenerator, Error, TEXT("Encrypted primary asset '%s' exists in two chunks. Only secondary assets should be shared between chunks."));
 						}
-						else
+						else if (AssetData->ChunkIDs.Num() == 1)
 						{
-							Guid = GuidCache.Add(ChunkID, AssetManager.GetChunkEncryptionKeyGuid(ChunkID));
-						}
+							int32 ChunkID = AssetData->ChunkIDs[0];
+							FGuid Guid;
 
-						if (Guid.IsValid())
-						{
-							FAssetDataTagMap TagsAndValues = AssetData->TagsAndValues.GetMap();
-							TagsAndValues.Add(UAssetManager::GetEncryptionKeyAssetTagName(), Guid.ToString());
-							FAssetData NewAssetData = FAssetData(AssetData->PackageName, AssetData->PackagePath, AssetData->AssetName, AssetData->AssetClass, TagsAndValues, AssetData->ChunkIDs, AssetData->PackageFlags);
-							TargetState.UpdateAssetData(AssetData, NewAssetData);
+							if (GuidCache.Contains(ChunkID))
+							{
+								Guid = GuidCache[ChunkID];
+							}
+							else
+							{
+								Guid = GuidCache.Add(ChunkID, AssetManager.GetChunkEncryptionKeyGuid(ChunkID));
+							}
+
+							if (Guid.IsValid())
+							{
+								FAssetDataTagMap TagsAndValues = AssetData->TagsAndValues.GetMap();
+								TagsAndValues.Add(UAssetManager::GetEncryptionKeyAssetTagName(), Guid.ToString());
+								FAssetData NewAssetData = FAssetData(AssetData->PackageName, AssetData->PackagePath, AssetData->AssetName, AssetData->AssetClass, TagsAndValues, AssetData->ChunkIDs, AssetData->PackageFlags);
+								TargetState.UpdateAssetData(AssetData, NewAssetData);
+							}
 						}
 					}
 				}
