Changes to how the asset manager specifies which assets should be encrypted
 - UAssetManager::GetEncryptedAssetSet() has changed to UAssetManager::GetContentEncryptionConfig()  which takes a FContentEncryptionConfig object and fills it out with details of what should be encrypted
 - Allow specification of non-UFS files that require encryption
 - AssetRegistryGenerator generates new manifests for non-UFS files that require encryption

diff --git a/Engine/Source/Editor/UnrealEd/Private/Commandlets/AssetRegistryGenerator.cpp b/Engine/Source/Editor/UnrealEd/Private/Commandlets/AssetRegistryGenerator.cpp
--- a/Engine/Source/Editor/UnrealEd/Private/Commandlets/AssetRegistryGenerator.cpp
+++ b/Engine/Source/Editor/UnrealEd/Private/Commandlets/AssetRegistryGenerator.cpp
@@ -216,6 +216,30 @@ bool FAssetRegistryGenerator::GenerateStreamingInstallManifest(int64 InExtraFlav
 		PlatformIniFile.GetArray(TEXT("/Script/UnrealEd.ProjectPackagingSettings"), TEXT("CompressedChunkWildcard"), CompressedChunkWildcards);
 	}
 
+	// Add manifests for any non-ufs file groups
+	if (bUseAssetManager && !TargetPlatform->HasSecurePackageFormat())
+	{
+		FContentEncryptionConfig ContentEncryptionConfig;
+		UAssetManager::Get().GetContentEncryptionConfig(ContentEncryptionConfig);
+		const TMap<FName, TSet<FName>>& EncryptedNonUFSFileGroups = ContentEncryptionConfig.GetNonUFSFileGroupMap();
+		
+		for (const TMap<FName, TSet<FName>>::ElementType& Element : EncryptedNonUFSFileGroups)
+		{
+			FName GroupName = Element.Key;
+			const TSet<FName>& Files = Element.Value;
+			FChunkPackageSet* NewManifest = new FChunkPackageSet();
+			
+			for (FName File : Files)
+			{
+				FString Filename = File.ToString();
+				FName LongPackageName = *FPackageName::FilenameToLongPackageName(Filename);
+				NewManifest->Add(LongPackageName, Filename);
+			}
+
+			FinalChunkManifests.Add(NewManifest);
+		}
+	}
+
 
 	// generate per-chunk pak list files
 	for (int32 Index = 0; Index < FinalChunkManifests.Num(); ++Index)
@@ -424,11 +448,10 @@ void FAssetRegistryGenerator::InjectEncryptionData(FAssetRegistryState& TargetSt
 		UAssetManager& AssetManager = UAssetManager::Get();
 
 		TMap<int32, FGuid> GuidCache;
-		TEncryptedAssetSet EncryptedAssetSet;
-		TSet<FGuid> ReleasedAssets;
-		AssetManager.GetEncryptedAssetSet(EncryptedAssetSet, ReleasedAssets);
+		FContentEncryptionConfig EncryptionConfig;
+		AssetManager.GetContentEncryptionConfig(EncryptionConfig);
 
-		for (TEncryptedAssetSet::ElementType EncryptedAssetSetElement : EncryptedAssetSet)
+		for (TMap<FName, TSet<FName>>::ElementType EncryptedAssetSetElement : EncryptionConfig.GetPackageGroupMap())
 		{
 			FName SetName = EncryptedAssetSetElement.Key;
 			TSet<FName>& EncryptedRootAssets = EncryptedAssetSetElement.Value;
diff --git a/Engine/Source/Runtime/Engine/Classes/Engine/AssetManager.h b/Engine/Source/Runtime/Engine/Classes/Engine/AssetManager.h
--- a/Engine/Source/Runtime/Engine/Classes/Engine/AssetManager.h
+++ b/Engine/Source/Runtime/Engine/Classes/Engine/AssetManager.h
@@ -9,6 +9,7 @@
 #include "AssetBundleData.h"
 #include "AssetRegistryModule.h"
 #include "GenericPlatform/GenericPlatformChunkInstall.h"
+#include "ContentEncryptionConfig.h"
 #include "AssetManager.generated.h"
 
 DECLARE_LOG_CATEGORY_EXTERN(LogAssetManager, Log, All);
@@ -21,9 +22,6 @@ struct FPrimaryAssetRulesCustomOverride;
 /** Delegate called when acquiring resources/chunks for assets, parameter will be true if all resources were acquired, false if any failed */
 DECLARE_DELEGATE_OneParam(FAssetManagerAcquireResourceDelegate, bool);
 
-/** Type that maps a named group to a set of primary assets */
-typedef TMap <FName, TSet<FName>> TEncryptedAssetSet;
-
 /** 
  * A singleton UObject that is responsible for loading and unloading PrimaryAssets, and maintaining game-specific asset references
  * Games should override this class and change the class reference
@@ -489,8 +487,7 @@ public:
 	/**
 	  * Gathers information about which assets the game wishes to encrypt into named groups
 	  */
-	virtual void GetEncryptedAssetSet(TEncryptedAssetSet& OutEncryptedAssets, TSet<FGuid>& OutReleasedAssets) {}
-
+	virtual void GetContentEncryptionConfig(FContentEncryptionConfig& OutContentEncryptionConfig) {}
 #endif
 
 protected:
diff --git a/Engine/Source/Runtime/Engine/Classes/Engine/ContentEncryptionConfig.h b/Engine/Source/Runtime/Engine/Classes/Engine/ContentEncryptionConfig.h
new file mode 100644
--- /dev/null
+++ b/Engine/Source/Runtime/Engine/Classes/Engine/ContentEncryptionConfig.h
@@ -0,0 +1,44 @@
+// Copyright 1998-2018 Epic Games, Inc. All Rights Reserved.
+#pragma once
+
+/** Project specific configuration for content encryption */
+class FContentEncryptionConfig
+{
+public:
+
+	void AddPackage(FName InGroupName, FName InPackageName)
+	{
+		PackageGroups.FindOrAdd(InGroupName).Add(InPackageName);
+	}
+
+	void AddNonUFSFile(FName InGroupName, FName InFileName)
+	{
+		NonUFSFileGroups.FindOrAdd(InGroupName).Add(InFileName);
+	}
+
+	void AddReleasedKey(FGuid InKey)
+	{
+		ReleasedKeys.Add(InKey);
+	}
+
+	const TMap<FName, TSet<FName>>& GetPackageGroupMap() const
+	{
+		return PackageGroups;
+	}
+
+	const TMap<FName, TSet<FName>>& GetNonUFSFileGroupMap() const
+	{
+		return NonUFSFileGroups;
+	}
+
+	const TSet<FGuid>& GetReleasedKeys() const
+	{
+		return ReleasedKeys;
+	}
+
+private:
+
+	TMap<FName, TSet<FName>> PackageGroups;
+	TMap<FName, TSet<FName>> NonUFSFileGroups;
+	TSet<FGuid> ReleasedKeys;
+};
\ No newline at end of file
