Restore the reduction settings for the skeletal mesh reimported LOD
Fix a bug in USkeletalMesh::IsReductionActive where it was using the wrong reduction module

diff --git a/Engine/Source/Editor/UnrealEd/Private/Factories/SkeletalMeshImport.cpp b/Engine/Source/Editor/UnrealEd/Private/Factories/SkeletalMeshImport.cpp
--- a/Engine/Source/Editor/UnrealEd/Private/Factories/SkeletalMeshImport.cpp
+++ b/Engine/Source/Editor/UnrealEd/Private/Factories/SkeletalMeshImport.cpp
@@ -427,6 +427,15 @@ ExistingSkelMeshData* SaveExistingSkelMeshData(USkeletalMesh* ExistingSkelMesh,
 			ExistingSkelMesh->GetLODNum() == ImportedResource->LODModels.Num() )
 		{
 			int32 OffsetReductionLODIndex = 0;
+			
+			FSkeletalMeshLODInfo* LODInfo = ExistingSkelMesh->GetLODInfo(ReimportSpecificLOD);
+			ExistingMeshDataPtr->bIsReimportLODReduced = LODInfo->bHasBeenSimplified;
+			if (LODInfo && LODInfo->bHasBeenSimplified)
+			{
+				//Save the imported LOD reduction settings
+				ExistingMeshDataPtr->ExistingReimportLODReductionSettings = LODInfo->ReductionSettings;
+			}
+
 			// Remove the zero'th LOD (ie: the LOD being reimported).
 			if (!ReimportSpecificLOD)
 			{
@@ -795,6 +804,20 @@ void RestoreExistingSkelMeshData(ExistingSkelMeshData* MeshData, USkeletalMesh*
 		SkeletalMesh->LODSettings->SetLODSettingsToMesh(SkeletalMesh, 0);
 	}
 
+	//Copy back the reimport LOD specific data
+	if (SkeletalMesh->GetLODInfoArray().IsValidIndex(ReimportLODIndex) && MeshData->bIsReimportLODReduced)
+	{
+		FSkeletalMeshLODInfo& BaseLODInfo = SkeletalMesh->GetLODInfoArray()[ReimportLODIndex];
+		//Restore the reimport LOD reduction settings
+		BaseLODInfo.ReductionSettings = MeshData->ExistingReimportLODReductionSettings;
+		//Regenerate the reimport LOD
+		GWarn->BeginSlowTask(LOCTEXT("RegenReimportedLOD", "Generating reimported LOD"), true);
+		FSkeletalMeshUpdateContext UpdateContext;
+		UpdateContext.SkeletalMesh = SkeletalMesh;
+		FLODUtilities::SimplifySkeletalMeshLOD(UpdateContext, ReimportLODIndex, false);
+		GWarn->EndSlowTask();
+	}
+
 	//Do everything we need for base LOD re-import
 	if (ReimportLODIndex < 1)
 	{
@@ -1146,7 +1169,7 @@ void RestoreExistingSkelMeshData(ExistingSkelMeshData* MeshData, USkeletalMesh*
 		}
 	}
 
-	if (!bImportSkinningOnly)
+	if (!bImportSkinningOnly && !MeshData->bIsReimportLODReduced)
 	{
 		//Fix the OriginalReductionSourceMeshData
 		if (ReimportLODIndex < 0)
diff --git a/Engine/Source/Editor/UnrealEd/Public/SkelImport.h b/Engine/Source/Editor/UnrealEd/Public/SkelImport.h
--- a/Engine/Source/Editor/UnrealEd/Public/SkelImport.h
+++ b/Engine/Source/Editor/UnrealEd/Public/SkelImport.h
@@ -53,6 +53,15 @@ struct ExistingSkelMeshData
 	USkeletalMeshLODSettings*				ExistingLODSettings;
 	TSubclassOf<UAnimInstance>				ExistingPostProcessAnimBlueprint;
 
+	//////////////////////////////////////////////////////////////////////////
+	//Reimport LOD specific data
+
+	//When the specific LOD is reduce, we want to apply the same reduction after the re-import of the LODs
+	bool bIsReimportLODReduced;
+	FSkeletalMeshOptimizationSettings		ExistingReimportLODReductionSettings;
+	
+	//////////////////////////////////////////////////////////////////////////
+
 	bool									bExistingUseFullPrecisionUVs;
 	bool									bExistingUseHighPrecisionTangentBasis;
 
diff --git a/Engine/Source/Runtime/Engine/Private/SkeletalMesh.cpp b/Engine/Source/Runtime/Engine/Private/SkeletalMesh.cpp
--- a/Engine/Source/Runtime/Engine/Private/SkeletalMesh.cpp
+++ b/Engine/Source/Runtime/Engine/Private/SkeletalMesh.cpp
@@ -379,7 +379,7 @@ void USkeletalMesh::ValidateBoundsExtension()
 bool USkeletalMesh::IsReductionActive(int32 LODIndex) const
 {
 	FSkeletalMeshOptimizationSettings ReductionSettings = GetReductionSettings(LODIndex);
-	IMeshReduction* ReductionModule = FModuleManager::Get().LoadModuleChecked<IMeshReductionManagerModule>("MeshReductionInterface").GetStaticMeshReductionInterface();
+	IMeshReduction* ReductionModule = FModuleManager::Get().LoadModuleChecked<IMeshReductionManagerModule>("MeshReductionInterface").GetSkeletalMeshReductionInterface();
 	return ReductionModule->IsReductionActive(ReductionSettings);
 }
 
