Avoid reduction when the reduction settings are not active.
Fix CIS static analysis

diff --git a/Engine/Source/Editor/Persona/Private/PersonaMeshDetails.cpp b/Engine/Source/Editor/Persona/Private/PersonaMeshDetails.cpp
--- a/Engine/Source/Editor/Persona/Private/PersonaMeshDetails.cpp
+++ b/Engine/Source/Editor/Persona/Private/PersonaMeshDetails.cpp
@@ -1441,6 +1441,12 @@ void FPersonaMeshDetails::RegenerateOneLOD(int32 LODIndex)
 			FLODUtilities::RestoreSkeletalMeshLODImportedData(SkelMesh, LODIndex, true);
 			return;
 		}
+		else if (!CurrentLODInfo.bHasBeenSimplified
+			&& !SkelMesh->IsReductionActive(LODIndex))
+		{
+			//Nothing to reduce
+			return;
+		}
 
 		FSkeletalMeshUpdateContext UpdateContext;
 		UpdateContext.SkeletalMesh = SkelMesh;
diff --git a/Engine/Source/Editor/UnrealEd/Private/Factories/SkeletalMeshImport.cpp b/Engine/Source/Editor/UnrealEd/Private/Factories/SkeletalMeshImport.cpp
--- a/Engine/Source/Editor/UnrealEd/Private/Factories/SkeletalMeshImport.cpp
+++ b/Engine/Source/Editor/UnrealEd/Private/Factories/SkeletalMeshImport.cpp
@@ -429,8 +429,8 @@ ExistingSkelMeshData* SaveExistingSkelMeshData(USkeletalMesh* ExistingSkelMesh,
 			int32 OffsetReductionLODIndex = 0;
 			
 			FSkeletalMeshLODInfo* LODInfo = ExistingSkelMesh->GetLODInfo(ReimportSpecificLOD);
-			ExistingMeshDataPtr->bIsReimportLODReduced = LODInfo->bHasBeenSimplified;
-			if (LODInfo && LODInfo->bHasBeenSimplified)
+			ExistingMeshDataPtr->bIsReimportLODReduced = (LODInfo && LODInfo->bHasBeenSimplified);
+			if (LODInfo && ExistingMeshDataPtr->bIsReimportLODReduced)
 			{
 				//Save the imported LOD reduction settings
 				ExistingMeshDataPtr->ExistingReimportLODReductionSettings = LODInfo->ReductionSettings;
