Make sure the remap morph target work well with inline base LOD reduction

diff --git a/Engine/Source/Editor/Persona/Private/PersonaMeshDetails.cpp b/Engine/Source/Editor/Persona/Private/PersonaMeshDetails.cpp
--- a/Engine/Source/Editor/Persona/Private/PersonaMeshDetails.cpp
+++ b/Engine/Source/Editor/Persona/Private/PersonaMeshDetails.cpp
@@ -1448,6 +1448,19 @@ void FPersonaMeshDetails::RegenerateOneLOD(int32 LODIndex)
 			return;
 		}
 
+		//If we are doing inline reduction, restore the skeletalmesh LOD if it was already reduced and we want to reduced it again (user change the options)
+		//This will ensure the data is the same when starting the reduction
+		if (LODIndex == CurrentLODInfo.ReductionSettings.BaseLOD
+			&& CurrentLODInfo.bHasBeenSimplified
+			&& SkelMesh->IsReductionActive(LODIndex)
+			&& SkelMesh->GetImportedModel()->OriginalReductionSourceMeshData.IsValidIndex(LODIndex)
+			&& !SkelMesh->GetImportedModel()->OriginalReductionSourceMeshData[LODIndex]->IsEmpty())
+		{
+			//Restore the base LOD data
+			CurrentLODInfo.bHasBeenSimplified = false;
+			FLODUtilities::RestoreSkeletalMeshLODImportedData(SkelMesh, LODIndex, false);
+		}
+
 		FSkeletalMeshUpdateContext UpdateContext;
 		UpdateContext.SkeletalMesh = SkelMesh;
 		UpdateContext.AssociatedComponents.Push(GetPersonaToolkit()->GetPreviewMeshComponent());
@@ -1508,7 +1521,7 @@ FReply FPersonaMeshDetails::RegenerateLOD(int32 LODIndex)
 			else if (bIsReductionActive)
 			{
 				//Ask user a special permission when the base LOD can be reduce 
-				const FText Text(LOCTEXT("Warning_ReductionApplyingToImportedMesh_ReduceBaseLOD", "Are you sure you'd like to apply mesh reduction to the base LOD?"));
+				const FText Text(LOCTEXT("Warning_ReductionApplyingToImportedMesh_ReduceBaseLOD", "Are you sure you'd like to apply mesh reduction to the non-generated base LOD?"));
 				EAppReturnType::Type Ret = FMessageDialog::Open(EAppMsgType::YesNo, Text);
 				if (Ret == EAppReturnType::No)
 				{
diff --git a/Engine/Source/Editor/UnrealEd/Private/Factories/SkeletalMeshImport.cpp b/Engine/Source/Editor/UnrealEd/Private/Factories/SkeletalMeshImport.cpp
--- a/Engine/Source/Editor/UnrealEd/Private/Factories/SkeletalMeshImport.cpp
+++ b/Engine/Source/Editor/UnrealEd/Private/Factories/SkeletalMeshImport.cpp
@@ -810,6 +810,10 @@ void RestoreExistingSkelMeshData(ExistingSkelMeshData* MeshData, USkeletalMesh*
 		FSkeletalMeshLODInfo& BaseLODInfo = SkeletalMesh->GetLODInfoArray()[ReimportLODIndex];
 		//Restore the reimport LOD reduction settings
 		BaseLODInfo.ReductionSettings = MeshData->ExistingReimportLODReductionSettings;
+		if (SkeletalMeshImportedModel && SkeletalMeshImportedModel->OriginalReductionSourceMeshData.IsValidIndex(ReimportLODIndex))
+		{
+			SkeletalMeshImportedModel->OriginalReductionSourceMeshData[ReimportLODIndex]->EmptyBulkData();
+		}
 		//Regenerate the reimport LOD
 		GWarn->BeginSlowTask(LOCTEXT("RegenReimportedLOD", "Generating reimported LOD"), true);
 		FSkeletalMeshUpdateContext UpdateContext;
diff --git a/Engine/Source/Editor/UnrealEd/Private/LODUtilities.cpp b/Engine/Source/Editor/UnrealEd/Private/LODUtilities.cpp
--- a/Engine/Source/Editor/UnrealEd/Private/LODUtilities.cpp
+++ b/Engine/Source/Editor/UnrealEd/Private/LODUtilities.cpp
@@ -422,6 +422,11 @@ void CreateLODMorphTarget(USkeletalMesh* SkeletalMesh, FReductionBaseSkeletalMes
 		ReductionBaseSkeletalMeshBulkData->LoadReductionData(TempBaseLODModel, BaseLODMorphTargetData);
 	}
 
+	FSkeletalMeshModel* SkeletalMeshModel = SkeletalMesh->GetImportedModel();
+	const FSkeletalMeshLODModel& TargetLODModel = SkeletalMeshModel->LODModels[DestinationLOD];
+
+	bool bInitializeMorphData = false;
+
 	for (UMorphTarget *MorphTarget : SkeletalMesh->MorphTargets)
 	{
 		if (!MorphTarget->HasDataForLOD(SourceLOD))
@@ -517,11 +522,17 @@ void CreateLODMorphTarget(USkeletalMesh* SkeletalMesh, FReductionBaseSkeletalMes
 			}
 		}
 		
-		FSkeletalMeshModel* SkeletalMeshModel = SkeletalMesh->GetImportedModel();
-		const FSkeletalMeshLODModel& TargetLODModel = SkeletalMeshModel->LODModels[DestinationLOD];
 		//Register the new morph target on the target LOD
 		MorphTarget->PopulateDeltas(NewMorphTargetDeltas, DestinationLOD, TargetLODModel.Sections, false, true);
-		SkeletalMesh->RegisterMorphTarget(MorphTarget);
+		if (MorphTarget->HasValidData())
+		{
+			bInitializeMorphData |= SkeletalMesh->RegisterMorphTarget(MorphTarget, false);
+		}
+	}
+
+	if (bInitializeMorphData)
+	{
+		SkeletalMesh->InitMorphTargetsAndRebuildRenderData();
 	}
 }
 
@@ -545,7 +556,7 @@ void FLODUtilities::ClearGeneratedMorphTarget(USkeletalMesh* SkeletalMesh, int32
 			continue;
 		}
 
-		if (MorphTarget->MorphLODModels[TargetLOD].bGeneratedByEngine)
+		//if (MorphTarget->MorphLODModels[TargetLOD].bGeneratedByEngine)
 		{
 			MorphTarget->MorphLODModels[TargetLOD].Reset();
 
@@ -870,6 +881,7 @@ void FLODUtilities::RestoreSkeletalMeshLODImportedData(USkeletalMesh* SkeletalMe
 		//Copy the SkeletalMeshLODModel
 		SkeletalMesh->GetImportedModel()->LODModels[LodIndex] = ImportedBaseLODModel;
 		//Copy the morph target deltas
+		bool bInitMorphTargetData = false;
 		for (UMorphTarget *MorphTarget : SkeletalMesh->MorphTargets)
 		{
 			if (!ImportedBaseLODMorphTargetData.Contains(MorphTarget->GetFullName()))
@@ -879,8 +891,9 @@ void FLODUtilities::RestoreSkeletalMeshLODImportedData(USkeletalMesh* SkeletalMe
 			TArray<FMorphTargetDelta>& ImportedDeltas = ImportedBaseLODMorphTargetData[MorphTarget->GetFullName()];
 
 			MorphTarget->PopulateDeltas(ImportedDeltas, LodIndex, SkeletalMesh->GetImportedModel()->LODModels[LodIndex].Sections, false, false);
-			SkeletalMesh->RegisterMorphTarget(MorphTarget);
+			bInitMorphTargetData |= SkeletalMesh->RegisterMorphTarget(MorphTarget, false);
 		}
+		SkeletalMesh->InitMorphTargetsAndRebuildRenderData();
 		
 		//Empty the bulkdata since we restore it
 		SkeletalMesh->GetImportedModel()->OriginalReductionSourceMeshData[LodIndex]->EmptyBulkData();
