Fbx: Fix skeletal mesh export

* true/false string metadata values on StaticMesh and SkeletalMesh are converted to Boolean attributes when exported to FBX.
* Now, the node name is appended to the metadata tag name only in the case of SkeletalMesh.

Copy from Dev-Enterprise

diff --git a/Engine/Source/Editor/UnrealEd/Private/Fbx/FbxMainExport.cpp b/Engine/Source/Editor/UnrealEd/Private/Fbx/FbxMainExport.cpp
--- a/Engine/Source/Editor/UnrealEd/Private/Fbx/FbxMainExport.cpp
+++ b/Engine/Source/Editor/UnrealEd/Private/Fbx/FbxMainExport.cpp
@@ -3558,11 +3558,22 @@ void FFbxExporter::ExportObjectMetadata(const UObject* ObjectToExport, FbxNode*
 					// Remaining tag follows the format NodeName.PropertyName, so replace '.' with '_'
 					TagAsString.ReplaceInline(TEXT("."), TEXT("_"));
 
-					FbxProperty Property = FbxProperty::Create(Node, FbxStringDT, TCHAR_TO_UTF8(*TagAsString));
-					FbxString ValueString(TCHAR_TO_UTF8(*MetadataIt.Value));
+					if (MetadataIt.Value == TEXT("true") || MetadataIt.Value == TEXT("false"))
+					{
+						FbxProperty Property = FbxProperty::Create(Node, FbxBoolDT, TCHAR_TO_UTF8(*TagAsString));
+						FbxBool ValueBool = MetadataIt.Value == TEXT("true") ? true : false;
+
+						Property.Set(ValueBool);
+						Property.ModifyFlag(FbxPropertyFlags::eUserDefined, true);
+					}
+					else
+					{
+						FbxProperty Property = FbxProperty::Create(Node, FbxStringDT, TCHAR_TO_UTF8(*TagAsString));
+						FbxString ValueString(TCHAR_TO_UTF8(*MetadataIt.Value));
 
-					Property.Set(ValueString);
-					Property.ModifyFlag(FbxPropertyFlags::eUserDefined, true);
+						Property.Set(ValueString);
+						Property.ModifyFlag(FbxPropertyFlags::eUserDefined, true);
+					}
 				}
 			}
 		}
diff --git a/Engine/Source/Editor/UnrealEd/Private/Fbx/FbxMainImport.cpp b/Engine/Source/Editor/UnrealEd/Private/Fbx/FbxMainImport.cpp
--- a/Engine/Source/Editor/UnrealEd/Private/Fbx/FbxMainImport.cpp
+++ b/Engine/Source/Editor/UnrealEd/Private/Fbx/FbxMainImport.cpp
@@ -3107,7 +3107,7 @@ FString GetFbxPropertyStringValue(const FbxProperty& Property)
 	return ValueStr;
 }
 
-void FFbxImporter::ImportNodeCustomProperties(UObject* Object, FbxNode* Node)
+void FFbxImporter::ImportNodeCustomProperties(UObject* Object, FbxNode* Node, bool bPrefixTagWithNodeName)
 {
 	if (!Object || !Node)
 	{
@@ -3125,7 +3125,7 @@ void FFbxImporter::ImportNodeCustomProperties(UObject* Object, FbxNode* Node)
 			// Prefix the FBX metadata tag to make it distinguishable from other metadata
 			// so that it can be exportable through FBX export
 			FString MetadataTag = UTF8_TO_TCHAR(CurrentProperty.GetName());
-			if (!MetadataTag.StartsWith(NodeName))
+			if (bPrefixTagWithNodeName && !MetadataTag.StartsWith(NodeName))
 			{
 				// Append the node name in the tag since all the metadata will be flattened on the Object
 				MetadataTag = NodeName + TEXT(".") + MetadataTag;
@@ -3141,7 +3141,7 @@ void FFbxImporter::ImportNodeCustomProperties(UObject* Object, FbxNode* Node)
 	int NumChildren = Node->GetChildCount();
 	for (int i = 0; i < NumChildren; ++i)
 	{
-		ImportNodeCustomProperties(Object, Node->GetChild(i));
+		ImportNodeCustomProperties(Object, Node->GetChild(i), bPrefixTagWithNodeName);
 	}
 }
 
diff --git a/Engine/Source/Editor/UnrealEd/Private/Fbx/FbxSkeletalMeshExport.cpp b/Engine/Source/Editor/UnrealEd/Private/Fbx/FbxSkeletalMeshExport.cpp
--- a/Engine/Source/Editor/UnrealEd/Private/Fbx/FbxSkeletalMeshExport.cpp
+++ b/Engine/Source/Editor/UnrealEd/Private/Fbx/FbxSkeletalMeshExport.cpp
@@ -530,11 +530,22 @@ void ExportObjectMetadataToBones(const UObject* ObjectToExport, const TArray<Fbx
 				FbxNode** Node = NameToNode.Find(NodeName);
 				if (Node)
 				{
-					FbxProperty Property = FbxProperty::Create(*Node, FbxStringDT, TCHAR_TO_UTF8(*TagAsString));
-					FbxString ValueString(TCHAR_TO_UTF8(*MetadataIt.Value));
+					if (MetadataIt.Value == TEXT("true") || MetadataIt.Value == TEXT("false"))
+					{
+						FbxProperty Property = FbxProperty::Create(*Node, FbxBoolDT, TCHAR_TO_UTF8(*TagAsString));
+						FbxBool ValueBool = MetadataIt.Value == TEXT("true") ? true : false;
 
-					Property.Set(ValueString);
-					Property.ModifyFlag(FbxPropertyFlags::eUserDefined, true);
+						Property.Set(ValueBool);
+						Property.ModifyFlag(FbxPropertyFlags::eUserDefined, true);
+					}
+					else
+					{
+						FbxProperty Property = FbxProperty::Create(*Node, FbxStringDT, TCHAR_TO_UTF8(*TagAsString));
+						FbxString ValueString(TCHAR_TO_UTF8(*MetadataIt.Value));
+
+						Property.Set(ValueString);
+						Property.ModifyFlag(FbxPropertyFlags::eUserDefined, true);
+					}
 				}
 			}
 		}
diff --git a/Engine/Source/Editor/UnrealEd/Private/Fbx/FbxSkeletalMeshImport.cpp b/Engine/Source/Editor/UnrealEd/Private/Fbx/FbxSkeletalMeshImport.cpp
--- a/Engine/Source/Editor/UnrealEd/Private/Fbx/FbxSkeletalMeshImport.cpp
+++ b/Engine/Source/Editor/UnrealEd/Private/Fbx/FbxSkeletalMeshImport.cpp
@@ -1905,13 +1905,13 @@ USkeletalMesh* UnFbx::FFbxImporter::ImportSkeletalMesh(FImportSkeletalMeshArgs &
 	// Import mesh metadata to SkeletalMesh
 	for (FbxNode* MeshNode : ImportSkeletalMeshArgs.NodeArray)
 	{
-		ImportNodeCustomProperties(SkeletalMesh, MeshNode);
+		ImportNodeCustomProperties(SkeletalMesh, MeshNode, true);
 	}
 
 	// Import bone metadata to Skeleton
 	for (FbxNode* SkeletonNode : ImportSkeletalMeshArgs.BoneNodeArray)
 	{
-		ImportNodeCustomProperties(SkeletalMesh->Skeleton, SkeletonNode);
+		ImportNodeCustomProperties(SkeletalMesh->Skeleton, SkeletonNode, true);
 	}
 
 	// ComponentContexts will now go out of scope, causing components to be re-registered
diff --git a/Engine/Source/Editor/UnrealEd/Public/FbxImporter.h b/Engine/Source/Editor/UnrealEd/Public/FbxImporter.h
--- a/Engine/Source/Editor/UnrealEd/Public/FbxImporter.h
+++ b/Engine/Source/Editor/UnrealEd/Public/FbxImporter.h
@@ -1719,7 +1719,7 @@ protected:
 	bool RetrievePoseFromBindPose(const TArray<FbxNode*>& NodeArray, FbxArray<FbxPose*> & PoseArray) const;
 
 	/** Import the user-defined properties on the node as FBX metadata on the object */
-	void ImportNodeCustomProperties(UObject* Object, FbxNode* Node);
+	void ImportNodeCustomProperties(UObject* Object, FbxNode* Node, bool bPrefixTagWithNodeName = false);
 
 public:
 	/** Import and set up animation related data from mesh **/
