Refactoring of content encryption config to remove separate specification of non-UFS files, and instead allow a group to be marked as "staging time only". The asset manager should still see those encryption groups and allocate them chunks, then the staging code will force all the specified files into those chunks.

diff --git a/Engine/Source/Editor/UnrealEd/Private/Commandlets/AssetRegistryGenerator.cpp b/Engine/Source/Editor/UnrealEd/Private/Commandlets/AssetRegistryGenerator.cpp
--- a/Engine/Source/Editor/UnrealEd/Private/Commandlets/AssetRegistryGenerator.cpp
+++ b/Engine/Source/Editor/UnrealEd/Private/Commandlets/AssetRegistryGenerator.cpp
@@ -216,31 +216,35 @@ bool FAssetRegistryGenerator::GenerateStreamingInstallManifest(int64 InExtraFlav
 		PlatformIniFile.GetArray(TEXT("/Script/UnrealEd.ProjectPackagingSettings"), TEXT("CompressedChunkWildcard"), CompressedChunkWildcards);
 	}
 
-	// Add manifests for any non-ufs file groups
+	// Add manifests for any staging-time only groups
 	if (bUseAssetManager && !TargetPlatform->HasSecurePackageFormat())
 	{
 		FContentEncryptionConfig ContentEncryptionConfig;
 		UAssetManager::Get().GetContentEncryptionConfig(ContentEncryptionConfig);
-		const TMap<FName, TSet<FName>>& EncryptedNonUFSFileGroups = ContentEncryptionConfig.GetNonUFSFileGroupMap();
+		const FContentEncryptionConfig::TGroupMap& EncryptedNonUFSFileGroups = ContentEncryptionConfig.GetPackageGroupMap();
 		
-		for (const TMap<FName, TSet<FName>>::ElementType& Element : EncryptedNonUFSFileGroups)
+		for (const FContentEncryptionConfig::TGroupMap::ElementType& Element : EncryptedNonUFSFileGroups)
 		{
-			FName GroupName = Element.Key;
-			const TSet<FName>& Files = Element.Value;
-			FChunkPackageSet* NewManifest = new FChunkPackageSet();
-			
-			for (FName File : Files)
+			if (Element.Value.bStageTimeOnly)
 			{
-				FString Filename = File.ToString();
-				FName LongPackageName = *FPackageName::FilenameToLongPackageName(Filename);
-				NewManifest->Add(LongPackageName, Filename);
-			}
+				FName GroupName = Element.Key;
+				const TSet<FName>& PackageNames = Element.Value.PackageNames;
+				
+				FChunkPackageSet* NewManifest = new FChunkPackageSet();
+				int32 ChunkID = UAssetManager::Get().GetContentEncryptionGroupChunkID(GroupName);
+				
+				for (FName PackageName : PackageNames)
+				{
+					NewManifest->Add(PackageName, FPackageName::LongPackageNameToFilename(PackageName.ToString()));
+				}
 
-			FinalChunkManifests.Add(NewManifest);
+				check(ChunkID < FinalChunkManifests.Num());
+				check(FinalChunkManifests[ChunkID] == nullptr);
+				FinalChunkManifests[ChunkID] = NewManifest;
+			}
 		}
 	}
 
-
 	// generate per-chunk pak list files
 	for (int32 Index = 0; Index < FinalChunkManifests.Num(); ++Index)
 	{
@@ -451,10 +455,10 @@ void FAssetRegistryGenerator::InjectEncryptionData(FAssetRegistryState& TargetSt
 		FContentEncryptionConfig EncryptionConfig;
 		AssetManager.GetContentEncryptionConfig(EncryptionConfig);
 
-		for (TMap<FName, TSet<FName>>::ElementType EncryptedAssetSetElement : EncryptionConfig.GetPackageGroupMap())
+		for (FContentEncryptionConfig::TGroupMap::ElementType EncryptedAssetSetElement : EncryptionConfig.GetPackageGroupMap())
 		{
 			FName SetName = EncryptedAssetSetElement.Key;
-			TSet<FName>& EncryptedRootAssets = EncryptedAssetSetElement.Value;
+			TSet<FName>& EncryptedRootAssets = EncryptedAssetSetElement.Value.PackageNames;
 
 			for (FName EncryptedRootPackageName : EncryptedRootAssets)
 			{
diff --git a/Engine/Source/Runtime/Engine/Classes/Engine/AssetManager.h b/Engine/Source/Runtime/Engine/Classes/Engine/AssetManager.h
--- a/Engine/Source/Runtime/Engine/Classes/Engine/AssetManager.h
+++ b/Engine/Source/Runtime/Engine/Classes/Engine/AssetManager.h
@@ -433,6 +433,11 @@ public:
 	 */
 	virtual FName GetUniqueAssetRegistryName(int32 InChunkIndex) const { return NAME_None; }
 
+	/**
+	 * For a given content encryption group name (as defined in the content encryption config that the project provides, return the relevant chunk ID
+	 */
+	virtual int32 GetContentEncryptionGroupChunkID(FName InGroupName) const { return INDEX_NONE; }
+
 	/** Returns the list of chunks assigned to the list of primary assets, which is usually a manager list. This is called by GetPackageChunkIds */
 	virtual bool GetPrimaryAssetSetChunkIds(const TSet<FPrimaryAssetId>& PrimaryAssetSet, const class ITargetPlatform* TargetPlatform, const TArray<int32>& ExistingChunkList, TArray<int32>& OutChunkList) const;
 
diff --git a/Engine/Source/Runtime/Engine/Classes/Engine/ContentEncryptionConfig.h b/Engine/Source/Runtime/Engine/Classes/Engine/ContentEncryptionConfig.h
--- a/Engine/Source/Runtime/Engine/Classes/Engine/ContentEncryptionConfig.h
+++ b/Engine/Source/Runtime/Engine/Classes/Engine/ContentEncryptionConfig.h
@@ -6,14 +6,22 @@ class FContentEncryptionConfig
 {
 public:
 
+	struct FGroup
+	{
+		TSet<FName> PackageNames;
+		bool bStageTimeOnly = false;
+	};
+
+	typedef TMap<FName, FGroup> TGroupMap;
+
 	void AddPackage(FName InGroupName, FName InPackageName)
 	{
-		PackageGroups.FindOrAdd(InGroupName).Add(InPackageName);
+		PackageGroups.FindOrAdd(InGroupName).PackageNames.Add(InPackageName);
 	}
 
-	void AddNonUFSFile(FName InGroupName, FName InFileName)
+	void SetGroupAsStageTimeOnly(FName InGroupName, bool bInStageTimeOnly)
 	{
-		NonUFSFileGroups.FindOrAdd(InGroupName).Add(InFileName);
+		PackageGroups.FindOrAdd(InGroupName).bStageTimeOnly = bInStageTimeOnly;
 	}
 
 	void AddReleasedKey(FGuid InKey)
@@ -21,15 +29,11 @@ public:
 		ReleasedKeys.Add(InKey);
 	}
 
-	const TMap<FName, TSet<FName>>& GetPackageGroupMap() const
+	const TGroupMap& GetPackageGroupMap() const
 	{
 		return PackageGroups;
 	}
 
-	const TMap<FName, TSet<FName>>& GetNonUFSFileGroupMap() const
-	{
-		return NonUFSFileGroups;
-	}
 
 	const TSet<FGuid>& GetReleasedKeys() const
 	{
@@ -38,7 +42,6 @@ public:
 
 private:
 
-	TMap<FName, TSet<FName>> PackageGroups;
-	TMap<FName, TSet<FName>> NonUFSFileGroups;
+	TGroupMap PackageGroups;
 	TSet<FGuid> ReleasedKeys;
 };
\ No newline at end of file
