Reverted bad loop-logic introduced by CL# 4721112

diff --git a/Engine/Source/Editor/UnrealEd/Private/Commandlets/FixConflictingLocalizationKeys.cpp b/Engine/Source/Editor/UnrealEd/Private/Commandlets/FixConflictingLocalizationKeys.cpp
--- a/Engine/Source/Editor/UnrealEd/Private/Commandlets/FixConflictingLocalizationKeys.cpp
+++ b/Engine/Source/Editor/UnrealEd/Private/Commandlets/FixConflictingLocalizationKeys.cpp
@@ -178,7 +178,7 @@ bool ReKeyTextProperty(UStruct* InOuterType, void* InAddrToUpdate, const TArray<
 						// ContainerIndex is the element index, but we need the sparse index
 						int32 SparseIndex = 0;
 						{
-							const int32 ElementCount = ScriptMapHelper.GetMaxIndex();
+							const int32 ElementCount = ScriptMapHelper.Num();
 							for (int32 ElementIndex = 0; ElementIndex < ElementCount; ++SparseIndex)
 							{
 								if (ScriptMapHelper.IsValidIndex(SparseIndex))
@@ -215,7 +215,7 @@ bool ReKeyTextProperty(UStruct* InOuterType, void* InAddrToUpdate, const TArray<
 						// ContainerIndex is the element index, but we need the sparse index
 						int32 SparseIndex = 0;
 						{
-							const int32 ElementCount = ScriptSetHelper.GetMaxIndex();
+							const int32 ElementCount = ScriptSetHelper.Num();
 							for (int32 ElementIndex = 0; ElementIndex < ElementCount; ++SparseIndex)
 							{
 								if (ScriptSetHelper.IsValidIndex(SparseIndex))
diff --git a/Engine/Source/Runtime/CoreUObject/Private/Serialization/PropertyLocalizationDataGathering.cpp b/Engine/Source/Runtime/CoreUObject/Private/Serialization/PropertyLocalizationDataGathering.cpp
--- a/Engine/Source/Runtime/CoreUObject/Private/Serialization/PropertyLocalizationDataGathering.cpp
+++ b/Engine/Source/Runtime/CoreUObject/Private/Serialization/PropertyLocalizationDataGathering.cpp
@@ -295,7 +295,7 @@ void FPropertyLocalizationDataGatherer::GatherLocalizationDataFromChildTextPrope
 
 			// Iterate over all elements of the map.
 			FScriptMapHelper ScriptMapHelper(MapProperty, ElementValueAddress);
-			const int32 ElementCount = ScriptMapHelper.GetMaxIndex();
+			const int32 ElementCount = ScriptMapHelper.Num();
 			for(int32 j = 0, ElementIndex = 0; ElementIndex < ElementCount; ++j)
 			{
 				if (!ScriptMapHelper.IsValidIndex(j))
@@ -338,7 +338,7 @@ void FPropertyLocalizationDataGatherer::GatherLocalizationDataFromChildTextPrope
 
 			// Iterate over all elements of the Set.
 			FScriptSetHelper ScriptSetHelper(SetProperty, ElementValueAddress);
-			const int32 ElementCount = ScriptSetHelper.GetMaxIndex();
+			const int32 ElementCount = ScriptSetHelper.Num();
 			for(int32 j = 0, ElementIndex = 0; ElementIndex < ElementCount; ++j)
 			{
 				if (!ScriptSetHelper.IsValidIndex(j))
