TBA: TextAssetCommandlet improvements

* Collect lists of broken assets during roundtrip tests and print a summary of packages that failed each phase at the end
* After resaving as text, load the file back as a plain JSON hierarchy to ensure the output was valid

diff --git a/Engine/Source/Editor/UnrealEd/Private/Commandlets/TextAssetCommandlet.cpp b/Engine/Source/Editor/UnrealEd/Private/Commandlets/TextAssetCommandlet.cpp
--- a/Engine/Source/Editor/UnrealEd/Private/Commandlets/TextAssetCommandlet.cpp
+++ b/Engine/Source/Editor/UnrealEd/Private/Commandlets/TextAssetCommandlet.cpp
@@ -14,6 +14,12 @@
 #include "UObject/UObjectIterator.h"
 #include "Stats/StatsMisc.h"
 #include "Misc/FileHelper.h"
+#include "Serialization/JsonReader.h"
+#include "Serialization/JsonSerializer.h"
+#include "Serialization/StructuredArchive.h"
+#include "Serialization/Formatters/JsonArchiveOutputFormatter.h"
+#include "Serialization/MemoryWriter.h"
+#include "Serialization/ArchiveUObjectFromStructuredArchive.h"
 
 DEFINE_LOG_CATEGORY(LogTextAsset);
 
@@ -35,6 +41,7 @@ bool HashFile(const TCHAR* InFilename, FSHAHash& OutHash)
 int32 UTextAssetCommandlet::Main(const FString& CmdLineParams)
 {
 	TArray<FString> Blacklist;
+	Blacklist.Add(TEXT("GM/SceneTest"));	// Broken EngineTest asset
 
 	FString ModeString = TEXT("ResaveText");
 	FString IterationsString = TEXT("1");
@@ -43,6 +50,7 @@ int32 UTextAssetCommandlet::Main(const FString& CmdLineParams)
 	FParse::Value(*CmdLineParams, TEXT("mode="), ModeString);
 	FParse::Value(*CmdLineParams, TEXT("filter="), FilenameFilterString);
 	FParse::Value(*CmdLineParams, TEXT("outputpath="), OutputPathString);
+	bool bVerifyJson = !FParse::Param(*CmdLineParams, TEXT("noverifyjson"));
 
 	enum class EMode
 	{
@@ -202,6 +210,10 @@ int32 UTextAssetCommandlet::Main(const FString& CmdLineParams)
 		float IterationPackageLoadTime = 0.0;
 		float IterationPackageSaveTime = 0.0;
 
+		TArray<FString> PhaseSuccess;
+		TArray<TArray<FString>> PhaseFails;
+		PhaseFails.AddDefaulted(3);
+
 		for (const TTuple<FString, FString>& FileToProcess : FilesToProcess)
 		{
 			FString SourceFilename = FileToProcess.Get<0>();
@@ -387,6 +399,21 @@ int32 UTextAssetCommandlet::Main(const FString& CmdLineParams)
 					UE_LOG(LogTextAsset, Error, TEXT("Binary determinism tests succeeded, but text and/or alternating tests failed for asset '%s'"), *SourceLongPackageName);
 				}
 
+				bool bSuccess = true;
+				for (int32 PhaseIndex = 0; PhaseIndex < NumPhases; ++PhaseIndex)
+				{
+					if (!bPhasesMatched[PhaseIndex])
+					{
+						bSuccess = false;
+						PhaseFails[PhaseIndex].Add(SourceLongPackageName);
+					}
+				}
+
+				if (bSuccess)
+				{
+					PhaseSuccess.Add(SourceLongPackageName);
+				}
+
 				UE_LOG(LogTextAsset, Display, TEXT("-----------------------------------------------------------------------------------------"));
 				UE_LOG(LogTextAsset, Display, TEXT("Completed roundtrip test for '%s'"), *SourceLongPackageName);
 				UE_LOG(LogTextAsset, Display, TEXT("-----------------------------------------------------------------------------------------"));
@@ -399,6 +426,8 @@ int32 UTextAssetCommandlet::Main(const FString& CmdLineParams)
 			{
 				UPackage* Package = nullptr;
 
+				UE_LOG(LogTextAsset, Display, TEXT("Resaving asset %s"), *SourceFilename);
+
 				double Timer = 0.0;
 				{
 					SCOPE_SECONDS_COUNTER(Timer);
@@ -419,6 +448,15 @@ int32 UTextAssetCommandlet::Main(const FString& CmdLineParams)
 					IterationPackageSaveTime += Timer;
 				}
 
+				if (bVerifyJson)
+				{
+					FArchive* File = IFileManager::Get().CreateFileReader(*DestinationFilename);
+					TSharedPtr< FJsonObject > RootObject;
+					TSharedRef< TJsonReader<char> > Reader = TJsonReaderFactory<char>::Create(File);
+					ensure(FJsonSerializer::Deserialize(Reader, RootObject));
+					delete File;
+				}
+
 				if (OutputPathString.Len() > 0)
 				{
 					FString CopyFilename = DestinationFilename;
@@ -436,12 +474,14 @@ int32 UTextAssetCommandlet::Main(const FString& CmdLineParams)
 			{
 				UPackage* Package = nullptr;
 
+				CollectGarbage(RF_NoFlags, true);
 				double Timer = 0.0;
 				{
 					SCOPE_SECONDS_COUNTER(Timer);
 					UE_LOG(LogTextAsset, Display, TEXT("Loading Text Asset '%s'"), *SourceFilename);
 					Package = LoadPackage(nullptr, *SourceFilename, 0);
 				}
+				CollectGarbage(RF_NoFlags, true);
 				IterationPackageLoadTime += Timer;
 				TotalPackageLoadTime += Timer;
 
@@ -469,7 +509,32 @@ int32 UTextAssetCommandlet::Main(const FString& CmdLineParams)
 			TotalTime += Time;
 			NumFiles++;
 		}
-		
+
+		if (Mode == EMode::RoundTrip)
+		{
+			UE_LOG(LogTextAsset, Display, TEXT("\t-----------------------------------------------------"));
+			UE_LOG(LogTextAsset, Display, TEXT("\tRoundTrip Results"));
+			UE_LOG(LogTextAsset, Display, TEXT("\tTotal Packages: %i"), FilesToProcess.Num());
+			UE_LOG(LogTextAsset, Display, TEXT("\tNum Successful Packages: %i"), PhaseSuccess.Num());
+			UE_LOG(LogTextAsset, Display, TEXT("\tPhase 0 Fails: %i (Binary Package Determinism Fails)"), PhaseFails[0].Num());
+			UE_LOG(LogTextAsset, Display, TEXT("\tPhase 1 Fails: %i (Text Package Determinism Fails)"), PhaseFails[1].Num());
+			UE_LOG(LogTextAsset, Display, TEXT("\tPhase 2 Fails: %i (Mixed Package Determinism Fails)"), PhaseFails[2].Num());
+			UE_LOG(LogTextAsset, Display, TEXT("\t-----------------------------------------------------"));
+
+			for (int32 PhaseIndex = 0; PhaseIndex < PhaseFails.Num(); ++PhaseIndex)
+			{
+				if (PhaseFails[PhaseIndex].Num() > 0)
+				{
+					UE_LOG(LogTextAsset, Display, TEXT("\tPhase %i Fails:"), PhaseIndex);
+					for (const FString& PhaseFail : PhaseFails[PhaseIndex])
+					{
+						UE_LOG(LogTextAsset, Display, TEXT("\t\t%s"), *PhaseFail);
+					}
+					UE_LOG(LogTextAsset, Display, TEXT("\t-----------------------------------------------------"));
+				}
+			}
+		}
+
 		UE_LOG(LogTextAsset, Display, TEXT("\tTotal Time:\t%.2fs"), TotalTime);
 		UE_LOG(LogTextAsset, Display, TEXT("\tAvg File Time:  \t%.2fms"), (TotalTime * 1000.0) / (double)NumFiles);
 		UE_LOG(LogTextAsset, Display, TEXT("\tMin File Time:  \t%.2fms (%s)"), MinTime * 1000.0, *MinTimePackage);
