TBA: Add 'FindMismatchedSerializers' mode to text asset commandlet, which dumps out a list of all UClasses which don't have the CLASS_MatchedSerializers flag, meaning we can't guarantee the have Serialize functions for FArchive AND FStructuredArchive, therefore we can't use the new structured archive based serialize path. Should only ever be native instrinsic classes as UHT takes care of all other cases.

diff --git a/Engine/Source/Editor/UnrealEd/Private/Commandlets/TextAssetCommandlet.cpp b/Engine/Source/Editor/UnrealEd/Private/Commandlets/TextAssetCommandlet.cpp
--- a/Engine/Source/Editor/UnrealEd/Private/Commandlets/TextAssetCommandlet.cpp
+++ b/Engine/Source/Editor/UnrealEd/Private/Commandlets/TextAssetCommandlet.cpp
@@ -38,6 +38,17 @@ bool HashFile(const TCHAR* InFilename, FSHAHash& OutHash)
 	return false;
 }
 
+void FindMismatchedSerializers()
+{
+	for (TObjectIterator<UClass> It; It; ++It)
+	{
+		if (!It->HasAnyClassFlags(CLASS_MatchedSerializers))
+		{
+			UE_LOG(LogTextAsset, Display, TEXT("Class Mismatched Serializers: %s"), *It->GetName());
+		}
+	}
+}
+
 int32 UTextAssetCommandlet::Main(const FString& CmdLineParams)
 {
 	TArray<FString> Blacklist;
@@ -58,6 +69,7 @@ int32 UTextAssetCommandlet::Main(const FString& CmdLineParams)
 		ResaveBinary,
 		RoundTrip,
 		LoadText,
+		FindMismatchedSerializers
 	};
 
 	TMap<FString, EMode> Modes;
@@ -65,10 +77,17 @@ int32 UTextAssetCommandlet::Main(const FString& CmdLineParams)
 	Modes.Add(TEXT("ResaveBinary"), EMode::ResaveBinary);
 	Modes.Add(TEXT("RoundTrip"), EMode::RoundTrip);
 	Modes.Add(TEXT("LoadText"), EMode::LoadText);
+	Modes.Add(TEXT("FindMismatchedSerializers"), EMode::FindMismatchedSerializers);
 
 	check(Modes.Contains(ModeString));
 	EMode Mode = Modes[ModeString];
 
+	if (Mode == EMode::FindMismatchedSerializers)
+	{
+		FindMismatchedSerializers();
+		return 0;
+	}
+
 	TArray<UObject*> Objects;	
 
 	int32 NumSaveIterations = 1;
