Quick fix for Linux build issue introduced in 4099927

diff --git a/Engine/Source/Editor/UnrealEd/Private/CookOnTheFlyServer.cpp b/Engine/Source/Editor/UnrealEd/Private/CookOnTheFlyServer.cpp
--- a/Engine/Source/Editor/UnrealEd/Private/CookOnTheFlyServer.cpp
+++ b/Engine/Source/Editor/UnrealEd/Private/CookOnTheFlyServer.cpp
@@ -240,7 +240,7 @@ void OutputHierarchyTimers(const FHierarchicalTimerInfo* TimerInfo, int32 Depth)
 
 	FString TimerName(TimerInfo->Name);
 
-	UE_LOG(LogCook, Display, TEXT("  %s%s: %.3fs (%u)"), &LeftPad[_countof(LeftPad) - 1 - Depth * 2], *TimerName, TimerInfo->Length, TimerInfo->HitCount);
+	UE_LOG(LogCook, Display, TEXT("  %s%s: %.3fs (%u)"), &LeftPad[sizeof(LeftPad)/sizeof(LeftPad[0]) - 1 - Depth * 2], *TimerName, TimerInfo->Length, TimerInfo->HitCount);
 
 	const int32 ChildDepth = Depth + TimerInfo->IncrementDepth;
 
