Cooker: Changed FHierarchicalTimerInfo so it uses a linked list for tracking child nodes, to be able to deal with any child count. Previously we assumed there would never be more than 9 children but it turns out there are cooker modes that need more.

Fixes check when using -FullLoadAndSave to cook

diff --git a/Engine/Source/Editor/UnrealEd/Private/CookOnTheFlyServer.cpp b/Engine/Source/Editor/UnrealEd/Private/CookOnTheFlyServer.cpp
--- a/Engine/Source/Editor/UnrealEd/Private/CookOnTheFlyServer.cpp
+++ b/Engine/Source/Editor/UnrealEd/Private/CookOnTheFlyServer.cpp
@@ -124,6 +124,7 @@ static FAutoConsoleVariableRef CVarCookDisplayMode(
 
 
 #if OUTPUT_TIMING
+#include <Containers/AllocatorFixedSizeFreeList.h>
 
 struct FHierarchicalTimerInfo
 {
@@ -131,8 +132,9 @@ public:
 	FHierarchicalTimerInfo(const FHierarchicalTimerInfo& InTimerInfo) = delete;
 	FHierarchicalTimerInfo(FHierarchicalTimerInfo&& InTimerInfo) = delete;
 
-	explicit FHierarchicalTimerInfo(const char* InName)
+	explicit FHierarchicalTimerInfo(const char* InName, uint16 InId)
 	:	Name(InName)
+	,	Id(InId)
 	{
 	}
 
@@ -141,49 +143,69 @@ public:
 		ClearChildren();
 	}
 
-	void ClearChildren()
-	{
-		for (int i = 0; i < ChildCount; ++i)
-			delete Children[i];
+	void							ClearChildren();
+	FHierarchicalTimerInfo*			GetChild(int InId, const char* InName);
 
-		ChildCount = 0;
-	}
+	uint32							HitCount = 0;
+	uint16							Id = 0;
+	bool							IncrementDepth = true;
+	double							Length = 0;
+	const char*						Name;
 
-	FHierarchicalTimerInfo* GetChild(int InId, const char* InName)
-	{
-		for (int i = 0; i < ChildCount; ++i)
-		{
-			if (ChildrenIds[i] == InId)
-			{
-				return Children[i];
-			}
-		}
+	FHierarchicalTimerInfo*			FirstChild = nullptr;
+	FHierarchicalTimerInfo*			NextSibling = nullptr;
 
-		FHierarchicalTimerInfo* Child = new FHierarchicalTimerInfo(InName);
-		const int Index = ChildCount++;
+private:
+	static FHierarchicalTimerInfo*	AllocNew(const char* InName, uint16 InId);
+	static void						DestroyAndFree(FHierarchicalTimerInfo*);
+};
+
+static FHierarchicalTimerInfo RootTimerInfo("Root", 0);
+static FHierarchicalTimerInfo* CurrentTimerInfo = &RootTimerInfo;
+static TAllocatorFixedSizeFreeList<sizeof(FHierarchicalTimerInfo), 256> TimerInfoAllocator;
 
-		check(Index < MaxChildCount);
+FHierarchicalTimerInfo* FHierarchicalTimerInfo::AllocNew(const char* InName, uint16 InId)
+{
+	return new(TimerInfoAllocator.Allocate()) FHierarchicalTimerInfo(InName, InId);
+}
+
+void FHierarchicalTimerInfo::DestroyAndFree(FHierarchicalTimerInfo* InPtr)
+{
+	InPtr->~FHierarchicalTimerInfo();
+	TimerInfoAllocator.Free(InPtr);
+}
+
+void FHierarchicalTimerInfo::ClearChildren()
+{
+	for (FHierarchicalTimerInfo* Child = FirstChild; Child;)
+	{
+		FHierarchicalTimerInfo* NextChild = Child->NextSibling;
 
-		Children[Index]		= Child;
-		ChildrenIds[Index]	= InId;
+		DestroyAndFree(Child);
 
-		return Child;
+		Child = NextChild;
 	}
 
-	enum { MaxChildCount = 9 };
+	FirstChild = nullptr;
+}
 
-	uint32							HitCount = 0;
-	uint16							ChildCount = 0;
-	uint8							ChildrenIds[MaxChildCount];
-	bool							IncrementDepth = true;
-	double							Length = 0;
-	const char*						Name;
+FHierarchicalTimerInfo* FHierarchicalTimerInfo::GetChild(int InId, const char* InName)
+{
+	for (FHierarchicalTimerInfo* Child = FirstChild; Child;)
+	{
+		if (Child->Id == InId)
+			return Child;
 
-	FHierarchicalTimerInfo*			Children[MaxChildCount];
-};
+		Child = Child->NextSibling;
+	}
 
-static FHierarchicalTimerInfo RootTimerInfo("Root");
-static FHierarchicalTimerInfo* CurrentTimerInfo = &RootTimerInfo;
+	FHierarchicalTimerInfo* Child = AllocNew(InName, InId);
+
+	Child->NextSibling	= FirstChild;
+	FirstChild			= Child;
+
+	return Child;
+}
 
 struct FScopeTimer
 {
@@ -193,11 +215,13 @@ public:
 
 	FScopeTimer(int InId, const char* InName, bool IncrementScope = false )
 	{
+		checkSlow(IsInGameThread());
+
 		HierarchyTimerInfo = CurrentTimerInfo->GetChild(InId, InName);
 
 		HierarchyTimerInfo->IncrementDepth = IncrementScope;
 
-		OldTimerInfo		= CurrentTimerInfo;
+		PrevTimerInfo		= CurrentTimerInfo;
 		CurrentTimerInfo	= HierarchyTimerInfo;
 	}
 
@@ -225,28 +249,38 @@ public:
 		Stop();
 
 		check(CurrentTimerInfo == HierarchyTimerInfo);
-		CurrentTimerInfo = OldTimerInfo;
+		CurrentTimerInfo = PrevTimerInfo;
 	}
 
 private:
 	uint64					StartTime = 0;
 	FHierarchicalTimerInfo* HierarchyTimerInfo;
-	FHierarchicalTimerInfo* OldTimerInfo;
+	FHierarchicalTimerInfo* PrevTimerInfo;
 };
 
 void OutputHierarchyTimers(const FHierarchicalTimerInfo* TimerInfo, int32 Depth)
 {
+	FString TimerName(TimerInfo->Name);
+
 	static const TCHAR LeftPad[] = TEXT("                                ");
+	const SIZE_T PadOffset = FMath::Max<int>(ARRAY_COUNT(LeftPad) - 1 - Depth * 2, 0);
 
-	FString TimerName(TimerInfo->Name);
+	UE_LOG(LogCook, Display, TEXT("  %s%s: %.3fs (%u)"), &LeftPad[PadOffset], *TimerName, TimerInfo->Length, TimerInfo->HitCount);
+
+	// We need to print in reverse order since the child list begins with the most recently added child
+
+	TArray<const FHierarchicalTimerInfo*> Stack;
 
-	UE_LOG(LogCook, Display, TEXT("  %s%s: %.3fs (%u)"), &LeftPad[sizeof(LeftPad)/sizeof(LeftPad[0]) - 1 - Depth * 2], *TimerName, TimerInfo->Length, TimerInfo->HitCount);
+	for (const FHierarchicalTimerInfo* Child = TimerInfo->FirstChild; Child; Child = Child->NextSibling)
+	{
+		Stack.Add(Child);
+	}
 
 	const int32 ChildDepth = Depth + TimerInfo->IncrementDepth;
 
-	for (int i = 0, n = TimerInfo->ChildCount; i < n; ++i)
+	for (size_t i = Stack.Num(); i > 0; --i)
 	{
-		OutputHierarchyTimers(TimerInfo->Children[i], ChildDepth);
+		OutputHierarchyTimers(Stack[i - 1], ChildDepth);
 	}
 }
 
