- Fixed Linux build warning re: member initialization order
- Also eliminated OUTPUT_HIERARCHYTIMERS/CLEAR_HIEARCHYTIMERS macros (plain functions are fine)
- Moved finishing-up code for FullLoadAndSave() to TickCookOnTheSide() call site to improve timer output. Previously some of the scopes would not have been closed before printing and thus the output was misleading.

diff --git a/Engine/Source/Editor/UnrealEd/Private/CookOnTheFlyServer.cpp b/Engine/Source/Editor/UnrealEd/Private/CookOnTheFlyServer.cpp
--- a/Engine/Source/Editor/UnrealEd/Private/CookOnTheFlyServer.cpp
+++ b/Engine/Source/Editor/UnrealEd/Private/CookOnTheFlyServer.cpp
@@ -133,8 +133,8 @@ public:
 	FHierarchicalTimerInfo(FHierarchicalTimerInfo&& InTimerInfo) = delete;
 
 	explicit FHierarchicalTimerInfo(const char* InName, uint16 InId)
-	:	Name(InName)
-	,	Id(InId)
+	:	Id(InId)
+	,	Name(InName)
 	{
 	}
 
@@ -238,7 +238,7 @@ public:
 		if (!StartTime)
 			return;
 
-		HierarchyTimerInfo->Length += (FPlatformTime::Cycles64() - StartTime) * FPlatformTime::GetSecondsPerCycle();
+		HierarchyTimerInfo->Length += FPlatformTime::ToSeconds64(FPlatformTime::Cycles64() - StartTime);
 		++HierarchyTimerInfo->HitCount;
 
 		StartTime = 0;
@@ -303,9 +303,6 @@ void ClearHierarchyTimers()
 #define ACCUMULATE_TIMER(name)			CREATE_TIMER(name, false);
 #define ACCUMULATE_TIMER_START(name)	ScopeTimer##name.Start();
 #define ACCUMULATE_TIMER_STOP(name)		ScopeTimer##name.Stop();
-
-#define OUTPUT_HIERARCHYTIMERS()		OutputHierarchyTimers();
-#define CLEAR_HIERARCHYTIMERS()			ClearHierarchyTimers();
 #else
 #define SCOPE_TIMER(name)
 
@@ -313,8 +310,8 @@ void ClearHierarchyTimers()
 #define ACCUMULATE_TIMER_START(name) 
 #define ACCUMULATE_TIMER_STOP(name) 
 
-#define OUTPUT_HIERARCHYTIMERS()
-#define CLEAR_HIERARCHYTIMERS() 
+void OutputHierarchyTimers() {}
+void ClearHierarchyTimers() {}
 #endif
 
 
@@ -1328,7 +1325,15 @@ uint32 UCookOnTheFlyServer::TickCookOnTheSide( const float TimeSlice, uint32 &Co
 {
 	if (IsCookByTheBookMode() && CookByTheBookOptions->bFullLoadAndSave)
 	{
-		return FullLoadAndSave(CookedPackageCount);
+		uint32 Result = FullLoadAndSave(CookedPackageCount);
+
+		UE_LOG(LogCook, Display, TEXT("Finishing up..."));
+
+		CookByTheBookFinished();
+
+		UE_LOG(LogCook, Display, TEXT("Done!"));
+
+		return Result;
 	}
 
 	COOK_STAT(FScopedDurationTimer TickTimer(DetailedCookStats::TickCookOnTheSideTimeSec));
@@ -3605,7 +3610,7 @@ void UCookOnTheFlyServer::DumpStats()
 	UE_LOG(LogCook, Display, TEXT("  %s=%d"), L"LoadPackage", this->StatLoadedPackageCount);
 	UE_LOG(LogCook, Display, TEXT("  %s=%d"), L"SavedPackage", this->StatSavedPackageCount);
 
-	OUTPUT_HIERARCHYTIMERS();
+	OutputHierarchyTimers();
 #if PROFILE_NETWORK
 	UE_LOG(LogCook, Display, TEXT("Network Stats \n"
 		"TimeTillRequestStarted %f\n"
@@ -5832,8 +5837,8 @@ void UCookOnTheFlyServer::CookByTheBookFinished()
 
 	UE_LOG(LogCook, Display, TEXT("Peak Used virtual %uMB Peak Used physical %uMB"), MemStats.PeakUsedVirtual / 1024 / 1024, MemStats.PeakUsedPhysical / 1024 / 1024 );
 
-	OUTPUT_HIERARCHYTIMERS();
-	CLEAR_HIERARCHYTIMERS();
+	OutputHierarchyTimers();
+	ClearHierarchyTimers();
 }
 
 void UCookOnTheFlyServer::BuildMapDependencyGraph(const FName& PlatformName)
@@ -7490,12 +7495,6 @@ uint32 UCookOnTheFlyServer::FullLoadAndSave(uint32& CookedPackageCount)
 		}
 	}
 
-	UE_LOG(LogCook, Display, TEXT("Finishing up..."));
-
-	CookByTheBookFinished();
-
-	UE_LOG(LogCook, Display, TEXT("Done!"));
-
 	return Result;
 }
 
