TBA: TextAssetCommandlet - When outputting converted assets to an output path, replicate the workspace relative path in the output directory

diff --git a/Engine/Source/Editor/UnrealEd/Private/Commandlets/TextAssetCommandlet.cpp b/Engine/Source/Editor/UnrealEd/Private/Commandlets/TextAssetCommandlet.cpp
--- a/Engine/Source/Editor/UnrealEd/Private/Commandlets/TextAssetCommandlet.cpp
+++ b/Engine/Source/Editor/UnrealEd/Private/Commandlets/TextAssetCommandlet.cpp
@@ -52,7 +52,6 @@ void FindMismatchedSerializers()
 int32 UTextAssetCommandlet::Main(const FString& CmdLineParams)
 {
 	TArray<FString> Blacklist;
-	Blacklist.Add(TEXT("GM/SceneTest"));	// Broken EngineTest asset
 
 	FString ModeString = TEXT("ResaveText");
 	FString IterationsString = TEXT("1");
@@ -479,8 +478,8 @@ int32 UTextAssetCommandlet::Main(const FString& CmdLineParams)
 				if (OutputPathString.Len() > 0)
 				{
 					FString CopyFilename = DestinationFilename;
-					FPaths::MakePathRelativeTo(CopyFilename, *FPaths::ProjectContentDir());
-					CopyFilename = OutputPathString / FApp::GetProjectName() / CopyFilename;
+					FPaths::MakePathRelativeTo(CopyFilename, *FPaths::RootDir());
+					CopyFilename = OutputPathString / CopyFilename;
 					CopyFilename.RemoveFromEnd(TEXT(".tmp"));
 					IFileManager::Get().MakeDirectory(*FPaths::GetPath(CopyFilename));
 					IFileManager::Get().Move(*CopyFilename, *DestinationFilename);
