Removed child cooker support.

In practice it is not a useful feature as it provides no performance improvement (quite the opposite in fact) and adds testing and maintenance complexity.

diff --git a/Engine/Source/Editor/UnrealEd/Classes/CookOnTheSide/CookOnTheFlyServer.h b/Engine/Source/Editor/UnrealEd/Classes/CookOnTheSide/CookOnTheFlyServer.h
--- a/Engine/Source/Editor/UnrealEd/Classes/CookOnTheSide/CookOnTheFlyServer.h
+++ b/Engine/Source/Editor/UnrealEd/Classes/CookOnTheSide/CookOnTheFlyServer.h
@@ -793,18 +793,6 @@ private:
 	//////////////////////////////////////////////////////////////////////////
 	// Cook by the book options
 public:
-	struct FChildCooker
-	{
-		FChildCooker() : ReadPipe(nullptr), ReturnCode(-1), bFinished(false),  Thread(nullptr) { }
-
-		FProcHandle ProcessHandle;
-		FString ResponseFileName;
-		FString BaseResponseFileName;
-		void* ReadPipe;
-		int32 ReturnCode;
-		bool bFinished;
-		FRunnableThread* Thread;
-	};
 private:
 	struct FCookByTheBookOptions
 	{
@@ -816,10 +804,8 @@ private:
 			CookTime( 0.0 ),
 			CookStartTime( 0.0 ), 
 			bErrorOnEngineContentUse(false),
-			bIsChildCooker(false),
 			bDisableUnsolicitedPackages(false),
-			bFullLoadAndSave(false),
-			ChildCookIdentifier(-1)
+			bFullLoadAndSave(false)
 		{ }
 
 		/** Should we test for UObject leaks */
@@ -849,13 +835,8 @@ private:
 		double CookStartTime;
 		/** error when detecting engine content being used in this cook */
 		bool bErrorOnEngineContentUse;
-		bool bIsChildCooker;
 		bool bDisableUnsolicitedPackages;
 		bool bFullLoadAndSave;
-		int32 ChildCookIdentifier;
-		FString ChildCookFilename;
-		TSet<FName> ChildUnsolicitedPackages;
-		TArray<FChildCooker> ChildCookers;
 		TArray<FName> TargetPlatformNames;
 		TArray<FName> StartupPackages;
 		/** Mapping from source packages to their localized variants (based on the culture list in FCookByTheBookStartupOptions) */
@@ -1032,7 +1013,6 @@ public:
 		COSR_ErrorLoadingPackage	= 0x00000004,
 		COSR_RequiresGC				= 0x00000008,
 		COSR_WaitingOnCache			= 0x00000010,
-		COSR_WaitingOnChildCookers	= 0x00000020,
 		COSR_MarkedUpKeepPackages	= 0x00000040
 	};
 
@@ -1098,8 +1078,6 @@ public:
 		FString DLCName;
 		FString CreateReleaseVersion;
 		FString BasedOnReleaseVersion;
-		FString ChildCookFileName; // if we are the child cooker 
-		int32 ChildCookIdentifier; // again, only if you are the child cooker
 		bool bGenerateStreamingInstallManifests; 
 		bool bGenerateDependenciesForMaps; 
 		bool bErrorOnEngineContentUse; // this is a flag for dlc, will cause the cooker to error if the dlc references engine content
@@ -1107,7 +1085,6 @@ public:
 		FCookByTheBookStartupOptions() :
 			CookOptions(ECookByTheBookOptions::None),
 			DLCName(FString()),
-			ChildCookIdentifier(-1),
 			bGenerateStreamingInstallManifests(false),
 			bGenerateDependenciesForMaps(false),
 			bErrorOnEngineContentUse(false),
@@ -1364,28 +1341,6 @@ private:
 	*/
 	void CookByTheBookFinished();
 
-	/**
-	* StartChildCookers to help out with cooking
-	* only valid in cook by the book (not from the editor)
-	* 
-	* @param NumChildCookersToSpawn, number of child cookers we want to use
-	*/
-	void StartChildCookers(int32 NumChildCookersToSpawn, const TArray<FName>& TargetPlatformNames, const FString& ExtraCmdParams = FString());
-
-	/**
-	* TickChildCookers
-	* output the information form the child cookers to the main cooker output
-	* 
-	* @return return true if all child cookers are finished
-	*/
-	bool TickChildCookers();
-
-	/**
-	* CleanUPChildCookers
-	* can only be called after TickChildCookers returns true
-	*/
-	void CleanUpChildCookers();
-
 	/**
 	* Get all the packages which are listed in asset registry passed in.  
 	*
@@ -1411,20 +1366,6 @@ private:
 	*/
 	void WriteMapDependencyGraph(const FName& PlatformName);
 
-	/**
-	* IsChildCooker, 
-	* returns if this cooker is a sue chef for some other master chef.
-	*/
-	bool IsChildCooker() const
-	{
-		if (IsCookByTheBookMode())
-		{
-			return !(CookByTheBookOptions->ChildCookFilename.IsEmpty());
-		}
-		return false;
-	}
-
-
 	//////////////////////////////////////////////////////////////////////////
 	// cook on the fly specific functions
 
diff --git a/Engine/Source/Editor/UnrealEd/Private/Commandlets/CookCommandlet.cpp b/Engine/Source/Editor/UnrealEd/Private/Commandlets/CookCommandlet.cpp
--- a/Engine/Source/Editor/UnrealEd/Private/Commandlets/CookCommandlet.cpp
+++ b/Engine/Source/Editor/UnrealEd/Private/Commandlets/CookCommandlet.cpp
@@ -859,8 +859,6 @@ bool UCookCommandlet::CookByTheBook( const TArray<ITargetPlatform*>& Platforms,
 	StartupOptions.bErrorOnEngineContentUse = bErrorOnEngineContentUse;
 	StartupOptions.bGenerateDependenciesForMaps = Switches.Contains(TEXT("GenerateDependenciesForMaps"));
 	StartupOptions.bGenerateStreamingInstallManifests = bGenerateStreamingInstallManifests;
-	StartupOptions.ChildCookFileName = ChildCookFile;
-	StartupOptions.ChildCookIdentifier = ChildCookIdentifier;
 	StartupOptions.NumProcesses = NumProcesses;
 
 	COOK_STAT(
diff --git a/Engine/Source/Editor/UnrealEd/Private/CookOnTheFlyServer.cpp b/Engine/Source/Editor/UnrealEd/Private/CookOnTheFlyServer.cpp
--- a/Engine/Source/Editor/UnrealEd/Private/CookOnTheFlyServer.cpp
+++ b/Engine/Source/Editor/UnrealEd/Private/CookOnTheFlyServer.cpp
@@ -481,18 +481,6 @@ const FString& GetAssetRegistryPath()
 	return AssetRegistryPath;
 }
 
-FString GetChildCookerResultFilename(const FString& ResponseFilename)
-{
-	FString Result = ResponseFilename + TEXT("Result.txt");
-	return Result;
-}
-
-FString GetChildCookerManifestFilename(const FString ResponseFilename)
-{
-	FString Result = ResponseFilename + TEXT("Manifest.txt");
-	return Result;
-}
-
 /**
  * Return the release asset registry filename for the release version supplied
  */
@@ -681,7 +669,6 @@ UCookOnTheFlyServer::~UCookOnTheFlyServer()
 	FCoreDelegates::OnFConfigCreated.RemoveAll(this);
 	FCoreDelegates::OnFConfigDeleted.RemoveAll(this);
 
-	check(TickChildCookers() == true);
 	if ( CookByTheBookOptions )
 	{		
 		delete CookByTheBookOptions;
@@ -1018,8 +1005,6 @@ public:
 
 void UCookOnTheFlyServer::GetDependentPackages(const TSet<UPackage*>& RootPackages, TSet<FName>& FoundPackages)
 {
-	check(IsChildCooker() == false);
-
 	TSet<FName> RootPackageFNames;
 	for (const UPackage* RootPackage : RootPackages)
 	{
@@ -1034,8 +1019,6 @@ void UCookOnTheFlyServer::GetDependentPackages(const TSet<UPackage*>& RootPackag
 
 void UCookOnTheFlyServer::GetDependentPackages( const TSet<FName>& RootPackages, TSet<FName>& FoundPackages )
 {
-	check(IsChildCooker() == false);
-
 	TArray<FName> FoundPackagesArray;
 	for (const FName& RootPackage : RootPackages)
 	{
@@ -1264,43 +1247,6 @@ bool UCookOnTheFlyServer::RequestPackage(const FName& StandardPackageFName, cons
 	return true;
 }
 
-// should only call this after tickchildcookers returns false
-void UCookOnTheFlyServer::CleanUpChildCookers()
-{
-	if (IsCookByTheBookMode())
-	{
-		for (FChildCooker& ChildCooker : CookByTheBookOptions->ChildCookers)
-		{
-			check(ChildCooker.bFinished == true);
-			if (ChildCooker.Thread != nullptr)
-			{
-				ChildCooker.Thread->WaitForCompletion();
-				delete ChildCooker.Thread;
-				ChildCooker.Thread = nullptr;
-			}
-
-		}
-
-		
-	}
-}
-
-bool UCookOnTheFlyServer::TickChildCookers()
-{
-	if (IsCookByTheBookMode())
-	{
-		for (FChildCooker& ChildCooker : CookByTheBookOptions->ChildCookers)
-		{
-			if (ChildCooker.bFinished == false)
-			{
-				return false;
-			}
-		}
-	}
-	return true;
-}
-
-
 // callback just before the garbage collector gets called.
 void UCookOnTheFlyServer::PreGarbageCollect()
 {
@@ -1341,7 +1287,6 @@ uint32 UCookOnTheFlyServer::TickCookOnTheSide( const float TimeSlice, uint32 &Co
 
 	uint32 Result = 0;
 
-	if (IsChildCooker() == false)
 	{
 		if (AssetRegistry == nullptr || AssetRegistry->IsLoadingAssets())
 		{
@@ -1354,27 +1299,6 @@ uint32 UCookOnTheFlyServer::TickCookOnTheSide( const float TimeSlice, uint32 &Co
 	// we use this in the unsolicited packages processing below
 	TArray<FName> AllTargetPlatformNames;
 	
-	if (CurrentCookMode == ECookMode::CookByTheBook)
-	{
-		if (CookRequests.HasItems() == false)
-		{
-			SCOPE_TIMER(WaitingForChildCookers);
-			// we have nothing left to cook so we are now waiting for child cookers to finish
-			if (TickChildCookers() == false)
-			{
-				Result |= COSR_WaitingOnChildCookers;
-			}
-			
-		}
-		else
-		{
-			if (TickChildCookers() == false)
-			{
-				Result |= COSR_WaitingOnChildCookers;
-			}
-		}
-	}
-
 	while (!GIsRequestingExit || CurrentCookMode == ECookMode::CookByTheBook)
 	{
 		if (CookRequests.Num() > 0)
@@ -1698,8 +1622,7 @@ uint32 UCookOnTheFlyServer::TickCookOnTheSide( const float TimeSlice, uint32 &Co
 	}
 
 	if (IsCookByTheBookRunning() &&
-		(CookRequests.HasItems() == false) &&
-		(!(Result & COSR_WaitingOnChildCookers)))
+		(CookRequests.HasItems() == false))
 	{
 		check(IsCookByTheBookMode());
 
@@ -2181,7 +2104,6 @@ bool UCookOnTheFlyServer::SaveCookedPackages(TArray<UPackage*>& PackagesToSave,
 				// also doesn't work in multiprocess cooking
 				KeepEditorOnlyPackages = !(IsCookByTheBookMode() && !IsCookingInEditor());
 				KeepEditorOnlyPackages |= IsCookFlagSet(ECookInitializationFlags::Iterative);
-				KeepEditorOnlyPackages |= IsChildCooker() || (CookByTheBookOptions && CookByTheBookOptions->ChildCookers.Num() > 0);
 				SaveFlags |= KeepEditorOnlyPackages ? SAVE_KeepEditorOnlyCookedPackages : SAVE_None;
 
 				GOutputCookingWarnings = IsCookFlagSet(ECookInitializationFlags::OutputVerboseCookerWarnings);
@@ -2313,8 +2235,6 @@ void UCookOnTheFlyServer::PostLoadPackageFixup(UPackage* Package)
 
 		World->PersistentLevel->HandleLegacyMapBuildData();
 
-		// child cookers can't process maps
-		// check(IsChildCooker() == false);
 		if (IsCookByTheBookMode())
 		{
 			GIsCookerLoadingPackage = true;
@@ -2340,16 +2260,7 @@ void UCookOnTheFlyServer::PostLoadPackageFixup(UPackage* Package)
 				FName StandardPackageFName = GetCachedStandardPackageFileFName(FName(*PackageName));
 				if (StandardPackageFName != NAME_None)
 				{
-					if (IsChildCooker())
-					{
-						check(IsCookByTheBookMode());
-						// notify the main cooker that it should make sure this package gets cooked
-						CookByTheBookOptions->ChildUnsolicitedPackages.Add(StandardPackageFName);
-					}
-					else
-					{
-						RequestPackage(StandardPackageFName, false);
-					}
+					RequestPackage(StandardPackageFName, false);
 				}
 			}
 		}
@@ -2541,22 +2452,14 @@ void UCookOnTheFlyServer::GetUnsolicitedPackages(TArray<UPackage*>& PackagesToSa
 				
 				if (StandardPackageFName != NAME_None) // if we have name none that means we are in core packages or something...
 				{
-					if (IsChildCooker())
+					// check if the package has already been saved
+					bool bAlreadyInSet = false;
+					PackagesToSaveSet.Add(Package, &bAlreadyInSet);
+					if ( bAlreadyInSet == false )
 					{
-						// notify the main cooker that it should make sure this package gets cooked
-						CookByTheBookOptions->ChildUnsolicitedPackages.Add(StandardPackageFName);
-					}
-					else
-					{
-						// check if the package has already been saved
-						bool bAlreadyInSet = false;
-						PackagesToSaveSet.Add(Package, &bAlreadyInSet);
-						if ( bAlreadyInSet == false )
-						{
-							UE_LOG(LogCook, Verbose, TEXT("Found unsolicited package to cook %s"), *Package->GetName());
-						}
-						continue;
+						UE_LOG(LogCook, Verbose, TEXT("Found unsolicited package to cook %s"), *Package->GetName());
 					}
+					continue;
 				}
 			}
 		}
@@ -4251,9 +4154,7 @@ bool UCookOnTheFlyServer::IniSettingsOutOfDate(const ITargetPlatform* TargetPlat
 
 bool UCookOnTheFlyServer::SaveCurrentIniSettings(const ITargetPlatform* TargetPlatform) const
 {
-
 	auto S = FScopeAssign<bool>(IniSettingRecurse, true);
-	check(IsChildCooker() == false);
 
 	TMap<FString, FString> AdditionalIniSettings;
 	GetAdditionalCurrentIniVersionStrings(TargetPlatform, AdditionalIniSettings);
@@ -4399,8 +4300,6 @@ void UCookOnTheFlyServer::GetAllCookedFiles(TMap<FName, FName>& UncookedPathToCo
 
 void UCookOnTheFlyServer::PopulateCookedPackagesFromDisk(const TArray<ITargetPlatform*>& Platforms)
 {
-	check(IsChildCooker() == false);
-
 	// See what files are out of date in the sandbox folder
 	for (int32 Index = 0; Index < Platforms.Num(); Index++)
 	{
@@ -4664,9 +4563,6 @@ const FString ExtractPackageNameFromObjectPath( const FString ObjectPath )
 
 void UCookOnTheFlyServer::CleanSandbox(const bool bIterative)
 {
-	//child cookers shouldn't clean sandbox we would be deleting the master cookers / other cookers data
-	check(IsChildCooker() == false);
-
 	const TArray<ITargetPlatform*>& Platforms = GetCookingTargetPlatforms();
 
 	// before we can delete any cooked files we need to make sure that we have finished writing them
@@ -4767,12 +4663,6 @@ void UCookOnTheFlyServer::GenerateAssetRegistry()
 	}
 	CookFlags |= ECookInitializationFlags::GeneratedAssetRegistry;
 
-	if (IsChildCooker())
-	{
-		// don't generate the asset registry
-		return;
-	}
-
 	double GenerateAssetRegistryTime = 0.0;
 	{
 		SCOPE_TIMER(GenerateAssetRegistryTime);
@@ -4886,25 +4776,7 @@ void UCookOnTheFlyServer::CollectFilesToCook(TArray<FName>& FilesInPath, const T
 
 	TArray<FName> InitialPackages = FilesInPath;
 
-	if (IsChildCooker())
-	{
-		const FString ChildCookFilename = CookByTheBookOptions->ChildCookFilename;
-		check(ChildCookFilename.Len());
-		FString ChildCookString;
-		ensure(FFileHelper::LoadFileToString(ChildCookString, *ChildCookFilename));
-	
-		TArray<FString> ChildCookArray;
-		ChildCookString.ParseIntoArrayLines(ChildCookArray);
 
-		for (const FString& ChildFile : ChildCookArray)
-		{
-			AddFileToCook(FilesInPath, ChildFile);
-		}
-		// if we are a child cooker then just add the files for the child cooker and early out 
-		return;
-	}
-
-	
 	TArray<FString> CookDirectories = InCookDirectories;
 	
 	if (!IsCookingDLC() && 
@@ -5423,14 +5295,9 @@ void UCookOnTheFlyServer::InitShaderCodeLibrary(void)
     }
 }
 
-static FString GenerateShaderCodeLibraryName(FString const& Name, bool bIsIterateSharedBuild, bool bIsChildCooker)
+static FString GenerateShaderCodeLibraryName(FString const& Name, bool bIsIterateSharedBuild)
 {
 	FString ActualName = (!bIsIterateSharedBuild) ? Name : Name + TEXT("_SC");
-	if (bIsChildCooker)
-	{
-		static FGuid Guid = FGuid::NewGuid();
-		ActualName += TEXT(".") + Guid.ToString();
-	}
 	return ActualName;
 }
 
@@ -5440,7 +5307,7 @@ void UCookOnTheFlyServer::OpenShaderCodeLibrary(FString const& Name)
     bool const bCacheShaderLibraries = (CurrentCookMode == ECookMode::CookByTheBook);
     if (bCacheShaderLibraries && PackagingSettings->bShareMaterialShaderCode)
 	{
-		FString ActualName = GenerateShaderCodeLibraryName(Name, IsCookFlagSet(ECookInitializationFlags::IterateSharedBuild), IsChildCooker());
+		FString ActualName = GenerateShaderCodeLibraryName(Name, IsCookFlagSet(ECookInitializationFlags::IterateSharedBuild));
 		
 		// The shader code library directory doesn't matter while cooking
 		FShaderCodeLibrary::OpenLibrary(ActualName, TEXT(""));
@@ -5534,7 +5401,7 @@ void UCookOnTheFlyServer::SaveShaderCodeLibrary(FString const& Name)
     ITargetPlatformManagerModule& TPM = GetTargetPlatformManagerRef();
     if (bCacheShaderLibraries && PackagingSettings->bShareMaterialShaderCode)
     {
-        FString ActualName = GenerateShaderCodeLibraryName(Name, IsCookFlagSet(ECookInitializationFlags::IterateSharedBuild), IsChildCooker());
+        FString ActualName = GenerateShaderCodeLibraryName(Name, IsCookFlagSet(ECookInitializationFlags::IterateSharedBuild));
         
         // Save shader code map - cleaning directories is deliebrately a separate loop here as we open the cache once per shader platform and we don't assume that they can't be shared across target platforms.
         for (const FName& TargetPlatformName : CookByTheBookOptions->TargetPlatformNames)
@@ -5555,25 +5422,19 @@ void UCookOnTheFlyServer::SaveShaderCodeLibrary(FString const& Name)
             if (ShaderFormats.Num() > 0)
             {
                 bool bSaved = false;
-                if (!IsChildCooker())
-                {
-					FString OutSCLCSVPath;
-                    bSaved = FShaderCodeLibrary::SaveShaderCodeMaster(ShaderCodeDir, MetaDataPath, ShaderFormats, OutSCLCSVPath);
-					if (OutSCLCSVPath.Len())
-					{
-						OutSCLCSVPaths.FindOrAdd(TargetPlatformName).Add(OutSCLCSVPath);
-					}
-                }
-                else
-                {
-                    bSaved = FShaderCodeLibrary::SaveShaderCodeChild(ShaderCodeDir, MetaDataPath, ShaderFormats);
-                }
+
+				FString OutSCLCSVPath;
+                bSaved = FShaderCodeLibrary::SaveShaderCodeMaster(ShaderCodeDir, MetaDataPath, ShaderFormats, OutSCLCSVPath);
+				if (OutSCLCSVPath.Len())
+				{
+					OutSCLCSVPaths.FindOrAdd(TargetPlatformName).Add(OutSCLCSVPath);
+				}
                 
                 if(!bSaved)
                 {
                     LogCookerMessage(FString::Printf(TEXT("Shared Material Shader Code Library failed for %s."),*TargetPlatformNameString), EMessageSeverity::Error);
                 }
-				else if (!IsChildCooker())
+				else
 				{
 					if (PackagingSettings->bSharedMaterialNativeLibraries)
 					{
@@ -5599,8 +5460,8 @@ void UCookOnTheFlyServer::CleanShaderCodeLibraries()
     const UProjectPackagingSettings* const PackagingSettings = GetDefault<UProjectPackagingSettings>();
     bool const bCacheShaderLibraries = (CurrentCookMode == ECookMode::CookByTheBook);
     ITargetPlatformManagerModule& TPM = GetTargetPlatformManagerRef();
-    // If not iterative and not a child cooker then clean up our temporary files
-    if (bCacheShaderLibraries && PackagingSettings->bShareMaterialShaderCode && !IsCookFlagSet(ECookInitializationFlags::Iterative) && !IsChildCooker())
+    // If not iterative then clean up our temporary files
+    if (bCacheShaderLibraries && PackagingSettings->bShareMaterialShaderCode && !IsCookFlagSet(ECookInitializationFlags::Iterative))
     {
         for (const FName& TargetPlatformName : CookByTheBookOptions->TargetPlatformNames)
         {
@@ -5632,34 +5493,8 @@ void UCookOnTheFlyServer::CookByTheBookFinished()
 
 	const UProjectPackagingSettings* const PackagingSettings = GetDefault<UProjectPackagingSettings>();
 	bool const bCacheShaderLibraries = (CurrentCookMode == ECookMode::CookByTheBook);
-	if (IsChildCooker())
-	{
-		// if we are the child cooker create a list of all the packages which we think wes hould have cooked but didn't
-		FString UncookedPackageList;
-		for (const FName& UncookedPackage : CookByTheBookOptions->ChildUnsolicitedPackages)
-		{
-			UncookedPackageList.Append(UncookedPackage.ToString() + TEXT("\n\r"));
-		}
-		FFileHelper::SaveStringToFile(UncookedPackageList, *GetChildCookerResultFilename(CookByTheBookOptions->ChildCookFilename));
-		if (IBlueprintNativeCodeGenModule::IsNativeCodeGenModuleLoaded())
-		{
-			IBlueprintNativeCodeGenModule::Get().SaveManifest();
-		}
-		
-		if (bCacheShaderLibraries && PackagingSettings->bShareMaterialShaderCode)
-		{
-			// Save shader code map
-			FString LibraryName = !IsCookingDLC() ? FApp::GetProjectName() : CookByTheBookOptions->DlcName;
-			SaveShaderCodeLibrary(LibraryName);
-			
-			FShaderCodeLibrary::Shutdown();
 
-		}
-	}
-	else
 	{
-		CleanUpChildCookers();
-
 		if (IBlueprintNativeCodeGenModule::IsNativeCodeGenModuleLoaded())
 		{
 			SCOPE_TIMER(GeneratingBlueprintAssets)
@@ -5668,20 +5503,12 @@ void UCookOnTheFlyServer::CookByTheBookFinished()
 			CodeGenModule.GenerateFullyConvertedClasses(); // While generating fully converted classes the list of necessary stubs is created.
 			CodeGenModule.GenerateStubs();
 
-			// merge the manifest for the blueprint code generator:
-			for (int32 I = 0; I < CookByTheBookOptions->ChildCookers.Num(); ++I )
-			{
-				CodeGenModule.MergeManifest(I);
-			}
-
 			CodeGenModule.FinalizeManifest();
 
 			// Unload the module as we only need it while cooking. This will also clear the current module's state in order to allow a new cooker pass to function properly.
 			FModuleManager::Get().UnloadModule(CodeGenModule.GetModuleName());
 		}
 
-		check(CookByTheBookOptions->ChildUnsolicitedPackages.Num() == 0);
-		
 		// Save modified asset registry with all streaming chunk info generated during cook
 		const FString& SandboxRegistryFilename = GetSandboxAssetRegistryFilename();
 
@@ -5814,7 +5641,7 @@ void UCookOnTheFlyServer::CookByTheBookFinished()
 		}
 	}
 
-	if (CookByTheBookOptions->bGenerateDependenciesForMaps && (IsChildCooker() == false))
+	if (CookByTheBookOptions->bGenerateDependenciesForMaps)
 	{
 		SCOPE_TIMER(GenerateMapDependencies);
 		for (auto& MapDependencyGraphIt : CookByTheBookOptions->MapDependencyGraphs)
@@ -6006,7 +5833,7 @@ void UCookOnTheFlyServer::InitializeSandbox()
 		CreateSandboxFile();
 
 		// When looking for deterministic cooking differences in cooked packages, don't delete the packages on disk
-		if (!IsChildCooker() && !FParse::Param(FCommandLine::Get(), TEXT("DIFFONLY")))
+		if (!FParse::Param(FCommandLine::Get(), TEXT("DIFFONLY")))
 		{
 			bIsInitializingSandbox = true;
 			CleanSandbox(IsCookFlagSet(ECookInitializationFlags::Iterative));
@@ -6034,12 +5861,6 @@ void UCookOnTheFlyServer::ValidateCookOnTheFlySettings() const
 
 void UCookOnTheFlyServer::ValidateCookByTheBookSettings() const
 {
-	if (IsChildCooker())
-	{
-		// should never be generating dependency maps / streaming install manifests for child cookers
-		check(CookByTheBookOptions->bGenerateDependenciesForMaps == false);
-		check(CookByTheBookOptions->bGenerateStreamingInstallManifests == false);
-	}
 }
 
 void UCookOnTheFlyServer::StartCookByTheBook( const FCookByTheBookStartupOptions& CookByTheBookStartupOptions )
@@ -6066,10 +5887,8 @@ void UCookOnTheFlyServer::StartCookByTheBook( const FCookByTheBookStartupOptions
 	CookByTheBookOptions->bGenerateStreamingInstallManifests = CookByTheBookStartupOptions.bGenerateStreamingInstallManifests;
 	CookByTheBookOptions->bGenerateDependenciesForMaps = CookByTheBookStartupOptions.bGenerateDependenciesForMaps;
 	CookByTheBookOptions->CreateReleaseVersion = CreateReleaseVersion;
-	CookByTheBookOptions->ChildCookFilename = CookByTheBookStartupOptions.ChildCookFileName;
 	CookByTheBookOptions->bDisableUnsolicitedPackages = !!(CookOptions & ECookByTheBookOptions::DisableUnsolicitedPackages);
 	CookByTheBookOptions->bFullLoadAndSave = !!(CookOptions & ECookByTheBookOptions::FullLoadAndSave);
-	CookByTheBookOptions->ChildCookIdentifier = CookByTheBookStartupOptions.ChildCookIdentifier;
 	CookByTheBookOptions->bErrorOnEngineContentUse = CookByTheBookStartupOptions.bErrorOnEngineContentUse;
 
 	GenerateAssetRegistry();
@@ -6205,7 +6024,6 @@ void UCookOnTheFlyServer::StartCookByTheBook( const FCookByTheBookStartupOptions
 			IBlueprintNativeCodeGenModule::Get().FillPlatformNativizationDetails(Entry, PlatformNativizationDetails);
 			CodeGenData.CodegenTargets.Push(PlatformNativizationDetails);
 		}
-		CodeGenData.ManifestIdentifier = CookByTheBookStartupOptions.ChildCookIdentifier;
 		IBlueprintNativeCodeGenModule::InitializeModule(CodeGenData);
 	}
 
@@ -6221,7 +6039,6 @@ void UCookOnTheFlyServer::StartCookByTheBook( const FCookByTheBookStartupOptions
 		}
 	}
 
-	if (!IsChildCooker())
 	{
 		for (ITargetPlatform* Platform : TargetPlatforms)
 		{
@@ -6290,7 +6107,7 @@ void UCookOnTheFlyServer::StartCookByTheBook( const FCookByTheBookStartupOptions
 	}
 	
 	// don't resave the global shader map files in dlc
-	if (!IsCookingDLC() && !IsChildCooker() && !(CookByTheBookStartupOptions.CookOptions & ECookByTheBookOptions::ForceDisableSaveGlobalShaders))
+	if (!IsCookingDLC() && !(CookByTheBookStartupOptions.CookOptions & ECookByTheBookOptions::ForceDisableSaveGlobalShaders))
 	{
 		OpenShaderCodeLibrary(TEXT("Global"));
 
@@ -6380,7 +6197,7 @@ void UCookOnTheFlyServer::StartCookByTheBook( const FCookByTheBookStartupOptions
 	}
 
 
-	if ( !IsCookingDLC() && !IsChildCooker() )
+	if (!IsCookingDLC())
 	{
 		// if we are not cooking dlc then basedOnRelease version just needs to make sure that we cook all the packages which are in the previous release (as well as the new ones)
 		if ( !BasedOnReleaseVersion.IsEmpty() )
@@ -6419,12 +6236,7 @@ void UCookOnTheFlyServer::StartCookByTheBook( const FCookByTheBookStartupOptions
 		CookRequests.EnqueueUnique( PreviousRequest );
 	}
 	CookByTheBookOptions->PreviousCookRequests.Empty();
-	
-	if (CookByTheBookStartupOptions.NumProcesses)
-	{
-		FString ExtraCommandLine;
-		StartChildCookers(CookByTheBookStartupOptions.NumProcesses, TargetPlatformNames, ExtraCommandLine);
-	}
+
 }
 
 bool UCookOnTheFlyServer::RecompileChangedShaders(const TArray<FName>& TargetPlatforms)
@@ -6437,288 +6249,6 @@ bool UCookOnTheFlyServer::RecompileChangedShaders(const TArray<FName>& TargetPla
 	return bShadersRecompiled;
 }
 
-class FChildCookerRunnable : public FRunnable
-{
-private:
-	UCookOnTheFlyServer* CookServer;
-	UCookOnTheFlyServer::FChildCooker* ChildCooker;
-
-	void ProcessChildLogOutput()
-	{
-		// process the log output from the child cooker even if we just finished to ensure that we get the end of the log
-		FString PipeContents = FPlatformProcess::ReadPipe(ChildCooker->ReadPipe);
-		if (PipeContents.Len())
-		{
-			TArray<FString> PipeLines;
-			PipeContents.ParseIntoArrayLines(PipeLines);
-
-			for (const FString& Line : PipeLines)
-			{
-				UE_LOG(LogCook, Display, TEXT("Cooker output %s: %s"), *ChildCooker->BaseResponseFileName, *Line);
-				if (ChildCooker->bFinished == false)
-				{
-					FPlatformProcess::Sleep(0.0f); // don't be greedy, log output isn't important compared to cooking performance
-				}
-			}
-		}
-	}
-
-public:
-	FChildCookerRunnable(UCookOnTheFlyServer::FChildCooker* InChildCooker, UCookOnTheFlyServer* InCookServer) : FRunnable(), CookServer(InCookServer), ChildCooker(InChildCooker)
-	{ }
-
-	
-	virtual uint32 Run() override
-	{
-		while (true)
-		{
-			check(ChildCooker != nullptr);
-			check(ChildCooker->bFinished == false); 
-
-			int32 ReturnCode;
-			if (FPlatformProcess::GetProcReturnCode(ChildCooker->ProcessHandle, &ReturnCode))
-			{
-				if (ReturnCode != 0)
-				{
-					UE_LOG(LogCook, Error, TEXT("Child cooker %s returned error code %d"), *ChildCooker->BaseResponseFileName, ReturnCode);
-				}
-				else
-				{
-					UE_LOG(LogCook, Display, TEXT("Child cooker %s returned %d"), *ChildCooker->BaseResponseFileName, ReturnCode);
-				}
-
-
-				// if the child cooker completed successfully then it would have output a list of files which the main cooker needs to process
-				const FString AdditionalPackagesFileName = GetChildCookerResultFilename(ChildCooker->ResponseFileName);
-
-				FString AdditionalPackages;
-				if (FFileHelper::LoadFileToString(AdditionalPackages, *AdditionalPackagesFileName) == false)
-				{
-					UE_LOG(LogCook, Fatal, TEXT("ChildCooker failed to write out additional packages file to location %s"), *AdditionalPackagesFileName);
-				}
-				TArray<FString> AdditionalPackageList;
-				AdditionalPackages.ParseIntoArrayLines(AdditionalPackageList);
-
-				for (const FString& AdditionalPackage : AdditionalPackageList)
-				{
-					UE_LOG(LogCook, Display, TEXT("Child cooker %s requested additional package %s to be cooked"), *ChildCooker->BaseResponseFileName, *AdditionalPackage);
-					FName Filename = FName(*AdditionalPackage);
-					CookServer->RequestPackage(Filename, false);
-				}
-				ProcessChildLogOutput();
-
-				ChildCooker->ReturnCode = ReturnCode;
-				ChildCooker->bFinished = true;
-				ChildCooker = nullptr;
-				return 1;
-			}
-
-			ProcessChildLogOutput();
-
-			FPlatformProcess::Sleep(0.05f);
-		}
-		return 1;
-	}
-
-};
-
-// sous chefs away!!
-void UCookOnTheFlyServer::StartChildCookers(int32 NumCookersToSpawn, const TArray<FName>& TargetPlatformNames, const FString& ExtraCmdParams)
-{
-	SCOPE_TIMER(StartingChildCookers);
-	// create a comprehensive list of all the files we need to cook
-	// then get the packages with least dependencies and give them to some sous chefs to handle
-	
-	check(!IsChildCooker());
-
-	// PackageNames contains a sorted list of the packages we want to distribute to child cookers
-	TArray<FName> PackageNames;
-	// PackageNamesSet is a set of the same information for quick lookup into packagenames array
-	TSet<FName> PackageNamesSet;
-	PackageNames.Empty(CookRequests.Num());
-
-	// TArray<FFilePlatformRequest>;
-
-	
-
-	for (const FName& CookRequest : CookRequests.GetQueue())
-	{
-		FString LongPackageName = FPackageName::FilenameToLongPackageName(CookRequest.ToString());
-
-		const FName PackageFName = FName(*LongPackageName);
-		PackageNames.Add(PackageFName);
-		PackageNamesSet.Add(PackageFName);
-	}
-
-
-	int32 PackageCounter = 0;
-	while (PackageCounter < PackageNames.Num())
-	{
-		const FName& PackageName = PackageNames[PackageCounter];
-		++PackageCounter;
-		TArray<FName> UnfilteredDependencies;
-		AssetRegistry->GetDependencies(PackageName, UnfilteredDependencies);
-
-		for (const FName& Dependency : UnfilteredDependencies)
-		{
-			if ((FPackageName::IsScriptPackage(Dependency.ToString()) == false) && (FPackageName::IsMemoryPackage(Dependency.ToString()) == false))
-			{
-				if (PackageNamesSet.Contains(Dependency) == false)
-				{
-					PackageNamesSet.Add(Dependency);
-					PackageNames.Insert(Dependency, PackageCounter);
-				}
-			}
-		}
-	}
-
-	TArray<FName> DistributeStandardFilenames;
-	for (const FName& DistributeCandidate : PackageNames)
-	{
-		FText OutReason;
-		FString LongPackageName = DistributeCandidate.ToString();
-		if (!FPackageName::IsValidLongPackageName(LongPackageName, true, &OutReason))
-		{
-			const FText FailMessage = FText::Format(LOCTEXT("UnableToGeneratePackageName", "Unable to generate long package name for {0}. {1}"),
-				FText::FromString(LongPackageName), OutReason);
-
-			LogCookerMessage(FailMessage.ToString(), EMessageSeverity::Warning);
-			UE_LOG(LogCook, Warning, TEXT("%s"), *(FailMessage.ToString()));
-			continue;
-		}
-		else if (FPackageName::IsScriptPackage(LongPackageName)|| FPackageName::IsMemoryPackage(LongPackageName))
-		{
-			continue;
-		}
-		DistributeStandardFilenames.Add(FName(*LongPackageName));
-	}
-
-	UE_LOG(LogCook, Display, TEXT("Distributing %d packages to %d cookers for processing"), DistributeStandardFilenames.Num(), NumCookersToSpawn);
-
-	FString TargetPlatformString;
-	for (const FName& TargetPlatformName : TargetPlatformNames)
-	{
-		if (TargetPlatformString.Len() != 0)
-		{
-			TargetPlatformString += TEXT("+");
-		}
-		TargetPlatformString += TargetPlatformName.ToString();
-	}
-
-	// allocate the memory here, this can't change while we are running the child threads which are handling input because they have a ptr to the childcookers array
-	CookByTheBookOptions->ChildCookers.Empty(NumCookersToSpawn);
-
-	// start the child cookers and give them each some distribution candidates
-	for (int32 CookerCounter = 0; CookerCounter < NumCookersToSpawn; ++CookerCounter )
-	{
-		// count our selves as a cooker
-		int32 NumFilesForCooker = DistributeStandardFilenames.Num() / ((NumCookersToSpawn+1) - CookerCounter);
-
-		// don't spawn a cooker unless it has a minimum amount of files to do
-		if (NumFilesForCooker < 5)
-		{
-			continue;
-		}
-
-
-		int32 ChildCookerIndex = CookByTheBookOptions->ChildCookers.AddDefaulted(1);
-		FChildCooker& ChildCooker = CookByTheBookOptions->ChildCookers[ChildCookerIndex];
-
-		ChildCooker.ResponseFileName = FPaths::CreateTempFilename(*(FPaths::ProjectSavedDir() / TEXT("CookingTemp")));
-		ChildCooker.BaseResponseFileName = FPaths::GetBaseFilename(ChildCooker.ResponseFileName);
-
-		// FArchive* ResponseFile = IFileManager::CreateFileWriter(ChildCooker.UniqueTempName);
-		FString ResponseFileText;
-
-		for (int32 I = 0; I < NumFilesForCooker; ++I)
-		{
-			FName PackageFName = DistributeStandardFilenames[I];
-			FString PackageName = PackageFName.ToString();
-
-			ResponseFileText += FString::Printf(TEXT("%s%s"), *PackageName, LINE_TERMINATOR);
-
-			// these are long package names
-			FName StandardPackageName = GetCachedStandardPackageFileFName(PackageFName);
-			if (StandardPackageName == NAME_None)
-				continue;
-#if 0 // validation code can be disabled
-			PackageName = FPackageName::LongPackageNameToFilename(PackageName,FPaths::GetExtension(PackageName));
-			FPaths::MakeStandardFilename(PackageName);
-			FName PackageFName = FName(*PackageName);
-			check(PackageFName== StandardPackageName);
-#endif	
-			TArray<bool> Succeeded;
-			for ( int32 T = 0; T < TargetPlatformNames.Num(); ++T )
-			{
-				Succeeded.Add(true);
-			}
-			CookedPackages.Add(FFilePlatformCookedPackage(StandardPackageName, TargetPlatformNames, MoveTemp(Succeeded)));
-		}
-		DistributeStandardFilenames.RemoveAt(0, NumFilesForCooker);
-
-		UE_LOG(LogCook, Display, TEXT("Child cooker %d working on %d files"), CookerCounter, NumFilesForCooker);
-
-		FFileHelper::SaveStringToFile(ResponseFileText, *ChildCooker.ResponseFileName);
-
-		// default commands
-		// multiprocess tells unreal in general we shouldn't do things like save ddc, clean shader working directory, and other various multiprocess unsafe things
-		FString CommandLine = FString::Printf(TEXT("\"%s\" -run=cook -multiprocess -targetplatform=%s -cookchild=\"%s\" -abslog=\"%sLog.txt\" -childIdentifier=%d %s"), 
-			*FPaths::GetProjectFilePath(), *TargetPlatformString, *ChildCooker.ResponseFileName, *ChildCooker.ResponseFileName, CookerCounter, *ExtraCmdParams);
-
-		auto KeepCommandlineValue = [&](const TCHAR* CommandlineToKeep)
-		{
-			FString CommandlineValue;
-			if (FParse::Value(FCommandLine::Get(), CommandlineToKeep, CommandlineValue ))
-			{
-				CommandLine += TEXT(" -");
-				CommandLine += CommandlineToKeep;
-				CommandLine += TEXT("=");
-				CommandLine += *CommandlineValue;
-			}
-		};
-
-		auto KeepCommandlineParam = [&](const TCHAR* CommandlineToKeep)
-		{
-			if (FParse::Param(FCommandLine::Get(), CommandlineToKeep))
-			{
-				CommandLine += TEXT(" -");
-				CommandLine += CommandlineToKeep;
-			}
-		};
-
-
-		KeepCommandlineParam(TEXT("NativizeAssets"));
-		KeepCommandlineValue(TEXT("ddc="));
-		KeepCommandlineParam(TEXT("SkipEditorContent"));
-		KeepCommandlineParam(TEXT("compressed"));
-		KeepCommandlineParam(TEXT("Unversioned"));
-		KeepCommandlineParam(TEXT("buildmachine"));
-		KeepCommandlineParam(TEXT("fileopenlog"));
-		KeepCommandlineParam(TEXT("stdout"));
-		KeepCommandlineParam(TEXT("FORCELOGFLUSH"));
-		KeepCommandlineParam(TEXT("CrashForUAT"));
-		KeepCommandlineParam(TEXT("AllowStdOutLogVerbosity"));
-		KeepCommandlineParam(TEXT("UTF8Output"));
-
-		FString ExecutablePath = FPlatformProcess::ExecutableName(true);
-
-		UE_LOG(LogCook, Display, TEXT("Launching cooker using commandline %s %s"), *ExecutablePath, *CommandLine);
-
-		void* ReadPipe = NULL;
-		void* WritePipe = NULL;
-		FPlatformProcess::CreatePipe(ReadPipe, WritePipe);
-
-		ChildCooker.ProcessHandle = FPlatformProcess::CreateProc(*ExecutablePath, *CommandLine, false, true, true, NULL, 0, NULL, WritePipe);
-		ChildCooker.ReadPipe = ReadPipe;
-
-		// start threads to monitor output and finished state
-		ChildCooker.Thread = FRunnableThread::Create(new FChildCookerRunnable(&ChildCooker, this), *FString::Printf(TEXT("ChildCookerInputHandleThreadFor:%s"), *ChildCooker.BaseResponseFileName));
-	}
-
-}
-
-
-
 /* UCookOnTheFlyServer callbacks
  *****************************************************************************/
 
