Removed FullGCAssetClasses logic - no longer necessary nor useful

diff --git a/Engine/Source/Editor/UnrealEd/Classes/CookOnTheSide/CookOnTheFlyServer.h b/Engine/Source/Editor/UnrealEd/Classes/CookOnTheSide/CookOnTheFlyServer.h
--- a/Engine/Source/Editor/UnrealEd/Classes/CookOnTheSide/CookOnTheFlyServer.h
+++ b/Engine/Source/Editor/UnrealEd/Classes/CookOnTheSide/CookOnTheFlyServer.h
@@ -853,7 +853,7 @@ private:
 
 	//////////////////////////////////////////////////////////////////////////
 	// General cook options
-	TArray<UClass*> FullGCAssetClasses;
+
 	/** Number of packages to load before performing a garbage collect. Set to 0 to never GC based on number of loaded packages */
 	uint32 PackagesPerGC;
 	/** Amount of time that is allowed to be idle before forcing a garbage collect. Set to 0 to never force GC due to idle time */
@@ -866,7 +866,7 @@ private:
 	uint64 MinFreeMemory;
 	/** Max number of packages to save before we partial gc */
 	int32 MaxNumPackagesBeforePartialGC;
-	/** Max number of conncurrent shader jobs reducing this too low will increase cook time */
+	/** Max number of concurrent shader jobs reducing this too low will increase cook time */
 	int32 MaxConcurrentShaderJobs;
 	ECookInitializationFlags CookFlags;
 	TUniquePtr<class FSandboxPlatformFile> SandboxFile;
@@ -1081,14 +1081,12 @@ public:
 		bool bGenerateStreamingInstallManifests; 
 		bool bGenerateDependenciesForMaps; 
 		bool bErrorOnEngineContentUse; // this is a flag for dlc, will cause the cooker to error if the dlc references engine content
-		int32 NumProcesses;
 		FCookByTheBookStartupOptions() :
 			CookOptions(ECookByTheBookOptions::None),
 			DLCName(FString()),
 			bGenerateStreamingInstallManifests(false),
 			bGenerateDependenciesForMaps(false),
-			bErrorOnEngineContentUse(false),
-			NumProcesses(0)
+			bErrorOnEngineContentUse(false)
 		{ }
 	};
 
@@ -1125,7 +1123,7 @@ public:
 	/**
 	* Get any packages which are in memory, these were probably required to be loaded because of the current package we are cooking, so we should probably cook them also
 	*/
-	void GetUnsolicitedPackages(TArray<UPackage*>& PackagesToSave, bool &ContainsFullGCAssetClasses, const TArray<FName>& TargetPlatformNames) const;
+	void GetUnsolicitedPackages(TArray<UPackage*>& PackagesToSave, const TArray<FName>& TargetPlatformNames) const;
 
 	/**
 	* PostLoadPackageFixup
@@ -1221,12 +1219,6 @@ public:
 
 	virtual void BeginDestroy() override;
 
-	/**
-	* SetFullGCAssetClasses FullGCAssetClasses is used to determine when TickCookOnTheSide returns RequiresGC
-	*   When one of these classes is saved it will return COSR_RequiresGC
-	*/
-	void SetFullGCAssetClasses( const TArray<UClass*>& InFullGCAssetClasses );
-
 	/** Returns the configured number of packages to process before GC */
 	uint32 GetPackagesPerGC() const;
 
@@ -1447,7 +1439,7 @@ private:
 	 * @param ContainsFullAssetGCClasses do these packages contain any of the assets which require a GC after cooking 
 	 *				(this is mostly historical for when objects like UWorld were global, almost nothing should require a GC to work correctly after being cooked anymore).
 	 */
-	void GetAllUnsolicitedPackages(TArray<UPackage*>& PackagesToSave, const TArray<FName>& TargetPlatformNames, bool& ContainsFullAssetGCClasses);
+	void GetAllUnsolicitedPackages(TArray<UPackage*>& PackagesToSave, const TArray<FName>& TargetPlatformNames);
 
 
 	/**
diff --git a/Engine/Source/Editor/UnrealEd/Private/Commandlets/CookCommandlet.cpp b/Engine/Source/Editor/UnrealEd/Private/Commandlets/CookCommandlet.cpp
--- a/Engine/Source/Editor/UnrealEd/Private/Commandlets/CookCommandlet.cpp
+++ b/Engine/Source/Editor/UnrealEd/Private/Commandlets/CookCommandlet.cpp
@@ -640,23 +640,6 @@ bool UCookCommandlet::CookByTheBook( const TArray<ITargetPlatform*>& Platforms,
 	CookFlags |= Switches.Contains(TEXT("IgnoreIniSettingsOutOfDate")) || CookerSettings->bIgnoreIniSettingsOutOfDateForIteration ? ECookInitializationFlags::IgnoreIniSettingsOutOfDate : ECookInitializationFlags::None;
 	CookFlags |= Switches.Contains(TEXT("IgnoreScriptPackagesOutOfDate")) || CookerSettings->bIgnoreScriptPackagesOutOfDateForIteration ? ECookInitializationFlags::IgnoreScriptPackagesOutOfDate : ECookInitializationFlags::None;
 
-	TArray<UClass*> FullGCAssetClasses;
-	if (FullGCAssetClassNames.Num())
-	{
-		for (const auto& ClassName : FullGCAssetClassNames)
-		{
-			UClass* ClassToForceFullGC = FindObject<UClass>(nullptr, *ClassName);
-			if (ClassToForceFullGC)
-			{
-				FullGCAssetClasses.Add(ClassToForceFullGC);
-			}
-			else
-			{
-				UE_LOG(LogCookCommandlet, Warning, TEXT("Configured to force full GC for assets of type (%s) but that class does not exist."), *ClassName);
-			}
-		}
-	}
-
 	//////////////////////////////////////////////////////////////////////////
 	// parse commandline options 
 
@@ -666,12 +649,6 @@ bool UCookCommandlet::CookByTheBook( const TArray<ITargetPlatform*>& Platforms,
 	FString ChildCookFile;
 	FParse::Value(*Params, TEXT("cookchild="), ChildCookFile);
 
-	int32 ChildCookIdentifier = -1;
-	FParse::Value(*Params, TEXT("childIdentifier="), ChildCookIdentifier);
-
-	int32 NumProcesses = 0;
-	FParse::Value(*Params, TEXT("numcookerstospawn="), NumProcesses);
-
 	FString BasedOnReleaseVersion;
 	FParse::Value( *Params, TEXT("BasedOnReleaseVersion="), BasedOnReleaseVersion);
 
@@ -747,12 +724,6 @@ bool UCookCommandlet::CookByTheBook( const TArray<ITargetPlatform*>& Platforms,
 
 	CookOnTheFlyServer->Initialize(ECookMode::CookByTheBook, CookFlags, OutputDirectoryOverride);
 
-	// for backwards compat use the FullGCAssetClasses that we got from the cook commandlet ini section
-	if (FullGCAssetClasses.Num() > 0)
-	{
-		CookOnTheFlyServer->SetFullGCAssetClasses(FullGCAssetClasses);
-	}
-
 	// Add any map sections specified on command line
 	TArray<FString> AlwaysCookMapList;
 
@@ -859,7 +830,6 @@ bool UCookCommandlet::CookByTheBook( const TArray<ITargetPlatform*>& Platforms,
 	StartupOptions.bErrorOnEngineContentUse = bErrorOnEngineContentUse;
 	StartupOptions.bGenerateDependenciesForMaps = Switches.Contains(TEXT("GenerateDependenciesForMaps"));
 	StartupOptions.bGenerateStreamingInstallManifests = bGenerateStreamingInstallManifests;
-	StartupOptions.NumProcesses = NumProcesses;
 
 	COOK_STAT(
 	{
diff --git a/Engine/Source/Editor/UnrealEd/Private/CookOnTheFlyServer.cpp b/Engine/Source/Editor/UnrealEd/Private/CookOnTheFlyServer.cpp
--- a/Engine/Source/Editor/UnrealEd/Private/CookOnTheFlyServer.cpp
+++ b/Engine/Source/Editor/UnrealEd/Private/CookOnTheFlyServer.cpp
@@ -1543,14 +1543,7 @@ uint32 UCookOnTheFlyServer::TickCookOnTheSide( const float TimeSlice, uint32 &Co
 		int32 FirstUnsolicitedPackage = PackagesToSave.Num();
 
 		UE_LOG(LogCook, Verbose, TEXT("Finding unsolicited packages for package %s"), *ToBuild.GetFilename().ToString());
-		bool ContainsFullAssetGCClasses = false;
-		GetAllUnsolicitedPackages(PackagesToSave, AllTargetPlatformNames, ContainsFullAssetGCClasses);
-		if (ContainsFullAssetGCClasses)
-		{
-			Result |= COSR_RequiresGC;
-		}
-
-
+		GetAllUnsolicitedPackages(PackagesToSave, AllTargetPlatformNames);
 		
 		// in cook by the book bail out early because shaders and compiled for the primary package we are trying to save. 
 		// note in this case we also put the package at the end of the queue that queue might be reordered if we do partial gc
@@ -1839,7 +1832,6 @@ void UCookOnTheFlyServer::OpportunisticSaveInMemoryPackages()
 	const float TimeSlice = 0.01f;
 	FCookerTimer Timer(TimeSlice, IsRealtimeMode());
 	TArray<UPackage*> PackagesToSave;
-	bool ContainsFullAssetGCClasses = false;
 
 	TArray<FName> TargetPlatformNames;
 	for (const auto& TargetPlatform : PresaveTargetPlatforms)
@@ -1848,7 +1840,7 @@ void UCookOnTheFlyServer::OpportunisticSaveInMemoryPackages()
 	}
 
 
-	GetAllUnsolicitedPackages(PackagesToSave, TargetPlatformNames, ContainsFullAssetGCClasses);
+	GetAllUnsolicitedPackages(PackagesToSave, TargetPlatformNames);
 
 	
 	uint32 CookedPackageCount = 0;
@@ -1856,7 +1848,7 @@ void UCookOnTheFlyServer::OpportunisticSaveInMemoryPackages()
 	SaveCookedPackages( PackagesToSave, TargetPlatformNames, PresaveTargetPlatforms, Timer, 0, CookedPackageCount, Result);
 }
 
-void UCookOnTheFlyServer::GetAllUnsolicitedPackages(TArray<UPackage*>& PackagesToSave, const TArray<FName>& TargetPlatformNames, bool& ContainsFullAssetGCClasses )
+void UCookOnTheFlyServer::GetAllUnsolicitedPackages(TArray<UPackage*>& PackagesToSave, const TArray<FName>& TargetPlatformNames)
 {
 	// generate a list of other packages which were loaded with this one
 	if (!IsCookByTheBookMode() || (CookByTheBookOptions->bDisableUnsolicitedPackages == false))
@@ -1874,7 +1866,7 @@ void UCookOnTheFlyServer::GetAllUnsolicitedPackages(TArray<UPackage*>& PackagesT
 			}
 		}
 		SCOPE_TIMER(UnsolicitedMarkup);
-		GetUnsolicitedPackages(PackagesToSave, ContainsFullAssetGCClasses, TargetPlatformNames);
+		GetUnsolicitedPackages(PackagesToSave, TargetPlatformNames);
 	}
 }
 
@@ -2408,7 +2400,7 @@ bool UCookOnTheFlyServer::HasExceededMaxMemory() const
 	return false;
 }
 
-void UCookOnTheFlyServer::GetUnsolicitedPackages(TArray<UPackage*>& PackagesToSave, bool &ContainsFullGCAssetClasses, const TArray<FName>& TargetPlatformNames) const
+void UCookOnTheFlyServer::GetUnsolicitedPackages(TArray<UPackage*>& PackagesToSave, const TArray<FName>& TargetPlatformNames) const
 {
 	TSet<UPackage*> PackagesToSaveSet;
 	for (UPackage* Package : PackagesToSave)
@@ -2431,12 +2423,6 @@ void UCookOnTheFlyServer::GetUnsolicitedPackages(TArray<UPackage*>& PackagesToSa
 		{
 			UPackage* Package = Cast<UPackage>(ObjectsInOuter[Index]);
 
-			UObject* Object = ObjectsInOuter[Index];
-			if (FullGCAssetClasses.Contains(Object->GetClass()))
-			{
-				ContainsFullGCAssetClasses = true;
-			}
-
 			if (Package)
 			{
 				FName StandardPackageFName = GetCachedStandardPackageFileFName(Package);
@@ -2621,11 +2607,6 @@ void UCookOnTheFlyServer::EndNetworkFileServer()
 	NetworkFileServers.Empty();
 }
 
-void UCookOnTheFlyServer::SetFullGCAssetClasses( const TArray<UClass*>& InFullGCAssetClasses)
-{
-	FullGCAssetClasses = InFullGCAssetClasses;
-}
-
 uint32 UCookOnTheFlyServer::GetPackagesPerGC() const
 {
 	return PackagesPerGC;
@@ -3289,39 +3270,12 @@ void UCookOnTheFlyServer::Initialize( ECookMode::Type DesiredCookMode, ECookInit
 	FCoreDelegates::OnFConfigCreated.AddUObject(this, &UCookOnTheFlyServer::OnFConfigCreated);
 	FCoreDelegates::OnFConfigDeleted.AddUObject(this, &UCookOnTheFlyServer::OnFConfigDeleted);
 
-	bool bUseFullGCAssetClassNames = true;
-	GConfig->GetBool(TEXT("CookSettings"), TEXT("bUseFullGCAssetClassNames"), bUseFullGCAssetClassNames, GEditorIni);
-	
 	MaxPrecacheShaderJobs = FPlatformMisc::NumberOfCores() - 1; // number of cores -1 is a good default allows the editor to still be responsive to other shader requests and allows cooker to take advantage of multiple processors while the editor is running
 	GConfig->GetInt(TEXT("CookSettings"), TEXT("MaxPrecacheShaderJobs"), MaxPrecacheShaderJobs, GEditorIni);
 
 	MaxConcurrentShaderJobs = FPlatformMisc::NumberOfCores() * 4; // number of cores -1 is a good default allows the editor to still be responsive to other shader requests and allows cooker to take advantage of multiple processors while the editor is running
 	GConfig->GetInt(TEXT("CookSettings"), TEXT("MaxConcurrentShaderJobs"), MaxConcurrentShaderJobs, GEditorIni);
 
-	if (bUseFullGCAssetClassNames)
-	{
-		TArray<FString> FullGCAssetClassNames;
-		GConfig->GetArray( TEXT("CookSettings"), TEXT("FullGCAssetClassNames"), FullGCAssetClassNames, GEditorIni );
-		for ( const FString& FullGCAssetClassName : FullGCAssetClassNames )
-		{
-			UClass* FullGCAssetClass = FindObject<UClass>(ANY_PACKAGE, *FullGCAssetClassName, true);
-			if( FullGCAssetClass == NULL )
-			{
-				UE_LOG(LogCook, Warning, TEXT("Unable to find full gc asset class name %s may result in bad cook"), *FullGCAssetClassName);
-			}
-			else
-			{
-				FullGCAssetClasses.Add( FullGCAssetClass );
-			}
-		}
-		if ( FullGCAssetClasses.Num() == 0 )
-		{
-			// default to UWorld
-			FullGCAssetClasses.Add( UWorld::StaticClass() );
-		}
-	}
-
-
 	ITargetPlatformManagerModule& TPM = GetTargetPlatformManagerRef();
 	TArray<FString> PresaveTargetPlatformNames;
 	if ( GConfig->GetArray(TEXT("CookSettings"), TEXT("PresaveTargetPlatforms"), PresaveTargetPlatformNames, GEditorIni ) )
@@ -5132,9 +5086,8 @@ void UCookOnTheFlyServer::CollectFilesToCook(TArray<FName>& FilesInPath, const T
 	{
 		// Gather initial unsolicited package list, this is needed in iterative mode because it may skip cooking all explicit packages and never hit this code
 		TArray<UPackage*> UnsolicitedPackages;
-		bool ContainsFullAssetGCClasses = false;
 		UE_LOG(LogCook, Verbose, TEXT("Finding initial unsolicited packages"));
-		GetUnsolicitedPackages(UnsolicitedPackages, ContainsFullAssetGCClasses, CookByTheBookOptions->TargetPlatformNames);
+		GetUnsolicitedPackages(UnsolicitedPackages, CookByTheBookOptions->TargetPlatformNames);
 
 		for (UPackage* UnsolicitedPackage : UnsolicitedPackages)
 		{
