Fixed bug in UCookOnTheFlyServer::GetCookOnTheFlyUnsolicitedFiles - UnsolicitedFiles should be passed by reference not by value

diff --git a/Engine/Source/Editor/UnrealEd/Classes/CookOnTheSide/CookOnTheFlyServer.h b/Engine/Source/Editor/UnrealEd/Classes/CookOnTheSide/CookOnTheFlyServer.h
--- a/Engine/Source/Editor/UnrealEd/Classes/CookOnTheSide/CookOnTheFlyServer.h
+++ b/Engine/Source/Editor/UnrealEd/Classes/CookOnTheSide/CookOnTheFlyServer.h
@@ -1369,7 +1369,7 @@ private:
 	 */
 	bool HandleNetworkFileServerNewConnection( const FString& VersionInfo, const FString& PlatformName );
 
-	void GetCookOnTheFlyUnsolicitedFiles(const FName& PlatformName, TArray<FString> UnsolicitedFiles, const FString& Filename);
+	void GetCookOnTheFlyUnsolicitedFiles(const FName& PlatformName, TArray<FString>& UnsolicitedFiles, const FString& Filename);
 
 	/**
 	* Cook requests for a package from network
diff --git a/Engine/Source/Editor/UnrealEd/Private/CookOnTheFlyServer.cpp b/Engine/Source/Editor/UnrealEd/Private/CookOnTheFlyServer.cpp
--- a/Engine/Source/Editor/UnrealEd/Private/CookOnTheFlyServer.cpp
+++ b/Engine/Source/Editor/UnrealEd/Private/CookOnTheFlyServer.cpp
@@ -6301,7 +6301,7 @@ bool UCookOnTheFlyServer::HandleNetworkFileServerNewConnection(const FString& Ve
 	return true;
 }
 
-void UCookOnTheFlyServer::GetCookOnTheFlyUnsolicitedFiles(const FName& PlatformName, TArray<FString> UnsolicitedFiles, const FString& Filename)
+void UCookOnTheFlyServer::GetCookOnTheFlyUnsolicitedFiles(const FName& PlatformName, TArray<FString>& UnsolicitedFiles, const FString& Filename)
 {
 	TArray<FName> UnsolicitedFilenames;
 	UnsolicitedCookedPackages.GetPackagesForPlatformAndRemove(PlatformName, UnsolicitedFilenames);
@@ -6324,8 +6324,8 @@ void UCookOnTheFlyServer::GetCookOnTheFlyUnsolicitedFiles(const FName& PlatformN
 		{
 			UE_LOG(LogCook, Warning, TEXT("Unsolicited file doesn't exist in sandbox, ignoring %s"), *Filename);
 		}
-
 	}
+
 	UPackage::WaitForAsyncFileWrites();
 }
 
