Removed accidentally checked in PRAGMA_DISABLE_OPTIMIZATION from CookOnTheFlyServer.cpp

diff --git a/Engine/Source/Editor/UnrealEd/Private/CookOnTheFlyServer.cpp b/Engine/Source/Editor/UnrealEd/Private/CookOnTheFlyServer.cpp
--- a/Engine/Source/Editor/UnrealEd/Private/CookOnTheFlyServer.cpp
+++ b/Engine/Source/Editor/UnrealEd/Private/CookOnTheFlyServer.cpp
@@ -2632,13 +2632,12 @@ uint64 UCookOnTheFlyServer::GetMaxMemoryAllowance() const
 {
 	return MaxMemoryAllowance;
 }
-PRAGMA_DISABLE_OPTIMIZATION;
+
 const TArray<FName>& UCookOnTheFlyServer::GetFullPackageDependencies(const FName& PackageName ) const
 {
 	TArray<FName>* PackageDependencies = CachedFullPackageDependencies.Find(PackageName);
 	if ( !PackageDependencies )
 	{
-
 		static const FName NAME_CircularReference(TEXT("CircularReference"));
 		static int32 UniqueArrayCounter = 0;
 		++UniqueArrayCounter;
@@ -2646,7 +2645,7 @@ const TArray<FName>& UCookOnTheFlyServer::GetFullPackageDependencies(const FName
 		{
 			// can't initialize the PackageDependencies array here because we call GetFullPackageDependencies below and that could recurse and resize CachedFullPackageDependencies
 			TArray<FName>& TempPackageDependencies = CachedFullPackageDependencies.Add(PackageName); // IMPORTANT READ ABOVE COMMENT
-			// initialize temppackagedependencies to a dummy dependency so that we can detect circular references
+			// initialize TempPackageDependencies to a dummy dependency so that we can detect circular references
 			TempPackageDependencies.Add(CircularReferenceArrayName);
 			// when someone finds the circular reference name they look for this array name in the CachedFullPackageDependencies map
 			// and add their own package name to it, so that they can get fixed up 
@@ -2663,16 +2662,17 @@ const TArray<FName>& UCookOnTheFlyServer::GetFullPackageDependencies(const FName
 				const TArray<FName>& ChildPackageDependencies = GetFullPackageDependencies(ChildDependency);
 				for ( const FName& ChildPackageDependency : ChildPackageDependencies )
 				{
-							if ( ChildPackageDependency == CircularReferenceArrayName )
-							{
-								continue;
-							}
-							if ( ChildPackageDependency.GetComparisonIndex() == NAME_CircularReference.GetComparisonIndex() )
-							{
-								// add our self to the package which we are circular referencing
-								TArray<FName>& TempCircularReference = CachedFullPackageDependencies.FindChecked(ChildPackageDependency);
-								TempCircularReference.AddUnique(PackageName); // add this package name so that it's dependencies get fixed up when the outer loop returns
-							}
+					if ( ChildPackageDependency == CircularReferenceArrayName )
+					{
+						continue;
+					}
+
+					if ( ChildPackageDependency.GetComparisonIndex() == NAME_CircularReference.GetComparisonIndex() )
+					{
+						// add our self to the package which we are circular referencing
+						TArray<FName>& TempCircularReference = CachedFullPackageDependencies.FindChecked(ChildPackageDependency);
+						TempCircularReference.AddUnique(PackageName); // add this package name so that it's dependencies get fixed up when the outer loop returns
+					}
 
 					Dependencies.AddUnique(ChildPackageDependency);
 				}
@@ -2709,9 +2709,10 @@ const TArray<FName>& UCookOnTheFlyServer::GetFullPackageDependencies(const FName
 			PackageDependencies->Add(PackageName);
 		}
 	}
+
 	return *PackageDependencies;
 }
-PRAGMA_ENABLE_OPTIMIZATION;
+
 void UCookOnTheFlyServer::MarkGCPackagesToKeepForCooker()
 {
 	// just saved this package will the cooker need this package again this cook?
