Fix crash when cooking nativized blueprints due to removal of child cooker system.

diff --git a/Engine/Source/Editor/UnrealEd/Private/CookOnTheFlyServer.cpp b/Engine/Source/Editor/UnrealEd/Private/CookOnTheFlyServer.cpp
--- a/Engine/Source/Editor/UnrealEd/Private/CookOnTheFlyServer.cpp
+++ b/Engine/Source/Editor/UnrealEd/Private/CookOnTheFlyServer.cpp
@@ -5975,6 +5975,7 @@ void UCookOnTheFlyServer::StartCookByTheBook( const FCookByTheBookStartupOptions
 			IBlueprintNativeCodeGenModule::Get().FillPlatformNativizationDetails(Entry, PlatformNativizationDetails);
 			CodeGenData.CodegenTargets.Push(PlatformNativizationDetails);
 		}
+		CodeGenData.ManifestIdentifier = -1;
 		IBlueprintNativeCodeGenModule::InitializeModule(CodeGenData);
 	}
 
