- Eliminating ShouldCook function which only returns true -- resulted in pointless temporary string allocations. Also, needless clutter.
- Eliminated redundant GetCachedStandardPackageFileFName call
- Fixed exception slicing rethrow in SaveCookedPackages

diff --git a/Engine/Source/Editor/UnrealEd/Classes/CookOnTheSide/CookOnTheFlyServer.h b/Engine/Source/Editor/UnrealEd/Classes/CookOnTheSide/CookOnTheFlyServer.h
--- a/Engine/Source/Editor/UnrealEd/Classes/CookOnTheSide/CookOnTheFlyServer.h
+++ b/Engine/Source/Editor/UnrealEd/Classes/CookOnTheSide/CookOnTheFlyServer.h
@@ -1408,15 +1408,6 @@ private:
 	//////////////////////////////////////////////////////////////////////////
 	// general functions
 
-	/**
-	* Determines if a package should be cooked
-	* 
-	* @param InFileName		package file name
-	* @param InPlatformName	desired platform to cook for
-	* @return If the package should be cooked
-	*/
-	bool ShouldCook(const FString& InFileName, const FName& InPlatformName);
-
 	/**
 	 * Tries to save all the UPackages in the PackagesToSave list
 	 *  uses the timer to time slice, any packages not saved are requeued in the CookRequests list
diff --git a/Engine/Source/Editor/UnrealEd/Private/CookOnTheFlyServer.cpp b/Engine/Source/Editor/UnrealEd/Private/CookOnTheFlyServer.cpp
--- a/Engine/Source/Editor/UnrealEd/Private/CookOnTheFlyServer.cpp
+++ b/Engine/Source/Editor/UnrealEd/Private/CookOnTheFlyServer.cpp
@@ -1385,15 +1385,7 @@ uint32 UCookOnTheFlyServer::TickCookOnTheSide( const float TimeSlice, uint32 &Co
 
 		const FString BuildFilename = ToBuild.GetFilename().ToString();
 
-		// if we have no target platforms then we want to cook because this will cook for all target platforms in that case
-		bool bShouldCook = TargetPlatformNames.Num() > 0 ? false : ShouldCook( BuildFilename, NAME_None );
-		{
-			SCOPE_TIMER(ShouldCook);
-			for ( int Index = 0; Index < TargetPlatformNames.Num(); ++Index )
-			{
-				bShouldCook |= ShouldCook( ToBuild.GetFilename().ToString(), TargetPlatformNames[Index] );
-			}
-		}
+		bool bShouldCook = true;
 		
 		if( CookByTheBookOptions && CookByTheBookOptions->bErrorOnEngineContentUse )
 		{
@@ -1871,6 +1863,8 @@ void UCookOnTheFlyServer::GetAllUnsolicitedPackages(TArray<UPackage*>& PackagesT
 bool UCookOnTheFlyServer::SaveCookedPackages(TArray<UPackage*>& PackagesToSave, const TArray<FName>& TargetPlatformNames, const TArray<const ITargetPlatform*>& TargetPlatformsToCache, FCookerTimer& Timer, 
 	int32 FirstUnsolicitedPackage, uint32& CookedPackageCount , uint32& Result)
 {
+	check(IsInGameThread());
+
 	bool bIsAllDataCached = true;
 
 	const TArray<FName>& AllTargetPlatformNames = TargetPlatformNames;
@@ -1880,7 +1874,7 @@ bool UCookOnTheFlyServer::SaveCookedPackages(TArray<UPackage*>& PackagesToSave,
 
 	if (PackagesToSave.Num())
 	{
-		int32 OriginalPackagesToSaveCount = PackagesToSave.Num();
+		const int32 OriginalPackagesToSaveCount = PackagesToSave.Num();
 		SCOPE_TIMER(SavingPackages);
 		for (int32 I = 0; I < PackagesToSave.Num(); ++I)
 		{
@@ -1895,16 +1889,13 @@ bool UCookOnTheFlyServer::SaveCookedPackages(TArray<UPackage*>& PackagesToSave,
 			// asset registry at the end of the cook
 			UncookedEditorOnlyPackages.Remove(Package->GetFName());
 
-			const FName StandardPackageFilename = GetCachedStandardPackageFileFName(Package);
-			check(IsInGameThread());
-			if (NeverCookPackageList.Contains(StandardPackageFilename))
+			const FName PackageFName = GetCachedStandardPackageFileFName(Package);
+			if (NeverCookPackageList.Contains(PackageFName))
 			{
 				// refuse to save this package, it's clearly one of the undesirables
 				continue;
 			}
 
-
-			FName PackageFName = GetCachedStandardPackageFileFName(Package);
 			TArray<FName> SaveTargetPlatformNames = AllTargetPlatformNames;
 			TArray<FName> CookedTargetPlatforms;
 			if (CookedPackages.GetCookedPlatforms(PackageFName, CookedTargetPlatforms))
@@ -1918,7 +1909,7 @@ bool UCookOnTheFlyServer::SaveCookedPackages(TArray<UPackage*>& PackagesToSave,
 			// we somehow already cooked this package not sure how that can happen because the PackagesToSave list should have already filtered this
 			if (SaveTargetPlatformNames.Num() == 0)
 			{
-				UE_LOG(LogCook, Warning, TEXT("Allready saved this package not sure how this got here!"));
+				UE_LOG(LogCook, Warning, TEXT("Already saved this package not sure how this got here!"));
 				// already saved this package
 				continue;
 			}
@@ -1938,12 +1929,12 @@ bool UCookOnTheFlyServer::SaveCookedPackages(TArray<UPackage*>& PackagesToSave,
 
 			// if we are cook the fly then save the package which was requested as fast as we can because the client is waiting on it
 
-			bool ProcessingUnsolicitedPackages = (I >= FirstUnsolicitedPackage);
+			const bool bProcessingUnsolicitedPackages = (I >= FirstUnsolicitedPackage);
 			bool bForceSavePackage = false;
 
 			if (IsCookOnTheFlyMode())
 			{
-				if (ProcessingUnsolicitedPackages)
+				if (bProcessingUnsolicitedPackages)
 				{
 					SCOPE_TIMER(WaitingForCachedCookedPlatformData);
 					if (CookRequests.HasItems())
@@ -2101,7 +2092,7 @@ bool UCookOnTheFlyServer::SaveCookedPackages(TArray<UPackage*>& PackagesToSave,
 				{
 					SaveCookedPackage(Package, SaveFlags, SaveTargetPlatformNames, SavePackageResults);
 				}
-				catch (std::exception& e)
+				catch (std::exception&)
 				{
 					FString TargetPlatforms;
 					for ( const FName& PlatformFName : SaveTargetPlatformNames )
@@ -2109,7 +2100,7 @@ bool UCookOnTheFlyServer::SaveCookedPackages(TArray<UPackage*>& PackagesToSave,
 						TargetPlatforms += FString::Printf( TEXT("%s, "), *PlatformFName.ToString());
 					}
 					UE_LOG(LogCook, Warning, TEXT("Tried to save package %s for target platforms %s but threw an exception"), *Package->GetPathName(), *TargetPlatforms);
-					throw e;
+					throw;
 				}
 				
 				GOutputCookingWarnings = false;
@@ -2802,11 +2793,6 @@ void UCookOnTheFlyServer::SaveCookedPackage(UPackage* Package, uint32 SaveFlags,
 	return SaveCookedPackage( Package, SaveFlags, TargetPlatformNames, SavePackageResults);
 }
 
-bool UCookOnTheFlyServer::ShouldCook(const FString& InFileName, const FName &InPlatformName)
-{
-	return true;
-}
-
 bool UCookOnTheFlyServer::ShouldConsiderCompressedPackageFileLengthRequirements() const
 {
 	bool bConsiderCompressedPackageFileLengthRequirements = true;
