Removed LeakTest functionality - not useful, also likely not functional since the code path has not been exercised for a long time

diff --git a/Engine/Source/Editor/UnrealEd/Classes/Commandlets/CookCommandlet.h b/Engine/Source/Editor/UnrealEd/Classes/Commandlets/CookCommandlet.h
--- a/Engine/Source/Editor/UnrealEd/Classes/Commandlets/CookCommandlet.h
+++ b/Engine/Source/Editor/UnrealEd/Classes/Commandlets/CookCommandlet.h
@@ -38,8 +38,6 @@ class UCookCommandlet
 	bool bCookAll;
 	/** Skip saving any packages in Engine/Content/Editor* UNLESS TARGET HAS EDITORONLY DATA (in which case it will save those anyway) */
 	bool bSkipEditorContent;
-	/** Test for UObject leaks */
-	bool bLeakTest;  
 	/** Save all cooked packages without versions. These are then assumed to be current version on load. This is dangerous but results in smaller patch sizes. */
 	bool bUnversioned;
 	/** Generate manifests for building streaming install packages */
diff --git a/Engine/Source/Editor/UnrealEd/Classes/CookOnTheSide/CookOnTheFlyServer.h b/Engine/Source/Editor/UnrealEd/Classes/CookOnTheSide/CookOnTheFlyServer.h
--- a/Engine/Source/Editor/UnrealEd/Classes/CookOnTheSide/CookOnTheFlyServer.h
+++ b/Engine/Source/Editor/UnrealEd/Classes/CookOnTheSide/CookOnTheFlyServer.h
@@ -54,7 +54,6 @@ enum class ECookByTheBookOptions
 	CookAll	=							0x00000001, // cook all maps and content in the content directory
 	MapsOnly =							0x00000002, // cook only maps
 	NoDevContent =						0x00000004, // don't include dev content
-	LeakTest =							0x00000008, // test for uobject leaks after each level load
 	ForceDisableCompressed =			0x00000010, // force compression to be disabled even if the cooker was initialized with it enabled
 	ForceEnableCompressed =				0x00000020, // force compression to be on even if the cooker was initialized with it disabled
 	ForceDisableSaveGlobalShaders =		0x00000040, // force global shaders to not be saved (used if cooking multiple times for the same platform and we know we are up todate)
@@ -790,14 +789,14 @@ private:
 	ECookMode::Type CurrentCookMode;
 	/** Directory to output to instead of the default should be empty in the case of DLC cooking */ 
 	FString OutputDirectoryOverride;
+
 	//////////////////////////////////////////////////////////////////////////
 	// Cook by the book options
-public:
-private:
+
 	struct FCookByTheBookOptions
 	{
 	public:
-		FCookByTheBookOptions() : bLeakTest(false),
+		FCookByTheBookOptions() : 
 			bGenerateStreamingInstallManifests(false),
 			bGenerateDependenciesForMaps(false),
 			bRunning(false),
@@ -808,8 +807,6 @@ private:
 			bFullLoadAndSave(false)
 		{ }
 
-		/** Should we test for UObject leaks */
-		bool bLeakTest;
 		/** Should we generate streaming install manifests (only valid option in cook by the book) */
 		bool bGenerateStreamingInstallManifests;
 		/** Should we generate a seperate manifest for map dependencies */
@@ -822,8 +819,6 @@ private:
 		FString DlcName;
 		/** Create a release from this manifest and store it in the releases directory for this cgame */
 		FString CreateReleaseVersion;
-		/** Leak test: last gc items (only valid when running from commandlet requires gc between each cooked package) */
-		TSet<FWeakObjectPtr> LastGCItems;
 		/** Dependency graph of maps as root objects. */
 		TMap<FName, TMap< FName, TSet <FName> > > MapDependencyGraphs; 
 		/** If a cook is cancelled next cook will need to resume cooking */ 
@@ -1421,7 +1416,7 @@ private:
 	 * @param Result (in+out) used to modify the result of the operation and add any relevant flags
 	 * @return returns true if we saved all the packages false if we bailed early for any reason
 	 */
-	bool SaveCookedPackages(TArray<UPackage*>& PackagesToSave, const TArray<FName>& TargetPlatformNames, const TArray<const ITargetPlatform*>& TargetPlatformsToCache, struct FCookerTimer& Timer, int32 FirstUnsolicitedPackage, uint32& CookedPackageCount, uint32& Result);
+	void SaveCookedPackages(TArray<UPackage*>& PackagesToSave, const TArray<FName>& TargetPlatformNames, const TArray<const ITargetPlatform*>& TargetPlatformsToCache, struct FCookerTimer& Timer, int32 FirstUnsolicitedPackage, uint32& CookedPackageCount, uint32& Result);
 
 	/**
 	 * Returns all packages which are found in memory which aren't cooked
diff --git a/Engine/Source/Editor/UnrealEd/Private/Commandlets/CookCommandlet.cpp b/Engine/Source/Editor/UnrealEd/Private/Commandlets/CookCommandlet.cpp
--- a/Engine/Source/Editor/UnrealEd/Private/Commandlets/CookCommandlet.cpp
+++ b/Engine/Source/Editor/UnrealEd/Private/Commandlets/CookCommandlet.cpp
@@ -533,7 +533,6 @@ int32 UCookCommandlet::Main(const FString& CmdLineParams)
 
 	bCookOnTheFly = Switches.Contains(TEXT("COOKONTHEFLY"));   // Prototype cook-on-the-fly server
 	bCookAll = Switches.Contains(TEXT("COOKALL"));   // Cook everything
-	bLeakTest = Switches.Contains(TEXT("LEAKTEST"));   // Test for UObject leaks
 	bUnversioned = Switches.Contains(TEXT("UNVERSIONED"));   // Save all cooked packages without versions. These are then assumed to be current version on load. This is dangerous but results in smaller patch sizes.
 	bGenerateStreamingInstallManifests = Switches.Contains(TEXT("MANIFESTS"));   // Generate manifests for building streaming install packages
 	bIterativeCooking = Switches.Contains(TEXT("ITERATE"));
@@ -806,7 +805,6 @@ bool UCookCommandlet::CookByTheBook( const TArray<ITargetPlatform*>& Platforms,
 	// start cook by the book 
 	ECookByTheBookOptions CookOptions = ECookByTheBookOptions::None;
 
-	CookOptions |= bLeakTest ? ECookByTheBookOptions::LeakTest : ECookByTheBookOptions::None; 
 	CookOptions |= bCookAll ? ECookByTheBookOptions::CookAll : ECookByTheBookOptions::None;
 	CookOptions |= Switches.Contains(TEXT("MAPSONLY")) ? ECookByTheBookOptions::MapsOnly : ECookByTheBookOptions::None;
 	CookOptions |= Switches.Contains(TEXT("NODEV")) ? ECookByTheBookOptions::NoDevContent : ECookByTheBookOptions::None;
diff --git a/Engine/Source/Editor/UnrealEd/Private/CookOnTheFlyServer.cpp b/Engine/Source/Editor/UnrealEd/Private/CookOnTheFlyServer.cpp
--- a/Engine/Source/Editor/UnrealEd/Private/CookOnTheFlyServer.cpp
+++ b/Engine/Source/Editor/UnrealEd/Private/CookOnTheFlyServer.cpp
@@ -1611,35 +1611,8 @@ uint32 UCookOnTheFlyServer::TickCookOnTheSide( const float TimeSlice, uint32 &Co
 			}
 		}
 
-		bool bFinishedSave = SaveCookedPackages( PackagesToSave, AllTargetPlatformNames, TargetPlatforms, Timer, FirstUnsolicitedPackage, CookedPackageCount, Result );
+		SaveCookedPackages( PackagesToSave, AllTargetPlatformNames, TargetPlatforms, Timer, FirstUnsolicitedPackage, CookedPackageCount, Result );
 
-
-		// TODO: Daniel: this is reference code needs to be reimplemented on the callee side.
-		//  cooker can't depend on callee being able to garbage collect
-		// collect garbage
-		if ( CookByTheBookOptions && CookByTheBookOptions->bLeakTest && bFinishedSave )
-		{
-			check( CurrentCookMode == ECookMode::CookByTheBook);
-			UE_LOG(LogCook, Display, TEXT("Full GC..."));
-
-			CollectGarbage( GARBAGE_COLLECTION_KEEPFLAGS );
-			for (FObjectIterator It; It; ++It)
-			{
-				if (!CookByTheBookOptions->LastGCItems.Contains(FWeakObjectPtr(*It)))
-				{
-					UE_LOG(LogCook, Warning, TEXT("\tLeaked %s"), *(It->GetFullName()));
-					CookByTheBookOptions->LastGCItems.Add(FWeakObjectPtr(*It));
-				}
-			}
-		}
-
-		
-		// if we are cook on the fly (even in editor) and we have cook requests, try process them straight away
-		/*if ( IsCookOnTheFlyMode() && CookRequests.HasItems() )
-		{
-			continue;
-		}*/
-		
 		if ( Timer.IsTimeUp() )
 		{
 			break;
@@ -1924,22 +1897,22 @@ void UCookOnTheFlyServer::GetAllUnsolicitedPackages(TArray<UPackage*>& PackagesT
 	GetUnsolicitedPackages(PackagesToSave, TargetPlatformNames);
 }
 
-bool UCookOnTheFlyServer::SaveCookedPackages(TArray<UPackage*>& PackagesToSave, const TArray<FName>& TargetPlatformNames, const TArray<const ITargetPlatform*>& TargetPlatformsToCache, FCookerTimer& Timer, 
+void UCookOnTheFlyServer::SaveCookedPackages(TArray<UPackage*>& PackagesToSave, const TArray<FName>& TargetPlatformNames, const TArray<const ITargetPlatform*>& TargetPlatformsToCache, FCookerTimer& Timer, 
 	int32 FirstUnsolicitedPackage, uint32& CookedPackageCount , uint32& Result)
 {
 	check(IsInGameThread());
 
+	if (PackagesToSave.Num() == 0)
+		return;
+
 	bool bIsAllDataCached = true;
 
 	const TArray<FName>& AllTargetPlatformNames = TargetPlatformNames;
-	
+	const int32 OriginalPackagesToSaveCount = PackagesToSave.Num();
 
-	bool bFinishedSave = true;
-
-	if (PackagesToSave.Num())
 	{
-		const int32 OriginalPackagesToSaveCount = PackagesToSave.Num();
 		SCOPE_TIMER(SavingPackages);
+
 		for (int32 I = 0; I < PackagesToSave.Num(); ++I)
 		{
 			UPackage* Package = PackagesToSave[I];
@@ -1960,6 +1933,10 @@ bool UCookOnTheFlyServer::SaveCookedPackages(TArray<UPackage*>& PackagesToSave,
 				continue;
 			}
 
+			// TODO: optimize this to avoid creating the full array first only to remove entries (multiple dynamic allocations)
+			//
+			// Could rewrite using FTargetPlatformSet
+			
 			TArray<FName> SaveTargetPlatformNames = AllTargetPlatformNames;
 			TArray<FName> CookedTargetPlatforms;
 			if (CookedPackages.GetCookedPlatforms(PackageFName, CookedTargetPlatforms))
@@ -1978,17 +1955,16 @@ bool UCookOnTheFlyServer::SaveCookedPackages(TArray<UPackage*>& PackagesToSave,
 				continue;
 			}
 
-
 			// if we are processing unsolicited packages we can optionally not save these right now
 			// the unsolicited packages which we missed now will be picked up on next run
 			// we want to do this in cook on the fly also, if there is a new network package request instead of saving unsolicited packages we can process the requested package
 
 			bool bShouldFinishTick = false;
 
-			if (Timer.IsTimeUp() && IsCookByTheBookMode() )
+			if (Timer.IsTimeUp() && IsCookByTheBookMode())
 			{
-				bShouldFinishTick = true;
 				// our timeslice is up
+				bShouldFinishTick = true;
 			}
 
 			// if we are cook the fly then save the package which was requested as fast as we can because the client is waiting on it
@@ -2040,7 +2016,7 @@ bool UCookOnTheFlyServer::SaveCookedPackages(TArray<UPackage*>& PackagesToSave,
 
 			MakePackageFullyLoaded(Package);
 
-			if ( IsCookOnTheFlyMode() )
+			if (IsCookOnTheFlyMode())
 			{
 				// never want to requeue packages
 				HasCheckedAllPackagesAreCached = true;
@@ -2080,16 +2056,14 @@ bool UCookOnTheFlyServer::SaveCookedPackages(TArray<UPackage*>& PackagesToSave,
 				for (int32 RemainingIndex = I; RemainingIndex < NumPackagesToRequeue; ++RemainingIndex)
 				{
 					FName StandardFilename = GetCachedStandardPackageFileFName(PackagesToSave[RemainingIndex]);
-							CookRequests.EnqueueUnique(FFilePlatformRequest(StandardFilename, SaveTargetPlatformNames));
+					CookRequests.EnqueueUnique(FFilePlatformRequest(StandardFilename, SaveTargetPlatformNames));
 				}
 				Result |= COSR_WaitingOnCache;
-
+			
 				// break out of the loop
-				bFinishedSave = false;
-				break;
+				return;
 			}
 
-
 			// don't precache other packages if our package isn't ready but we are going to save it.   This will fill up the worker threads with extra shaders which we may need to flush on 
 			if ((!IsCookOnTheFlyMode()) &&
 				(!IsRealtimeMode() || AllObjectsCookedDataCached == true))
@@ -2266,7 +2240,6 @@ bool UCookOnTheFlyServer::SaveCookedPackages(TArray<UPackage*>& PackagesToSave,
 			}
 		}
 	}
-	return true;
 }
 
 
@@ -5656,7 +5629,6 @@ void UCookOnTheFlyServer::CookByTheBookFinished()
 		}
 	}
 
-	CookByTheBookOptions->LastGCItems.Empty();
 	const float TotalCookTime = (float)(FPlatformTime::Seconds() - CookByTheBookOptions->CookStartTime);
 	UE_LOG(LogCook, Display, TEXT("Cook by the book total time in tick %fs total time %f"), CookByTheBookOptions->CookTime, TotalCookTime);
 
@@ -6054,18 +6026,6 @@ void UCookOnTheFlyServer::StartCookByTheBook( const FCookByTheBookStartupOptions
 		IBlueprintNativeCodeGenModule::InitializeModule(CodeGenData);
 	}
 
-	CookByTheBookOptions->bLeakTest = (CookOptions & ECookByTheBookOptions::LeakTest) != ECookByTheBookOptions::None; // this won't work from the editor this needs to be standalone
-	check(!CookByTheBookOptions->bLeakTest || CurrentCookMode == ECookMode::CookByTheBook);
-
-	CookByTheBookOptions->LastGCItems.Empty();
-	if (CookByTheBookOptions->bLeakTest)
-	{
-		for (FObjectIterator It; It; ++It)
-		{
-			CookByTheBookOptions->LastGCItems.Add(FWeakObjectPtr(*It));
-		}
-	}
-
 	{
 		for (ITargetPlatform* Platform : TargetPlatforms)
 		{
