Refactoring: moved some CookOnTheFlyServer definitions from header to cpp (no behaviour change)

Makes it easier to change implementation and reduces iteration time - when the header is touched then we end up running UHT

- FFilePlatformCookedPackage moved to cpp and made it not derive from FFilePlatformRequest (they are not really related so inheritance doesn't make sense IMO)
- FThreadSafeFilenameSet -> CookedPackageSet
- FFilenameQueue -> CookRequestQueue

diff --git a/Engine/Source/Editor/UnrealEd/Classes/CookOnTheSide/CookOnTheFlyServer.h b/Engine/Source/Editor/UnrealEd/Classes/CookOnTheSide/CookOnTheFlyServer.h
--- a/Engine/Source/Editor/UnrealEd/Classes/CookOnTheSide/CookOnTheFlyServer.h
+++ b/Engine/Source/Editor/UnrealEd/Classes/CookOnTheSide/CookOnTheFlyServer.h
@@ -295,7 +295,7 @@ public:
 			return false;
 		}
 
-		FORCEINLINE FString ToString() const
+		FString ToString() const
 		{
 			FString Result = FString::Printf(TEXT("%s;"), *Filename.ToString());
 
@@ -305,373 +305,9 @@ public:
 			}
 			return Result;
 		}
-
 	};
 
-
 private:
-
-	struct FFilePlatformCookedPackage : public FFilePlatformRequest
-	{
-	private:
-		TArray<bool> bSucceededSavePackage; // one bool for each platform
-	public:
-		FFilePlatformCookedPackage(const FFilePlatformRequest& InFilePlatformRequest, const TArray<bool>& bInSuccededSavePackage) : FFilePlatformRequest(InFilePlatformRequest.GetFilename(), InFilePlatformRequest.GetPlatformNames()), bSucceededSavePackage(bInSuccededSavePackage) { check(PlatformNames.Num() == bSucceededSavePackage.Num()); }
-		FFilePlatformCookedPackage(const FName& InFilename, const TArray<FName>& InPlatformName) : FFilePlatformRequest(InFilename, InPlatformName)
-		{ 
-			// only use this constructor to short hand when packages fail
-			for ( int32 I = 0; I < InPlatformName.Num(); ++I )
-			{
-				bSucceededSavePackage.Add(false);
-			}
-			check(PlatformNames.Num() == bSucceededSavePackage.Num());
-		}
-		FFilePlatformCookedPackage(const FName& InFilename, const TArray<FName>& InPlatformName, const TArray<bool>& bInSuccededSavePackage) : FFilePlatformRequest(InFilename, InPlatformName), bSucceededSavePackage(bInSuccededSavePackage) { check(PlatformNames.Num() == bSucceededSavePackage.Num()); }
-		FFilePlatformCookedPackage(const FName& InFilename, TArray<FName>&& InPlatformName, TArray<bool>&& bInSuccededSavePackage) : FFilePlatformRequest(InFilename, MoveTemp(InPlatformName)), bSucceededSavePackage(MoveTemp(bInSuccededSavePackage)) { check(PlatformNames.Num() == bSucceededSavePackage.Num()); }
-		FFilePlatformCookedPackage(const FFilePlatformCookedPackage& InFilePlatformRequest) : FFilePlatformRequest(InFilePlatformRequest.Filename, InFilePlatformRequest.PlatformNames), bSucceededSavePackage(InFilePlatformRequest.bSucceededSavePackage) { check(PlatformNames.Num() == bSucceededSavePackage.Num());  }
-		FFilePlatformCookedPackage(FFilePlatformCookedPackage&& InFilePlatformRequest) : FFilePlatformRequest(MoveTemp(InFilePlatformRequest.Filename), MoveTemp(InFilePlatformRequest.PlatformNames)), bSucceededSavePackage(InFilePlatformRequest.bSucceededSavePackage) { check(PlatformNames.Num() == bSucceededSavePackage.Num()); }
-
-		inline const void AddPlatform( const FName& Platform, bool bSucceeded )
-		{
-			check(PlatformNames.Num() == bSucceededSavePackage.Num());
-			check(Platform != NAME_None);
-			PlatformNames.Add(Platform);
-			bSucceededSavePackage.Add(bSucceeded);
-			check( PlatformNames.Num() == bSucceededSavePackage.Num() );
-		}
-
-		inline void RemovePlatform(const FName &Platform)
-		{
-			check(PlatformNames.Num() == bSucceededSavePackage.Num());
-			int32 Index = PlatformNames.IndexOfByKey(Platform);
-			if (Index != -1)
-			{
-				PlatformNames.RemoveAt(Index);
-				bSucceededSavePackage.RemoveAt(Index);
-			}
-			check(PlatformNames.Num() == bSucceededSavePackage.Num());
-		}
-
-		inline const bool HasSucceededSavePackage(const FName& PlatformName) const 
-		{ 
-			int32 Index = PlatformNames.IndexOfByKey(PlatformName);
-			if ( (Index != -1) && (Index < bSucceededSavePackage.Num()) )
-			{
-				return bSucceededSavePackage[Index]; 
-			}
-			return false;
-		}
-	};
-
-	/** Helper list of all files which have been cooked */
-	struct FThreadSafeFilenameSet
-	{
-	private:
-		mutable FCriticalSection	SynchronizationObject;
-		TMap<FName, FFilePlatformCookedPackage> FilesProcessed;
-
-		void Lock()
-		{
-			SynchronizationObject.Lock();
-		}
-		void Unlock()
-		{
-			SynchronizationObject.Unlock();
-		}
-
-	public:
-		int32 Num()
-		{ 
-			return FilesProcessed.Num();
-		}
-
-		void Add(const FFilePlatformCookedPackage& Request)
-		{
-			FScopeLock ScopeLock(&SynchronizationObject);
-			check(Request.IsValid());
-
-			// see if it's already in the requests list
-			FFilePlatformCookedPackage *ExistingRequest = FilesProcessed.Find(Request.GetFilename() );
-
-			if ( ExistingRequest )
-			{
-				check( ExistingRequest->GetFilename() == Request.GetFilename() );
-				for ( const FName& Platform : Request.GetPlatformNames() )
-				{
-					const bool bSucceeded = Request.HasSucceededSavePackage(Platform);
-					ExistingRequest->AddPlatform(Platform, bSucceeded);
-				}
-			}
-			else
-				FilesProcessed.Add(Request.GetFilename(), Request);
-		}
-		bool Exists(const FFilePlatformRequest & Request) const
-		{
-			FScopeLock ScopeLock(&SynchronizationObject);
-
-			const FFilePlatformCookedPackage* OurRequest = FilesProcessed.Find(Request.GetFilename());
-
-			if (!OurRequest)
-			{
-				return false;
-			}
-
-			// make sure all the platforms are completed
-			for ( const FName& Platform : Request.GetPlatformNames() )
-			{
-				if ( !OurRequest->GetPlatformNames().Contains( Platform ) )
-				{
-					return false;
-				}
-			}
-
-			return true;
-			// return FilesProcessed.Contains(Filename);
-		}
-		// two versions of this function so I don't have to create temporary FFIleplatformRequest in some cases to call the exists function
-		bool Exists( const FName& Filename, const TArray<FName>& PlatformNames ) const
-		{
-			FScopeLock ScopeLock(&SynchronizationObject);
-
-			const FFilePlatformRequest* OurRequest = FilesProcessed.Find( Filename );
-
-			if (!OurRequest)
-			{
-				return false;
-			}
-
-			// make sure all the platforms are completed
-			for ( const FName& Platform : PlatformNames )
-			{
-				if ( !OurRequest->GetPlatformNames().Contains( Platform ) )
-				{
-					return false;
-				}
-			}
-
-			return true;
-		}
-
-		// do we want failed packages or not
-		bool Exists( const FName& Filename, const TArray<FName>& PlatformNames, bool bIncludeFailed ) const
-		{
-			FScopeLock ScopeLock(&SynchronizationObject);
-
-			const FFilePlatformCookedPackage* OurRequest = FilesProcessed.Find(Filename);
-
-			if (!OurRequest)
-			{
-				return false;
-			}
-
-			if (bIncludeFailed == false)
-			{
-				bool allFailed = true;
-				for ( const auto& PlatformName : PlatformNames )
-				{
-					if (OurRequest->HasSucceededSavePackage(PlatformName))
-					{
-						allFailed = false;
-						break;
-					}
-				}
-				if ( allFailed )
-				{
-					return false;
-				}
-			}
-
-			// make sure all the platforms are completed
-			for (const FName& Platform : PlatformNames)
-			{
-				if (!OurRequest->GetPlatformNames().Contains(Platform))
-				{
-					return false;
-				}
-			}
-			return true;
-		}
-
-		void RemoveAllFilesForPlatform( const FName& PlatformName )
-		{
-			FScopeLock ScopeLock(&SynchronizationObject);
-
-			for ( auto& Request : FilesProcessed )
-			{	
-				Request.Value.RemovePlatform( PlatformName );
-			}
-		}
-
-		bool GetCookedPlatforms( const FName& Filename, TArray<FName>& PlatformList ) const
-		{
-			FScopeLock ScopeLock( &SynchronizationObject );
-			const FFilePlatformCookedPackage* Request = FilesProcessed.Find(Filename);
-			if ( Request )
-			{
-				PlatformList = Request->GetPlatformNames();
-				return true;
-			}
-			return false;
-		}
-		int RemoveFile( const FName& Filename )
-		{
-			FScopeLock ScopeLock( &SynchronizationObject );
-			return FilesProcessed.Remove( Filename );
-		}
-
-		bool RemoveFileForPlatform( const FName& Filename, const FName& PlatformName )
-		{
-			FScopeLock ScopeLock( &SynchronizationObject );
-			FFilePlatformCookedPackage* ProcessedFile = FilesProcessed.Find(Filename);
-			if( ProcessedFile )
-			{
-				ProcessedFile->RemovePlatform(PlatformName);
-				return true; 
-			}
-			return false;
-		}
-
-		void GetCookedFilesForPlatform(const FName& PlatformName, TArray<FName>& CookedFiles, bool bGetFailedCookedPackages = true, bool bGetSuccessfulCookedPackages = true)
-		{
-			FScopeLock ScopeLock(&SynchronizationObject);
-			for (const auto& CookedFile : FilesProcessed)
-			{
-				if (CookedFile.Value.HasPlatform(PlatformName))
-				{
-					bool bHasSucceededSavePackage = CookedFile.Value.HasSucceededSavePackage(PlatformName);
-					if ((bHasSucceededSavePackage && bGetSuccessfulCookedPackages) ||
-						((bHasSucceededSavePackage == false) && bGetFailedCookedPackages) )
-						CookedFiles.Add(CookedFile.Value.GetFilename());
-				}
-			}
-		}
-		void Empty(int32 ExpectedNumElements = 0)
-		{
-			FScopeLock ScopeLock( &SynchronizationObject );
-			FilesProcessed.Empty(ExpectedNumElements);
-		}
-	};
-
-
-	struct FFilenameQueue
-	{
-	private:
-		TArray<FName>			Queue;
-		TMap<FName, TArray<FName>> PlatformList;
-		mutable FCriticalSection SynchronizationObject;
-	public:
-
-		template <class PREDICATE_CLASS>
-		void Sort(const PREDICATE_CLASS& Predicate)
-		{
-			FScopeLock ScopeLock(&SynchronizationObject);
-			Queue.Sort(Predicate);
-		}
-
-		const TArray<FName>& GetQueue() const { return Queue;  }
-		void EnqueueUnique(const FFilePlatformRequest& Request, bool ForceEnqueFront = false)
-		{
-			FScopeLock ScopeLock( &SynchronizationObject );
-			TArray<FName>* Platforms = PlatformList.Find(Request.GetFilename());
-			if ( Platforms == NULL )
-			{
-				PlatformList.Add(Request.GetFilename(), Request.GetPlatformNames());
-				Queue.Add(Request.GetFilename());
-			}
-			else
-			{
-				// add the requested platforms to the platform list
-				for ( const FName& Platform : Request.GetPlatformNames() )
-				{
-					Platforms->AddUnique(Platform);
-				}
-			}
-
-			if ( ForceEnqueFront )
-			{
-				int32 Index = Queue.Find(Request.GetFilename());
-				check( Index != INDEX_NONE );
-				if ( Index != 0 )
-				{
-					Queue.Swap(0, Index);
-				}
-			}
-		}
-
-		bool Dequeue(FFilePlatformRequest* Result)
-		{
-			FScopeLock ScopeLock( &SynchronizationObject );
-			if (Queue.Num())
-			{
-				FName Filename = Queue[0];
-				Queue.RemoveAt(0);
-				TArray<FName> Platforms = PlatformList.FindChecked(Filename);
-				PlatformList.Remove(Filename);
-				*Result = FFilePlatformRequest(MoveTemp(Filename), MoveTemp(Platforms));
-				return true;
-			}
-			return false;
-		}
-
-		void DequeueAllRequests( TArray<FFilePlatformRequest>& RequestArray )
-		{
-			FScopeLock ScopeLock( &SynchronizationObject );
-			if ( Queue.Num() )
-			{
-				for ( const auto& Request : PlatformList )
-				{
-					RequestArray.Add( FFilePlatformRequest( Request.Key, Request.Value ) );
-				}
-				PlatformList.Empty();
-				Queue.Empty();
-			}
-		}
-
-		bool Exists( const FName& Filename, const TArray<FName>& PlatformNames ) const 
-		{
-			FScopeLock ScopeLock(&SynchronizationObject);
-			const TArray<FName>* Platforms = PlatformList.Find( Filename );
-			if ( Platforms == NULL )
-				return false;
-
-			for ( const FName& PlatformName : PlatformNames )
-			{
-				if ( !Platforms->Contains( PlatformName ) )
-					return false;
-			}
-			return true;
-		}
-
-		bool Exists(const FName& Filename)
-		{
-			FScopeLock ScopeLock(&SynchronizationObject);
-			const TArray<FName>* Platforms = PlatformList.Find(Filename);
-			if (Platforms == NULL)
-				return false;
-			return true;
-		}
-
-		bool HasItems() const
-		{
-			FScopeLock ScopeLock( &SynchronizationObject );
-			return Queue.Num() > 0;
-		}
-
-		int Num() const 
-		{
-			FScopeLock ScopeLock( &SynchronizationObject );
-			return Queue.Num();
-		}
-
-
-		void Empty()  
-		{
-			FScopeLock ScopeLock( &SynchronizationObject );
-			Queue.Empty();
-			PlatformList.Empty();
-		}
-	};
-
 	struct FThreadSafeUnsolicitedPackagesList
 	{
 	private:
@@ -922,9 +558,7 @@ private:
 	FReentryData& GetReentryData(const UPackage* Package) const;
 
 	FThreadSafeQueue<struct FRecompileRequest*> RecompileRequests;
-	FFilenameQueue CookRequests; // list of requested files
 	FThreadSafeUnsolicitedPackagesList UnsolicitedCookedPackages;
-	FThreadSafeFilenameSet CookedPackages; // set of files which have been cooked when needing to recook a file the entry will need to be removed from here
 	FThreadSafeSet<FName> NeverCookPackageList;
 	FThreadSafeSet<FName> UncookedEditorOnlyPackages; // set of packages that have been rejected due to being referenced by editor-only properties
 
@@ -1177,9 +811,9 @@ public:
 	*/
 	void TickRecompileShaderRequests();
 
-	bool HasCookRequests() const { return CookRequests.HasItems(); }
+	bool HasRecompileShaderRequests() const;
 
-	bool HasRecompileShaderRequests() const { return RecompileRequests.HasItems(); }
+	bool HasCookRequests() const;
 
 	uint32 NumConnections() const;
 
diff --git a/Engine/Source/Editor/UnrealEd/Private/CookOnTheFlyServer.cpp b/Engine/Source/Editor/UnrealEd/Private/CookOnTheFlyServer.cpp
--- a/Engine/Source/Editor/UnrealEd/Private/CookOnTheFlyServer.cpp
+++ b/Engine/Source/Editor/UnrealEd/Private/CookOnTheFlyServer.cpp
@@ -521,6 +521,427 @@ void LogCookerMessage( const FString& MessageText, EMessageSeverity::Type Severi
 	MessageLog.Notify(FText(), EMessageSeverity::Warning, false);
 }
 
+struct FFilePlatformCookedPackage
+{
+public:
+	FFilePlatformCookedPackage(const UCookOnTheFlyServer::FFilePlatformRequest& InFilePlatformRequest, const TArray<bool>& bInSuccededSavePackage)
+	: Filename(InFilePlatformRequest.GetFilename())
+	, PlatformNames(InFilePlatformRequest.GetPlatformNames())
+	, bSucceededSavePackage(bInSuccededSavePackage)
+	{ 
+		check(PlatformNames.Num() == bSucceededSavePackage.Num()); 
+	}
+
+	FFilePlatformCookedPackage(const FName& InFilename, const TArray<FName>& InPlatformName) 
+	: Filename(InFilename)
+	, PlatformNames(InPlatformName)
+	{
+		// only use this constructor to short hand when packages fail
+		for (int32 I = 0; I < InPlatformName.Num(); ++I)
+		{
+			bSucceededSavePackage.Add(false);
+		}
+
+		check(PlatformNames.Num() == bSucceededSavePackage.Num());
+	}
+
+	FFilePlatformCookedPackage(const FName& InFilename, const TArray<FName>& InPlatformName, const TArray<bool>& bInSuccededSavePackage)
+	: Filename(InFilename)
+	, PlatformNames(InPlatformName)
+	, bSucceededSavePackage(bInSuccededSavePackage) 
+	{ 
+		check(PlatformNames.Num() == bSucceededSavePackage.Num()); 
+	}
+
+	FFilePlatformCookedPackage(const FName& InFilename, TArray<FName>&& InPlatformName, TArray<bool>&& bInSuccededSavePackage)
+	: Filename(InFilename)
+	, PlatformNames(MoveTemp(InPlatformName))
+	, bSucceededSavePackage(MoveTemp(bInSuccededSavePackage))
+	{ 
+		check(PlatformNames.Num() == bSucceededSavePackage.Num());
+	}
+
+	FFilePlatformCookedPackage(const FFilePlatformCookedPackage& InFilePlatformRequest) = default;
+	FFilePlatformCookedPackage(FFilePlatformCookedPackage&& InFilePlatformRequest) = default;
+
+	bool IsValid() const
+	{
+		return Filename != NAME_None;
+	}
+
+	const FName& GetFilename() const
+	{
+		return Filename;
+	}
+
+	const TArray<FName>& GetPlatformNames() const
+	{
+		return PlatformNames;
+	}
+
+	bool HasPlatform(const FName& Platform) const
+	{
+		return PlatformNames.Find(Platform) != INDEX_NONE;
+	}
+
+	const void AddPlatform(const FName& Platform, bool bSucceeded)
+	{
+		check(PlatformNames.Num() == bSucceededSavePackage.Num());
+		check(Platform != NAME_None);
+		PlatformNames.Add(Platform);
+		bSucceededSavePackage.Add(bSucceeded);
+		check(PlatformNames.Num() == bSucceededSavePackage.Num());
+	}
+
+	void RemovePlatform(const FName& Platform)
+	{
+		check(PlatformNames.Num() == bSucceededSavePackage.Num());
+		const int32 Index = PlatformNames.IndexOfByKey(Platform);
+
+		if (Index != -1)
+		{
+			PlatformNames.RemoveAt(Index);
+			bSucceededSavePackage.RemoveAt(Index);
+		}
+
+		check(PlatformNames.Num() == bSucceededSavePackage.Num());
+	}
+
+	const bool HasSucceededSavePackage(const FName& PlatformName) const
+	{
+		const int32 Index = PlatformNames.IndexOfByKey(PlatformName);
+
+		if ((Index != -1) && (Index < bSucceededSavePackage.Num()))
+		{
+			return bSucceededSavePackage[Index];
+		}
+		return false;
+	}
+
+private:
+	FName			Filename;
+	TArray<FName>	PlatformNames;
+	TArray<bool>	bSucceededSavePackage; // one bool for each platform
+};
+
+struct CookedPackageSet
+{
+private:
+	mutable FCriticalSection				SynchronizationObject;
+	TMap<FName, FFilePlatformCookedPackage> FilesProcessed;
+
+public:
+	int32 Num()
+	{
+		return FilesProcessed.Num();
+	}
+
+	void Add(const FFilePlatformCookedPackage& Request)
+	{
+		FScopeLock ScopeLock(&SynchronizationObject);
+		check(Request.IsValid());
+
+		// see if it's already in the requests list
+		FFilePlatformCookedPackage* ExistingRequest = FilesProcessed.Find(Request.GetFilename());
+
+		if (ExistingRequest)
+		{
+			check(ExistingRequest->GetFilename() == Request.GetFilename());
+			for (const FName& Platform : Request.GetPlatformNames())
+			{
+				const bool bSucceeded = Request.HasSucceededSavePackage(Platform);
+				ExistingRequest->AddPlatform(Platform, bSucceeded);
+			}
+		}
+		else
+		{
+			FilesProcessed.Add(Request.GetFilename(), Request);
+		}
+	}
+
+	bool Exists(const UCookOnTheFlyServer::FFilePlatformRequest& Request) const
+	{
+		FScopeLock ScopeLock(&SynchronizationObject);
+
+		const FFilePlatformCookedPackage* OurRequest = FilesProcessed.Find(Request.GetFilename());
+
+		if (!OurRequest)
+		{
+			return false;
+		}
+
+		// make sure all the platforms are completed
+		for (const FName& Platform : Request.GetPlatformNames())
+		{
+			if (!OurRequest->GetPlatformNames().Contains(Platform))
+			{
+				return false;
+			}
+		}
+
+		return true;
+		// return FilesProcessed.Contains(Filename);
+	}
+
+	// There are two versions of this function so I don't have to create temporary FFileplatformRequest 
+	// in some cases to call the exists function
+	bool Exists(const FName& Filename, const TArray<FName>& PlatformNames) const
+	{
+		FScopeLock ScopeLock(&SynchronizationObject);
+
+		const FFilePlatformCookedPackage* OurRequest = FilesProcessed.Find(Filename);
+
+		if (!OurRequest)
+		{
+			return false;
+		}
+
+		// make sure all the platforms are completed
+		for (const FName& Platform : PlatformNames)
+		{
+			if (!OurRequest->GetPlatformNames().Contains(Platform))
+			{
+				return false;
+			}
+		}
+
+		return true;
+	}
+
+	// do we want failed packages or not
+	bool Exists(const FName& Filename, const TArray<FName>& PlatformNames, bool bIncludeFailed) const
+	{
+		FScopeLock ScopeLock(&SynchronizationObject);
+
+		const FFilePlatformCookedPackage* OurRequest = FilesProcessed.Find(Filename);
+
+		if (!OurRequest)
+		{
+			return false;
+		}
+
+		if (bIncludeFailed == false)
+		{
+			bool allFailed = true;
+			for (const FName& PlatformName : PlatformNames)
+			{
+				if (OurRequest->HasSucceededSavePackage(PlatformName))
+				{
+					allFailed = false;
+					break;
+				}
+			}
+
+			if (allFailed)
+			{
+				return false;
+			}
+		}
+
+		// make sure all the platforms are completed
+		for (const FName& Platform : PlatformNames)
+		{
+			if (!OurRequest->GetPlatformNames().Contains(Platform))
+			{
+				return false;
+			}
+		}
+
+		return true;
+	}
+
+	void RemoveAllFilesForPlatform(const FName& PlatformName)
+	{
+		FScopeLock ScopeLock(&SynchronizationObject);
+
+		for (auto& Request : FilesProcessed)
+		{
+			Request.Value.RemovePlatform(PlatformName);
+		}
+	}
+
+	bool GetCookedPlatforms(const FName& Filename, TArray<FName>& OutPlatformList) const
+	{
+		FScopeLock ScopeLock(&SynchronizationObject);
+
+		if (const FFilePlatformCookedPackage* Request = FilesProcessed.Find(Filename))
+		{
+			OutPlatformList = Request->GetPlatformNames();
+			return true;
+		}
+
+		return false;
+	}
+
+	int RemoveFile(const FName& Filename)
+	{
+		FScopeLock ScopeLock(&SynchronizationObject);
+
+		return FilesProcessed.Remove(Filename);
+	}
+
+	bool RemoveFileForPlatform(const FName& Filename, const FName& PlatformName)
+	{
+		FScopeLock ScopeLock(&SynchronizationObject);
+
+		if (FFilePlatformCookedPackage* ProcessedFile = FilesProcessed.Find(Filename))
+		{
+			ProcessedFile->RemovePlatform(PlatformName);
+			return true;
+		}
+
+		return false;
+	}
+
+	void GetCookedFilesForPlatform(const FName& PlatformName, TArray<FName>& CookedFiles, bool bGetFailedCookedPackages, bool bGetSuccessfulCookedPackages)
+	{
+		FScopeLock ScopeLock(&SynchronizationObject);
+		for (const auto& CookedFile : FilesProcessed)
+		{
+			if (CookedFile.Value.HasPlatform(PlatformName))
+			{
+				const bool bHasSucceededSavePackage = CookedFile.Value.HasSucceededSavePackage(PlatformName);
+
+				if (	(bHasSucceededSavePackage && bGetSuccessfulCookedPackages)
+					|| ((bHasSucceededSavePackage == false) && bGetFailedCookedPackages))
+				{
+					CookedFiles.Add(CookedFile.Value.GetFilename());
+				}
+			}
+		}
+	}
+
+	void Empty(int32 ExpectedNumElements = 0)
+	{
+		FScopeLock ScopeLock(&SynchronizationObject);
+		FilesProcessed.Empty(ExpectedNumElements);
+	}
+};
+
+struct CookRequestQueue
+{
+	typedef UCookOnTheFlyServer::FFilePlatformRequest FFilePlatformRequest;
+
+	template <class PREDICATE_CLASS>
+	void Sort(const PREDICATE_CLASS& Predicate)
+	{
+		FScopeLock ScopeLock(&SynchronizationObject);
+		Queue.Sort(Predicate);
+	}
+
+	const TArray<FName>& GetQueue() const 
+	{ 
+		return Queue;
+	}
+
+	void EnqueueUnique(const FFilePlatformRequest& Request, bool ForceEnqueFront = false)
+	{
+		FScopeLock ScopeLock(&SynchronizationObject);
+		TArray<FName>* Platforms = PlatformList.Find(Request.GetFilename());
+		if (Platforms == NULL)
+		{
+			PlatformList.Add(Request.GetFilename(), Request.GetPlatformNames());
+			Queue.Add(Request.GetFilename());
+		}
+		else
+		{
+			// add the requested platforms to the platform list
+			for (const FName& Platform : Request.GetPlatformNames())
+			{
+				Platforms->AddUnique(Platform);
+			}
+		}
+
+		if (ForceEnqueFront)
+		{
+			int32 Index = Queue.Find(Request.GetFilename());
+			check(Index != INDEX_NONE);
+			if (Index != 0)
+			{
+				Queue.Swap(0, Index);
+			}
+		}
+	}
+
+	bool Dequeue(FFilePlatformRequest& OutResult)
+	{
+		FScopeLock ScopeLock(&SynchronizationObject);
+		if (Queue.Num())
+		{
+			FName Filename = Queue[0];
+			Queue.RemoveAt(0);
+			TArray<FName> Platforms = PlatformList.FindChecked(Filename);
+			PlatformList.Remove(Filename);
+
+			OutResult = FFilePlatformRequest(MoveTemp(Filename), MoveTemp(Platforms));
+			return true;
+		}
+		return false;
+	}
+
+	void DequeueAllRequests(TArray<FFilePlatformRequest>& RequestArray)
+	{
+		FScopeLock ScopeLock(&SynchronizationObject);
+		if (Queue.Num())
+		{
+			for (const auto& Request : PlatformList)
+			{
+				RequestArray.Add(FFilePlatformRequest(Request.Key, Request.Value));
+			}
+			PlatformList.Empty();
+			Queue.Empty();
+		}
+	}
+
+	bool Exists(const FName& Filename, const TArray<FName>& PlatformNames) const
+	{
+		FScopeLock ScopeLock(&SynchronizationObject);
+		const TArray<FName>* Platforms = PlatformList.Find(Filename);
+		if (Platforms == NULL)
+			return false;
+
+		for (const FName& PlatformName : PlatformNames)
+		{
+			if (!Platforms->Contains(PlatformName))
+				return false;
+		}
+		return true;
+	}
+
+	bool Exists(const FName& Filename)
+	{
+		FScopeLock ScopeLock(&SynchronizationObject);
+		const TArray<FName>* Platforms = PlatformList.Find(Filename);
+		if (Platforms == NULL)
+			return false;
+		return true;
+	}
+
+	bool HasItems() const
+	{
+		FScopeLock ScopeLock(&SynchronizationObject);
+		return Queue.Num() > 0;
+	}
+
+	int Num() const
+	{
+		FScopeLock ScopeLock(&SynchronizationObject);
+		return Queue.Num();
+	}
+
+	void Empty()
+	{
+		FScopeLock ScopeLock(&SynchronizationObject);
+		Queue.Empty();
+		PlatformList.Empty();
+	}
+
+private:
+	TArray<FName>				Queue;
+	TMap<FName, TArray<FName>>	PlatformList;
+	mutable FCriticalSection	SynchronizationObject;
+};
+
 struct UCookOnTheFlyServer::FImpl : public FUObjectArray::FUObjectCreateListener, public FUObjectArray::FUObjectDeleteListener
 {
 	// This is the set of packages which have already had PostLoadFixup called 
@@ -532,6 +953,13 @@ struct UCookOnTheFlyServer::FImpl : public FUObjectArray::FUObjectCreateListener
 	// This list contains the UPackages loaded since last call to GetNewPackages
 	TArray<UPackage*> NewPackages;
 
+	// Set of files which have been cooked 
+	// when needing to recook a file the entry will need to be removed from here
+	CookedPackageSet CookedPackages;
+
+	// List of requested files
+	CookRequestQueue CookRequests;
+
 	FImpl()
 	{
 		for (TObjectIterator<UPackage> It; It; ++It)
@@ -1288,10 +1716,15 @@ FString UCookOnTheFlyServer::GetContentDirecctoryForDLC() const
 
 COREUOBJECT_API extern bool GOutputCookingWarnings;
 
+bool UCookOnTheFlyServer::HasCookRequests() const 
+{ 
+	return Impl->CookRequests.HasItems(); 
+}
+
 bool UCookOnTheFlyServer::RequestPackage(const FName& StandardPackageFName, const TArray<FName>& TargetPlatforms, const bool bForceFrontOfQueue)
 {
 	FFilePlatformRequest FileRequest(StandardPackageFName, TargetPlatforms);
-	CookRequests.EnqueueUnique(FileRequest, bForceFrontOfQueue);
+	Impl->CookRequests.EnqueueUnique(FileRequest, bForceFrontOfQueue);
 	return true;
 }
 
@@ -1301,7 +1734,7 @@ bool UCookOnTheFlyServer::RequestPackage(const FName& StandardPackageFName, cons
 	check(IsCookByTheBookMode());
 	// need target platforms if we are not in cook by the book mode
 	FFilePlatformRequest FileRequest(StandardPackageFName, CookByTheBookOptions->TargetPlatformNames);
-	CookRequests.EnqueueUnique(FileRequest, bForceFrontOfQueue);
+	Impl->CookRequests.EnqueueUnique(FileRequest, bForceFrontOfQueue);
 	return true;
 }
 
@@ -1359,10 +1792,10 @@ uint32 UCookOnTheFlyServer::TickCookOnTheSide( const float TimeSlice, uint32 &Co
 	
 	while (!GIsRequestingExit || CurrentCookMode == ECookMode::CookByTheBook)
 	{
-		if (CookRequests.Num() > 0)
+		if (HasCookRequests())
 		{
 			UE_CLOG(!(TickFlags & ECookTickFlags::HideProgressDisplay) && (GCookProgressDisplay & (int32)ECookProgressDisplayMode::RemainingPackages), 
-			LogCook, Display, TEXT("Cooked packages %d Packages Remain %d Total %d"), CookedPackages.Num(), CookRequests.Num(), CookedPackages.Num() + CookRequests.Num());
+			LogCook, Display, TEXT("Cooked packages %d Packages Remain %d Total %d"), Impl->CookedPackages.Num(), Impl->CookRequests.Num(), Impl->CookedPackages.Num() + Impl->CookRequests.Num());
 		}
 
 		// if we just cooked a map then don't process anything the rest of this tick
@@ -1382,10 +1815,9 @@ uint32 UCookOnTheFlyServer::TickCookOnTheSide( const float TimeSlice, uint32 &Co
 
 		FFilePlatformRequest ToBuild;
 		
-		//bool bIsUnsolicitedRequest = false;
-		if ( CookRequests.HasItems() )
+		if (HasCookRequests())
 		{
-			CookRequests.Dequeue( &ToBuild );
+			Impl->CookRequests.Dequeue(/* out */ ToBuild);
 		}
 		else
 		{
@@ -1408,7 +1840,7 @@ uint32 UCookOnTheFlyServer::TickCookOnTheSide( const float TimeSlice, uint32 &Co
 			GUnrealEd->GetPackageAutoSaver().ForceMinimumTimeTillAutoSave(SecondsWarningTillAutosave);
 		}
 
-		if (CookedPackages.Exists(ToBuild))
+		if (Impl->CookedPackages.Exists(ToBuild))
 		{
 #if DEBUG_COOKONTHEFLY
 			UE_LOG(LogCook, Display, TEXT("Package for platform already cooked %s, discarding request"), *ToBuild.GetFilename().ToString());
@@ -1483,7 +1915,7 @@ uint32 UCookOnTheFlyServer::TickCookOnTheSide( const float TimeSlice, uint32 &Co
 					//	and then also save our new request as processed so we don't do it again
 					UE_LOG(LogCook, Verbose, TEXT("Request for %s received going to save %s"), *BuildFilename, *PackageFilename);
 
-					CookedPackages.Add(FFilePlatformCookedPackage(ToBuild.GetFilename(), TargetPlatformNames));
+					Impl->CookedPackages.Add(FFilePlatformCookedPackage(ToBuild.GetFilename(), TargetPlatformNames));
 
 					ToBuild.SetFilename(PackageFilename);
 				}
@@ -1521,7 +1953,7 @@ uint32 UCookOnTheFlyServer::TickCookOnTheSide( const float TimeSlice, uint32 &Co
 				}
 			}
 #endif
-			CookedPackages.Add( FFilePlatformCookedPackage( ToBuild.GetFilename(), TargetPlatformNames) );
+			Impl->CookedPackages.Add( FFilePlatformCookedPackage( ToBuild.GetFilename(), TargetPlatformNames) );
 			continue;
 		}
 
@@ -1606,7 +2038,7 @@ uint32 UCookOnTheFlyServer::TickCookOnTheSide( const float TimeSlice, uint32 &Co
 			{
 				GShaderCompilingManager->ProcessAsyncResults(true, false); // we can afford to do work here because we are essentually requing this package for processing later
 				Timer.SavedPackage();  // this is a special case to prevent infinite loop, if we only have one package we might fall through this and could loop forever.  
-				CookRequests.EnqueueUnique(ToBuild, false);
+				Impl->CookRequests.EnqueueUnique(ToBuild, false);
 				continue;
 			}
 		}
@@ -1638,7 +2070,7 @@ uint32 UCookOnTheFlyServer::TickCookOnTheSide( const float TimeSlice, uint32 &Co
 	}
 
 	if (IsCookByTheBookRunning() &&
-		(CookRequests.HasItems() == false))
+		(HasCookRequests() == false))
 	{
 		check(IsCookByTheBookMode());
 
@@ -1939,7 +2371,7 @@ void UCookOnTheFlyServer::SaveCookedPackages(TArray<UPackage*>& PackagesToSave,
 			
 			TArray<FName> SaveTargetPlatformNames = AllTargetPlatformNames;
 			TArray<FName> CookedTargetPlatforms;
-			if (CookedPackages.GetCookedPlatforms(PackageFName, CookedTargetPlatforms))
+			if (Impl->CookedPackages.GetCookedPlatforms(PackageFName, CookedTargetPlatforms))
 			{
 				for (const FName& CookedPlatform : CookedTargetPlatforms)
 				{
@@ -1961,7 +2393,7 @@ void UCookOnTheFlyServer::SaveCookedPackages(TArray<UPackage*>& PackagesToSave,
 
 			bool bShouldFinishTick = false;
 
-			if (Timer.IsTimeUp() && IsCookByTheBookMode())
+			if (Timer.IsTimeUp() && IsCookByTheBookMode() )
 			{
 				// our timeslice is up
 				bShouldFinishTick = true;
@@ -1977,7 +2409,7 @@ void UCookOnTheFlyServer::SaveCookedPackages(TArray<UPackage*>& PackagesToSave,
 				if (bProcessingUnsolicitedPackages)
 				{
 					SCOPE_TIMER(WaitingForCachedCookedPlatformData);
-					if (CookRequests.HasItems())
+					if (HasCookRequests())
 					{
 						bShouldFinishTick = true;
 					}
@@ -2056,10 +2488,10 @@ void UCookOnTheFlyServer::SaveCookedPackages(TArray<UPackage*>& PackagesToSave,
 				for (int32 RemainingIndex = I; RemainingIndex < NumPackagesToRequeue; ++RemainingIndex)
 				{
 					FName StandardFilename = GetCachedStandardPackageFileFName(PackagesToSave[RemainingIndex]);
-					CookRequests.EnqueueUnique(FFilePlatformRequest(StandardFilename, SaveTargetPlatformNames));
+					Impl->CookRequests.EnqueueUnique(FFilePlatformRequest(StandardFilename, SaveTargetPlatformNames));
 				}
 				Result |= COSR_WaitingOnCache;
-			
+
 				// break out of the loop
 				return;
 			}
@@ -2213,13 +2645,14 @@ void UCookOnTheFlyServer::SaveCookedPackages(TArray<UPackage*>& PackagesToSave,
 				}
 				if (!bWasReferencedOnlyByEditorOnlyData)
 				{
-					CookedPackages.Add(FileRequest);
+					Impl->CookedPackages.Add(FileRequest);
 					if ((CurrentCookMode == ECookMode::CookOnTheFly) && (I >= FirstUnsolicitedPackage))
 					{
 						// this is an unsolicited package
 						if (FPaths::FileExists(FileRequest.GetFilename().ToString()) == true)
 						{
-							UnsolicitedCookedPackages.AddCookedPackage(FileRequest);
+							UnsolicitedCookedPackages.AddCookedPackage(FFilePlatformRequest(FileRequest.GetFilename(), FileRequest.GetPlatformNames()));
+
 #if DEBUG_COOKONTHEFLY
 							UE_LOG(LogCook, Display, TEXT("UnsolicitedCookedPackages: %s"), *FileRequest.GetFilename().ToString());
 #endif
@@ -2459,14 +2892,11 @@ void UCookOnTheFlyServer::GetUnsolicitedPackages(TArray<UPackage*>& PackagesToSa
 		{
 			ensure(Package != nullptr);
 
-			if (PackagesToSaveSet.Find(Package))
-				continue;
-
-			FName StandardPackageFName = GetCachedStandardPackageFileFName(Package);
+			const FName StandardPackageFName = GetCachedStandardPackageFileFName(Package);
 			if (StandardPackageFName == NAME_None)
 				continue;	// if we have name none that means we are in core packages or something...
 
-			if (CookedPackages.Exists(StandardPackageFName, TargetPlatformNames))
+			if (Impl->CookedPackages.Exists(StandardPackageFName, TargetPlatformNames))
 				continue;
 				
 			PackagesToSaveSet.Add(Package);
@@ -2584,13 +3014,13 @@ void UCookOnTheFlyServer::MarkPackageDirtyForCooker( UPackage *Package )
 			{
 				TArray<FName> CookedPlatforms;
 				// if we have already cooked this package and we have made changes then recook ;)
-				if ( CookedPackages.GetCookedPlatforms(PackageFFileName, CookedPlatforms) )
+				if (Impl->CookedPackages.GetCookedPlatforms(PackageFFileName, CookedPlatforms) )
 				{
 					if (IsCookByTheBookRunning())
 					{
 						// if this package was previously cooked and we are doing a cook by the book 
 						// we need to recook this package before finishing cook by the book
-						CookRequests.EnqueueUnique(FFilePlatformRequest(PackageFFileName, CookedPlatforms));
+						Impl->CookRequests.EnqueueUnique(FFilePlatformRequest(PackageFFileName, CookedPlatforms));
 					}
 					else
 					{
@@ -2618,7 +3048,7 @@ void UCookOnTheFlyServer::MarkPackageDirtyForCooker( UPackage *Package )
 				check( true);
 			}
 		}
-		CookedPackages.RemoveFile( PackageFFileName );
+		Impl->CookedPackages.RemoveFile( PackageFFileName );
 	}
 }
 
@@ -2755,7 +3185,7 @@ void UCookOnTheFlyServer::MarkGCPackagesToKeepForCooker()
 	// then see if the package is needed by any of the packages which are required to be saved
 
 	TMap<FName, int32> PackageDependenciesCount;
-	for (const FName& QueuedPackage : CookRequests.GetQueue())
+	for (const FName& QueuedPackage : Impl->CookRequests.GetQueue())
 	{
 		const FName* PackageName = GetCachedPackageFilenameToPackageFName(QueuedPackage);
 		if ( !PackageName )
@@ -2785,10 +3215,11 @@ void UCookOnTheFlyServer::MarkGCPackagesToKeepForCooker()
 		}
 	}
 
-	// sort the cook requests by the packages which are loaded first
+	// Sort the cook requests by the packages which are loaded first
 	// then sort by the number of dependencies which are referenced by the package
-	// we want to process the packages with the highest dependencies so that they can be evicted from memory and are likely to be able to be released on next gc pass
-	CookRequests.Sort([=, &PackageDependenciesCount,&LoadedPackages](const FName& A, const FName& B)
+	// we want to process the packages with the highest dependencies so that they can
+	// be evicted from memory and are likely to be able to be released on next GC pass
+	Impl->CookRequests.Sort([&PackageDependenciesCount,&LoadedPackages](const FName& A, const FName& B)
 		{
 			int32 ADependencies = PackageDependenciesCount.FindChecked(A);
 			int32 BDependencies = PackageDependenciesCount.FindChecked(B);
@@ -2825,6 +3256,10 @@ void UCookOnTheFlyServer::TickRecompileShaderRequests()
 	}
 }
 
+bool UCookOnTheFlyServer::HasRecompileShaderRequests() const 
+{ 
+	return RecompileRequests.HasItems(); 
+}
 
 void UCookOnTheFlyServer::SaveCookedPackage(UPackage* Package, uint32 SaveFlags, TArray<FSavePackageResultStruct>& SavePackageResults)
 {
@@ -4457,7 +4892,7 @@ void UCookOnTheFlyServer::PopulateCookedPackagesFromDisk(const TArray<ITargetPla
 
 				if (IdenticalCookedPackages.Contains(SourcePackageName))
 				{
-					CookedPackages.Add(FFilePlatformCookedPackage(UncookedFilename, MoveTemp(PlatformNames), MoveTemp(Succeeded)));
+					Impl->CookedPackages.Add(FFilePlatformCookedPackage(UncookedFilename, MoveTemp(PlatformNames), MoveTemp(Succeeded)));
 					++NumPackagesKept;
 				}
 			}
@@ -4466,7 +4901,7 @@ void UCookOnTheFlyServer::PopulateCookedPackagesFromDisk(const TArray<ITargetPla
 				if (SourcePackageName != NAME_None && IsCookByTheBookMode()) // cook on the fly will requeue this package when it wants it 
 				{
 					// Force cook the modified file
-					CookRequests.EnqueueUnique(FFilePlatformRequest(UncookedFilename, PlatformNames));
+					Impl->CookRequests.EnqueueUnique(FFilePlatformRequest(UncookedFilename, PlatformNames));
 				}
 				if (CookedFile != NAME_DummyCookedFilename)
 				{
@@ -4493,9 +4928,9 @@ void UCookOnTheFlyServer::PopulateCookedPackagesFromDisk(const TArray<ITargetPla
 			TArray<FName> PlatformNames;
 			PlatformNames.Add(PlatformFName);
 
-			ensure(CookedPackages.Exists(UncookedFilename, PlatformNames, false) == false);
+			ensure(Impl->CookedPackages.Exists(UncookedFilename, PlatformNames, false) == false);
 
-			CookedPackages.Add(FFilePlatformCookedPackage(UncookedFilename, MoveTemp(PlatformNames)));
+			Impl->CookedPackages.Add(FFilePlatformCookedPackage(UncookedFilename, MoveTemp(PlatformNames)));
 			++NumMarkedFailedSaveKept;
 		}
 
@@ -4598,7 +5033,7 @@ void UCookOnTheFlyServer::CleanSandbox(const bool bIterative)
 			}
 
 			// This is fast in asset registry iterate, so just reconstruct
-			CookedPackages.Empty();
+			Impl->CookedPackages.Empty();
 			PopulateCookedPackagesFromDisk(Platforms);
 		}
 	}
@@ -5521,10 +5956,10 @@ void UCookOnTheFlyServer::CookByTheBookFinished()
 				const FName& PlatformName = Pair.Key;
 				FAssetRegistryGenerator& Generator = *Pair.Value;
 
-				CookedPackages.GetCookedFilesForPlatform(PlatformName, CookedPackagesFilenames, false, true);
+				Impl->CookedPackages.GetCookedFilesForPlatform(PlatformName, CookedPackagesFilenames, false, /* include successful */ true);
 
 				// ignore any packages which failed to cook
-				CookedPackages.GetCookedFilesForPlatform(PlatformName, IgnorePackageFilenames, true, false);
+				Impl->CookedPackages.GetCookedFilesForPlatform(PlatformName, IgnorePackageFilenames, /* include failed */ true, false);
 
 				bool bForceNoFilterAssetsFromAssetRegistry = false;
 
@@ -5650,7 +6085,7 @@ void UCookOnTheFlyServer::BuildMapDependencyGraph(const FName& PlatformName)
 	auto& MapDependencyGraph = CookByTheBookOptions->MapDependencyGraphs.FindChecked(PlatformName);
 
 	TArray<FName> PlatformCookedPackages;
-	CookedPackages.GetCookedFilesForPlatform(PlatformName, PlatformCookedPackages);
+	Impl->CookedPackages.GetCookedFilesForPlatform(PlatformName, PlatformCookedPackages, /* include failed */ true, /* include successful */ true);
 
 	// assign chunks for all the map packages
 	for (const FName& CookedPackage : PlatformCookedPackages)
@@ -5716,7 +6151,7 @@ void UCookOnTheFlyServer::CancelCookByTheBook()
 		check( IsInGameThread() );
 
 		// save the cook requests 
-		CookRequests.DequeueAllRequests(CookByTheBookOptions->PreviousCookRequests);
+		Impl->CookRequests.DequeueAllRequests(CookByTheBookOptions->PreviousCookRequests);
 		CookByTheBookOptions->bRunning = false;
 
 		SandboxFile = nullptr;
@@ -5734,9 +6169,9 @@ void UCookOnTheFlyServer::StopAndClearCookedData()
 	}
 
 	RecompileRequests.Empty();
-	CookRequests.Empty();
+	Impl->CookRequests.Empty();
 	UnsolicitedCookedPackages.Empty();
-	CookedPackages.Empty(); // set of files which have been cooked when needing to recook a file the entry will need to be removed from here
+	Impl->CookedPackages.Empty(); // set of files which have been cooked when needing to recook a file the entry will need to be removed from here
 }
 
 void UCookOnTheFlyServer::ClearAllCookedData()
@@ -5746,7 +6181,7 @@ void UCookOnTheFlyServer::ClearAllCookedData()
 	UPackage::WaitForAsyncFileWrites();
 
 	UnsolicitedCookedPackages.Empty();
-	CookedPackages.Empty(); // set of files which have been cooked when needing to recook a file the entry will need to be removed from here
+	Impl->CookedPackages.Empty(); // set of files which have been cooked when needing to recook a file the entry will need to be removed from here
 }
 
 void UCookOnTheFlyServer::ClearPlatformCookedData( const FName& PlatformName )
@@ -5755,7 +6190,7 @@ void UCookOnTheFlyServer::ClearPlatformCookedData( const FName& PlatformName )
 	// that means it's also conceivable that we will recook the same package which currently has an outstanding async write request
 	UPackage::WaitForAsyncFileWrites();
 
-	CookedPackages.RemoveAllFilesForPlatform( PlatformName );
+	Impl->CookedPackages.RemoveAllFilesForPlatform( PlatformName );
 	TArray<FName> PackageNames;
 	UnsolicitedCookedPackages.GetPackagesForPlatformAndRemove(PlatformName, PackageNames);
 }
@@ -6086,7 +6521,7 @@ void UCookOnTheFlyServer::StartCookByTheBook( const FCookByTheBookStartupOptions
 				Succeeded.Add(true);
 				for (const FName& PackageFilename : PackageList)
 				{
-					CookedPackages.Add( FFilePlatformCookedPackage( PackageFilename, PlatformNames, Succeeded) );
+					Impl->CookedPackages.Add( FFilePlatformCookedPackage( PackageFilename, PlatformNames, Succeeded) );
 				}
 			}
 			CookByTheBookOptions->BasedOnReleaseCookedPackages.Add(PlatformName, MoveTemp(PackageList));
@@ -6173,7 +6608,7 @@ void UCookOnTheFlyServer::StartCookByTheBook( const FCookByTheBookStartupOptions
 		
 		if (PackageFileFName != NAME_None)
 		{
-			CookRequests.EnqueueUnique( FFilePlatformRequest( PackageFileFName, TargetPlatformNames ) );
+			Impl->CookRequests.EnqueueUnique( FFilePlatformRequest( PackageFileFName, TargetPlatformNames ) );
 		}
 		else if (!FLinkerLoad::IsKnownMissingPackage(FileFName))
 		{
@@ -6206,7 +6641,7 @@ void UCookOnTheFlyServer::StartCookByTheBook( const FCookByTheBookStartupOptions
 
 				for ( const FName& PackageFilename : PackageFiles )
 				{
-					CookRequests.EnqueueUnique( FFilePlatformRequest( PackageFilename, PlatformArray) );
+					Impl->CookRequests.EnqueueUnique( FFilePlatformRequest( PackageFilename, PlatformArray) );
 				}
 
 			}
@@ -6220,7 +6655,7 @@ void UCookOnTheFlyServer::StartCookByTheBook( const FCookByTheBookStartupOptions
 	// if too many packages are being recooked after resume then we may need to figure out a different way to do this
 	for ( const FFilePlatformRequest& PreviousRequest : CookByTheBookOptions->PreviousCookRequests )
 	{
-		CookRequests.EnqueueUnique( PreviousRequest );
+		Impl->CookRequests.EnqueueUnique( PreviousRequest );
 	}
 	CookByTheBookOptions->PreviousCookRequests.Empty();
 
@@ -6268,7 +6703,7 @@ void UCookOnTheFlyServer::MaybeMarkPackageAsAlreadyLoaded(UPackage *Package)
 	bool bShouldMarkAsAlreadyProcessed = false;
 
 	TArray<FName> CookedPlatforms;
-	if (CookedPackages.GetCookedPlatforms(StandardName, CookedPlatforms))
+	if (Impl->CookedPackages.GetCookedPlatforms(StandardName, CookedPlatforms))
 	{
 		bShouldMarkAsAlreadyProcessed = true;
 		for (const FName& TargetPlatform : CookByTheBookOptions->TargetPlatformNames)
@@ -6392,7 +6827,7 @@ void UCookOnTheFlyServer::HandleNetworkFileServerFileRequest( const FString& Fil
 	
 	UE_LOG(LogCook, Display, TEXT("Requesting file from cooker %s"), *StandardFileName);
 
-	CookRequests.EnqueueUnique(FileRequest, true);
+	Impl->CookRequests.EnqueueUnique(FileRequest, true);
 
 
 #if PROFILE_NETWORK
@@ -6400,7 +6835,7 @@ void UCookOnTheFlyServer::HandleNetworkFileServerFileRequest( const FString& Fil
 	while (NetworkRequestEvent->Wait(1) == false)
 	{
 		// for some reason we missed the stat
-		if (CookedPackages.Exists(FileRequest))
+		if (Impl->CookedPackages.Exists(FileRequest))
 		{
 			double DeltaTimeTillRequestForfilled = FPlatformTime::Seconds() - StartTime;
 			TimeTillRequestForfilled += DeltaTimeTillRequestForfilled;
@@ -6417,7 +6852,7 @@ void UCookOnTheFlyServer::HandleNetworkFileServerFileRequest( const FString& Fil
 	StartTime = FPlatformTime::Seconds();
 #endif
 
-	while (!CookedPackages.Exists(FileRequest))
+	while (!Impl->CookedPackages.Exists(FileRequest))
 	{
 		FPlatformProcess::Sleep(0.0001f);
 	}
@@ -6454,7 +6889,7 @@ void UCookOnTheFlyServer::HandleNetworkGetPrecookedList(const FString& PlatformN
 	FName PlatformFName = FName(*PlatformName);
 
 	TArray<FName> CookedPlatformFiles;
-	CookedPackages.GetCookedFilesForPlatform(PlatformFName, CookedPlatformFiles);
+	Impl->CookedPackages.GetCookedFilesForPlatform(PlatformFName, CookedPlatformFiles, /* include failed */ true, /* include successful */ true);
 
 
 	for ( const FName& CookedFile : CookedPlatformFiles)
@@ -6569,10 +7004,10 @@ uint32 UCookOnTheFlyServer::FullLoadAndSave(uint32& CookedPackageCount)
 	{
 		UE_LOG(LogCook, Display, TEXT("Loading requested packages..."));
 		SCOPE_TIMER(FullLoadAndSave_RequestedLoads);
-		while (CookRequests.HasItems())
+		while (HasCookRequests())
 		{
 			FFilePlatformRequest ToBuild;
-			CookRequests.Dequeue(&ToBuild);
+			Impl->CookRequests.Dequeue(/* out */ ToBuild);
 			check(ToBuild.IsValid());
 
 			const FName BuildFilenameFName = ToBuild.GetFilename();
@@ -6973,7 +7408,7 @@ uint32 UCookOnTheFlyServer::FullLoadAndSave(uint32& CookedPackageCount)
 				check(CachedPackageFilename->StandardFileFName != NAME_None);
 
 				FFilePlatformCookedPackage FileRequest(CachedPackageFilename->StandardFileFName, CookByTheBookOptions->TargetPlatformNames, SavePackageSuccessPerPlatform);
-				CookedPackages.Add(FileRequest);
+				Impl->CookedPackages.Add(FileRequest);
 
 				if (SavePackageSuccessPerPlatform.Contains(false))
 				{
@@ -7015,4 +7450,3 @@ uint32 UCookOnTheFlyServer::FullLoadAndSave(uint32& CookedPackageCount)
 }
 
 #undef LOCTEXT_NAMESPACE
-
