Added some filtering logic to ensure the new package tracking logic matches the old iteration code.

Exposed by CI on EngineTest, which evidently has more elaborate data setups than InfiltratorDemo

diff --git a/Engine/Source/Editor/UnrealEd/Private/CookOnTheFlyServer.cpp b/Engine/Source/Editor/UnrealEd/Private/CookOnTheFlyServer.cpp
--- a/Engine/Source/Editor/UnrealEd/Private/CookOnTheFlyServer.cpp
+++ b/Engine/Source/Editor/UnrealEd/Private/CookOnTheFlyServer.cpp
@@ -964,7 +964,12 @@ struct UCookOnTheFlyServer::FImpl : public FUObjectArray::FUObjectCreateListener
 	{
 		for (TObjectIterator<UPackage> It; It; ++It)
 		{
-			LoadedPackages.Add(*It);
+			UPackage* Package = *It;
+
+			if (Package->GetOuter() == nullptr)
+			{
+				LoadedPackages.Add(Package);
+			}
 		}
 
 		NewPackages = LoadedPackages;
@@ -990,8 +995,11 @@ struct UCookOnTheFlyServer::FImpl : public FUObjectArray::FUObjectCreateListener
 		{
 			auto Package = const_cast<UPackage*>(static_cast<const UPackage*>(Object));
 
-			LoadedPackages.Add(Package);
-			NewPackages.Add(Package);
+			if (Package->GetOuter() == nullptr)
+			{
+				LoadedPackages.Add(Package);
+				NewPackages.Add(Package);
+			}
 		}
 	}
 
