Moving more state out of the public header for CookOnTheFlyServer

No functional change -- just restructuring code in preparation for refactoring+optimization.

diff --git a/Engine/Source/Editor/UnrealEd/Classes/CookOnTheSide/CookOnTheFlyServer.h b/Engine/Source/Editor/UnrealEd/Classes/CookOnTheSide/CookOnTheFlyServer.h
--- a/Engine/Source/Editor/UnrealEd/Classes/CookOnTheSide/CookOnTheFlyServer.h
+++ b/Engine/Source/Editor/UnrealEd/Classes/CookOnTheSide/CookOnTheFlyServer.h
@@ -33,6 +33,7 @@ enum class ECookInitializationFlags
 	Unversioned =								0x00000008, // save the cooked packages without a version number
 	AutoTick =									0x00000010, // enable ticking (only works in the editor)
 	AsyncSave =									0x00000020, // save packages async
+	// unused =									0x00000040,
 	IncludeServerMaps =							0x00000080, // should we include the server maps when cooking
 	UseSerializationForPackageDependencies =	0x00000100, // should we use the serialization code path for generating package dependencies (old method will be deprecated)
 	BuildDDCInBackground =						0x00000200, // build ddc content in background while the editor is running (only valid for modes which are in editor IsCookingInEditor())
@@ -102,251 +103,9 @@ class UNREALED_API UCookOnTheFlyServer : public UObject, public FTickableEditorO
 {
 	GENERATED_BODY()
 
-		UCookOnTheFlyServer(const FObjectInitializer& ObjectInitializer = FObjectInitializer::Get());
+	UCookOnTheFlyServer(const FObjectInitializer& ObjectInitializer = FObjectInitializer::Get());
 
 private:
-
-	/** Array which has been made thread safe :) */
-	template<typename Type, typename SynchronizationObjectType, typename ScopeLockType>
-	struct FUnsynchronizedQueue
-	{
-	private:
-		mutable SynchronizationObjectType	SynchronizationObject; // made this mutable so this class can have const functions and still be thread safe
-		TArray<Type>		Items;
-	public:
-		void Enqueue(const Type& Item)
-		{
-			ScopeLockType ScopeLock(&SynchronizationObject);
-			Items.Add(Item);
-		}
-		void EnqueueUnique( const Type& Item )
-		{
-			ScopeLockType ScopeLock(&SynchronizationObject);
-			Items.AddUnique(Item);
-		}
-		bool Dequeue(Type* Result)
-		{
-			ScopeLockType ScopeLock(&SynchronizationObject);
-			if (Items.Num())
-			{
-				*Result = Items[0];
-				Items.RemoveAt(0);
-				return true;
-			}
-			return false;
-		}
-		void DequeueAll(TArray<Type>& Results)
-		{
-			ScopeLockType ScopeLock(&SynchronizationObject);
-			Results += Items;
-			Items.Empty();
-		}
-
-		bool HasItems() const
-		{
-			ScopeLockType ScopeLock( &SynchronizationObject );
-			return Items.Num() > 0;
-		}
-
-		void Remove( const Type& Item ) 
-		{
-			ScopeLockType ScopeLock( &SynchronizationObject );
-			Items.Remove( Item );
-		}
-
-		void CopyItems( TArray<Type> &InItems ) const
-		{
-			ScopeLockType ScopeLock( &SynchronizationObject );
-			InItems = Items;
-		}
-
-		int Num() const 
-		{
-			ScopeLockType ScopeLock( &SynchronizationObject );
-			return Items.Num();
-		}
-
-		void Empty()
-		{
-			ScopeLockType ScopeLock( &SynchronizationObject );
-			Items.Empty();
-		}
-	};
-
-
-
-	struct FDummyCriticalSection
-	{
-	public:
-		FORCEINLINE void Lock() { }
-		FORCEINLINE void Unlock() { }
-	};
-
-	struct FDummyScopeLock
-	{
-	public:
-		FDummyScopeLock( FDummyCriticalSection * ) { }
-	};
-
-public:
-
-
-	template<typename Type>
-	struct FThreadSafeQueue : public FUnsynchronizedQueue<Type, FCriticalSection, FScopeLock>
-	{
-		/**
-		* Don't add any functions here, this is just a overqualified typedef
-		* Add functions / functionality to the FUnsynchronizedQueue
-		*/
-	};
-
-	template<typename Type>
-	struct FQueue : public FUnsynchronizedQueue<Type, FDummyCriticalSection, FDummyScopeLock>
-	{
-		/**
-		* Don't add any functions here, this is just a overqualified typedef
-		* Add functions / functionality to the FUnsynchronizedQueue
-		*/
-	};
-
-public:
-	/** cooked file requests which includes platform which file is requested for */
-	struct FFilePlatformRequest
-	{
-	protected:
-		FName Filename;
-		TArray<FName> PlatformNames;
-	public:
-
-		// yes we have some friends
-		friend uint32 GetTypeHash(const UCookOnTheFlyServer::FFilePlatformRequest& Key);
-
-		FFilePlatformRequest() { }
-
-
-		FFilePlatformRequest( const FName& InFileName, const FName& InPlatformName ) : Filename( InFileName )
-		{ PlatformNames.Add( InPlatformName ); }
-
-		FFilePlatformRequest( const FName& InFilename, const TArray<FName>& InPlatformName ) : Filename( InFilename )
-		{ PlatformNames = InPlatformName; }
-
-		FFilePlatformRequest( const FName& InFilename, TArray<FName>&& InPlatformName ) : Filename( InFilename ), PlatformNames(MoveTemp(InPlatformName)) { }
-		FFilePlatformRequest( const FFilePlatformRequest& InFilePlatformRequest ) : Filename( InFilePlatformRequest.Filename ), PlatformNames( InFilePlatformRequest.PlatformNames ) { }
-		FFilePlatformRequest( FFilePlatformRequest&& InFilePlatformRequest ) : Filename( MoveTemp(InFilePlatformRequest.Filename) ), PlatformNames( MoveTemp(InFilePlatformRequest.PlatformNames) ) { }
-
-		void SetFilename( const FString &InFilename ) 
-		{
-			Filename = FName(*InFilename);
-		}
-
-		const FName &GetFilename() const
-		{
-			return Filename;
-		}
-
-		const TArray<FName>& GetPlatformNames() const
-		{
-			return PlatformNames;
-		}
-
-		void RemovePlatform( const FName &Platform )
-		{
-			PlatformNames.Remove(Platform);
-		}
-
-		void AddPlatform( const FName &Platform )
-		{
-			check( Platform != NAME_None );
-			PlatformNames.Add(Platform );
-		}
-
-		bool HasPlatform( const FName &Platform ) const
-		{
-			return PlatformNames.Find(Platform) != INDEX_NONE;
-		}
-
-		bool IsValid()  const
-		{
-			return Filename != NAME_None;
-		}
-
-		void Clear()
-		{
-			Filename = TEXT("");
-			PlatformNames.Empty();
-		}
-
-		FFilePlatformRequest &operator=( FFilePlatformRequest &&InFileRequest )
-		{
-			Filename = MoveTemp( InFileRequest.Filename );
-			PlatformNames = MoveTemp( InFileRequest.PlatformNames );
-			return *this;
-		}
-
-		bool operator ==( const FFilePlatformRequest &InFileRequest ) const
-		{
-			if ( InFileRequest.Filename == Filename )
-			{
-				if ( InFileRequest.PlatformNames == PlatformNames )
-				{
-					return true;
-				}
-			}
-			return false;
-		}
-
-		FString ToString() const
-		{
-			FString Result = FString::Printf(TEXT("%s;"), *Filename.ToString());
-
-			for ( const FName& Platform : PlatformNames )
-			{
-				Result += FString::Printf(TEXT("%s,"), *Platform.ToString() );
-			}
-			return Result;
-		}
-	};
-
-private:
-	struct FThreadSafeUnsolicitedPackagesList
-	{
-	private:
-		FCriticalSection SyncObject;
-		TArray<FFilePlatformRequest> CookedPackages;
-	public:
-		void AddCookedPackage( const FFilePlatformRequest& PlatformRequest )
-		{
-			FScopeLock S( &SyncObject );
-			CookedPackages.Add( PlatformRequest );
-		}
-		void GetPackagesForPlatformAndRemove( const FName& Platform, TArray<FName> PackageNames )
-		{
-			FScopeLock S( &SyncObject );
-
-			for ( int I = CookedPackages.Num()-1; I >= 0; --I )
-			{
-				FFilePlatformRequest &Request = CookedPackages[I];
-
-				if ( Request.GetPlatformNames().Contains( Platform ) )
-				{
-					// remove the platform
-					Request.RemovePlatform( Platform );
-
-					if ( Request.GetPlatformNames().Num() == 0 )
-					{
-						CookedPackages.RemoveAt(I);
-					}
-				}
-			}
-		}
-		void Empty()
-		{
-			FScopeLock S( &SyncObject );
-			CookedPackages.Empty();
-		}
-	};
-
-
 	struct FCachedPackageFilename
 	{
 	public:
@@ -375,109 +134,19 @@ private:
 		FString StandardFilename;
 		FName StandardFileFName;
 	};
-
-	/** Simple thread safe proxy for TSet<FName> */
-	template <typename T>
-	class FThreadSafeSet
-	{
-		TSet<T> InnerSet;
-		FCriticalSection SetCritical;
-	public:
-		void Add(T InValue)
-		{
-			FScopeLock SetLock(&SetCritical);
-			InnerSet.Add(InValue);
-		}
-		bool AddUnique(T InValue)
-		{
-			FScopeLock SetLock(&SetCritical);
-			if (!InnerSet.Contains(InValue))
-			{
-				InnerSet.Add(InValue);
-				return true;
-			}
-			return false;
-		}
-		bool Contains(T InValue)
-		{
-			FScopeLock SetLock(&SetCritical);
-			return InnerSet.Contains(InValue);
-		}
-		void Remove(T InValue)
-		{
-			FScopeLock SetLock(&SetCritical);
-			InnerSet.Remove(InValue);
-		}
-		void Empty()
-		{
-			FScopeLock SetLock(&SetCritical);
-			InnerSet.Empty();
-		}
-
-		void GetValues(TSet<T>& OutSet)
-		{
-			FScopeLock SetLock(&SetCritical);
-			OutSet.Append(InnerSet);
-		}
-	};
 private:
 	/** Current cook mode the cook on the fly server is running in */
-	ECookMode::Type CurrentCookMode;
+	ECookMode::Type CurrentCookMode = ECookMode::CookOnTheFly;
 	/** Directory to output to instead of the default should be empty in the case of DLC cooking */ 
 	FString OutputDirectoryOverride;
 
-	//////////////////////////////////////////////////////////////////////////
-	// Cook by the book options
-
-	struct FCookByTheBookOptions
-	{
-	public:
-		FCookByTheBookOptions() : 
-			bGenerateStreamingInstallManifests(false),
-			bGenerateDependenciesForMaps(false),
-			bRunning(false),
-			CookTime( 0.0 ),
-			CookStartTime( 0.0 ), 
-			bErrorOnEngineContentUse(false),
-			bDisableUnsolicitedPackages(false),
-			bFullLoadAndSave(false)
-		{ }
-
-		/** Should we generate streaming install manifests (only valid option in cook by the book) */
-		bool bGenerateStreamingInstallManifests;
-		/** Should we generate a seperate manifest for map dependencies */
-		bool bGenerateDependenciesForMaps;
-		/** Is cook by the book currently running */
-		bool bRunning;
-		/** Cancel has been queued will be processed next tick */
-		bool bCancel;
-		/** DlcName setup if we are cooking dlc will be used as the directory to save cooked files to */
-		FString DlcName;
-		/** Create a release from this manifest and store it in the releases directory for this cgame */
-		FString CreateReleaseVersion;
-		/** Dependency graph of maps as root objects. */
-		TMap<FName, TMap< FName, TSet <FName> > > MapDependencyGraphs; 
-		/** If a cook is cancelled next cook will need to resume cooking */ 
-		TArray<FFilePlatformRequest> PreviousCookRequests; 
-		/** If we are based on a release version of the game this is the set of packages which were cooked in that release. Map from platform name to list of uncooked package filenames */
-		TMap<FName,TArray<FName> > BasedOnReleaseCookedPackages;
-		/** Timing information about cook by the book */
-		double CookTime;
-		double CookStartTime;
-		/** error when detecting engine content being used in this cook */
-		bool bErrorOnEngineContentUse;
-		bool bDisableUnsolicitedPackages;
-		bool bFullLoadAndSave;
-		TArray<FName> TargetPlatformNames;
-		TArray<FName> StartupPackages;
-		/** Mapping from source packages to their localized variants (based on the culture list in FCookByTheBookStartupOptions) */
-		TMap<FName, TArray<FName>> SourceToLocalizedPackageVariants;
-	};
-	FCookByTheBookOptions* CookByTheBookOptions;
+	struct FCookByTheBookOptions;
 
+	FCookByTheBookOptions* CookByTheBookOptions = nullptr;
 
 	//////////////////////////////////////////////////////////////////////////
 	// Cook on the fly options
+
 	/** Cook on the fly server uses the NetworkFileServer */
 	TArray<class INetworkFileServer*> NetworkFileServers;
 	FOnFileModifiedDelegate FileModifiedDelegate;
@@ -499,11 +168,12 @@ private:
 	int32 MaxNumPackagesBeforePartialGC;
 	/** Max number of concurrent shader jobs reducing this too low will increase cook time */
 	int32 MaxConcurrentShaderJobs;
-	ECookInitializationFlags CookFlags;
+
+	ECookInitializationFlags CookFlags = ECookInitializationFlags::None;
 	TUniquePtr<class FSandboxPlatformFile> SandboxFile;
-	bool bIsInitializingSandbox; // stop recursion into callbacks when we are initializing sandbox
-	mutable bool bIgnoreMarkupPackageAlreadyLoaded; // avoid marking up packages as already loaded (want to put this around some functionality as we want to load packages fully some times)
-	bool bIsSavingPackage; // used to stop recursive mark package dirty functions
+	bool bIsInitializingSandbox = false; // stop recursion into callbacks when we are initializing sandbox
+	mutable bool bIgnoreMarkupPackageAlreadyLoaded = false; // avoid marking up packages as already loaded (want to put this around some functionality as we want to load packages fully some times)
+	bool bIsSavingPackage = false; // used to stop recursive mark package dirty functions
 
 
 	TMap<FName, int32> MaxAsyncCacheForType; // max number of objects of a specific type which are allowed to async cache at once
@@ -514,11 +184,14 @@ private:
 
 	//////////////////////////////////////////////////////////////////////////
 	// precaching system
-	// this system precaches materials and textures before we have considered the object as requiring save so as to utilize the system when it's idle
+	//
+	// this system precaches materials and textures before we have considered the object 
+	// as requiring save so as to utilize the system when it's idle
+	
 	TArray<FWeakObjectPtr> CachedMaterialsToCacheArray;
 	TArray<FWeakObjectPtr> CachedTexturesToCacheArray;
-	int32 LastUpdateTick;
-	int32 MaxPrecacheShaderJobs;
+	int32 LastUpdateTick = 0;
+	int32 MaxPrecacheShaderJobs = 0;
 	void TickPrecacheObjectsForPlatforms(const float TimeSlice, const TArray<const ITargetPlatform*>& TargetPlatform);
 
 	// presave system
@@ -557,12 +230,6 @@ private:
 
 	FReentryData& GetReentryData(const UPackage* Package) const;
 
-	FThreadSafeQueue<struct FRecompileRequest*> RecompileRequests;
-	FThreadSafeUnsolicitedPackagesList UnsolicitedCookedPackages;
-	FThreadSafeSet<FName> NeverCookPackageList;
-	FThreadSafeSet<FName> UncookedEditorOnlyPackages; // set of packages that have been rejected due to being referenced by editor-only properties
-
-
 	FString GetCachedPackageFilename( const FName& PackageName ) const;
 	FString GetCachedStandardPackageFilename( const FName& PackageName ) const;
 	FName GetCachedStandardPackageFileFName( const FName& PackageName ) const;
@@ -588,7 +255,7 @@ private:
 	mutable TMap<FName, FName> PackageFilenameToPackageFNameCache;
 
 	/** Cached copy of asset registry */
-	IAssetRegistry* AssetRegistry;
+	IAssetRegistry* AssetRegistry = nullptr;
 
 	/** Map of platform name to asset registry generators, which hold the state of asset registry data for a platform */
 	TMap<FName, FAssetRegistryGenerator*> RegistryGenerators;
@@ -606,7 +273,7 @@ private:
 	// tmap of the Config name, Section name, Key name, to the value
 	typedef TMap<FName, TMap<FName, TMap<FName, TArray<FString>>>> FIniSettingContainer;
 
-	mutable bool IniSettingRecurse;
+	mutable bool IniSettingRecurse = false;
 	mutable FIniSettingContainer AccessedIniStrings;
 	TArray<const FConfigFile*> OpenConfigFiles;
 	TArray<FString> ConfigSettingBlacklist;
@@ -703,20 +370,13 @@ public:
 		TArray<FString> CookCultures; 
 		TArray<FString> IniMapSections;
 		TArray<FString> CookPackages; // list of packages we should cook, used to specify specific packages to cook
-		ECookByTheBookOptions CookOptions;
+		ECookByTheBookOptions CookOptions = ECookByTheBookOptions::None;
 		FString DLCName;
 		FString CreateReleaseVersion;
 		FString BasedOnReleaseVersion;
-		bool bGenerateStreamingInstallManifests; 
-		bool bGenerateDependenciesForMaps; 
-		bool bErrorOnEngineContentUse; // this is a flag for dlc, will cause the cooker to error if the dlc references engine content
-		FCookByTheBookStartupOptions() :
-			CookOptions(ECookByTheBookOptions::None),
-			DLCName(FString()),
-			bGenerateStreamingInstallManifests(false),
-			bGenerateDependenciesForMaps(false),
-			bErrorOnEngineContentUse(false)
-		{ }
+		bool bGenerateStreamingInstallManifests = false;
+		bool bGenerateDependenciesForMaps = false;
+		bool bErrorOnEngineContentUse = false; // this is a flag for dlc, will cause the cooker to error if the dlc references engine content
 	};
 
 	/**
@@ -854,9 +514,6 @@ public:
 	/** Returns the configured number of packages to process before partial GC */
 	uint32 GetPackagesPerPartialGC() const;
 
-	/** Returns the target max concurrent shader jobs */
-	int32 GetMaxConcurrentShaderJobs() const;
-
 	/** Returns the configured amount of idle time before forcing a GC */
 	double GetIdleTimeToGC() const;
 
@@ -1198,16 +855,7 @@ private:
 	*/
 	FString GetSandboxDirectory( const FString& PlatformName ) const;
 
-	inline bool IsCookingDLC() const
-	{
-		// can only cook dlc in cook by the book
-		// we are cooking dlc when the dlc name is setup
-		if ( CookByTheBookOptions )
-		{
-			return  !CookByTheBookOptions->DlcName.IsEmpty();
-		}
-		return false;
-	}
+	bool IsCookingDLC() const;
 
 	/**
 	* GetBaseDirectoryForDLC
@@ -1218,14 +866,7 @@ private:
 
 	FString GetContentDirecctoryForDLC() const;
 
-	inline bool IsCreatingReleaseVersion()
-	{
-		if ( CookByTheBookOptions )
-		{
-			return !CookByTheBookOptions->CreateReleaseVersion.IsEmpty();
-		}
-		return false;
-	}
+	bool IsCreatingReleaseVersion();
 
 	/**
 	* Loads the cooked ini version settings maps into the Ini settings cache
@@ -1340,13 +981,3 @@ private:
 
 	FImpl*		Impl;
 };
-
-FORCEINLINE uint32 GetTypeHash(const UCookOnTheFlyServer::FFilePlatformRequest &Key)
-{
-	uint32 Hash = GetTypeHash( Key.Filename );
-	for ( const FName& PlatformName : Key.PlatformNames )
-	{
-		Hash += Hash << 2 ^ GetTypeHash( PlatformName );
-	}
-	return Hash;
-}
diff --git a/Engine/Source/Editor/UnrealEd/Private/CookOnTheFlyServer.cpp b/Engine/Source/Editor/UnrealEd/Private/CookOnTheFlyServer.cpp
--- a/Engine/Source/Editor/UnrealEd/Private/CookOnTheFlyServer.cpp
+++ b/Engine/Source/Editor/UnrealEd/Private/CookOnTheFlyServer.cpp
@@ -401,7 +401,7 @@ DECLARE_CYCLE_STAT(TEXT("Tick cooking"), STAT_TickCooker, STATGROUP_Cooking);
 struct FRecompileRequest
 {
 	struct FShaderRecompileData RecompileData;
-	volatile bool bComplete;
+	volatile bool bComplete = false;
 };
 
 
@@ -521,10 +521,122 @@ void LogCookerMessage( const FString& MessageText, EMessageSeverity::Type Severi
 	MessageLog.Notify(FText(), EMessageSeverity::Warning, false);
 }
 
+//////////////////////////////////////////////////////////////////////////
+
+/** cooked file requests which includes platform which file is requested for */
+struct FFilePlatformRequest
+{
+protected:
+	FName			Filename;
+	TArray<FName>	PlatformNames;
+
+public:
+	// yes we have some friends
+	friend uint32 GetTypeHash(const FFilePlatformRequest& Key);
+
+	FFilePlatformRequest() = default;
+
+	FFilePlatformRequest(const FName& InFileName, const FName& InPlatformName) : Filename(InFileName)
+	{
+		PlatformNames.Add(InPlatformName);
+	}
+
+	FFilePlatformRequest(const FName& InFilename, const TArray<FName>& InPlatformName) : Filename(InFilename)
+	{
+		PlatformNames = InPlatformName;
+	}
+
+	FFilePlatformRequest(const FName& InFilename, TArray<FName>&& InPlatformName) : Filename(InFilename), PlatformNames(MoveTemp(InPlatformName)) { }
+	FFilePlatformRequest(const FFilePlatformRequest& InFilePlatformRequest) : Filename(InFilePlatformRequest.Filename), PlatformNames(InFilePlatformRequest.PlatformNames) { }
+	FFilePlatformRequest(FFilePlatformRequest&& InFilePlatformRequest) : Filename(MoveTemp(InFilePlatformRequest.Filename)), PlatformNames(MoveTemp(InFilePlatformRequest.PlatformNames)) { }
+
+	void SetFilename(const FString& InFilename)
+	{
+		Filename = FName(*InFilename);
+	}
+
+	const FName& GetFilename() const
+	{
+		return Filename;
+	}
+
+	const TArray<FName>& GetPlatformNames() const
+	{
+		return PlatformNames;
+	}
+
+	void RemovePlatform(const FName& Platform)
+	{
+		PlatformNames.Remove(Platform);
+	}
+
+	void AddPlatform(const FName& Platform)
+	{
+		check(Platform != NAME_None);
+		PlatformNames.Add(Platform);
+	}
+
+	bool HasPlatform(const FName& Platform) const
+	{
+		return PlatformNames.Find(Platform) != INDEX_NONE;
+	}
+
+	bool IsValid()  const
+	{
+		return Filename != NAME_None;
+	}
+
+	void Clear()
+	{
+		Filename = TEXT("");
+		PlatformNames.Empty();
+	}
+
+	FFilePlatformRequest& operator=(FFilePlatformRequest&& InFileRequest)
+	{
+		Filename = MoveTemp(InFileRequest.Filename);
+		PlatformNames = MoveTemp(InFileRequest.PlatformNames);
+		return *this;
+	}
+
+	bool operator==(const FFilePlatformRequest& InFileRequest) const
+	{
+		if (InFileRequest.Filename == Filename)
+		{
+			if (InFileRequest.PlatformNames == PlatformNames)
+			{
+				return true;
+			}
+		}
+		return false;
+	}
+
+	FString ToString() const
+	{
+		FString Result = FString::Printf(TEXT("%s;"), *Filename.ToString());
+
+		for (const FName& Platform : PlatformNames)
+		{
+			Result += FString::Printf(TEXT("%s,"), *Platform.ToString());
+		}
+		return Result;
+	}
+};
+
+FORCEINLINE uint32 GetTypeHash(const FFilePlatformRequest& Key)
+{
+	uint32 Hash = GetTypeHash(Key.Filename);
+	for (const FName& PlatformName : Key.PlatformNames)
+	{
+		Hash += Hash << 2 ^ GetTypeHash(PlatformName);
+	}
+	return Hash;
+}
+
 struct FFilePlatformCookedPackage
 {
 public:
-	FFilePlatformCookedPackage(const UCookOnTheFlyServer::FFilePlatformRequest& InFilePlatformRequest, const TArray<bool>& bInSuccededSavePackage)
+	FFilePlatformCookedPackage(const FFilePlatformRequest& InFilePlatformRequest, const TArray<bool>& bInSuccededSavePackage)
 	: Filename(InFilePlatformRequest.GetFilename())
 	, PlatformNames(InFilePlatformRequest.GetPlatformNames())
 	, bSucceededSavePackage(bInSuccededSavePackage)
@@ -659,7 +771,7 @@ public:
 		}
 	}
 
-	bool Exists(const UCookOnTheFlyServer::FFilePlatformRequest& Request) const
+	bool Exists(const FFilePlatformRequest& Request) const
 	{
 		FScopeLock ScopeLock(&SynchronizationObject);
 
@@ -820,8 +932,6 @@ public:
 
 struct CookRequestQueue
 {
-	typedef UCookOnTheFlyServer::FFilePlatformRequest FFilePlatformRequest;
-
 	template <class PREDICATE_CLASS>
 	void Sort(const PREDICATE_CLASS& Predicate)
 	{
@@ -942,24 +1052,182 @@ private:
 	mutable FCriticalSection	SynchronizationObject;
 };
 
-struct UCookOnTheFlyServer::FImpl : public FUObjectArray::FUObjectCreateListener, public FUObjectArray::FUObjectDeleteListener
+struct FThreadSafeUnsolicitedPackagesList
 {
-	// This is the set of packages which have already had PostLoadFixup called 
-	TSet<UPackage*> PostLoadFixupPackages;
+	void AddCookedPackage(const FFilePlatformRequest& PlatformRequest)
+	{
+		FScopeLock S(&SyncObject);
+		CookedPackages.Add(PlatformRequest);
+	}
 
-	// This is a complete list of currently loaded UPackages
-	TArray<UPackage*> LoadedPackages;
+	void GetPackagesForPlatformAndRemove(const FName& Platform, TArray<FName> PackageNames)
+	{
+		FScopeLock _(&SyncObject);
 
-	// This list contains the UPackages loaded since last call to GetNewPackages
-	TArray<UPackage*> NewPackages;
+		for (int I = CookedPackages.Num() - 1; I >= 0; --I)
+		{
+			FFilePlatformRequest &Request = CookedPackages[I];
 
-	// Set of files which have been cooked 
-	// when needing to recook a file the entry will need to be removed from here
-	CookedPackageSet CookedPackages;
+			if (Request.GetPlatformNames().Contains(Platform))
+			{
+				// remove the platform
+				Request.RemovePlatform(Platform);
 
-	// List of requested files
-	CookRequestQueue CookRequests;
+				if (Request.GetPlatformNames().Num() == 0)
+				{
+					CookedPackages.RemoveAt(I);
+				}
+			}
+		}
+	}
+
+	void Empty()
+	{
+		FScopeLock _(&SyncObject);
+		CookedPackages.Empty();
+	}
 
+private:
+	FCriticalSection				SyncObject;
+	TArray<FFilePlatformRequest>	CookedPackages;
+};
+
+template<typename Type, typename SynchronizationObjectType, typename ScopeLockType>
+struct FUnsynchronizedQueue
+{
+private:
+	mutable SynchronizationObjectType	SynchronizationObject; // made this mutable so this class can have const functions and still be thread safe
+	TArray<Type>		Items;
+public:
+	void Enqueue(const Type& Item)
+	{
+		ScopeLockType ScopeLock(&SynchronizationObject);
+		Items.Add(Item);
+	}
+	void EnqueueUnique(const Type& Item)
+	{
+		ScopeLockType ScopeLock(&SynchronizationObject);
+		Items.AddUnique(Item);
+	}
+	bool Dequeue(Type* Result)
+	{
+		ScopeLockType ScopeLock(&SynchronizationObject);
+		if (Items.Num())
+		{
+			*Result = Items[0];
+			Items.RemoveAt(0);
+			return true;
+		}
+		return false;
+	}
+	void DequeueAll(TArray<Type>& Results)
+	{
+		ScopeLockType ScopeLock(&SynchronizationObject);
+		Results += Items;
+		Items.Empty();
+	}
+
+	bool HasItems() const
+	{
+		ScopeLockType ScopeLock(&SynchronizationObject);
+		return Items.Num() > 0;
+	}
+
+	void Remove(const Type& Item)
+	{
+		ScopeLockType ScopeLock(&SynchronizationObject);
+		Items.Remove(Item);
+	}
+
+	void CopyItems(TArray<Type> &InItems) const
+	{
+		ScopeLockType ScopeLock(&SynchronizationObject);
+		InItems = Items;
+	}
+
+	int Num() const
+	{
+		ScopeLockType ScopeLock(&SynchronizationObject);
+		return Items.Num();
+	}
+
+	void Empty()
+	{
+		ScopeLockType ScopeLock(&SynchronizationObject);
+		Items.Empty();
+	}
+};
+
+
+struct FDummyCriticalSection
+{
+public:
+	FORCEINLINE void Lock() { }
+	FORCEINLINE void Unlock() { }
+};
+
+struct FDummyScopeLock
+{
+public:
+	FDummyScopeLock(FDummyCriticalSection *) { }
+};
+
+template<typename Type>
+struct FThreadSafeQueue : public FUnsynchronizedQueue<Type, FCriticalSection, FScopeLock>
+{
+	/**
+	* Don't add any functions here, this is just a overqualified typedef
+	* Add functions / functionality to the FUnsynchronizedQueue
+	*/
+};
+
+/** Simple thread safe proxy for TSet<FName> */
+template <typename T>
+class FThreadSafeSet
+{
+	TSet<T> InnerSet;
+	FCriticalSection SetCritical;
+public:
+	void Add(T InValue)
+	{
+		FScopeLock SetLock(&SetCritical);
+		InnerSet.Add(InValue);
+	}
+	bool AddUnique(T InValue)
+	{
+		FScopeLock SetLock(&SetCritical);
+		if (!InnerSet.Contains(InValue))
+		{
+			InnerSet.Add(InValue);
+			return true;
+		}
+		return false;
+	}
+	bool Contains(T InValue)
+	{
+		FScopeLock SetLock(&SetCritical);
+		return InnerSet.Contains(InValue);
+	}
+	void Remove(T InValue)
+	{
+		FScopeLock SetLock(&SetCritical);
+		InnerSet.Remove(InValue);
+	}
+	void Empty()
+	{
+		FScopeLock SetLock(&SetCritical);
+		InnerSet.Empty();
+	}
+
+	void GetValues(TSet<T>& OutSet)
+	{
+		FScopeLock SetLock(&SetCritical);
+		OutSet.Append(InnerSet);
+	}
+};
+
+struct UCookOnTheFlyServer::FImpl : public FUObjectArray::FUObjectCreateListener, public FUObjectArray::FUObjectDeleteListener
+{
 	FImpl()
 	{
 		for (TObjectIterator<UPackage> It; It; ++It)
@@ -1014,6 +1282,28 @@ struct UCookOnTheFlyServer::FImpl : public FUObjectArray::FUObjectCreateListener
 			PostLoadFixupPackages.Remove(Package);
 		}
 	}
+
+	// This is the set of packages which have already had PostLoadFixup called 
+	TSet<UPackage*>			PostLoadFixupPackages;
+
+	// This is a complete list of currently loaded UPackages
+	TArray<UPackage*>		LoadedPackages;
+
+	// This list contains the UPackages loaded since last call to GetNewPackages
+	TArray<UPackage*>		NewPackages;
+
+	// Set of files which have been cooked 
+	// when needing to recook a file the entry will need to be removed from here
+	CookedPackageSet		CookedPackages;
+
+	// List of requested files
+	CookRequestQueue		CookRequests;
+
+	FThreadSafeUnsolicitedPackagesList			UnsolicitedCookedPackages;
+	FThreadSafeQueue<struct FRecompileRequest*> RecompileRequests;
+
+	FThreadSafeSet<FName>						NeverCookPackageList;
+	FThreadSafeSet<FName>						UncookedEditorOnlyPackages; // set of packages that have been rejected due to being referenced by editor-only properties
 };
 
 /* FIlename caching functions
@@ -1139,19 +1429,59 @@ void UCookOnTheFlyServer::ClearPackageFilenameCache() const
 	PackageFilenameToPackageFNameCache.Empty();
 }
 
+//////////////////////////////////////////////////////////////////////////
+// Cook by the book options
+
+struct UCookOnTheFlyServer::FCookByTheBookOptions
+{
+public:
+	/** Should we generate streaming install manifests (only valid option in cook by the book) */
+	bool							bGenerateStreamingInstallManifests = false;
+
+	/** Should we generate a seperate manifest for map dependencies */
+	bool							bGenerateDependenciesForMaps = false;
+
+	/** Is cook by the book currently running */
+	bool							bRunning = false;
+
+	/** Cancel has been queued will be processed next tick */
+	bool							bCancel = false;
+
+	/** DlcName setup if we are cooking dlc will be used as the directory to save cooked files to */
+	FString							DlcName;
+
+	/** Create a release from this manifest and store it in the releases directory for this cgame */
+	FString							CreateReleaseVersion;
+
+	/** Dependency graph of maps as root objects. */
+	TMap<FName,TMap<FName,TSet<FName>>> MapDependencyGraphs;
+
+	/** If a cook is cancelled next cook will need to resume cooking */
+	TArray<FFilePlatformRequest>		PreviousCookRequests;
+
+	/** If we are based on a release version of the game this is the set of packages which were cooked in that release. Map from platform name to list of uncooked package filenames */
+	TMap<FName, TArray<FName>>			BasedOnReleaseCookedPackages;
+
+	/** Timing information about cook by the book */
+	double							CookTime = 0.0;
+	double							CookStartTime = 0.0;
+
+	/** error when detecting engine content being used in this cook */
+	bool							bErrorOnEngineContentUse = false;
+	bool							bDisableUnsolicitedPackages = false;
+	bool							bFullLoadAndSave = false;
+	TArray<FName>					TargetPlatformNames;
+	TArray<FName>					StartupPackages;
+
+	/** Mapping from source packages to their localized variants (based on the culture list in FCookByTheBookStartupOptions) */
+	TMap<FName, TArray<FName>>		SourceToLocalizedPackageVariants;
+};
 
 /* UCookOnTheFlyServer functions
  *****************************************************************************/
 
 UCookOnTheFlyServer::UCookOnTheFlyServer(const FObjectInitializer& ObjectInitializer)
-	: Super(ObjectInitializer),
-	CurrentCookMode(ECookMode::CookOnTheFly),
-	CookByTheBookOptions(nullptr),
-	CookFlags(ECookInitializationFlags::None),
-	bIsInitializingSandbox(false),
-	bIgnoreMarkupPackageAlreadyLoaded(false),
-	bIsSavingPackage(false),
-	AssetRegistry(nullptr)
+	: Super(ObjectInitializer)
 {
 	Impl = new FImpl;
 }
@@ -1704,6 +2034,29 @@ bool UCookOnTheFlyServer::IsCookOnTheFlyMode() const
 	return CurrentCookMode == ECookMode::CookOnTheFly || CurrentCookMode == ECookMode::CookOnTheFlyFromTheEditor; 
 }
 
+bool UCookOnTheFlyServer::IsCreatingReleaseVersion()
+{
+	if (CookByTheBookOptions)
+	{
+		return !CookByTheBookOptions->CreateReleaseVersion.IsEmpty();
+	}
+
+	return false;
+}
+
+bool UCookOnTheFlyServer::IsCookingDLC() const
+{
+	// can only cook DLC in cook by the book
+	// we are cooking DLC when the DLC name is setup
+	
+	if (CookByTheBookOptions)
+	{
+		return !CookByTheBookOptions->DlcName.IsEmpty();
+	}
+
+	return false;
+}
+
 FString UCookOnTheFlyServer::GetBaseDirectoryForDLC() const
 {
 	TSharedPtr<IPlugin> Plugin = IPluginManager::Get().FindPlugin(CookByTheBookOptions->DlcName);
@@ -1897,7 +2250,7 @@ uint32 UCookOnTheFlyServer::TickCookOnTheSide( const float TimeSlice, uint32 &Co
 		}
 
 		check(IsInGameThread());
-		if (NeverCookPackageList.Contains(ToBuild.GetFilename()))
+		if (Impl->NeverCookPackageList.Contains(ToBuild.GetFilename()))
 		{
 #if DEBUG_COOKONTHEFLY
 			UE_LOG(LogCook, Display, TEXT("Package %s requested but is in the never cook package list, discarding request"), *ToBuild.GetFilename().ToString());
@@ -2356,7 +2709,7 @@ void UCookOnTheFlyServer::SaveCookedPackages(TArray<UPackage*>& PackagesToSave,
 		for (int32 I = 0; I < PackagesToSave.Num(); ++I)
 		{
 			UPackage* Package = PackagesToSave[I];
-			if (Package->IsLoadedByEditorPropertiesOnly() && UncookedEditorOnlyPackages.Contains(Package->GetFName()))
+			if (Package->IsLoadedByEditorPropertiesOnly() && Impl->UncookedEditorOnlyPackages.Contains(Package->GetFName()))
 			{
 				// We already attempted to cook this package and it's still not referenced by any non editor-only properties.
 				continue;
@@ -2364,10 +2717,10 @@ void UCookOnTheFlyServer::SaveCookedPackages(TArray<UPackage*>& PackagesToSave,
 
 			// This package is valid, so make sure it wasn't previously marked as being an uncooked editor only package or it would get removed from the
 			// asset registry at the end of the cook
-			UncookedEditorOnlyPackages.Remove(Package->GetFName());
+			Impl->UncookedEditorOnlyPackages.Remove(Package->GetFName());
 
 			const FName PackageFName = GetCachedStandardPackageFileFName(Package);
-			if (NeverCookPackageList.Contains(PackageFName))
+			if (Impl->NeverCookPackageList.Contains(PackageFName))
 			{
 				// refuse to save this package, it's clearly one of the undesirables
 				continue;
@@ -2659,7 +3012,7 @@ void UCookOnTheFlyServer::SaveCookedPackages(TArray<UPackage*>& PackagesToSave,
 						// this is an unsolicited package
 						if (FPaths::FileExists(FileRequest.GetFilename().ToString()) == true)
 						{
-							UnsolicitedCookedPackages.AddCookedPackage(FFilePlatformRequest(FileRequest.GetFilename(), FileRequest.GetPlatformNames()));
+							Impl->UnsolicitedCookedPackages.AddCookedPackage(FFilePlatformRequest(FileRequest.GetFilename(), FileRequest.GetPlatformNames()));
 
 #if DEBUG_COOKONTHEFLY
 							UE_LOG(LogCook, Display, TEXT("UnsolicitedCookedPackages: %s"), *FileRequest.GetFilename().ToString());
@@ -2669,7 +3022,7 @@ void UCookOnTheFlyServer::SaveCookedPackages(TArray<UPackage*>& PackagesToSave,
 				}
 				else
 				{
-					UncookedEditorOnlyPackages.AddUnique(Package->GetFName());
+					Impl->UncookedEditorOnlyPackages.AddUnique(Package->GetFName());
 				}
 			}
 			else
@@ -2716,7 +3069,7 @@ void UCookOnTheFlyServer::PostLoadPackageFixup(UPackage* Package)
 	if (World->GetStreamingLevels().Num())
 	{
 		TSet<FName> NeverCookPackageNames;
-		NeverCookPackageList.GetValues(NeverCookPackageNames);
+		Impl->NeverCookPackageList.GetValues(NeverCookPackageNames);
 
 		UE_LOG(LogCook, Display, TEXT("Loading secondary levels for package '%s'"), *World->GetName());
 
@@ -3084,11 +3437,6 @@ uint32 UCookOnTheFlyServer::GetPackagesPerPartialGC() const
 }
 
 
-int32 UCookOnTheFlyServer::GetMaxConcurrentShaderJobs() const
-{
-	return MaxConcurrentShaderJobs;
-}
-
 double UCookOnTheFlyServer::GetIdleTimeToGC() const
 {
 	return IdleTimeToGC;
@@ -3251,7 +3599,7 @@ void UCookOnTheFlyServer::TickRecompileShaderRequests()
 	// try to pull off a request
 	FRecompileRequest* Request = NULL;
 
-	RecompileRequests.Dequeue(&Request);
+	Impl->RecompileRequests.Dequeue(&Request);
 
 	// process it
 	if (Request)
@@ -3266,7 +3614,7 @@ void UCookOnTheFlyServer::TickRecompileShaderRequests()
 
 bool UCookOnTheFlyServer::HasRecompileShaderRequests() const 
 { 
-	return RecompileRequests.HasItems(); 
+	return Impl->RecompileRequests.HasItems();
 }
 
 void UCookOnTheFlyServer::SaveCookedPackage(UPackage* Package, uint32 SaveFlags, TArray<FSavePackageResultStruct>& SavePackageResults)
@@ -5312,7 +5660,7 @@ void UCookOnTheFlyServer::CollectFilesToCook(TArray<FName>& FilesInPath, const T
 
 				if (StandardPackageFilename != NAME_None)
 				{
-					NeverCookPackageList.Add(StandardPackageFilename);
+					Impl->NeverCookPackageList.Add(StandardPackageFilename);
 				}
 			}
 		}
@@ -6005,7 +6353,7 @@ void UCookOnTheFlyServer::CookByTheBookFinished()
 
 				// ignore packages that weren't cooked because they were only referenced by editor-only properties
 				TSet<FName> UncookedEditorOnlyPackageNames;
-				UncookedEditorOnlyPackages.GetValues(UncookedEditorOnlyPackageNames);
+				Impl->UncookedEditorOnlyPackages.GetValues(UncookedEditorOnlyPackageNames);
 				for (FName UncookedEditorOnlyPackage : UncookedEditorOnlyPackageNames)
 				{
 					IgnorePackageNames.Add(UncookedEditorOnlyPackage);
@@ -6176,9 +6524,9 @@ void UCookOnTheFlyServer::StopAndClearCookedData()
 		CookByTheBookOptions->PreviousCookRequests.Empty();
 	}
 
-	RecompileRequests.Empty();
+	Impl->RecompileRequests.Empty();
 	Impl->CookRequests.Empty();
-	UnsolicitedCookedPackages.Empty();
+	Impl->UnsolicitedCookedPackages.Empty();
 	Impl->CookedPackages.Empty(); // set of files which have been cooked when needing to recook a file the entry will need to be removed from here
 }
 
@@ -6188,19 +6536,20 @@ void UCookOnTheFlyServer::ClearAllCookedData()
 	// that means it's also conceivable that we will recook the same package which currently has an outstanding async write request
 	UPackage::WaitForAsyncFileWrites();
 
-	UnsolicitedCookedPackages.Empty();
+	Impl->UnsolicitedCookedPackages.Empty();
 	Impl->CookedPackages.Empty(); // set of files which have been cooked when needing to recook a file the entry will need to be removed from here
 }
 
-void UCookOnTheFlyServer::ClearPlatformCookedData( const FName& PlatformName )
+void UCookOnTheFlyServer::ClearPlatformCookedData(const FName& PlatformName)
 {
 	// if we are going to clear the cooked packages it is conceivable that we will recook the packages which we just cooked 
 	// that means it's also conceivable that we will recook the same package which currently has an outstanding async write request
 	UPackage::WaitForAsyncFileWrites();
 
-	Impl->CookedPackages.RemoveAllFilesForPlatform( PlatformName );
+	Impl->CookedPackages.RemoveAllFilesForPlatform(PlatformName);
+
 	TArray<FName> PackageNames;
-	UnsolicitedCookedPackages.GetPackagesForPlatformAndRemove(PlatformName, PackageNames);
+	Impl->UnsolicitedCookedPackages.GetPackagesForPlatformAndRemove(PlatformName, PackageNames);
 }
 
 void UCookOnTheFlyServer::ClearCachedCookedPlatformDataForPlatform( const FName& PlatformName )
@@ -6371,7 +6720,7 @@ void UCookOnTheFlyServer::StartCookByTheBook( const FCookByTheBookStartupOptions
 
 	const UProjectPackagingSettings* const PackagingSettings = GetDefault<UProjectPackagingSettings>();
 
-	NeverCookPackageList.Empty();
+	Impl->NeverCookPackageList.Empty();
 	{
 		const FString AbsoluteGameContentDir = FPaths::ConvertRelativePathToFull(FPaths::ProjectContentDir());
 
@@ -6399,7 +6748,7 @@ void UCookOnTheFlyServer::StartCookByTheBook( const FCookByTheBookStartupOptions
 			struct FNeverCookDirectoryWalker : public IPlatformFile::FDirectoryVisitor
 			{
 			private:
-				FThreadSafeSet<FName> &NeverCookPackageList;
+				FThreadSafeSet<FName>& NeverCookPackageList;
 			public:
 				FNeverCookDirectoryWalker(FThreadSafeSet<FName> &InNeverCookPackageList) : NeverCookPackageList(InNeverCookPackageList) { }
 
@@ -6417,7 +6766,7 @@ void UCookOnTheFlyServer::StartCookByTheBook( const FCookByTheBookStartupOptions
 					return true;
 				}
 
-			} NeverCookDirectoryWalker(NeverCookPackageList);
+			} NeverCookDirectoryWalker(Impl->NeverCookPackageList);
 
 			IPlatformFile& PlatformFile = FPlatformFileManager::Get().GetPlatformFile();
 
@@ -6743,7 +7092,7 @@ void UCookOnTheFlyServer::MaybeMarkPackageAsAlreadyLoaded(UPackage *Package)
 	}
 
 	check(IsInGameThread());
-	if (NeverCookPackageList.Contains(StandardName))
+	if (Impl->NeverCookPackageList.Contains(StandardName))
 	{
 		bShouldMarkAsAlreadyProcessed = true;
 		UE_LOG(LogCook, Verbose, TEXT("Marking %s as reloading for cooker because it was requested as never cook package."), *StandardName.ToString());
@@ -6779,7 +7128,7 @@ bool UCookOnTheFlyServer::HandleNetworkFileServerNewConnection(const FString& Ve
 void UCookOnTheFlyServer::GetCookOnTheFlyUnsolicitedFiles(const FName& PlatformName, TArray<FString>& UnsolicitedFiles, const FString& Filename)
 {
 	TArray<FName> UnsolicitedFilenames;
-	UnsolicitedCookedPackages.GetPackagesForPlatformAndRemove(PlatformName, UnsolicitedFilenames);
+	Impl->UnsolicitedCookedPackages.GetPackagesForPlatformAndRemove(PlatformName, UnsolicitedFilenames);
 
 	for (const FName& UnsolicitedFile : UnsolicitedFilenames)
 	{
@@ -6929,7 +7278,7 @@ void UCookOnTheFlyServer::HandleNetworkFileServerRecompileShaders(const FShaderR
 		Request->bComplete = false;
 
 		// push the request for the game thread to process
-		RecompileRequests.Enqueue(Request);
+		Impl->RecompileRequests.Enqueue(Request);
 
 		// wait for it to complete (the game thread will pull it out of the TArray, but I will delete it)
 		while (!Request->bComplete)
@@ -7019,7 +7368,7 @@ uint32 UCookOnTheFlyServer::FullLoadAndSave(uint32& CookedPackageCount)
 			check(ToBuild.IsValid());
 
 			const FName BuildFilenameFName = ToBuild.GetFilename();
-			if (!NeverCookPackageList.Contains(BuildFilenameFName))
+			if (!Impl->NeverCookPackageList.Contains(BuildFilenameFName))
 			{
 				const FString BuildFilename = BuildFilenameFName.ToString();
 				GIsCookerLoadingPackage = true;
@@ -7084,7 +7433,7 @@ uint32 UCookOnTheFlyServer::FullLoadAndSave(uint32& CookedPackageCount)
 
 				FName PackageName = Package->GetFName();
 				FName StandardPackageName = GetCachedStandardPackageFileFName(PackageName);
-				if (NeverCookPackageList.Contains(StandardPackageName))
+				if (Impl->NeverCookPackageList.Contains(StandardPackageName))
 				{
 					// refuse to save this package
 					continue;
@@ -7227,7 +7576,7 @@ uint32 UCookOnTheFlyServer::FullLoadAndSave(uint32& CookedPackageCount)
 			for (const FString& ToLoad : PackagesToLoad)
 			{
 				FName BuildFilenameFName = GetCachedStandardPackageFileFName(FName(*ToLoad));
-				if (!NeverCookPackageList.Contains(BuildFilenameFName))
+				if (!Impl->NeverCookPackageList.Contains(BuildFilenameFName))
 				{
 					LoadPackage(nullptr, *ToLoad, LOAD_None);
 					if (GShaderCompilingManager)
@@ -7420,7 +7769,7 @@ uint32 UCookOnTheFlyServer::FullLoadAndSave(uint32& CookedPackageCount)
 
 				if (SavePackageSuccessPerPlatform.Contains(false))
 				{
-					UncookedEditorOnlyPackages.Add(PackageName);
+					Impl->UncookedEditorOnlyPackages.Add(PackageName);
 				}
 
 				Package->SetPackageFlagsTo(OriginalPackageFlags);
