Extracted package name cache related code into separate header

Small part of refactoring UCookOnTheFlyServer monolith into more manageable components to aid optimization and maintenance

diff --git a/Engine/Source/Editor/UnrealEd/Classes/CookOnTheSide/CookOnTheFlyServer.h b/Engine/Source/Editor/UnrealEd/Classes/CookOnTheSide/CookOnTheFlyServer.h
--- a/Engine/Source/Editor/UnrealEd/Classes/CookOnTheSide/CookOnTheFlyServer.h
+++ b/Engine/Source/Editor/UnrealEd/Classes/CookOnTheSide/CookOnTheFlyServer.h
@@ -4,14 +4,11 @@
 
 #include "CoreMinimal.h"
 #include "Misc/EnumClassFlags.h"
-#include "Stats/Stats.h"
 #include "UObject/ObjectMacros.h"
 #include "UObject/UObjectGlobals.h"
 #include "UObject/Object.h"
 #include "UObject/WeakObjectPtr.h"
 #include "UObject/Package.h"
-#include "Misc/ScopeLock.h"
-#include "HAL/PlatformProcess.h"
 #include "TickableEditorObject.h"
 #include "IPlatformFileSandboxWrapper.h"
 #include "INetworkFileSystemModule.h"
@@ -24,6 +21,8 @@ struct FPropertyChangedEvent;
 class IPlugin;
 class IAssetRegistry;
 
+struct FPackageNameCache;
+
 enum class ECookInitializationFlags
 {
 	None =										0x00000000, // No flags
@@ -94,10 +93,6 @@ enum class ECookTickFlags : uint8
 };
 ENUM_CLASS_FLAGS(ECookTickFlags);
 
-// hudson is the name of my favorite dwagon
-
-DECLARE_STATS_GROUP(TEXT("Cooking"), STATGROUP_Cooking, STATCAT_Advanced);
-
 UCLASS()
 class UNREALED_API UCookOnTheFlyServer : public UObject, public FTickableEditorObject, public FExec
 {
@@ -105,35 +100,6 @@ class UNREALED_API UCookOnTheFlyServer : public UObject, public FTickableEditorO
 
 	UCookOnTheFlyServer(const FObjectInitializer& ObjectInitializer = FObjectInitializer::Get());
 
-private:
-	struct FCachedPackageFilename
-	{
-	public:
-		FCachedPackageFilename(FString &&InPackageFilename, FString &&InStandardFilename, FName InStandardFileFName ) :
-			PackageFilename( MoveTemp( InPackageFilename )),
-			StandardFilename(MoveTemp(InStandardFilename)),
-			StandardFileFName( InStandardFileFName )
-		{
-		}
-
-		FCachedPackageFilename( const FCachedPackageFilename &In )
-		{
-			PackageFilename = In.PackageFilename;
-			StandardFilename = In.StandardFilename;
-			StandardFileFName = In.StandardFileFName;
-		}
-
-		FCachedPackageFilename( FCachedPackageFilename &&In )
-		{
-			PackageFilename = MoveTemp(In.PackageFilename);
-			StandardFilename = MoveTemp(In.StandardFilename);
-			StandardFileFName = In.StandardFileFName;
-		}
-
-		FString PackageFilename; // this is also full path
-		FString StandardFilename;
-		FName StandardFileFName;
-	};
 private:
 	/** Current cook mode the cook on the fly server is running in */
 	ECookMode::Type CurrentCookMode = ECookMode::CookOnTheFly;
@@ -230,30 +196,12 @@ private:
 
 	FReentryData& GetReentryData(const UPackage* Package) const;
 
-	FString GetCachedPackageFilename( const FName& PackageName ) const;
-	FString GetCachedStandardPackageFilename( const FName& PackageName ) const;
-	FName GetCachedStandardPackageFileFName( const FName& PackageName ) const;
-	FString GetCachedPackageFilename( const UPackage* Package ) const;
-	FString GetCachedStandardPackageFilename( const UPackage* Package ) const;
-	FName GetCachedStandardPackageFileFName( const UPackage* Package ) const;
-	const FString& GetCachedSandboxFilename( const UPackage* Package, TUniquePtr<class FSandboxPlatformFile>& SandboxFile ) const;
-	const FName* GetCachedPackageFilenameToPackageFName(const FName& StandardPackageFilename) const;
-	const FCachedPackageFilename& Cache(const FName& PackageName) const;
-	void ClearPackageFilenameCache() const;
-	bool ClearPackageFilenameCacheForPackage( const UPackage* Package ) const;
-	bool ClearPackageFilenameCacheForPackage( const FName& PackageName ) const;
-
 	FString ConvertCookedPathToUncookedPath(const FString& CookedPackageName) const;
 
-	// get dependencies for this package 
+	/** Get dependencies for package */
 	const TArray<FName>& GetFullPackageDependencies(const FName& PackageName) const;
 	mutable TMap<FName, TArray<FName>> CachedFullPackageDependencies;
 
-	// declared mutable as it's used to cache package filename strings and I don't want to declare all functions using it as non const
-	// used by GetCached * Filename functions
-	mutable TMap<FName, FCachedPackageFilename> PackageFilenameCache; // filename cache (only process the string operations once)
-	mutable TMap<FName, FName> PackageFilenameToPackageFNameCache;
-
 	/** Cached copy of asset registry */
 	IAssetRegistry* AssetRegistry = nullptr;
 
@@ -385,18 +333,6 @@ public:
 	*/
 	void StartCookByTheBook( const FCookByTheBookStartupOptions& CookByTheBookStartupOptions );
 
-	/**
-	* ValidateCookByTheBookSettings
-	* look at the cookByTheBookOptions and ensure there isn't any conflicting settings
-	*/
-	void ValidateCookByTheBookSettings() const;
-
-	/**
-	* ValidateCookOnTheFlySettings
-	* look at the initialization flags and other cooker settings make sure the programmer that thought of checking them are ok
-	*/
-	void ValidateCookOnTheFlySettings() const;
-
 	/**
 	* Queue a cook by the book cancel (you might want to do this instead of calling cancel directly so that you don't have to be in the game thread when canceling
 	*/
@@ -978,6 +914,11 @@ private:
 	uint32		StatSavedPackageCount = 0;
 
 	struct FImpl;
+	FImpl*				Impl;
+
+	FPackageNameCache*	PackageNameCache;
 
-	FImpl*		Impl;
+	// temporary -- should eliminate the need for this. Only required right now because FullLoadAndSave 
+	// accesses maps directly
+	friend FPackageNameCache;
 };
diff --git a/Engine/Source/Editor/UnrealEd/Private/CookOnTheFlyServer.cpp b/Engine/Source/Editor/UnrealEd/Private/CookOnTheFlyServer.cpp
--- a/Engine/Source/Editor/UnrealEd/Private/CookOnTheFlyServer.cpp
+++ b/Engine/Source/Editor/UnrealEd/Private/CookOnTheFlyServer.cpp
@@ -5,19 +5,24 @@
 =============================================================================*/
 
 #include "CookOnTheSide/CookOnTheFlyServer.h"
+#include "Cooker/PackagenameCache.h"
+
 #include "HAL/PlatformFilemanager.h"
 #include "HAL/FileManager.h"
 #include "Misc/CommandLine.h"
 #include "Misc/FileHelper.h"
 #include "HAL/Runnable.h"
 #include "HAL/RunnableThread.h"
+#include "Stats/Stats.h"
 #include "Stats/StatsMisc.h"
 #include "Misc/ConfigCacheIni.h"
 #include "HAL/IConsoleManager.h"
+#include "HAL/PlatformProcess.h"
 #include "Misc/LocalTimestampDirectoryVisitor.h"
 #include "Serialization/CustomVersion.h"
 #include "Misc/App.h"
 #include "Misc/CoreDelegates.h"
+#include "Misc/ScopeLock.h"
 #include "Modules/ModuleManager.h"
 #include "Serialization/ArchiveUObject.h"
 #include "UObject/GarbageCollection.h"
@@ -120,6 +125,7 @@ static FAutoConsoleVariableRef CVarCookDisplayMode(
 	TEXT("  3: Both\n"),
 	ECVF_Default);
 
+DECLARE_STATS_GROUP(TEXT("Cooking"), STATGROUP_Cooking, STATCAT_Advanced);
 
 #if OUTPUT_TIMING
 #include <Containers/AllocatorFixedSizeFreeList.h>
@@ -623,13 +629,15 @@ public:
 	}
 };
 
-FORCEINLINE uint32 GetTypeHash(const FFilePlatformRequest& Key)
+uint32 GetTypeHash(const FFilePlatformRequest& Key)
 {
 	uint32 Hash = GetTypeHash(Key.Filename);
+
 	for (const FName& PlatformName : Key.PlatformNames)
 	{
 		Hash += Hash << 2 ^ GetTypeHash(PlatformName);
 	}
+
 	return Hash;
 }
 
@@ -1306,129 +1314,6 @@ struct UCookOnTheFlyServer::FImpl : public FUObjectArray::FUObjectCreateListener
 	FThreadSafeSet<FName>						UncookedEditorOnlyPackages; // set of packages that have been rejected due to being referenced by editor-only properties
 };
 
-/* FIlename caching functions
- *****************************************************************************/
-
-FString UCookOnTheFlyServer::GetCachedPackageFilename( const FName& PackageName ) const 
-{
-	return Cache( PackageName ).PackageFilename;
-}
-
-FString UCookOnTheFlyServer::GetCachedStandardPackageFilename( const FName& PackageName ) const 
-{
-	return Cache( PackageName ).StandardFilename;
-}
-
-FName UCookOnTheFlyServer::GetCachedStandardPackageFileFName( const FName& PackageName ) const 
-{
-	return Cache( PackageName ).StandardFileFName;
-}
-
-
-FString UCookOnTheFlyServer::GetCachedPackageFilename( const UPackage* Package ) const 
-{
-	// check( Package->GetName() == Package->GetFName().ToString() );
-	return Cache( Package->GetFName() ).PackageFilename;
-}
-
-FString UCookOnTheFlyServer::GetCachedStandardPackageFilename( const UPackage* Package ) const 
-{
-	// check( Package->GetName() == Package->GetFName().ToString() );
-	return Cache( Package->GetFName() ).StandardFilename;
-}
-
-FName UCookOnTheFlyServer::GetCachedStandardPackageFileFName( const UPackage* Package ) const 
-{
-	// check( Package->GetName() == Package->GetFName().ToString() );
-	return Cache( Package->GetFName() ).StandardFileFName;
-}
-
-
-bool UCookOnTheFlyServer::ClearPackageFilenameCacheForPackage( const FName& PackageName ) const
-{
-	check(IsInGameThread());
-	return PackageFilenameCache.Remove( PackageName ) >= 1;
-}
-
-bool UCookOnTheFlyServer::ClearPackageFilenameCacheForPackage( const UPackage* Package ) const
-{
-	check(IsInGameThread());
-	return PackageFilenameCache.Remove( Package->GetFName() ) >= 1;
-}
-
-const FString& UCookOnTheFlyServer::GetCachedSandboxFilename( const UPackage* Package, TUniquePtr<class FSandboxPlatformFile>& InSandboxFile ) const 
-{
-	check(IsInGameThread());
-
-	FName PackageFName = Package->GetFName();
-	static TMap<FName, FString> CachedSandboxFilenames;
-	FString* CachedSandboxFilename = CachedSandboxFilenames.Find(PackageFName);
-	if ( CachedSandboxFilename )
-		return *CachedSandboxFilename;
-
-	const FString& PackageFilename = GetCachedPackageFilename(Package);
-	FString SandboxFilename = ConvertToFullSandboxPath(*PackageFilename, true );
-
-	return CachedSandboxFilenames.Add( PackageFName, MoveTemp(SandboxFilename) );
-}
-
-const UCookOnTheFlyServer::FCachedPackageFilename& UCookOnTheFlyServer::Cache(const FName& PackageName) const 
-{
-	check( IsInGameThread() );
-
-	FCachedPackageFilename *Cached = PackageFilenameCache.Find( PackageName );
-	if ( Cached != NULL )
-	{
-		return *Cached;
-	}
-	// cache all the things, like it's your birthday!
-
-	FString Filename;
-	FString PackageFilename;
-	FString StandardFilename;
-	FName StandardFileFName = NAME_None;
-	if (FPackageName::DoesPackageExist(PackageName.ToString(), NULL, &Filename, false))
-	{
-		StandardFilename = PackageFilename = FPaths::ConvertRelativePathToFull(Filename);
-
-
-		FPaths::MakeStandardFilename(StandardFilename);
-		StandardFileFName = FName(*StandardFilename);
-	}
-	PackageFilenameToPackageFNameCache.Add(StandardFileFName, PackageName);
-	return PackageFilenameCache.Emplace( PackageName, FCachedPackageFilename(MoveTemp(PackageFilename),MoveTemp(StandardFilename), StandardFileFName) );
-}
-
-
-const FName* UCookOnTheFlyServer::GetCachedPackageFilenameToPackageFName(const FName& StandardPackageFilename) const
-{
-	check(IsInGameThread());
-	const FName* Result = PackageFilenameToPackageFNameCache.Find(StandardPackageFilename);
-	if ( Result )
-	{
-		return Result;
-	}
-
-	FName PackageName = StandardPackageFilename;
-	FString PotentialLongPackageName = StandardPackageFilename.ToString();
-	if (FPackageName::IsValidLongPackageName(PotentialLongPackageName) == false)
-	{
-		PotentialLongPackageName = FPackageName::FilenameToLongPackageName(PotentialLongPackageName);
-		PackageName = FName(*PotentialLongPackageName);
-	}
-
-	Cache(PackageName);
-
-	return PackageFilenameToPackageFNameCache.Find(StandardPackageFilename);
-}
-
-void UCookOnTheFlyServer::ClearPackageFilenameCache() const 
-{
-	check(IsInGameThread());
-	PackageFilenameCache.Empty();
-	PackageFilenameToPackageFNameCache.Empty();
-}
-
 //////////////////////////////////////////////////////////////////////////
 // Cook by the book options
 
@@ -1484,6 +1369,7 @@ UCookOnTheFlyServer::UCookOnTheFlyServer(const FObjectInitializer& ObjectInitial
 	: Super(ObjectInitializer)
 {
 	Impl = new FImpl;
+	PackageNameCache = new FPackageNameCache;
 }
 
 UCookOnTheFlyServer::~UCookOnTheFlyServer()
@@ -1495,6 +1381,9 @@ UCookOnTheFlyServer::~UCookOnTheFlyServer()
 	delete CookByTheBookOptions;
 	CookByTheBookOptions = nullptr;
 
+	delete PackageNameCache;
+	PackageNameCache = nullptr;
+
 	delete Impl;
 	Impl = nullptr;
 }
@@ -1611,7 +1500,6 @@ bool UCookOnTheFlyServer::StartNetworkFileServer( const bool BindAnyPort )
 #if PROFILE_NETWORK
 	NetworkRequestEvent = FPlatformProcess::GetSynchEventFromPool();
 #endif
-	ValidateCookOnTheFlySettings();
 
 	GenerateAssetRegistry();
 
@@ -1989,7 +1877,7 @@ bool UCookOnTheFlyServer::ContainsRedirector(const FName& PackageName, TMap<FNam
 					// we can;t call GetCachedStandardPackageFileFName with None
 					if (RedirectedPath != NAME_None)
 					{
-						FName StandardPackageName = GetCachedStandardPackageFileFName(FName(*FPackageName::ObjectPathToPackageName(RedirectedPathString)));
+						FName StandardPackageName = PackageNameCache->GetCachedStandardPackageFileFName(FName(*FPackageName::ObjectPathToPackageName(RedirectedPathString)));
 						if (StandardPackageName != NAME_None)
 						{
 							bDestinationValid = true;
@@ -2267,7 +2155,7 @@ uint32 UCookOnTheFlyServer::TickCookOnTheSide( const float TimeSlice, uint32 &Co
 			if ( Package )
 			{
 				FString Name = Package->GetPathName();
-				FString PackageFilename(GetCachedStandardPackageFilename(Package));
+				FString PackageFilename(PackageNameCache->GetCachedStandardPackageFilename(Package));
 				if (PackageFilename != BuildFilename)
 				{
 					// we have saved something which we didn't mean to load 
@@ -2719,7 +2607,7 @@ void UCookOnTheFlyServer::SaveCookedPackages(TArray<UPackage*>& PackagesToSave,
 			// asset registry at the end of the cook
 			Impl->UncookedEditorOnlyPackages.Remove(Package->GetFName());
 
-			const FName PackageFName = GetCachedStandardPackageFileFName(Package);
+			const FName PackageFName = PackageNameCache->GetCachedStandardPackageFileFName(Package);
 			if (Impl->NeverCookPackageList.Contains(PackageFName))
 			{
 				// refuse to save this package, it's clearly one of the undesirables
@@ -2848,7 +2736,7 @@ void UCookOnTheFlyServer::SaveCookedPackages(TArray<UPackage*>& PackagesToSave,
 				
 				for (int32 RemainingIndex = I; RemainingIndex < NumPackagesToRequeue; ++RemainingIndex)
 				{
-					FName StandardFilename = GetCachedStandardPackageFileFName(PackagesToSave[RemainingIndex]);
+					FName StandardFilename = PackageNameCache->GetCachedStandardPackageFileFName(PackagesToSave[RemainingIndex]);
 					Impl->CookRequests.EnqueueUnique(FFilePlatformRequest(StandardFilename, SaveTargetPlatformNames));
 				}
 				Result |= COSR_WaitingOnCache;
@@ -2987,7 +2875,7 @@ void UCookOnTheFlyServer::SaveCookedPackages(TArray<UPackage*>& PackagesToSave,
 			}
 
 				//@todo ResetLoaders outside of this (ie when Package is NULL) causes problems w/ default materials
-			FName StandardFilename = GetCachedStandardPackageFileFName(Package);
+			FName StandardFilename = PackageNameCache->GetCachedStandardPackageFileFName(Package);
 
 			// We always want to mark package as processed unless it wasn't saved because it was referenced by editor-only data
 			// in which case we may still need to save it later when new content loads it through non editor-only references
@@ -3087,7 +2975,7 @@ void UCookOnTheFlyServer::PostLoadPackageFixup(UPackage* Package)
 
 	for (const FString& PackageName : NewPackagesToCook)
 	{
-		FName StandardPackageFName = GetCachedStandardPackageFileFName(FName(*PackageName));
+		FName StandardPackageFName = PackageNameCache->GetCachedStandardPackageFileFName(FName(*PackageName));
 
 		if (StandardPackageFName != NAME_None)
 		{
@@ -3253,7 +3141,7 @@ void UCookOnTheFlyServer::GetUnsolicitedPackages(TArray<UPackage*>& PackagesToSa
 		{
 			ensure(Package != nullptr);
 
-			const FName StandardPackageFName = GetCachedStandardPackageFileFName(Package);
+			const FName StandardPackageFName = PackageNameCache->GetCachedStandardPackageFileFName(Package);
 			if (StandardPackageFName == NAME_None)
 				continue;	// if we have name none that means we are in core packages or something...
 
@@ -3359,11 +3247,12 @@ void UCookOnTheFlyServer::MarkPackageDirtyForCooker( UPackage *Package )
 		// force that package to be recooked
 		const FString Name = Package->GetPathName();
 
-		FName PackageFFileName = GetCachedStandardPackageFileFName(Package);
+		FName PackageFFileName = PackageNameCache->GetCachedStandardPackageFileFName(Package);
 
 		if ( PackageFFileName == NAME_None )
 		{
-			ClearPackageFilenameCacheForPackage( Package );
+			PackageNameCache->ClearPackageFilenameCacheForPackage(Package);
+
 			return;
 		}
 
@@ -3543,7 +3432,7 @@ void UCookOnTheFlyServer::MarkGCPackagesToKeepForCooker()
 	TMap<FName, int32> PackageDependenciesCount;
 	for (const FName& QueuedPackage : Impl->CookRequests.GetQueue())
 	{
-		const FName* PackageName = GetCachedPackageFilenameToPackageFName(QueuedPackage);
+		const FName* PackageName = PackageNameCache->GetCachedPackageFilenameToPackageFName(QueuedPackage);
 		if ( !PackageName )
 		{
 			PackageDependenciesCount.Add(QueuedPackage, 0);
@@ -3561,7 +3450,7 @@ void UCookOnTheFlyServer::MarkGCPackagesToKeepForCooker()
 		UPackage* Package = (UPackage*)(*It);
 		if ( KeepPackages.Contains(Package->GetFName()) )
 		{
-			LoadedPackages.Add( GetCachedStandardPackageFileFName(Package->GetFName()) );
+			LoadedPackages.Add(PackageNameCache->GetCachedStandardPackageFileFName(Package->GetFName()) );
 			const FReentryData& ReentryData = GetReentryData(Package);
 			Package->SetFlags(RF_KeepForCooker);
 			for (UObject* Obj : ReentryData.CachedObjectsInOuter)
@@ -3829,7 +3718,7 @@ void UCookOnTheFlyServer::SaveCookedPackage(UPackage* Package, uint32 SaveFlags,
 	bIsSavingPackage = true;
 
 	const FString PackagePathName = Package->GetPathName();
-	FString Filename(GetCachedPackageFilename(Package));
+	FString Filename(PackageNameCache->GetCachedPackageFilename(Package));
 
 	// Also request any localized variants of this package
 	if (IsCookByTheBookMode() && !CookByTheBookOptions->bDisableUnsolicitedPackages && !FPackageName::IsLocalizedPackage(PackagePathName))
@@ -3839,7 +3728,7 @@ void UCookOnTheFlyServer::SaveCookedPackage(UPackage* Package, uint32 SaveFlags,
 		{
 			for (const FName LocalizedPackageName : *LocalizedVariants)
 			{
-				const FName LocalizedPackageFile = GetCachedStandardPackageFileFName(LocalizedPackageName);
+				const FName LocalizedPackageFile = PackageNameCache->GetCachedStandardPackageFileFName(LocalizedPackageName);
 				RequestPackage(LocalizedPackageFile, false);
 			}
 		}
@@ -3864,7 +3753,7 @@ void UCookOnTheFlyServer::SaveCookedPackage(UPackage* Package, uint32 SaveFlags,
 		}
 
 		// Verify package actually exists
-		FName StandardPackageName = GetCachedStandardPackageFileFName(SoftObjectPackage);
+		FName StandardPackageName = PackageNameCache->GetCachedStandardPackageFileFName(SoftObjectPackage);
 
 		if (StandardPackageName != NAME_None && IsCookByTheBookMode() && !CookByTheBookOptions->bDisableUnsolicitedPackages)
 		{
@@ -4257,7 +4146,7 @@ bool UCookOnTheFlyServer::Exec(class UWorld* InWorld, const TCHAR* Cmd, FOutputD
 		StartupOptions.TargetPlatforms.Add(TargetPlatform);
 		for (const FName& StandardPackageName : PackageNames)
 		{
-			FName PackageFileFName = GetCachedStandardPackageFileFName(StandardPackageName);
+			FName PackageFileFName = PackageNameCache->GetCachedStandardPackageFileFName(StandardPackageName);
 			StartupOptions.CookMaps.Add(StandardPackageName.ToString());
 		}
 		StartupOptions.CookOptions = ECookByTheBookOptions::NoAlwaysCookMaps | ECookByTheBookOptions::NoDefaultMaps | ECookByTheBookOptions::NoGameAlwaysCookPackages | ECookByTheBookOptions::NoInputPackages | ECookByTheBookOptions::NoSlatePackages | ECookByTheBookOptions::DisableUnsolicitedPackages | ECookByTheBookOptions::ForceDisableSaveGlobalShaders;
@@ -5199,7 +5088,7 @@ void UCookOnTheFlyServer::PopulateCookedPackagesFromDisk(const TArray<ITargetPla
 		{
 			const FName CookedFile = CookedPaths.Value;
 			const FName UncookedFilename = CookedPaths.Key;
-			const FName* FoundPackageName = GetCachedPackageFilenameToPackageFName(UncookedFilename);
+			const FName* FoundPackageName = PackageNameCache->GetCachedPackageFilenameToPackageFName(UncookedFilename);
 			bool bShouldKeep = true;
 			const FName SourcePackageName = FoundPackageName ? *FoundPackageName : NAME_None;
 			if ( !FoundPackageName )
@@ -5279,7 +5168,7 @@ void UCookOnTheFlyServer::PopulateCookedPackagesFromDisk(const TArray<ITargetPla
 		// Register identical uncooked packages from previous run
 		for (FName UncookedPackage : IdenticalUncookedPackages)
 		{
-			const FName UncookedFilename = GetCachedStandardPackageFileFName(UncookedPackage);
+			const FName UncookedFilename = PackageNameCache->GetCachedStandardPackageFileFName(UncookedPackage);
 
 			TArray<FName> PlatformNames;
 			PlatformNames.Add(PlatformFName);
@@ -5656,7 +5545,7 @@ void UCookOnTheFlyServer::CollectFilesToCook(TArray<FName>& FilesInPath, const T
 
 			for (FName NeverCookPackage : PackagesToNeverCook)
 			{
-				const FName StandardPackageFilename = GetCachedStandardPackageFileFName(NeverCookPackage);
+				const FName StandardPackageFilename = PackageNameCache->GetCachedStandardPackageFileFName(NeverCookPackage);
 
 				if (StandardPackageFilename != NAME_None)
 				{
@@ -6340,14 +6229,14 @@ void UCookOnTheFlyServer::CookByTheBookFinished()
 				TSet<FName> CookedPackageNames;
 				for (FName PackageFilename : CookedPackagesFilenames)
 				{
-					const FName *FoundLongPackageFName = GetCachedPackageFilenameToPackageFName(PackageFilename);
+					const FName *FoundLongPackageFName = PackageNameCache->GetCachedPackageFilenameToPackageFName(PackageFilename);
 					CookedPackageNames.Add(*FoundLongPackageFName);
 				}
 
 				TSet<FName> IgnorePackageNames;
 				for (FName PackageFilename : IgnorePackageFilenames)
 				{
-					const FName *FoundLongPackageFName = GetCachedPackageFilenameToPackageFName(PackageFilename);
+					const FName *FoundLongPackageFName = PackageNameCache->GetCachedPackageFilenameToPackageFName(PackageFilename);
 					IgnorePackageNames.Add(*FoundLongPackageFName);
 				}
 
@@ -6637,18 +6526,10 @@ void UCookOnTheFlyServer::InitializeTargetPlatforms()
 void UCookOnTheFlyServer::TermSandbox()
 {
 	ClearAllCookedData();
-	ClearPackageFilenameCache();
+	PackageNameCache->ClearPackageFilenameCache();
 	SandboxFile = nullptr;
 }
 
-void UCookOnTheFlyServer::ValidateCookOnTheFlySettings() const
-{
-}
-
-void UCookOnTheFlyServer::ValidateCookByTheBookSettings() const
-{
-}
-
 void UCookOnTheFlyServer::StartCookByTheBook( const FCookByTheBookStartupOptions& CookByTheBookStartupOptions )
 {
 	SCOPE_TIMER(StartCookByTheBookTime);
@@ -6784,8 +6665,6 @@ void UCookOnTheFlyServer::StartCookByTheBook( const FCookByTheBookStartupOptions
 	}
 	const TArray<FName>& TargetPlatformNames = CookByTheBookOptions->TargetPlatformNames;
 
-	ValidateCookByTheBookSettings();
-
 	if ( CookByTheBookOptions->DlcName != DLCName )
 	{
 		// we are going to change the state of dlc we need to clean out our package filename cache (the generated filename cache is dependent on this key)
@@ -6961,7 +6840,7 @@ void UCookOnTheFlyServer::StartCookByTheBook( const FCookByTheBookStartupOptions
 		if ( FileFName == NAME_None )
 			continue;
 
-		const FName PackageFileFName = GetCachedStandardPackageFileFName(FileFName);
+		const FName PackageFileFName = PackageNameCache->GetCachedStandardPackageFileFName(FileFName);
 		
 		if (PackageFileFName != NAME_None)
 		{
@@ -7053,7 +6932,7 @@ void UCookOnTheFlyServer::MaybeMarkPackageAsAlreadyLoaded(UPackage *Package)
 		return;
 	}
 
-	FName StandardName = GetCachedStandardPackageFileFName(Package);
+	FName StandardName = PackageNameCache->GetCachedStandardPackageFileFName(Package);
 
 	// UE_LOG(LogCook, Display, TEXT("Loading package %s"), *StandardName.ToString());
 
@@ -7322,7 +7201,7 @@ bool UCookOnTheFlyServer::GetAllPackageFilenamesFromAssetRegistry( const FString
 		for (const TPair<FName, const FAssetData*>& RegistryData : RegistryDataMap)
 		{
 			const FAssetData* NewAssetData = RegistryData.Value;
-			FName CachedPackageFileFName = GetCachedStandardPackageFileFName(NewAssetData->ObjectPath);
+			FName CachedPackageFileFName = PackageNameCache->GetCachedStandardPackageFileFName(NewAssetData->ObjectPath);
 			if (CachedPackageFileFName != NAME_None)
 			{
 				OutPackageFilenames.Add(CachedPackageFileFName);
@@ -7432,7 +7311,7 @@ uint32 UCookOnTheFlyServer::FullLoadAndSave(uint32& CookedPackageCount)
 				}
 
 				FName PackageName = Package->GetFName();
-				FName StandardPackageName = GetCachedStandardPackageFileFName(PackageName);
+				FName StandardPackageName = PackageNameCache->GetCachedStandardPackageFileFName(PackageName);
 				if (Impl->NeverCookPackageList.Contains(StandardPackageName))
 				{
 					// refuse to save this package
@@ -7575,7 +7454,7 @@ uint32 UCookOnTheFlyServer::FullLoadAndSave(uint32& CookedPackageCount)
 			GIsCookerLoadingPackage = true;
 			for (const FString& ToLoad : PackagesToLoad)
 			{
-				FName BuildFilenameFName = GetCachedStandardPackageFileFName(FName(*ToLoad));
+				FName BuildFilenameFName = PackageNameCache->GetCachedStandardPackageFileFName(FName(*ToLoad));
 				if (!Impl->NeverCookPackageList.Contains(BuildFilenameFName))
 				{
 					LoadPackage(nullptr, *ToLoad, LOAD_None);
@@ -7683,7 +7562,7 @@ uint32 UCookOnTheFlyServer::FullLoadAndSave(uint32& CookedPackageCount)
 			}
 
 			FName PackageName = Package->GetFName();
-			FCachedPackageFilename* CachedPackageFilename = PackageFilenameCache.Find(PackageName);
+			FCachedPackageFilename* CachedPackageFilename = PackageNameCache->PackageFilenameCache.Find(PackageName);
 			check(CachedPackageFilename);
 	
 			if (CachedPackageFilename->PackageFilename.Len())
diff --git a/Engine/Source/Editor/UnrealEd/Private/Cooker/PackageNameCache.h b/Engine/Source/Editor/UnrealEd/Private/Cooker/PackageNameCache.h
new file mode 100644
--- /dev/null
+++ b/Engine/Source/Editor/UnrealEd/Private/Cooker/PackageNameCache.h
@@ -0,0 +1,145 @@
+// Copyright 1998-2018 Epic Games, Inc. All Rights Reserved.
+
+#pragma once
+
+#include "CoreMinimal.h"
+#include "Misc/PackageName.h"
+
+struct FCachedPackageFilename
+{
+	FCachedPackageFilename(FString&& InPackageFilename, FString&& InStandardFilename, FName InStandardFileFName)
+	:	PackageFilename(MoveTemp(InPackageFilename))
+	,	StandardFilename(MoveTemp(InStandardFilename))
+	,	StandardFileFName(InStandardFileFName)
+	{
+	}
+
+	FCachedPackageFilename(const FCachedPackageFilename& In) = default;
+
+	FCachedPackageFilename(FCachedPackageFilename&& In)
+	:	PackageFilename(MoveTemp(In.PackageFilename))
+	,	StandardFilename(MoveTemp(In.StandardFilename))
+	,	StandardFileFName(In.StandardFileFName)
+	{
+	}
+
+	FString		PackageFilename; // this is also full path
+	FString		StandardFilename;
+	FName		StandardFileFName;
+};
+
+struct FPackageNameCache
+{
+	// Package name cache
+
+	FString			GetCachedPackageFilename(const UPackage* Package) const;
+
+	FString			GetCachedStandardPackageFilename(const UPackage* Package) const;
+
+	FName			GetCachedStandardPackageFileFName(const FName& PackageName) const;
+	FName			GetCachedStandardPackageFileFName(const UPackage* Package) const;
+
+	const FName*	GetCachedPackageFilenameToPackageFName(const FName& StandardPackageFilename) const;
+
+	void			ClearPackageFilenameCache() const;
+	bool			ClearPackageFilenameCacheForPackage(const UPackage* Package) const;
+
+private:
+	const FCachedPackageFilename& Cache(const FName& PackageName) const;
+
+	mutable TMap<FName, FCachedPackageFilename> PackageFilenameCache; // filename cache (only process the string operations once)
+	mutable TMap<FName, FName>					PackageFilenameToPackageFNameCache;
+
+	// temporary -- should eliminate the need for this
+	friend uint32 UCookOnTheFlyServer::FullLoadAndSave(uint32& CookedPackageCount);
+};
+
+FName FPackageNameCache::GetCachedStandardPackageFileFName(const FName& PackageName) const
+{
+	return Cache(PackageName).StandardFileFName;
+}
+
+FString FPackageNameCache::GetCachedPackageFilename(const UPackage* Package) const
+{
+	// check( Package->GetName() == Package->GetFName().ToString() );
+	return Cache(Package->GetFName()).PackageFilename;
+}
+
+FString FPackageNameCache::GetCachedStandardPackageFilename(const UPackage* Package) const
+{
+	// check( Package->GetName() == Package->GetFName().ToString() );
+	return Cache(Package->GetFName()).StandardFilename;
+}
+
+FName FPackageNameCache::GetCachedStandardPackageFileFName(const UPackage* Package) const
+{
+	// check( Package->GetName() == Package->GetFName().ToString() );
+	return Cache(Package->GetFName()).StandardFileFName;
+}
+
+bool FPackageNameCache::ClearPackageFilenameCacheForPackage(const UPackage* Package) const
+{
+	check(IsInGameThread());
+
+	return PackageFilenameCache.Remove(Package->GetFName()) >= 1;
+}
+
+const FCachedPackageFilename& FPackageNameCache::Cache(const FName& PackageName) const
+{
+	check(IsInGameThread());
+
+	FCachedPackageFilename *Cached = PackageFilenameCache.Find(PackageName);
+
+	if (Cached != NULL)
+	{
+		return *Cached;
+	}
+
+	// cache all the things, like it's your birthday!
+
+	FString Filename;
+	FString PackageFilename;
+	FString StandardFilename;
+	FName StandardFileFName = NAME_None;
+
+	if (FPackageName::DoesPackageExist(PackageName.ToString(), NULL, &Filename, false))
+	{
+		StandardFilename = PackageFilename = FPaths::ConvertRelativePathToFull(Filename);
+
+		FPaths::MakeStandardFilename(StandardFilename);
+		StandardFileFName = FName(*StandardFilename);
+	}
+
+	PackageFilenameToPackageFNameCache.Add(StandardFileFName, PackageName);
+
+	return PackageFilenameCache.Emplace(PackageName, FCachedPackageFilename(MoveTemp(PackageFilename), MoveTemp(StandardFilename), StandardFileFName));
+}
+
+const FName* FPackageNameCache::GetCachedPackageFilenameToPackageFName(const FName& StandardPackageFilename) const
+{
+	check(IsInGameThread());
+	const FName* Result = PackageFilenameToPackageFNameCache.Find(StandardPackageFilename);
+	if (Result)
+	{
+		return Result;
+	}
+
+	FName PackageName = StandardPackageFilename;
+	FString PotentialLongPackageName = StandardPackageFilename.ToString();
+	if (FPackageName::IsValidLongPackageName(PotentialLongPackageName) == false)
+	{
+		PotentialLongPackageName = FPackageName::FilenameToLongPackageName(PotentialLongPackageName);
+		PackageName = FName(*PotentialLongPackageName);
+	}
+
+	Cache(PackageName);
+
+	return PackageFilenameToPackageFNameCache.Find(StandardPackageFilename);
+}
+
+void FPackageNameCache::ClearPackageFilenameCache() const
+{
+	check(IsInGameThread());
+	PackageFilenameCache.Empty();
+	PackageFilenameToPackageFNameCache.Empty();
+}
