Fixed include file casing for "Cooker/PackageNameCache.h"

diff --git a/Engine/Source/Editor/UnrealEd/Private/CookOnTheFlyServer.cpp b/Engine/Source/Editor/UnrealEd/Private/CookOnTheFlyServer.cpp
--- a/Engine/Source/Editor/UnrealEd/Private/CookOnTheFlyServer.cpp
+++ b/Engine/Source/Editor/UnrealEd/Private/CookOnTheFlyServer.cpp
@@ -5,7 +5,7 @@
 =============================================================================*/
 
 #include "CookOnTheSide/CookOnTheFlyServer.h"
-#include "Cooker/PackagenameCache.h"
+#include "Cooker/PackageNameCache.h"
 
 #include "HAL/PlatformFilemanager.h"
 #include "HAL/FileManager.h"
