The FPackageTracker can get informed about objects of a blueprint-derived type where the type object has already been garbage collected. This makes it unsafe to use the UClass pointer. Instead we just check the UClass pointer to see if it matches the UPackage's UClass object. This does make deriving off UPackage unsupported however.

diff --git a/Engine/Source/Editor/UnrealEd/Private/CookOnTheFlyServer.cpp b/Engine/Source/Editor/UnrealEd/Private/CookOnTheFlyServer.cpp
--- a/Engine/Source/Editor/UnrealEd/Private/CookOnTheFlyServer.cpp
+++ b/Engine/Source/Editor/UnrealEd/Private/CookOnTheFlyServer.cpp
@@ -1216,7 +1216,7 @@ struct FPackageTracker : public FUObjectArray::FUObjectCreateListener, public FU
 
 	virtual void NotifyUObjectCreated(const class UObjectBase *Object, int32 Index) override
 	{
-		if (Object->GetClass()->IsChildOf(UPackage::StaticClass()))
+		if (Object->GetClass() == UPackage::StaticClass())
 		{
 			auto Package = const_cast<UPackage*>(static_cast<const UPackage*>(Object));
 
@@ -1232,7 +1232,7 @@ struct FPackageTracker : public FUObjectArray::FUObjectCreateListener, public FU
 
 	virtual void NotifyUObjectDeleted(const class UObjectBase *Object, int32 Index) override
 	{
-		if (Object->GetClass()->IsChildOf(UPackage::StaticClass()))
+		if (Object->GetClass() == UPackage::StaticClass())
 		{
 			auto Package = const_cast<UPackage*>(static_cast<const UPackage*>(Object));
 
