Fixing static analysis warnings

diff --git a/Engine/Source/Editor/UnrealEd/Private/Commandlets/ContentCommandlets.cpp b/Engine/Source/Editor/UnrealEd/Private/Commandlets/ContentCommandlets.cpp
--- a/Engine/Source/Editor/UnrealEd/Private/Commandlets/ContentCommandlets.cpp
+++ b/Engine/Source/Editor/UnrealEd/Private/Commandlets/ContentCommandlets.cpp
@@ -575,7 +575,7 @@ void UResavePackagesCommandlet::LoadAndSaveOnePackage(const FString& Filename)
 		TRefCountPtr<FUObjectSerializeContext> LoadContext(new FUObjectSerializeContext());
 		BeginLoad(LoadContext);
 		FLinkerLoad* Linker = GetPackageLinker(NULL, *Filename, LOAD_NoVerify, nullptr, nullptr, nullptr, LoadContext);
-		EndLoad(Linker->GetSerializeContext());
+		EndLoad(Linker ? Linker->GetSerializeContext() : LoadContext);
 	
 		// Bail early if we don't have a valid linker (package was out of date, etc)
 		if( !Linker )
