Static analysis fixes

diff --git a/Engine/Source/Editor/UnrealEd/Private/Commandlets/ContentCommandlets.cpp b/Engine/Source/Editor/UnrealEd/Private/Commandlets/ContentCommandlets.cpp
--- a/Engine/Source/Editor/UnrealEd/Private/Commandlets/ContentCommandlets.cpp
+++ b/Engine/Source/Editor/UnrealEd/Private/Commandlets/ContentCommandlets.cpp
@@ -575,7 +575,7 @@ void UResavePackagesCommandlet::LoadAndSaveOnePackage(const FString& Filename)
 		TRefCountPtr<FUObjectSerializeContext> LoadContext(new FUObjectSerializeContext());
 		BeginLoad(LoadContext);
 		FLinkerLoad* Linker = GetPackageLinker(NULL, *Filename, LOAD_NoVerify, nullptr, nullptr, nullptr, LoadContext);
-		EndLoad(Linker ? Linker->GetSerializeContext() : LoadContext);
+		EndLoad(Linker ? Linker->GetSerializeContext() : LoadContext.GetReference());
 	
 		// Bail early if we don't have a valid linker (package was out of date, etc)
 		if( !Linker )
@@ -2176,7 +2176,7 @@ int32 UWrangleContentCommandlet::Main( const FString& Params )
 					TRefCountPtr<FUObjectSerializeContext> LoadContext(new FUObjectSerializeContext());
 					BeginLoad(LoadContext);
 					FLinkerLoad* Linker = GetPackageLinker(nullptr, *PackageFilename, LOAD_Quiet | LOAD_NoWarn | LOAD_NoVerify, nullptr, nullptr, nullptr, LoadContext);
-					EndLoad(Linker->GetSerializeContext());
+					EndLoad(Linker ? Linker->GetSerializeContext() : LoadContext.GetReference());
 
 					// look for special package types
 					bool bIsMap = Linker->ContainsMap();
@@ -2331,7 +2331,7 @@ int32 UWrangleContentCommandlet::Main( const FString& Params )
 			TRefCountPtr<FUObjectSerializeContext> LoadContext(new FUObjectSerializeContext());
 			BeginLoad(LoadContext);
 			FLinkerLoad* Linker = GetPackageLinker(nullptr, *PackageFilename, LOAD_Quiet | LOAD_NoWarn | LOAD_NoVerify, nullptr, nullptr, nullptr, LoadContext);
-			EndLoad(Linker->GetSerializeContext());
+			EndLoad(Linker ? Linker->GetSerializeContext() : LoadContext.GetReference());
 
 			// go through the exports in the package, looking for public objects
 			for (int32 ExportIndex = 0; ExportIndex < Linker->ExportMap.Num(); ExportIndex++)
