Fixing CIS - missing files from my previous checkin

diff --git a/Engine/Source/Editor/UnrealEd/Private/Commandlets/ContentCommandlets.cpp b/Engine/Source/Editor/UnrealEd/Private/Commandlets/ContentCommandlets.cpp
--- a/Engine/Source/Editor/UnrealEd/Private/Commandlets/ContentCommandlets.cpp
+++ b/Engine/Source/Editor/UnrealEd/Private/Commandlets/ContentCommandlets.cpp
@@ -572,10 +572,7 @@ void UResavePackagesCommandlet::LoadAndSaveOnePackage(const FString& Filename)
 		// Get the package linker.
 		VerboseMessage(TEXT("Pre GetPackageLinker"));
 
-		TRefCountPtr<FUObjectSerializeContext> LoadContext(new FUObjectSerializeContext());
-		BeginLoad(LoadContext);
-		FLinkerLoad* Linker = GetPackageLinker(NULL, *Filename, LOAD_NoVerify, nullptr, nullptr, nullptr, LoadContext);
-		EndLoad(Linker ? Linker->GetSerializeContext() : LoadContext.GetReference());
+		FLinkerLoad* Linker = LoadPackageLinker(NULL, *Filename, LOAD_NoVerify, nullptr, nullptr, nullptr);
 	
 		// Bail early if we don't have a valid linker (package was out of date, etc)
 		if( !Linker )
@@ -2173,10 +2170,7 @@ int32 UWrangleContentCommandlet::Main( const FString& Params )
 					// load the package fully
 					UPackage* Package = LoadPackage(nullptr, *PackageFilename, LOAD_None);
 
-					TRefCountPtr<FUObjectSerializeContext> LoadContext(new FUObjectSerializeContext());
-					BeginLoad(LoadContext);
-					FLinkerLoad* Linker = GetPackageLinker(nullptr, *PackageFilename, LOAD_Quiet | LOAD_NoWarn | LOAD_NoVerify, nullptr, nullptr, nullptr, LoadContext);
-					EndLoad(Linker ? Linker->GetSerializeContext() : LoadContext.GetReference());
+					FLinkerLoad* Linker = LoadPackageLinker(nullptr, *PackageFilename, LOAD_Quiet | LOAD_NoWarn | LOAD_NoVerify, nullptr, nullptr, nullptr);
 
 					// look for special package types
 					bool bIsMap = Linker->ContainsMap();
@@ -2328,10 +2322,7 @@ int32 UWrangleContentCommandlet::Main( const FString& Params )
 				}
 			}
 
-			TRefCountPtr<FUObjectSerializeContext> LoadContext(new FUObjectSerializeContext());
-			BeginLoad(LoadContext);
-			FLinkerLoad* Linker = GetPackageLinker(nullptr, *PackageFilename, LOAD_Quiet | LOAD_NoWarn | LOAD_NoVerify, nullptr, nullptr, nullptr, LoadContext);
-			EndLoad(Linker ? Linker->GetSerializeContext() : LoadContext.GetReference());
+			FLinkerLoad* Linker = LoadPackageLinker(nullptr, *PackageFilename, LOAD_Quiet | LOAD_NoWarn | LOAD_NoVerify, nullptr, nullptr, nullptr);
 
 			// go through the exports in the package, looking for public objects
 			for (int32 ExportIndex = 0; ExportIndex < Linker->ExportMap.Num(); ExportIndex++)
diff --git a/Engine/Source/Editor/UnrealEd/Private/Commandlets/PackageUtilities.cpp b/Engine/Source/Editor/UnrealEd/Private/Commandlets/PackageUtilities.cpp
--- a/Engine/Source/Editor/UnrealEd/Private/Commandlets/PackageUtilities.cpp
+++ b/Engine/Source/Editor/UnrealEd/Private/Commandlets/PackageUtilities.cpp
@@ -781,10 +781,7 @@ int32 ULoadPackageCommandlet::Main( const FString& Params )
 
 		if (bCheckForLegacyPackages)
 		{
-			TRefCountPtr<FUObjectSerializeContext> LoadContext(new FUObjectSerializeContext());
-			BeginLoad(LoadContext);
-			auto Linker = GetPackageLinker(nullptr, *Filename, LOAD_NoVerify, nullptr, nullptr, nullptr, LoadContext);
-			EndLoad(Linker->GetSerializeContext());
+			FLinkerLoad* Linker = LoadPackageLinker(nullptr, *Filename, LOAD_NoVerify, nullptr, nullptr, nullptr);
 			MinVersion = FMath::Min<int32>(MinVersion, Linker->Summary.GetFileVersionUE4());
 		}
 		else
