Deprecation warning fix for placement new on TIndirectArray.

diff --git a/Engine/Source/Editor/UnrealEd/Private/Factories/SkeletalMeshImport.cpp b/Engine/Source/Editor/UnrealEd/Private/Factories/SkeletalMeshImport.cpp
--- a/Engine/Source/Editor/UnrealEd/Private/Factories/SkeletalMeshImport.cpp
+++ b/Engine/Source/Editor/UnrealEd/Private/Factories/SkeletalMeshImport.cpp
@@ -858,7 +858,8 @@ void RestoreExistingSkelMeshData(ExistingSkelMeshData* MeshData, USkeletalMesh*
 					}
 					FSkeletalMeshLODModel& LODModel = MeshData->ExistingLODModels[Index];
 					FSkeletalMeshLODInfo& LODInfo = MeshData->ExistingLODInfo[Index];
-					FSkeletalMeshLODModel* NewLODModel = new(SkeletalMeshImportedModel->LODModels) FSkeletalMeshLODModel(LODModel);
+					FSkeletalMeshLODModel* NewLODModel = new FSkeletalMeshLODModel(LODModel);
+					SkeletalMeshImportedModel->LODModels.Add(NewLODModel);
 					// add LOD info back
 					SkeletalMesh->AddLODInfo(LODInfo);
 
