Fixes for various deprecated features.

diff --git a/Engine/Plugins/Developer/Concert/ConcertMain/Source/Concert/Private/ConcertMessageData.cpp b/Engine/Plugins/Developer/Concert/ConcertMain/Source/Concert/Private/ConcertMessageData.cpp
--- a/Engine/Plugins/Developer/Concert/ConcertMain/Source/Concert/Private/ConcertMessageData.cpp
+++ b/Engine/Plugins/Developer/Concert/ConcertMain/Source/Concert/Private/ConcertMessageData.cpp
@@ -123,9 +123,9 @@ namespace PayloadDetail
 				if (UncompressedData.Num() > 0)
 				{
 					// Compress the result to send on the wire
-					int32 CompressedSize = FCompression::CompressMemoryBound(COMPRESS_ZLIB, UncompressedData.Num());
+					int32 CompressedSize = FCompression::CompressMemoryBound(NAME_Zlib, UncompressedData.Num());
 					OutCompressedData.SetNumUninitialized(CompressedSize);
-					if (FCompression::CompressMemory(COMPRESS_ZLIB, OutCompressedData.GetData(), CompressedSize, UncompressedData.GetData(), UncompressedData.Num()))
+					if (FCompression::CompressMemory(NAME_Zlib, OutCompressedData.GetData(), CompressedSize, UncompressedData.GetData(), UncompressedData.Num()))
 					{
 						OutUncompressedDataSizeBytes = UncompressedData.Num();
 						OutCompressedData.SetNum(CompressedSize, false);
@@ -162,7 +162,7 @@ namespace PayloadDetail
 				// Uncompress the data
 				TArray<uint8> UncompressedData;
 				UncompressedData.SetNumUninitialized(InUncompressedDataSizeBytes);
-				if (FCompression::UncompressMemory(COMPRESS_ZLIB, UncompressedData.GetData(), UncompressedData.Num(), InCompressedData.GetData(), InCompressedData.Num()))
+				if (FCompression::UncompressMemory(NAME_Zlib, UncompressedData.GetData(), UncompressedData.Num(), InCompressedData.GetData(), InCompressedData.Num()))
 				{
 					// Deserialize the uncompressed data
 					{
diff --git a/Engine/Source/Editor/UnrealEd/Private/Factories/SkeletalMeshImport.cpp b/Engine/Source/Editor/UnrealEd/Private/Factories/SkeletalMeshImport.cpp
--- a/Engine/Source/Editor/UnrealEd/Private/Factories/SkeletalMeshImport.cpp
+++ b/Engine/Source/Editor/UnrealEd/Private/Factories/SkeletalMeshImport.cpp
@@ -561,7 +561,8 @@ void TryRegenerateLODs(ExistingSkelMeshData* MeshData, USkeletalMesh* SkeletalMe
 				// reset material maps, it won't work anyway. 
 				ExistLODInfo.LODMaterialMap.Empty();
 
-				FSkeletalMeshLODModel* NewLODModel = new(SkeletalMeshImportedModel->LODModels) FSkeletalMeshLODModel(ExistLODModel);
+				FSkeletalMeshLODModel* NewLODModel = new FSkeletalMeshLODModel(ExistLODModel);
+				SkeletalMeshImportedModel->LODModels.Add(NewLODModel);
 				// add LOD info back
 				SkeletalMesh->AddLODInfo(ExistLODInfo);
 				check(LODIndex < SkeletalMesh->GetLODInfoArray().Num());
@@ -872,7 +873,8 @@ void RestoreExistingSkelMeshData(ExistingSkelMeshData* MeshData, USkeletalMesh*
 					FSkeletalMeshLODModel* NewLODModel = nullptr;
 					if (RealIndex >= SkeletalMesh->GetLODNum())
 					{
-						NewLODModel = new(SkeletalMeshImportedModel->LODModels) FSkeletalMeshLODModel(LODModel);
+						NewLODModel = new FSkeletalMeshLODModel(LODModel);
+						SkeletalMeshImportedModel->LODModels.Add(NewLODModel);
 						// add LOD info back
 						SkeletalMesh->AddLODInfo(LODInfo);
 					}
@@ -1003,7 +1005,8 @@ void RestoreExistingSkelMeshData(ExistingSkelMeshData* MeshData, USkeletalMesh*
 							if (bRegenLODs)
 							{
 								//We need to add LODInfo
-								FSkeletalMeshLODModel* NewLODModel = new(SkeletalMeshImportedModel->LODModels) FSkeletalMeshLODModel(LODModel);
+								FSkeletalMeshLODModel* NewLODModel = new FSkeletalMeshLODModel(LODModel);
+								SkeletalMeshImportedModel->LODModels.Add(NewLODModel);
 								SkeletalMesh->AddLODInfo(LODInfo);
 								RestoreReductionSourceData(i, SkeletalMesh->GetLODNum() - 1);
 							}
diff --git a/Engine/Source/Runtime/Core/Private/ProfilingDebugging/CsvProfiler.cpp b/Engine/Source/Runtime/Core/Private/ProfilingDebugging/CsvProfiler.cpp
--- a/Engine/Source/Runtime/Core/Private/ProfilingDebugging/CsvProfiler.cpp
+++ b/Engine/Source/Runtime/Core/Private/ProfilingDebugging/CsvProfiler.cpp
@@ -25,6 +25,7 @@
 #include "Misc/EngineVersion.h"
 #include "Stats/Stats.h"
 #include "HAL/LowLevelMemTracker.h"
+#include "Misc/Compression.h"
 
 #if CSV_PROFILER
 
@@ -1074,9 +1075,10 @@ private:
 						// Compress the data in Buffer into the GZipBuffer array
 						CompressedSize = GZipBuffer.Num();
 						if (FCompression::CompressMemory(
-							(ECompressionFlags)(ECompressionFlags::COMPRESS_GZIP | ECompressionFlags::COMPRESS_BiasSpeed),
+							NAME_Zlib,
 							GZipBuffer.GetData(), CompressedSize,
-							Buffer.GetData(), BytesInBuffer))
+							Buffer.GetData(), BytesInBuffer,
+							ECompressionFlags::COMPRESS_BiasSpeed))
 						{
 							break;
 						}
