Weighted Tangent Import Fix.

Fix a slight error with calculation of our tangent weights not getting scaled by the tangent correctly so they could be a little big.

diff --git a/Engine/Source/Editor/UnrealEd/Private/Fbx/FbxCurveDataImport.cpp b/Engine/Source/Editor/UnrealEd/Private/Fbx/FbxCurveDataImport.cpp
--- a/Engine/Source/Editor/UnrealEd/Private/Fbx/FbxCurveDataImport.cpp
+++ b/Engine/Source/Editor/UnrealEd/Private/Fbx/FbxCurveDataImport.cpp
@@ -332,14 +332,16 @@ namespace UnFbx {
 				//so we just need to scale it by the normalized weight value.
 				if (!FMath::IsNearlyZero(ArriveTangentWeight))
 				{
-					const float Y = ArriveTangent * ArriveTangentWeight;
-					ArriveTangentWeight = FMath::Sqrt(Y*Y + ArriveTimeDiff * ArriveTimeDiff);
+					const float X = ArriveTangentWeight * ArriveTimeDiff;
+					const float Y = ArriveTangent * ArriveTangentWeight;  //timediff already there
+					ArriveTangentWeight = FMath::Sqrt(Y*Y + X * X);
 				}
 				NewKey.ArriveTangentWeight = ArriveTangentWeight;
 				if (!FMath::IsNearlyZero(LeaveTangentWeight))
 				{
-					const float Y = LeaveTangent * LeaveTangentWeight;
-					LeaveTangentWeight = FMath::Sqrt(Y*Y + LeaveTimeDiff * LeaveTimeDiff);
+					const float X = LeaveTangentWeight * LeaveTimeDiff;
+					const float Y = LeaveTangent * LeaveTangentWeight; //timediff already there
+					LeaveTangentWeight = FMath::Sqrt(Y*Y + X * X);
 				}
 				NewKey.LeaveTangentWeight = LeaveTangentWeight;
 			}
