Fix for asset registry generator duplicating manifest data when SaveManifests is called twice. This happens when -ExtraFlavorChunkSize is specified on the command line.
Added some more logging and debugging information to stagetime-only encrypted chunk manifest generation

diff --git a/Engine/Source/Editor/UnrealEd/Private/Commandlets/AssetRegistryGenerator.cpp b/Engine/Source/Editor/UnrealEd/Private/Commandlets/AssetRegistryGenerator.cpp
--- a/Engine/Source/Editor/UnrealEd/Private/Commandlets/AssetRegistryGenerator.cpp
+++ b/Engine/Source/Editor/UnrealEd/Private/Commandlets/AssetRegistryGenerator.cpp
@@ -219,6 +219,7 @@ bool FAssetRegistryGenerator::GenerateStreamingInstallManifest(int64 InExtraFlav
 	// Add manifests for any staging-time only groups
 	if (bUseAssetManager && !TargetPlatform->HasSecurePackageFormat())
 	{
+		UE_LOG(LogAssetRegistryGenerator, Log, TEXT("Updating stage-only encryption manifests for platform %s"), *TargetPlatform->IniPlatformName());
 		FContentEncryptionConfig ContentEncryptionConfig;
 		UAssetManager::Get().GetContentEncryptionConfig(ContentEncryptionConfig);
 		const FContentEncryptionConfig::TGroupMap& EncryptedNonUFSFileGroups = ContentEncryptionConfig.GetPackageGroupMap();
@@ -227,19 +228,24 @@ bool FAssetRegistryGenerator::GenerateStreamingInstallManifest(int64 InExtraFlav
 		{
 			if (Element.Value.bStageTimeOnly)
 			{
+				UE_LOG(LogAssetRegistryGenerator, Log, TEXT("Adding stage-time only manifest for group '%s'"), *Element.Key.ToString());
+
 				FName GroupName = Element.Key;
 				const TSet<FName>& PackageNames = Element.Value.PackageNames;
 				
 				FChunkPackageSet* NewManifest = new FChunkPackageSet();
-				int32 ChunkID = UAssetManager::Get().GetContentEncryptionGroupChunkID(GroupName);
-				
 				for (FName PackageName : PackageNames)
 				{
 					NewManifest->Add(PackageName, FPackageName::LongPackageNameToFilename(PackageName.ToString()));
 				}
 
-				check(ChunkID < FinalChunkManifests.Num());
-				check(FinalChunkManifests[ChunkID] == nullptr);
+				int32 ChunkID = UAssetManager::Get().GetContentEncryptionGroupChunkID(GroupName);
+				if (ChunkID >= FinalChunkManifests.Num())
+				{
+					FinalChunkManifests.AddZeroed(ChunkID - FinalChunkManifests.Num() + 1);
+				}
+				checkf(ChunkID < FinalChunkManifests.Num(), TEXT("Chunk %i out of range. %i manifests available"), ChunkID, FinalChunkManifests.Num() - 1);
+				checkf(FinalChunkManifests[ChunkID] == nullptr, TEXT("Manifest already exists for chunk %i"), ChunkID);
 				FinalChunkManifests[ChunkID] = NewManifest;
 			}
 		}
@@ -1397,6 +1403,13 @@ void FAssetRegistryGenerator::FixupPackageDependenciesForChunks(FSandboxPlatform
 	UE_LOG(LogAssetRegistryGenerator, Log, TEXT("Starting FixupPackageDependenciesForChunks..."));
 	SCOPE_LOG_TIME_IN_SECONDS(TEXT("... FixupPackageDependenciesForChunks complete."), nullptr);
 
+	// Clear any existing manifests from the final array
+	for (FChunkPackageSet* Manifest : FinalChunkManifests)
+	{
+		delete Manifest;
+	}
+	FinalChunkManifests.Empty();
+
 	for (int32 ChunkID = 0, MaxChunk = ChunkManifests.Num(); ChunkID < MaxChunk; ++ChunkID)
 	{
 		FinalChunkManifests.Add(nullptr);
