Made the UMaterialShaderQualitySettings for the editor's startup feature level available earlier

Previously, the quality settings did not become available until certain plugins were loaded, so the editor's initial material previews could be incorrect when a preview platform from the previous editor session was needed at editor startup.

diff --git a/Engine/Source/Editor/UnrealEd/Private/EditorEngine.cpp b/Engine/Source/Editor/UnrealEd/Private/EditorEngine.cpp
--- a/Engine/Source/Editor/UnrealEd/Private/EditorEngine.cpp
+++ b/Engine/Source/Editor/UnrealEd/Private/EditorEngine.cpp
@@ -7400,9 +7400,15 @@ void UEditorEngine::LoadEditorFeatureLevel()
 	EShaderPlatform ShaderPlatform = ShaderFormatToLegacyShaderPlatform(Settings->PreviewShaderPlatformName);
 	if (ShaderPlatform != SP_NumPlatforms)
 	{
-		const FName MaterialQualityPlatform = Settings->bIsMaterialQualityOverridePlatform ? Settings->PreviewShaderPlatformName : NAME_None;
+		FName MaterialQualityPlatform = NAME_None;
 		const ERHIFeatureLevel::Type FeatureLevel = GetMaxSupportedFeatureLevel(ShaderPlatform);
 
+		if (Settings->bIsMaterialQualityOverridePlatform)
+		{
+			MaterialQualityPlatform = Settings->PreviewShaderPlatformName;
+			UMaterialShaderQualitySettings::Get()->GetShaderPlatformQualitySettings(Settings->PreviewShaderPlatformName);
+		}
+
 		SetPreviewPlatform(MaterialQualityPlatform, FeatureLevel, false);
 	}
 }
