Ensure the package localization ID is correctly migrated when renaming an asset

diff --git a/Engine/Source/Editor/UnrealEd/Private/ObjectTools.cpp b/Engine/Source/Editor/UnrealEd/Private/ObjectTools.cpp
--- a/Engine/Source/Editor/UnrealEd/Private/ObjectTools.cpp
+++ b/Engine/Source/Editor/UnrealEd/Private/ObjectTools.cpp
@@ -59,6 +59,7 @@
 #include "Dialogs/Dialogs.h"
 #include "UnrealEdGlobals.h"
 #include "PackageTools.h"
+#include "Internationalization/TextPackageNamespaceUtil.h"
 #include "Framework/Application/SlateApplication.h"
 
 #include "BusyCursor.h"
@@ -3111,6 +3112,10 @@ namespace ObjectTools
 
 			if (OldPackage && OldPackage->MetaData)
 			{
+				// Migrate the localization ID to the new package
+				TextNamespaceUtil::ForcePackageNamespace(NewPackage, TextNamespaceUtil::GetPackageNamespace(OldPackage));
+				TextNamespaceUtil::ClearPackageNamespace(OldPackage);
+
 				// Remove any metadata from old package pointing to moved objects
 				OldPackage->MetaData->RemoveMetaDataOutsidePackage();
 			}
