Destroy PIE world net driver later in the shutdown process so the net mode will be correct for actors during EndPlay

diff --git a/Engine/Source/Editor/UnrealEd/Private/EditorEngine.cpp b/Engine/Source/Editor/UnrealEd/Private/EditorEngine.cpp
--- a/Engine/Source/Editor/UnrealEd/Private/EditorEngine.cpp
+++ b/Engine/Source/Editor/UnrealEd/Private/EditorEngine.cpp
@@ -782,7 +782,7 @@ void UEditorEngine::Init(IEngineLoop* InEngineLoop)
 	{
 		if (GetPIEWorldContext() != nullptr && GetPIEWorldContext()->World() != nullptr)
 		{
-			GEngine->ShutdownWorldNetDriver(GetPIEWorldContext()->World());
+			GetPIEWorldContext()->World()->DestroyDemoNetDriver();
 		}
 	});
 
diff --git a/Engine/Source/Editor/UnrealEd/Private/PlayLevel.cpp b/Engine/Source/Editor/UnrealEd/Private/PlayLevel.cpp
--- a/Engine/Source/Editor/UnrealEd/Private/PlayLevel.cpp
+++ b/Engine/Source/Editor/UnrealEd/Private/PlayLevel.cpp
@@ -303,6 +303,8 @@ void UEditorEngine::EndPlayMap()
 			if (ThisContext.World())
 			{
 				TeardownPlaySession(ThisContext);
+
+				ShutdownWorldNetDriver(ThisContext.World());
 			}
 
 			// Cleanup online subsystems instantiated during PIE
