#DEV-VR - Fix for CIS shadowed member variable error

diff --git a/Engine/Source/Editor/UnrealEd/Private/CookOnTheFlyServer.cpp b/Engine/Source/Editor/UnrealEd/Private/CookOnTheFlyServer.cpp
--- a/Engine/Source/Editor/UnrealEd/Private/CookOnTheFlyServer.cpp
+++ b/Engine/Source/Editor/UnrealEd/Private/CookOnTheFlyServer.cpp
@@ -6192,12 +6192,12 @@ void UCookOnTheFlyServer::InitializeSandbox()
 
 void UCookOnTheFlyServer::InitializeTargetPlatforms()
 {
-	const TArray<ITargetPlatform*>& CookingTargetPlatforms = GetCookingTargetPlatforms();
+	const TArray<ITargetPlatform*>& TargetPlatforms = GetCookingTargetPlatforms();
 
 	//allow each platform to update its internals before cooking
-	for (int32 TargetPlatformIndex = 0; TargetPlatformIndex < CookingTargetPlatforms.Num(); ++TargetPlatformIndex)
+	for (int32 TargetPlatformIndex = 0; TargetPlatformIndex < TargetPlatforms.Num(); ++TargetPlatformIndex)
 	{
-		ITargetPlatform* TargetPlatform = CookingTargetPlatforms[TargetPlatformIndex];
+		ITargetPlatform* TargetPlatform = TargetPlatforms[TargetPlatformIndex];
 		if (TargetPlatform)
 		{
 			TargetPlatform->RefreshSettings();
