[NO NON-PII COMMIT MESSAGE]

diff --git a/Engine/Source/Editor/ContentBrowser/Private/SFilterList.cpp b/Engine/Source/Editor/ContentBrowser/Private/SFilterList.cpp
--- a/Engine/Source/Editor/ContentBrowser/Private/SFilterList.cpp
+++ b/Engine/Source/Editor/ContentBrowser/Private/SFilterList.cpp
@@ -1117,7 +1117,7 @@ TSharedRef<SWidget> SFilterList::MakeAddFilterMenu(EAssetTypeCategories::Type Me
 		}
 	}
 
-	FMenuBuilder MenuBuilder(/*bInShouldCloseWindowAfterMenuSelection=*/true, NULL);
+	FMenuBuilder MenuBuilder(/*bInShouldCloseWindowAfterMenuSelection=*/true, nullptr, nullptr, /*bCloseSelfOnly=*/true);
 
 	MenuBuilder.BeginSection( "ContentBrowserResetFilters" );
 	{
diff --git a/Engine/Source/Editor/UnrealEd/Private/Toolkits/SGlobalOpenAssetDialog.cpp b/Engine/Source/Editor/UnrealEd/Private/Toolkits/SGlobalOpenAssetDialog.cpp
--- a/Engine/Source/Editor/UnrealEd/Private/Toolkits/SGlobalOpenAssetDialog.cpp
+++ b/Engine/Source/Editor/UnrealEd/Private/Toolkits/SGlobalOpenAssetDialog.cpp
@@ -28,6 +28,7 @@ void SGlobalOpenAssetDialog::Construct(const FArguments& InArgs, FVector2D InSiz
 	AssetPickerConfig.bShowBottomToolbar = true;
 	AssetPickerConfig.bAutohideSearchBar = false;
 	AssetPickerConfig.bCanShowClasses = false;
+	AssetPickerConfig.bAddFilterUI = true;
 	AssetPickerConfig.SaveSettingsName = TEXT("GlobalAssetPicker");
 
 	ChildSlot
