Fbx importer: Hide the properties that are not supported when reimporting.

diff --git a/Engine/Source/Editor/DetailCustomizations/Private/FbxImportUIDetails.cpp b/Engine/Source/Editor/DetailCustomizations/Private/FbxImportUIDetails.cpp
--- a/Engine/Source/Editor/DetailCustomizations/Private/FbxImportUIDetails.cpp
+++ b/Engine/Source/Editor/DetailCustomizations/Private/FbxImportUIDetails.cpp
@@ -342,6 +342,10 @@ void FFbxImportUIDetails::CustomizeDetails( IDetailLayoutBuilder& DetailBuilder
 		{
 			DetailBuilder.HideProperty(Handle);
 		}
+		else if (ImportUI->bIsReimport && Handle->GetBoolMetaData(TEXT("ReimportRestrict")))
+		{
+			DetailBuilder.HideProperty(Handle);
+		}
 		else
 		{
 			SetupRefreshForHandle(Handle);
@@ -363,6 +367,10 @@ void FFbxImportUIDetails::CustomizeDetails( IDetailLayoutBuilder& DetailBuilder
 		{
 			bSkip = true;
 		}
+
+		//Skip the variable that is ReimportRestrick when we are in reimport mode
+		bSkip |= ImportUI->bIsReimport && Handle->GetBoolMetaData(TEXT("ReimportRestrict"));
+		
 		if(!bSkip && IsImportTypeMetaDataValid(ImportType, ImportTypeMetaData))
 		{
 			// Decide on category
diff --git a/Engine/Source/Editor/UnrealEd/Classes/Factories/FbxStaticMeshImportData.h b/Engine/Source/Editor/UnrealEd/Classes/Factories/FbxStaticMeshImportData.h
--- a/Engine/Source/Editor/UnrealEd/Classes/Factories/FbxStaticMeshImportData.h
+++ b/Engine/Source/Editor/UnrealEd/Classes/Factories/FbxStaticMeshImportData.h
@@ -31,17 +31,17 @@ class UFbxStaticMeshImportData : public UFbxMeshImportData
 	FColor VertexOverrideColor;
 
 	/** Disabling this option will keep degenerate triangles found.  In general you should leave this option on. */
-	UPROPERTY(EditAnywhere, BlueprintReadWrite, AdvancedDisplay, Category = Mesh, meta = (ImportType = "StaticMesh"))
+	UPROPERTY(EditAnywhere, BlueprintReadWrite, AdvancedDisplay, Category = Mesh, meta = (ImportType = "StaticMesh", ReimportRestrict = "true"))
 	uint32 bRemoveDegenerates:1;
 
 	/** Required for PNT tessellation but can be slow. Recommend disabling for larger meshes. */
-	UPROPERTY(EditAnywhere, BlueprintReadWrite, AdvancedDisplay, Category = Mesh, meta = (ImportType = "StaticMesh"))
+	UPROPERTY(EditAnywhere, BlueprintReadWrite, AdvancedDisplay, Category = Mesh, meta = (ImportType = "StaticMesh", ReimportRestrict = "true"))
 	uint32 bBuildAdjacencyBuffer:1;
 
-	UPROPERTY(EditAnywhere, BlueprintReadWrite, AdvancedDisplay, Category = Mesh, meta = (ImportType = "StaticMesh"))
+	UPROPERTY(EditAnywhere, BlueprintReadWrite, AdvancedDisplay, Category = Mesh, meta = (ImportType = "StaticMesh", ReimportRestrict = "true"))
 	uint32 bBuildReversedIndexBuffer:1;
 
-	UPROPERTY(EditAnywhere, BlueprintReadWrite, config, AdvancedDisplay, Category= Mesh, meta=(ImportType="StaticMesh"))
+	UPROPERTY(EditAnywhere, BlueprintReadWrite, config, AdvancedDisplay, Category= Mesh, meta=(ImportType="StaticMesh", ReimportRestrict = "true"))
 	uint32 bGenerateLightmapUVs:1;
 
 	/** If checked, one convex hull per UCX_ prefixed collision mesh will be generated instead of decomposing into multiple hulls */
diff --git a/Engine/Source/Editor/UnrealEd/Private/Factories/EditorFactories.cpp b/Engine/Source/Editor/UnrealEd/Private/Factories/EditorFactories.cpp
--- a/Engine/Source/Editor/UnrealEd/Private/Factories/EditorFactories.cpp
+++ b/Engine/Source/Editor/UnrealEd/Private/Factories/EditorFactories.cpp
@@ -5343,7 +5343,14 @@ EReimportResult::Type UReimportFbxStaticMeshFactory::Reimport( UObject* Obj )
 	{
 		ReimportUI->bIsReimport = true;
 		ReimportUI->ReimportMesh = Mesh;
+
+		//Make sure the outer is the ImportUI, because there is some logic in the meta data needing this outer
+		UObject* OriginalOuter = ImportData != nullptr ? ImportData->GetOuter() : nullptr;
 		ReimportUI->StaticMeshImportData = ImportData;
+		if (ReimportUI->StaticMeshImportData && OriginalOuter)
+		{
+			ReimportUI->StaticMeshImportData->Rename(nullptr, ReimportUI);
+		}
 		
 		//Force the bAutoGenerateCollision to false if the Mesh Customize collision is true
 		bool bOldAutoGenerateCollision = ReimportUI->StaticMeshImportData->bAutoGenerateCollision;
@@ -5366,6 +5373,12 @@ EReimportResult::Type UReimportFbxStaticMeshFactory::Reimport( UObject* Obj )
 		{
 			ReimportUI->StaticMeshImportData->bAutoGenerateCollision = bOldAutoGenerateCollision;
 		}
+
+		//Put back the original SM outer
+		if (ReimportUI->StaticMeshImportData && OriginalOuter)
+		{
+			ReimportUI->StaticMeshImportData->Rename(nullptr, OriginalOuter);
+		}
 	}
 	ImportOptions->bCanShowDialog = !IsUnattended;
 	//We do not touch bAutoComputeLodDistances when we re-import, setting it to true will make sure we do not change anything.
