Prevent fail LOD import message if the user cancel the operation.

diff --git a/Engine/Source/Editor/UnrealEd/Private/FbxMeshUtils.cpp b/Engine/Source/Editor/UnrealEd/Private/FbxMeshUtils.cpp
--- a/Engine/Source/Editor/UnrealEd/Private/FbxMeshUtils.cpp
+++ b/Engine/Source/Editor/UnrealEd/Private/FbxMeshUtils.cpp
@@ -706,7 +706,8 @@ namespace FbxMeshUtils
 			}
 		}
 
-		if(!bImportSuccess)
+		//If the filename is empty it mean the user cancel the file selection
+		if(!bImportSuccess && !FilenameToImport.IsEmpty())
 		{
 			// Failed to import a LOD, even after retries (if applicable)
 			FMessageDialog::Open(EAppMsgType::Ok, FText::Format(LOCTEXT("LODImport_Failure", "Failed to import LOD{0}"), FText::AsNumber(LODLevel)));
