Prevent adding more then once the same collision primitive when importing them from fbx

diff --git a/Engine/Source/Editor/UnrealEd/Private/Fbx/FbxStaticMeshImport.cpp b/Engine/Source/Editor/UnrealEd/Private/Fbx/FbxStaticMeshImport.cpp
--- a/Engine/Source/Editor/UnrealEd/Private/Fbx/FbxStaticMeshImport.cpp
+++ b/Engine/Source/Editor/UnrealEd/Private/Fbx/FbxStaticMeshImport.cpp
@@ -2328,7 +2328,9 @@ bool UnFbx::FFbxImporter::FillCollisionModelList(FbxNode* Node)
 				{
 					Record = Models->GetValue();
 				}
-				Record->Add(Node);
+
+				//Unique add
+				Record->AddUnique(Node);
 			}
 		}
 
