Make sure we can export skeletal mesh whereever the clothing section are located.

diff --git a/Engine/Source/Editor/UnrealEd/Private/Fbx/FbxSkeletalMeshExport.cpp b/Engine/Source/Editor/UnrealEd/Private/Fbx/FbxSkeletalMeshExport.cpp
--- a/Engine/Source/Editor/UnrealEd/Private/Fbx/FbxSkeletalMeshExport.cpp
+++ b/Engine/Source/Editor/UnrealEd/Private/Fbx/FbxSkeletalMeshExport.cpp
@@ -199,11 +199,16 @@ FbxNode* FFbxExporter::CreateMesh(const USkeletalMesh* SkelMesh, const TCHAR* Me
 
 
 	// Create the per-material polygons sets.
-	int32 SectionCount = SourceModel.NumNonClothingSections();
+	int32 SectionCount = SourceModel.Sections.Num();
+	int32 ClothSectionVertexRemoveOffset = 0;
 	for (int32 SectionIndex = 0; SectionIndex < SectionCount; ++SectionIndex)
 	{
 		const FSkelMeshSection& Section = SourceModel.Sections[SectionIndex];
-
+		if (Section.HasClothingData())
+		{
+			ClothSectionVertexRemoveOffset += Section.GetNumVertices();
+			continue;
+		}
 		int32 MatIndex = Section.MaterialIndex;
 
 		// Static meshes contain one triangle list per element.
@@ -215,7 +220,9 @@ FbxNode* FFbxExporter::CreateMesh(const USkeletalMesh* SkelMesh, const TCHAR* Me
 			Mesh->BeginPolygon(MatIndex);
 			for (int32 PointIndex = 0; PointIndex < 3; PointIndex++)
 			{
-				Mesh->AddPolygon(SourceModel.IndexBuffer[Section.BaseIndex + ((TriangleIndex * 3) + PointIndex)]);
+				int32 VertexPositionIndex = SourceModel.IndexBuffer[Section.BaseIndex + ((TriangleIndex * 3) + PointIndex)] - ClothSectionVertexRemoveOffset;
+				check(VertexPositionIndex >= 0);
+				Mesh->AddPolygon(VertexPositionIndex);
 			}
 			Mesh->EndPolygon();
 		}
@@ -284,7 +291,6 @@ void FFbxExporter::BindMeshToSkeleton(const USkeletalMesh* SkelMesh, FbxNode* Me
 {
 	const FSkeletalMeshModel* SkelMeshResource = SkelMesh->GetImportedModel();
 	const FSkeletalMeshLODModel& SourceModel = SkelMeshResource->LODModels[0];
-	const int32 VertexCount = SourceModel.NumVertices;
 
 	FbxAMatrix MeshMatrix;
 
@@ -316,6 +322,10 @@ void FFbxExporter::BindMeshToSkeleton(const USkeletalMesh* SkelMesh, FbxNode* Me
 		for(int32 SectionIndex = 0; SectionIndex < SectionCount; ++SectionIndex)
 		{
 			const FSkelMeshSection& Section = SourceModel.Sections[SectionIndex];
+			if (Section.HasClothingData())
+			{
+				continue;
+			}
 
 			for(int32 SoftIndex = 0; SoftIndex < Section.SoftVertices.Num(); ++SoftIndex)
 			{
diff --git a/Engine/Source/Runtime/Engine/Private/SkeletalMeshLODModel.cpp b/Engine/Source/Runtime/Engine/Private/SkeletalMeshLODModel.cpp
--- a/Engine/Source/Runtime/Engine/Private/SkeletalMeshLODModel.cpp
+++ b/Engine/Source/Runtime/Engine/Private/SkeletalMeshLODModel.cpp
@@ -839,7 +839,7 @@ int32 FSkeletalMeshLODModel::GetNumNonClothingVertices() const
 		// Stop when we hit clothing sections
 		if (Section.ClothingData.AssetGuid.IsValid())
 		{
-			break;
+			continue;
 		}
 
 		NumVerts += Section.SoftVertices.Num();
@@ -851,15 +851,10 @@ int32 FSkeletalMeshLODModel::GetNumNonClothingVertices() const
 void FSkeletalMeshLODModel::GetNonClothVertices(TArray<FSoftSkinVertex>& OutVertices) const
 {
 	// Get the number of sections to copy
-	int32 NumSections = NumNonClothingSections();
+	int32 NumSections = Sections.Num();
 
 	// Count number of verts
-	int32 NumVertsToCopy = 0;
-	for (int32 SectionIndex = 0; SectionIndex < NumSections; SectionIndex++)
-	{
-		const FSkelMeshSection& Section = Sections[SectionIndex];
-		NumVertsToCopy += Section.SoftVertices.Num();
-	}
+	int32 NumVertsToCopy = GetNumNonClothingVertices();
 
 	OutVertices.Empty(NumVertsToCopy);
 	OutVertices.AddUninitialized(NumVertsToCopy);
@@ -870,6 +865,10 @@ void FSkeletalMeshLODModel::GetNonClothVertices(TArray<FSoftSkinVertex>& OutVert
 	for (int32 SectionIndex = 0; SectionIndex < NumSections; SectionIndex++)
 	{
 		const FSkelMeshSection& Section = Sections[SectionIndex];
+		if (Section.HasClothingData())
+		{
+			continue;
+		}
 		FMemory::Memcpy(DestVertex, Section.SoftVertices.GetData(), Section.SoftVertices.Num() * sizeof(FSoftSkinVertex));
 		DestVertex += Section.SoftVertices.Num();
 	}
