Only show the used material and remove the namespace when building the preview material conflict

diff --git a/Engine/Source/Editor/UnrealEd/Private/Fbx/FbxFactory.cpp b/Engine/Source/Editor/UnrealEd/Private/Fbx/FbxFactory.cpp
--- a/Engine/Source/Editor/UnrealEd/Private/Fbx/FbxFactory.cpp
+++ b/Engine/Source/Editor/UnrealEd/Private/Fbx/FbxFactory.cpp
@@ -1117,16 +1117,45 @@ namespace ImportCompareHelper
 		SkeletonCompareData.bHasConflict = HasConflictRecursive(SkeletonCompareData.ResultAssetRoot, SkeletonCompareData.CurrentAssetRoot);
 	}
 
-	void FillFbxMaterials(const TArray<FbxNode*>& MeshNodes, FMaterialCompareData& MaterialCompareData)
+	void FillFbxMaterials(UnFbx::FFbxImporter* FFbxImporter, const TArray<FbxNode*>& MeshNodes, FMaterialCompareData& MaterialCompareData)
 	{
 		TArray<FName> NodeMaterialNames;
 		for (int32 NodeIndex = 0; NodeIndex < MeshNodes.Num(); ++NodeIndex)
 		{
 			FbxNode* Node = MeshNodes[NodeIndex];
-			for (int32 MaterialIndex = 0; MaterialIndex < Node->GetMaterialCount(); ++MaterialIndex)
+			if (Node->GetMesh() == nullptr)
 			{
+				continue;
+			}
+
+			int32 MaterialCount = Node->GetMaterialCount();
+			TArray<int32> MaterialUseByMesh;
+			FbxLayer* BaseLayer = Node->GetMesh()->GetLayer(0);
+			FbxLayerElementMaterial* MateriallayerElement = BaseLayer->GetMaterials();
+			FbxLayerElement::EMappingMode MaterialMappingMode = MateriallayerElement ? MateriallayerElement->GetMappingMode() : FbxLayerElement::eByPolygon;
+
+			if (MaterialMappingMode == FbxLayerElement::eAllSame || MaterialCount == 0 || MateriallayerElement == nullptr)
+			{
+				MaterialUseByMesh.Add(0);
+			}
+			else
+			{
+				int32 PolygonCount = Node->GetMesh()->GetPolygonCount();
+				for (int32 PolygonIndex = 0; PolygonIndex < PolygonCount; PolygonIndex++)
+				{
+					MaterialUseByMesh.AddUnique(MateriallayerElement->GetIndexArray().GetAt(PolygonIndex));
+				}
+			}
+
+			for (int32 MaterialIndex = 0; MaterialIndex < MaterialCount; ++MaterialIndex)
+			{
+				//Skip unused mesh material
+				if (!MaterialUseByMesh.Contains(MaterialIndex))
+				{
+					continue;
+				}
 				FbxSurfaceMaterial* SurfaceMaterial = Node->GetMaterial(MaterialIndex);
-				FName SurfaceMaterialName = FName(UTF8_TO_TCHAR(SurfaceMaterial->GetName()));
+				FName SurfaceMaterialName = FName(UTF8_TO_TCHAR(FFbxImporter->MakeName(SurfaceMaterial->GetName())));
 				if (!NodeMaterialNames.Contains(SurfaceMaterialName))
 				{
 					FMaterialData& MaterialData = MaterialCompareData.ResultAsset.AddDefaulted_GetRef();
@@ -1392,7 +1421,7 @@ namespace ImportCompareHelper
 			StaticMeshNodes.Append(FbxMeshArray);
 		}
 
-		FillFbxMaterials(StaticMeshNodes, ImportUI->MaterialCompareData);
+		FillFbxMaterials(FFbxImporter, StaticMeshNodes, ImportUI->MaterialCompareData);
 		//Compare the result and set the conflict status
 		SetHasConflict(ImportUI->MaterialCompareData);
 	}
@@ -1439,7 +1468,7 @@ namespace ImportCompareHelper
 			}
 
 			//Fill the result fbx data
-			FillFbxMaterials(FlattenSkeletalMeshNodes, ImportUI->MaterialCompareData);
+			FillFbxMaterials(FFbxImporter, FlattenSkeletalMeshNodes, ImportUI->MaterialCompareData);
 
 			//Compare the result and set the conflict status
 			SetHasConflict(ImportUI->MaterialCompareData);
diff --git a/Engine/Source/Editor/UnrealEd/Private/Fbx/FbxMainImport.cpp b/Engine/Source/Editor/UnrealEd/Private/Fbx/FbxMainImport.cpp
--- a/Engine/Source/Editor/UnrealEd/Private/Fbx/FbxMainImport.cpp
+++ b/Engine/Source/Editor/UnrealEd/Private/Fbx/FbxMainImport.cpp
@@ -1047,13 +1047,21 @@ bool FFbxImporter::OpenFile(FString Filename)
 
 void FFbxImporter::FixMaterialClashName()
 {
+	const bool bKeepNamespace = GetDefault<UEditorPerProjectUserSettings>()->bKeepFbxNamespace;
+
 	FbxArray<FbxSurfaceMaterial*> MaterialArray;
 	Scene->FillMaterialArray(MaterialArray);
 	TSet<FString> AllMaterialName;
 	for (int32 MaterialIndex = 0; MaterialIndex < MaterialArray.Size(); ++MaterialIndex)
 	{
 		FbxSurfaceMaterial *Material = MaterialArray[MaterialIndex];
-		FString MaterialName = UTF8_TO_TCHAR(Material->GetName());
+		FString MaterialName = UTF8_TO_TCHAR(MakeName(Material->GetName()));
+
+		if (!bKeepNamespace)
+		{
+			Material->SetName(TCHAR_TO_UTF8(*MaterialName));
+		}
+
 		if (AllMaterialName.Contains(MaterialName))
 		{
 			FString OriginalMaterialName = MaterialName;
