Fix CIS static analysis

diff --git a/Engine/Source/Editor/UnrealEd/Private/Factories/EditorFactories.cpp b/Engine/Source/Editor/UnrealEd/Private/Factories/EditorFactories.cpp
--- a/Engine/Source/Editor/UnrealEd/Private/Factories/EditorFactories.cpp
+++ b/Engine/Source/Editor/UnrealEd/Private/Factories/EditorFactories.cpp
@@ -5369,7 +5369,7 @@ EReimportResult::Type UReimportFbxStaticMeshFactory::Reimport( UObject* Obj )
 		GetImportOptions( FFbxImporter, ReimportUI, bShowOptionDialog, bIsAutomated, Obj->GetPathName(), bOperationCanceled, bOutImportAll, bIsObjFormat, Filename, bForceImportType, FBXIT_StaticMesh);
 		
 		//Put back the original bAutoGenerateCollision settings since the user cancel the re-import
-		if (bOperationCanceled && Mesh->bCustomizedCollision)
+		if (ReimportUI->StaticMeshImportData && bOperationCanceled && Mesh->bCustomizedCollision)
 		{
 			ReimportUI->StaticMeshImportData->bAutoGenerateCollision = bOldAutoGenerateCollision;
 		}
