Selecting multiple keys of the same time now displays the time instead of "multiple values"

diff --git a/Engine/Source/Editor/UnrealEd/Private/SCurveEditor.cpp b/Engine/Source/Editor/UnrealEd/Private/SCurveEditor.cpp
--- a/Engine/Source/Editor/UnrealEd/Private/SCurveEditor.cpp
+++ b/Engine/Source/Editor/UnrealEd/Private/SCurveEditor.cpp
@@ -2053,13 +2053,24 @@ void SCurveEditor::ProcessClick(const FGeometry& InMyGeometry, const FPointerEve
 
 TOptional<float> SCurveEditor::OnGetTime() const
 {
-	if ( SelectedKeys.Num() == 1 )
+	TOptional<float> Time;
+
+	// Return the time if all selected keys have the same time, otherwise return an unset value
+	if (SelectedKeys.Num() > 0)
 	{
-		return GetKeyTime(SelectedKeys[0]);
+		Time = GetKeyTime(SelectedKeys[0]);
+		for (int32 i = 1; i < SelectedKeys.Num(); i++)
+		{
+			TOptional<float> NewTime = GetKeyTime(SelectedKeys[i]);
+			bool bAreEqual = ((!Time.IsSet() && !NewTime.IsSet()) || (Time.IsSet() && NewTime.IsSet() && Time.GetValue() == NewTime.GetValue()));
+			if (!bAreEqual)
+			{
+				return TOptional<float>();
+			}
+		}
 	}
 
-	// Value couldn't be accessed.  Return an unset value
-	return TOptional<float>();
+	return Time;
 }
 
 void SCurveEditor::OnTimeComitted(float NewTime, ETextCommit::Type CommitType)
@@ -2089,16 +2100,28 @@ void SCurveEditor::OnTimeChanged(float NewTime)
 
 TOptional<int32> SCurveEditor::OnGetTimeInFrames() const
 {
-	if ( SelectedKeys.Num() == 1 )
+	TOptional<float> Time;
+
+	// Return the time in frames if all selected keys have the same time, otherwise return an unset value
+	if (SelectedKeys.Num() > 0)
 	{
-		TOptional<float> KeyTime = GetKeyTime(SelectedKeys[0]);
-		if (KeyTime.IsSet())
+		Time = GetKeyTime(SelectedKeys[0]);
+		for (int32 i = 1; i < SelectedKeys.Num(); i++)
 		{
-			return TOptional<int32>(TimeToFrame(KeyTime.GetValue()));
+			TOptional<float> NewTime = GetKeyTime(SelectedKeys[i]);
+			bool bAreEqual = ((!Time.IsSet() && !NewTime.IsSet()) || (Time.IsSet() && NewTime.IsSet() && Time.GetValue() == NewTime.GetValue()));
+			if (!bAreEqual)
+			{
+				return TOptional<int32>();
+			}
 		}
 	}
 
-	// Value couldn't be accessed.  Return an unset value
+	if (Time.IsSet())
+	{
+		return TOptional<int32>(TimeToFrame(Time.GetValue()));
+	}
+
 	return TOptional<int32>();
 }
 
