Remove unused analytic fbx data

diff --git a/Engine/Source/Editor/UnrealEd/Private/Fbx/FbxMainImport.cpp b/Engine/Source/Editor/UnrealEd/Private/Fbx/FbxMainImport.cpp
--- a/Engine/Source/Editor/UnrealEd/Private/Fbx/FbxMainImport.cpp
+++ b/Engine/Source/Editor/UnrealEd/Private/Fbx/FbxMainImport.cpp
@@ -1402,14 +1402,10 @@ bool FFbxImporter::ImportFromFile(const FString& Filename, const FString& Type,
 				  * @EventName Editor.Usage.FBX
 				  * @Trigger Fires when the user clicks OK in the FBX Import Dialog
 				  * @Type Editor
-				  * @EventParam OriginalVendor string Returns the name of the vendor that manufactures the original application used to generate the imported FBX
-				  * @EventParam OriginalAppName string Returns the name of the original application used to generate the imported FBX
-				  * @EventParam OriginalAppVersion string Returns the revision of the original application used to generate the imported FBX
 				  * @EventParam LastSavedVendor string Returns the name of the vendor that manufactures the last application used to modify the imported FBX
 				  * @EventParam LastSavedAppName string Returns the name of the last application used to modify the imported FBX
 				  * @EventParam LastSavedAppVersion string Returns the revision of the last application used to modify the imported FBX
 				  * @EventParam FBXFileVersion string Returns the FBX SDK used to generate the imported FBX
-				  * @EventParam FilenameHash string Returns the filename of the meshes imported from the FBX 
 				  * @EventParam ImportType string Returns the mesh data type (Static, Skeletal, Animation) being imported from the FBX
 				  * @EventParam ConvertScene boolean Returns whether the import fbx should be converted to unreal axis system
 				  * @EventParam ConvertSceneUnit boolean Returns whether the import fbx should converted the unit to unreal unit (cm)
@@ -1484,26 +1480,15 @@ bool FFbxImporter::ImportFromFile(const FString& Filename, const FString& Type,
 						
 						TArray<FAnalyticsEventAttribute> Attribs;
 
-						FString OriginalVendor(UTF8_TO_TCHAR(DocInfo->Original_ApplicationVendor.Get().Buffer()));
-						FString OriginalAppName(UTF8_TO_TCHAR(DocInfo->Original_ApplicationName.Get().Buffer()));
-						FString OriginalAppVersion(UTF8_TO_TCHAR(DocInfo->Original_ApplicationVersion.Get().Buffer()));
-
 						FString LastSavedVendor(UTF8_TO_TCHAR(DocInfo->LastSaved_ApplicationVendor.Get().Buffer()));
 						FString LastSavedAppName(UTF8_TO_TCHAR(DocInfo->LastSaved_ApplicationName.Get().Buffer()));
 						FString LastSavedAppVersion(UTF8_TO_TCHAR(DocInfo->LastSaved_ApplicationVersion.Get().Buffer()));
 
-						FString FilenameHash = FMD5::HashAnsiString(*Filename);
-
-						Attribs.Add(FAnalyticsEventAttribute(TEXT("Original Application Vendor"), OriginalVendor));
-						Attribs.Add(FAnalyticsEventAttribute(TEXT("Original Application Name"), OriginalAppName));
-						Attribs.Add(FAnalyticsEventAttribute(TEXT("Original Application Version"), OriginalAppVersion));
-
 						Attribs.Add(FAnalyticsEventAttribute(TEXT("LastSaved Application Vendor"), LastSavedVendor));
 						Attribs.Add(FAnalyticsEventAttribute(TEXT("LastSaved Application Name"), LastSavedAppName));
 						Attribs.Add(FAnalyticsEventAttribute(TEXT("LastSaved Application Version"), LastSavedAppVersion));
 
 						Attribs.Add(FAnalyticsEventAttribute(TEXT("FBX Version"), FbxFileVersion));
-						Attribs.Add(FAnalyticsEventAttribute(TEXT("Filename Hash"), FilenameHash));
 
 						//////////////////////////////////////////////////////////////////////////
 						//FBX import options
