Checking asset validity to avoid crashes on level reopen

diff --git a/Engine/Source/Editor/UnrealEd/Private/Toolkits/AssetEditorManager.cpp b/Engine/Source/Editor/UnrealEd/Private/Toolkits/AssetEditorManager.cpp
--- a/Engine/Source/Editor/UnrealEd/Private/Toolkits/AssetEditorManager.cpp
+++ b/Engine/Source/Editor/UnrealEd/Private/Toolkits/AssetEditorManager.cpp
@@ -397,8 +397,9 @@ bool FAssetEditorManager::OpenEditorForAsset(UObject* Asset, const EToolkitMode:
 	{
 		GWarn->EndSlowTask();
 	}
+	// Must check Asset here in addition to at the beginning of the function, because if the asset was destroyed and recreated it might not be found correctly
 	// Do not add to recently opened asset list if this is a level-associated asset like Level Blueprint or Built Data. Their naming is not compatible
-	if (Asset->IsAsset() && !Asset->IsA(UMapBuildDataRegistry::StaticClass()))
+	if (Asset && Asset->IsAsset() && !Asset->IsA(UMapBuildDataRegistry::StaticClass()))
 	{
 		FString AssetPath = Asset->GetOuter()->GetPathName();
 		FContentBrowserModule& CBModule = FModuleManager::LoadModuleChecked<FContentBrowserModule>(TEXT("ContentBrowser"));
