Guards against nullptr and out of array bounds access when getting node titles

diff --git a/Engine/Source/Editor/UnrealEd/Private/MaterialGraphNode.cpp b/Engine/Source/Editor/UnrealEd/Private/MaterialGraphNode.cpp
--- a/Engine/Source/Editor/UnrealEd/Private/MaterialGraphNode.cpp
+++ b/Engine/Source/Editor/UnrealEd/Private/MaterialGraphNode.cpp
@@ -147,8 +147,10 @@ bool UMaterialGraphNode::CanPasteHere(const UEdGraph* TargetGraph) const
 FText UMaterialGraphNode::GetNodeTitle(ENodeTitleType::Type TitleType) const
 {
 	TArray<FString> Captions;
-	MaterialExpression->GetCaption(Captions);
-
+	if (MaterialExpression)
+	{
+		MaterialExpression->GetCaption(Captions);
+	}
 	if (TitleType == ENodeTitleType::EditableTitle)
 	{
 		return FText::FromString(GetParameterName());
@@ -161,11 +163,13 @@ FText UMaterialGraphNode::GetNodeTitle(ENodeTitleType::Type TitleType) const
 	{
 		// More useful to display multi line parameter captions in reverse order
 		// TODO: May have to choose order based on expression type if others need correct order
-		int32 CaptionIndex = Captions.Num() -1;
+		int32 CaptionIndex = Captions.Num() - 1;
 
 		FTextBuilder NodeTitle;
-		NodeTitle.AppendLine(Captions[CaptionIndex]);
-
+		if (Captions.IsValidIndex(CaptionIndex))
+		{
+			NodeTitle.AppendLine(Captions[CaptionIndex]);
+		}
 		for (; CaptionIndex > 0; )
 		{
 			CaptionIndex--;
