FBX export animation curve created within the editor
Thanks to PR author [REDACTED_PERSON]
PR #4654

diff --git a/Engine/Source/Editor/UnrealEd/Private/Fbx/FbxAnimationExport.cpp b/Engine/Source/Editor/UnrealEd/Private/Fbx/FbxAnimationExport.cpp
--- a/Engine/Source/Editor/UnrealEd/Private/Fbx/FbxAnimationExport.cpp
+++ b/Engine/Source/Editor/UnrealEd/Private/Fbx/FbxAnimationExport.cpp
@@ -33,9 +33,12 @@ void FFbxExporter::ExportAnimSequenceToFbx(const UAnimSequence* AnimSeq,
 									 float AnimPlayRate,
 									 float StartTime)
 {
+	// stack allocator for extracting curve
+	FMemMark Mark(FMemStack::Get());
+
 	USkeleton* Skeleton = AnimSeq->GetSkeleton();
 
-	if (AnimSeq->SequenceLength == 0.f)
+	if (AnimSeq->SequenceLength == 0.f || Skeleton == nullptr)
 	{
 		// something is wrong
 		return;	
@@ -62,7 +65,35 @@ void FFbxExporter::ExportAnimSequenceToFbx(const UAnimSequence* AnimSeq,
 	FbxTimeSpan ExportedTimeSpan;
 	ExportedTimeSpan.Set(ExportedStartTime, ExportedStopTime);
 	AnimStack->SetLocalTimeSpan(ExportedTimeSpan);
-	
+
+	//Prepare root anim curves data to be exported
+	TArray<FName> AnimCurveNames;
+	TArray<SmartName::UID_Type> AnimCurveUIDs;
+	TMap<FName, FbxAnimCurve*> CustomCurveMap;
+	if (BoneNodes.Num() > 0)
+	{
+		const FSmartNameMapping* AnimCurveMapping = Skeleton->GetSmartNameContainer(USkeleton::AnimCurveMappingName);
+		
+		if (AnimCurveMapping)
+		{
+			AnimCurveMapping->FillNameArray(AnimCurveNames);
+			AnimCurveMapping->FillUidArray(AnimCurveUIDs);
+		}
+		
+		if (AnimCurveNames.Num() > 0)
+		{
+			for (auto AnimCurveName : AnimCurveNames)
+			{
+				FbxProperty AnimCurveFbxProp = FbxProperty::Create(BoneNodes[0], FbxDoubleDT, TCHAR_TO_ANSI(*AnimCurveName.ToString()));
+				AnimCurveFbxProp.ModifyFlag(FbxPropertyFlags::eAnimatable, true);
+				AnimCurveFbxProp.ModifyFlag(FbxPropertyFlags::eUserDefined, true);
+				FbxAnimCurve* AnimFbxCurve = AnimCurveFbxProp.GetCurve(InAnimLayer, true);
+				CustomCurveMap.Add(AnimCurveName, AnimFbxCurve);
+			}
+		}
+	}
+
+
 	// Add the animation data to the bone nodes
 	for(int32 BoneIndex = 0; BoneIndex < BoneNodes.Num(); ++BoneIndex)
 	{
@@ -123,6 +154,30 @@ void FFbxExporter::ExportAnimSequenceToFbx(const UAnimSequence* AnimSeq,
 			FbxVector4 Scale = Converter.ConvertToFbxScale(BoneAtom.GetScale3D());
 			FbxVector4 Vectors[3] = { Translation, Rotation, Scale };
 		
+			//Add custom curve keys only to the root bone
+			if (BoneIndex == 0 && CustomCurveMap.Num() > 0)
+			{
+				FBlendedCurve BlendedCurve;
+				BlendedCurve.InitFrom(&AnimCurveUIDs);
+				AnimSeq->EvaluateCurveData(BlendedCurve, AnimTime, true);
+				if (BlendedCurve.IsValid())
+				{
+					//Loop over the custom curves and add the actual keys
+					for (auto CustomCurve : CustomCurveMap)
+					{
+						CustomCurve.Value->KeyModifyBegin();
+						SmartName::UID_Type NameUID = Skeleton->GetUIDByName(USkeleton::AnimCurveMappingName, CustomCurve.Key);
+						if (NameUID != SmartName::MaxUID)
+						{
+							float CurveValueAtTime = BlendedCurve.Get(NameUID);
+							int32 KeyIndex = CustomCurve.Value->KeyAdd(ExportTime);
+							CustomCurve.Value->KeySetValue(KeyIndex, CurveValueAtTime);
+						}
+					}
+
+				}
+			}
+
 			int32 lKeyIndex;
 			bLastKey = (AnimTime + KINDA_SMALL_NUMBER) > AnimEndTime;
 			
@@ -153,6 +208,14 @@ void FFbxExporter::ExportAnimSequenceToFbx(const UAnimSequence* AnimSeq,
 		{
 			Curve->KeyModifyEnd();
 		}
+
+		if (BoneIndex == 0 && CustomCurveMap.Num() > 0)
+		{
+			for (auto CustomCurve : CustomCurveMap)
+			{
+				CustomCurve.Value->KeyModifyEnd();
+			}
+		}
 	}
 }
 
