In the material conflict window, do not allow user to clear a unmatch material.

diff --git a/Engine/Source/Editor/UnrealEd/Private/Fbx/FbxMaterialConflictWindow.cpp b/Engine/Source/Editor/UnrealEd/Private/Fbx/FbxMaterialConflictWindow.cpp
--- a/Engine/Source/Editor/UnrealEd/Private/Fbx/FbxMaterialConflictWindow.cpp
+++ b/Engine/Source/Editor/UnrealEd/Private/Fbx/FbxMaterialConflictWindow.cpp
@@ -328,11 +328,13 @@ FReply FMaterialConflictData::OnMouseButtonDown(const FGeometry& MyGeometry, con
 
 	//Gather the possible match item
 	FMenuBuilder ContextMenu(true, TSharedPtr<const FUICommandList>());
-
-	//INDEX_NONE clear the remapping
 	FText EntryName = FText(LOCTEXT("OnMouseButtonDown_menuClear", "Clear"));
-	ContextMenu.AddMenuEntry(EntryName, FText::GetEmpty(), FSlateIcon(), FUIAction(FExecuteAction::CreateRaw(this, &FMaterialConflictData::AssignMaterialMatch, (int32)INDEX_NONE)));
-	ContextMenu.AddMenuSeparator();
+	if (RemapMaterials[ResultMaterialIndex] != INDEX_NONE)
+	{
+		//INDEX_NONE clear the remapping
+		ContextMenu.AddMenuEntry(EntryName, FText::GetEmpty(), FSlateIcon(), FUIAction(FExecuteAction::CreateRaw(this, &FMaterialConflictData::AssignMaterialMatch, (int32)INDEX_NONE)));
+		ContextMenu.AddMenuSeparator();
+	}
 	for (int32 OriginalMaterialIndex = 0; OriginalMaterialIndex < SourceMaterials.Num(); ++OriginalMaterialIndex)
 	{
 		EntryName = FText::FromName(SourceMaterials[OriginalMaterialIndex].ImportedMaterialSlotName);
@@ -359,7 +361,7 @@ void FMaterialConflictData::AssignMaterialMatch(int32 OriginalMaterialIndex)
 	}
 	else if (OriginalMaterialIndex == INDEX_NONE)
 	{
-		RemapMaterials[ResultMaterialIndex] = ResultMaterialIndex;
+		RemapMaterials[ResultMaterialIndex] = INDEX_NONE;
 	}
 }
 
