Fix for static analysis

diff --git a/Engine/Source/Editor/UnrealEd/Private/MaterialGraphNode.cpp b/Engine/Source/Editor/UnrealEd/Private/MaterialGraphNode.cpp
--- a/Engine/Source/Editor/UnrealEd/Private/MaterialGraphNode.cpp
+++ b/Engine/Source/Editor/UnrealEd/Private/MaterialGraphNode.cpp
@@ -151,11 +151,12 @@ FText UMaterialGraphNode::GetNodeTitle(ENodeTitleType::Type TitleType) const
 	{
 		MaterialExpression->GetCaption(Captions);
 	}
+
 	if (TitleType == ENodeTitleType::EditableTitle)
 	{
 		return FText::FromString(GetParameterName());
 	}
-	else if (TitleType == ENodeTitleType::ListView || TitleType == ENodeTitleType::MenuTitle)
+	else if (MaterialExpression && (TitleType == ENodeTitleType::ListView || TitleType == ENodeTitleType::MenuTitle))
 	{
 		return FText::FromString(MaterialExpression->GetClass()->GetDescription());
 	}
@@ -176,7 +177,7 @@ FText UMaterialGraphNode::GetNodeTitle(ENodeTitleType::Type TitleType) const
 			NodeTitle.AppendLine(Captions[CaptionIndex]);
 		}
 
-		if ( MaterialExpression->bShaderInputData && (MaterialExpression->bHidePreviewWindow || MaterialExpression->bCollapsed))
+		if (MaterialExpression && MaterialExpression->bShaderInputData && (MaterialExpression->bHidePreviewWindow || MaterialExpression->bCollapsed))
 		{
 			if (MaterialExpression->IsA<UMaterialExpressionTextureProperty>())
 			{
