fix fbx import full name issue

diff --git a/Engine/Source/Editor/UnrealEd/Private/Fbx/FbxFactory.cpp b/Engine/Source/Editor/UnrealEd/Private/Fbx/FbxFactory.cpp
--- a/Engine/Source/Editor/UnrealEd/Private/Fbx/FbxFactory.cpp
+++ b/Engine/Source/Editor/UnrealEd/Private/Fbx/FbxFactory.cpp
@@ -369,7 +369,7 @@ UObject* UFbxFactory::FactoryCreateFile
 				else
 				{
 					// count meshes in lod groups if we dont care about importing LODs
-					bool bCountLODGroupMeshes = !bImportStaticMeshLODs;
+					bool bCountLODGroupMeshes = !bImportStaticMeshLODs && bCombineMeshes;
 					int32 NumLODGroups = 0;
 					InterestingNodeCount = FbxImporter->GetFbxMeshCount(RootNodeToImport,bCountLODGroupMeshes,NumLODGroups);
 
@@ -1296,7 +1296,7 @@ namespace ImportCompareHelper
 		else
 		{
 			// count meshes in lod groups if we dont care about importing LODs
-			bool bCountLODGroupMeshes = !bImportStaticMeshLODs;
+			bool bCountLODGroupMeshes = !bImportStaticMeshLODs && bCombineMeshes;
 			int32 NumLODGroups = 0;
 			FFbxImporter->GetFbxMeshCount(FFbxImporter->Scene->GetRootNode(), bCountLODGroupMeshes, NumLODGroups);
 			// if there were LODs in the file, do not combine meshes even if requested
diff --git a/Engine/Source/Editor/UnrealEd/Private/Fbx/FbxStaticMeshImport.cpp b/Engine/Source/Editor/UnrealEd/Private/Fbx/FbxStaticMeshImport.cpp
--- a/Engine/Source/Editor/UnrealEd/Private/Fbx/FbxStaticMeshImport.cpp
+++ b/Engine/Source/Editor/UnrealEd/Private/Fbx/FbxStaticMeshImport.cpp
@@ -1114,7 +1114,7 @@ UStaticMesh* UnFbx::FFbxImporter::ReimportStaticMesh(UStaticMesh* Mesh, UFbxStat
 	else
 	{
 		// count meshes in lod groups if we dont care about importing LODs
-		bool bCountLODGroupMeshes = !bImportStaticMeshLODs;
+		bool bCountLODGroupMeshes = !bImportStaticMeshLODs && bCombineMeshes;
 		int32 NumLODGroups = 0;
 		GetFbxMeshCount(Scene->GetRootNode(), bCountLODGroupMeshes, NumLODGroups);
 		// if there were LODs in the file, do not combine meshes even if requested
