Fix a crash when removing morph target from LODs

diff --git a/Engine/Source/Editor/UnrealEd/Private/LODUtilities.cpp b/Engine/Source/Editor/UnrealEd/Private/LODUtilities.cpp
--- a/Engine/Source/Editor/UnrealEd/Private/LODUtilities.cpp
+++ b/Engine/Source/Editor/UnrealEd/Private/LODUtilities.cpp
@@ -791,30 +791,36 @@ void FLODUtilities::SimplifySkeletalMeshLOD( USkeletalMesh* SkeletalMesh, int32
 	if (MeshReduction->ReduceSkeletalMesh(SkeletalMesh, DesiredLOD, bReregisterComponent))
 	{
 		check(SkeletalMesh->GetLODNum() >= 1);
-		
-		FSkeletalMeshOptimizationSettings& ReductionSettings = SkeletalMesh->GetLODInfo(DesiredLOD)->ReductionSettings;
-		//Apply morph to the new LOD. Force it if we reduce the base LOD, base LOD must apply the morph target
-		if (ReductionSettings.bRemapMorphTargets)
+
+		auto ApplyMorphTargetOption = [&SkeletalMesh, &DesiredLOD]()
 		{
-			if (bReregisterComponent)
+			FSkeletalMeshOptimizationSettings& ReductionSettings = SkeletalMesh->GetLODInfo(DesiredLOD)->ReductionSettings;
+			//Apply morph to the new LOD. Force it if we reduce the base LOD, base LOD must apply the morph target
+			if (ReductionSettings.bRemapMorphTargets)
 			{
-				TComponentReregisterContext<USkinnedMeshComponent> ReregisterContext;
-				SkeletalMesh->ReleaseResources();
-				SkeletalMesh->ReleaseResourcesFence.Wait();
 				ApplyMorphTargetsToLOD(SkeletalMesh, ReductionSettings.BaseLOD, DesiredLOD);
-				SkeletalMesh->PostEditChange();
-				SkeletalMesh->InitResources();
 			}
 			else
 			{
-				ApplyMorphTargetsToLOD(SkeletalMesh, ReductionSettings.BaseLOD, DesiredLOD);
+				ClearGeneratedMorphTarget(SkeletalMesh, DesiredLOD);
 			}
+		};
+
+		if (bReregisterComponent)
+		{
+			TComponentReregisterContext<USkinnedMeshComponent> ReregisterContext;
+			SkeletalMesh->ReleaseResources();
+			SkeletalMesh->ReleaseResourcesFence.Wait();
+
+			ApplyMorphTargetOption();
+			
+			SkeletalMesh->PostEditChange();
+			SkeletalMesh->InitResources();
 		}
 		else
 		{
-			ClearGeneratedMorphTarget(SkeletalMesh, DesiredLOD);
+			ApplyMorphTargetOption();
 		}
-
 		SkeletalMesh->MarkPackageDirty();
 	}
 	else
