PR #4269: Prevent dropping of non-placeable uobjects in the editor viewport (Contributed by [REDACTED_PERSON])

diff --git a/Engine/Source/Editor/UnrealEd/Private/LevelEditorViewport.cpp b/Engine/Source/Editor/UnrealEd/Private/LevelEditorViewport.cpp
--- a/Engine/Source/Editor/UnrealEd/Private/LevelEditorViewport.cpp
+++ b/Engine/Source/Editor/UnrealEd/Private/LevelEditorViewport.cpp
@@ -1172,7 +1172,12 @@ FDropQuery FLevelEditorViewportClient::CanDropObjectsAtCoordinates(int32 MouseX,
 	}
 
 	UObject* AssetObj = AssetData.GetAsset();
-	UClass* ClassObj = Cast<UClass>( AssetObj );
+	UClass* ClassObj = Cast<UClass>(AssetObj);
+	if (!ClassObj)
+	{
+		UBlueprint* BlueprintObj = Cast<UBlueprint>(AssetObj);
+		ClassObj = BlueprintObj ? BlueprintObj->ParentClass : nullptr;
+	}
 
 	if ( ClassObj )
 	{
