Fix the transform when exporting staticmesh component

diff --git a/Engine/Source/Editor/UnrealEd/Private/Fbx/FbxMainExport.cpp b/Engine/Source/Editor/UnrealEd/Private/Fbx/FbxMainExport.cpp
--- a/Engine/Source/Editor/UnrealEd/Private/Fbx/FbxMainExport.cpp
+++ b/Engine/Source/Editor/UnrealEd/Private/Fbx/FbxMainExport.cpp
@@ -1812,8 +1812,7 @@ FbxNode* FFbxExporter::ExportActor(AActor* Actor, bool bExportComponents, INodeN
 					const FTransform RelativeTransform = RotationDirectionConvert * Component->GetComponentToWorld().GetRelativeTransform(Actor->GetTransform());
 
 					FTransform ActorTransform(FRotator::MakeFromEuler(ActorRotation).Quaternion(), ActorLocation, ActorScale);
-					FTransform TotalTransform = RelativeTransform;
-					TotalTransform.Accumulate(ActorTransform);
+					FTransform TotalTransform = RelativeTransform * ActorTransform;
 
 					ActorNode->LclTranslation.Set(Converter.ConvertToFbxPos(TotalTransform.GetLocation()));
 					ActorNode->LclRotation.Set(Converter.ConvertToFbxRot(TotalTransform.GetRotation().Euler()));
