Compile fix

diff --git a/Engine/Source/Editor/UnrealEd/Public/AssetImportTask.h b/Engine/Source/Editor/UnrealEd/Public/AssetImportTask.h
--- a/Engine/Source/Editor/UnrealEd/Public/AssetImportTask.h
+++ b/Engine/Source/Editor/UnrealEd/Public/AssetImportTask.h
@@ -9,6 +9,8 @@
 
 DECLARE_LOG_CATEGORY_EXTERN(LogAssetImportTask, Log, All);
 
+class UFactory;
+
 /**
  * Contains data for a group of assets to import
  */ 
