FBX importer: Add Warning when there is too much UV channels

diff --git a/Engine/Source/Editor/UnrealEd/Private/Fbx/FbxStaticMeshImport.cpp b/Engine/Source/Editor/UnrealEd/Private/Fbx/FbxStaticMeshImport.cpp
--- a/Engine/Source/Editor/UnrealEd/Private/Fbx/FbxStaticMeshImport.cpp
+++ b/Engine/Source/Editor/UnrealEd/Private/Fbx/FbxStaticMeshImport.cpp
@@ -89,7 +89,7 @@ UStaticMesh* UnFbx::FFbxImporter::ImportStaticMesh(UObject* InParent, FbxNode* N
 struct FFBXUVs
 {
 	// constructor
-	FFBXUVs(FbxMesh* Mesh)
+	FFBXUVs(UnFbx::FFbxImporter* FbxImporter, FbxMesh* Mesh)
 		: UniqueUVCount(0)
 	{
 		check(Mesh);
@@ -157,7 +157,7 @@ struct FFBXUVs
 		}
 	}
 
-	void Phase2(FbxMesh* Mesh)
+	void Phase2(UnFbx::FFbxImporter* FbxImporter, FbxMesh* Mesh)
 	{
 		//
 		//	store the UVs in arrays for fast access in the later looping of triangles 
@@ -202,6 +202,12 @@ struct FFBXUVs
 				}
 			}
 		}
+
+		if (UniqueUVCount > MAX_MESH_TEXTURE_COORDS_MD)
+		{
+			FbxImporter->AddTokenizedErrorMessage(FTokenizedMessage::Create(EMessageSeverity::Warning, FText::Format(LOCTEXT("Error_TooMuchUVChannel", "Reached the maximum number of UV Channels for a Static Mesh({0}) - discarding {1} UV Channels"), FText::AsNumber(MAX_MESH_TEXTURE_COORDS_MD), FText::AsNumber(UniqueUVCount-MAX_MESH_TEXTURE_COORDS_MD))), FFbxErrors::Generic_Mesh_TooMuchUVChannels);
+		}
+
 		UniqueUVCount = FMath::Min<int32>(UniqueUVCount, MAX_MESH_TEXTURE_COORDS_MD);
 	}
 
@@ -280,7 +286,7 @@ bool UnFbx::FFbxImporter::BuildStaticMeshFromGeometry(FbxNode* Node, UStaticMesh
 		return false;
 	}
 
-	FFBXUVs FBXUVs(Mesh);
+	FFBXUVs FBXUVs(this, Mesh);
 	int32 FBXNamedLightMapCoordinateIndex = FBXUVs.FindLightUVIndex();
 	if (FBXNamedLightMapCoordinateIndex != INDEX_NONE)
 	{
@@ -399,7 +405,7 @@ bool UnFbx::FFbxImporter::BuildStaticMeshFromGeometry(FbxNode* Node, UStaticMesh
 		LayerElementMaterial->GetMappingMode() : FbxLayerElement::eByPolygon;
 
 	//	todo second phase UV, ok to put in first phase?
-	FBXUVs.Phase2(Mesh);
+	FBXUVs.Phase2(this, Mesh);
 
 	//
 	// get the smoothing group layer
diff --git a/Engine/Source/Runtime/Core/Private/Misc/FbxErrors.cpp b/Engine/Source/Runtime/Core/Private/Misc/FbxErrors.cpp
--- a/Engine/Source/Runtime/Core/Private/Misc/FbxErrors.cpp
+++ b/Engine/Source/Runtime/Core/Private/Misc/FbxErrors.cpp
@@ -55,6 +55,8 @@ FName FFbxErrors::Generic_Mesh_TangentsComputeError(TEXT("Generic_Mesh_TangentsC
 
 FName FFbxErrors::Generic_Mesh_NoReductionModuleAvailable(TEXT("Generic_Mesh_NoReductionModuleAvailable"));
 
+FName FFbxErrors::Generic_Mesh_TooMuchUVChannels(TEXT("Generic_Mesh_TooMuchUVChannels"));
+
 /** Static Mesh **/
 FName FFbxErrors::StaticMesh_TooManyMaterials(TEXT("StaticMesh_TooManyMaterials"));
 
diff --git a/Engine/Source/Runtime/Core/Public/Misc/FbxErrors.h b/Engine/Source/Runtime/Core/Public/Misc/FbxErrors.h
--- a/Engine/Source/Runtime/Core/Public/Misc/FbxErrors.h
+++ b/Engine/Source/Runtime/Core/Public/Misc/FbxErrors.h
@@ -67,6 +67,8 @@ struct CORE_API FFbxErrors
 
 	static FName Generic_Mesh_NoReductionModuleAvailable;
 
+	static FName Generic_Mesh_TooMuchUVChannels;
+
 	/** Static Mesh **/
 	static FName StaticMesh_TooManyMaterials;
 
