Engine and Editor subsystems
- Dynamic subsystems for fixed module startup wrt subsystem initialization
- Python Gettter exposed for both eg: import_subsystem = unreal.get_editor_subsystem(unreal.ImportSubsystem)
- Import Subsystem with python exposed import events

diff --git a/Engine/Plugins/2D/Paper2D/Source/Paper2DEditor/Private/Paper2DEditorModule.cpp b/Engine/Plugins/2D/Paper2D/Source/Paper2DEditor/Private/Paper2DEditorModule.cpp
--- a/Engine/Plugins/2D/Paper2D/Source/Paper2DEditor/Private/Paper2DEditorModule.cpp
+++ b/Engine/Plugins/2D/Paper2D/Source/Paper2DEditor/Private/Paper2DEditorModule.cpp
@@ -122,6 +122,11 @@ private:
 
 public:
 	virtual void StartupModule() override
+	{
+		FCoreDelegates::OnPostEngineInit.AddRaw(this, &FPaper2DEditor::OnPostEngineInit);
+	}
+
+	void OnPostEngineInit()
 	{
 		SpriteEditor_MenuExtensibilityManager = MakeShareable(new FExtensibilityManager);
 		SpriteEditor_ToolBarExtensibilityManager = MakeShareable(new FExtensibilityManager);
@@ -176,7 +181,7 @@ public:
 		OnPropertyChangedDelegateHandle = FCoreUObjectDelegates::OnObjectPropertyChanged.AddRaw(this, &FPaper2DEditor::OnPropertyChanged);
 
 		// Register to be notified when an asset is reimported
-		OnAssetReimportDelegateHandle = FEditorDelegates::OnAssetReimport.AddRaw(this, &FPaper2DEditor::OnObjectReimported);
+		OnAssetReimportDelegateHandle = GEditor->GetEditorSubsystem<UImportSubsystem>()->OnAssetReimport.AddRaw(this, &FPaper2DEditor::OnObjectReimported);
 
 		// Register the thumbnail renderers
 		UThumbnailManager::Get().RegisterCustomRenderer(UPaperSprite::StaticClass(), UPaperSpriteThumbnailRenderer::StaticClass());
@@ -216,6 +221,8 @@ public:
 
 	virtual void ShutdownModule() override
 	{
+		FCoreDelegates::OnPostEngineInit.RemoveAll(this);
+
 		SpriteEditor_MenuExtensibilityManager.Reset();
 		SpriteEditor_ToolBarExtensibilityManager.Reset();
 
@@ -254,7 +261,7 @@ public:
 			FCoreUObjectDelegates::OnObjectPropertyChanged.Remove(OnPropertyChangedDelegateHandle);
 
 			// Unregister the asset reimport handler
-			FEditorDelegates::OnAssetReimport.Remove(OnAssetReimportDelegateHandle);
+			GEditor->GetEditorSubsystem<UImportSubsystem>()->OnAssetReimport.Remove(OnAssetReimportDelegateHandle);
 		}
 
 		// Unregister the details customization
diff --git a/Engine/Plugins/2D/Paper2D/Source/PaperSpriteSheetImporter/Private/PaperSpriteSheetImportFactory.cpp b/Engine/Plugins/2D/Paper2D/Source/PaperSpriteSheetImporter/Private/PaperSpriteSheetImportFactory.cpp
--- a/Engine/Plugins/2D/Paper2D/Source/PaperSpriteSheetImporter/Private/PaperSpriteSheetImportFactory.cpp
+++ b/Engine/Plugins/2D/Paper2D/Source/PaperSpriteSheetImporter/Private/PaperSpriteSheetImportFactory.cpp
@@ -47,7 +47,7 @@ UObject* UPaperSpriteSheetImportFactory::FactoryCreateText(UClass* InClass, UObj
 {
 	Flags |= RF_Transactional;
 
-	FEditorDelegates::OnAssetPreImport.Broadcast(this, InClass, InParent, InName, Type);
+	GEditor->GetEditorSubsystem<UImportSubsystem>()->BroadcastAssetPreImport(this, InClass, InParent, InName, Type);
 
 	FAssetToolsModule& AssetToolsModule = FModuleManager::GetModuleChecked<FAssetToolsModule>("AssetTools");
 
@@ -77,7 +77,7 @@ UObject* UPaperSpriteSheetImportFactory::FactoryCreateText(UClass* InClass, UObj
 		}
 	}
 	
-	FEditorDelegates::OnAssetPostImport.Broadcast(this, Result);
+	GEditor->GetEditorSubsystem<UImportSubsystem>()->BroadcastAssetPostImport(this, Result);
 
 	// Reset the importer to ensure that no leftover data can contaminate future imports.
 	Importer = FPaperJsonSpriteSheetImporter();
diff --git a/Engine/Plugins/2D/Paper2D/Source/PaperTiledImporter/Private/PaperTiledImporterFactory.cpp b/Engine/Plugins/2D/Paper2D/Source/PaperTiledImporter/Private/PaperTiledImporterFactory.cpp
--- a/Engine/Plugins/2D/Paper2D/Source/PaperTiledImporter/Private/PaperTiledImporterFactory.cpp
+++ b/Engine/Plugins/2D/Paper2D/Source/PaperTiledImporter/Private/PaperTiledImporterFactory.cpp
@@ -132,7 +132,7 @@ UObject* UPaperTiledImporterFactory::FactoryCreateText(UClass* InClass, UObject*
 {
 	Flags |= RF_Transactional;
 
-	FEditorDelegates::OnAssetPreImport.Broadcast(this, InClass, InParent, InName, Type);
+	GEditor->GetEditorSubsystem<UImportSubsystem>()->BroadcastAssetPreImport(this, InClass, InParent, InName, Type);
 
  	FAssetToolsModule& AssetToolsModule = FModuleManager::GetModuleChecked<FAssetToolsModule>("AssetTools");
  
@@ -285,7 +285,7 @@ UObject* UPaperTiledImporterFactory::FactoryCreateText(UClass* InClass, UObject*
 		//GlobalInfo.TileSets
 	}
 
-	FEditorDelegates::OnAssetPostImport.Broadcast(this, Result);
+	GEditor->GetEditorSubsystem<UImportSubsystem>()->BroadcastAssetPostImport(this, Result);
 
 	return Result;
 }
diff --git a/Engine/Plugins/2D/Paper2D/Source/SpriterImporter/Private/SpriterImporterFactory.cpp b/Engine/Plugins/2D/Paper2D/Source/SpriterImporter/Private/SpriterImporterFactory.cpp
--- a/Engine/Plugins/2D/Paper2D/Source/SpriterImporter/Private/SpriterImporterFactory.cpp
+++ b/Engine/Plugins/2D/Paper2D/Source/SpriterImporter/Private/SpriterImporterFactory.cpp
@@ -227,7 +227,7 @@ UObject* USpriterImporterFactory::FactoryCreateText(UClass* InClass, UObject* In
 {
 	Flags |= RF_Transactional;
 
-	FEditorDelegates::OnAssetPreImport.Broadcast(this, InClass, InParent, InName, Type);
+	GEditor->GetEditorSubsystem<UImportSubsystem>()->BroadcastAssetPreImport(this, InClass, InParent, InName, Type);
 
  	FAssetToolsModule& AssetToolsModule = FModuleManager::GetModuleChecked<FAssetToolsModule>("AssetTools");
  
@@ -569,7 +569,7 @@ UObject* USpriterImporterFactory::FactoryCreateText(UClass* InClass, UObject* In
 // 		ImportData->SourceFileTimestamp = IFileManager::Get().GetTimeStamp(*CurrentFilename).ToString();
 	}
 
-	FEditorDelegates::OnAssetPostImport.Broadcast(this, Result);
+	GEditor->GetEditorSubsystem<UImportSubsystem>()->BroadcastAssetPostImport(this, Result);
 
 	return Result;
 }
diff --git a/Engine/Plugins/Editor/GLTFImporter/Source/GLTFImporter/Private/GLTFImportFactory.cpp b/Engine/Plugins/Editor/GLTFImporter/Source/GLTFImporter/Private/GLTFImportFactory.cpp
--- a/Engine/Plugins/Editor/GLTFImporter/Source/GLTFImporter/Private/GLTFImportFactory.cpp
+++ b/Engine/Plugins/Editor/GLTFImporter/Source/GLTFImporter/Private/GLTFImportFactory.cpp
@@ -171,7 +171,7 @@ UObject* UGLTFImportFactory::FactoryCreateFile(UClass* InClass, UObject* InParen
 {
 	UStaticMesh* StaticMesh = nullptr;
 
-	FEditorDelegates::OnAssetPreImport.Broadcast(this, InClass, InParent, InName, Parms);
+	GEditor->GetEditorSubsystem<UImportSubsystem>()->BroadcastAssetPreImport(this, InClass, InParent, InName, Parms);
 
 	Warn->Log(Filename);
 	const FString Extension = FPaths::GetExtension(Filename);
@@ -232,7 +232,7 @@ UObject* UGLTFImportFactory::FactoryCreateFile(UClass* InClass, UObject* InParen
 	}
 
 	// Is this necessary? Should call for all new assets, not just the first StaticMesh?
-	FEditorDelegates::OnAssetPostImport.Broadcast(this, StaticMesh);
+	GEditor->GetEditorSubsystem<UImportSubsystem>()->BroadcastAssetPostImport(this, StaticMesh);
 
 	return StaticMesh;
 }
diff --git a/Engine/Plugins/Editor/GameplayTagsEditor/Source/GameplayTagsEditor/Private/GameplayTagsEditorModule.cpp b/Engine/Plugins/Editor/GameplayTagsEditor/Source/GameplayTagsEditor/Private/GameplayTagsEditorModule.cpp
--- a/Engine/Plugins/Editor/GameplayTagsEditor/Source/GameplayTagsEditor/Private/GameplayTagsEditorModule.cpp
+++ b/Engine/Plugins/Editor/GameplayTagsEditor/Source/GameplayTagsEditor/Private/GameplayTagsEditorModule.cpp
@@ -45,6 +45,11 @@ public:
 	// IModuleInterface
 
 	virtual void StartupModule() override
+	{
+		FCoreDelegates::OnPostEngineInit.AddRaw(this, &FGameplayTagsEditorModule::OnPostEngineInit);
+	}
+	
+	void OnPostEngineInit()
 	{
 		// Register the details customizer
 		{
@@ -84,7 +89,7 @@ public:
 		// Hook into notifications for object re-imports so that the gameplay tag tree can be reconstructed if the table changes
 		if (GIsEditor)
 		{
-			FEditorDelegates::OnAssetPostImport.AddRaw(this, &FGameplayTagsEditorModule::OnObjectReimported);
+			GEditor->GetEditorSubsystem<UImportSubsystem>()->OnAssetPostImport.AddRaw(this, &FGameplayTagsEditorModule::OnObjectReimported);
 			FEditorDelegates::OnEditAssetIdentifiers.AddRaw(this, &FGameplayTagsEditorModule::OnEditGameplayTag);
 			IGameplayTagsModule::OnTagSettingsChanged.AddRaw(this, &FGameplayTagsEditorModule::OnEditorSettingsChanged);
 			UPackage::PackageSavedEvent.AddRaw(this, &FGameplayTagsEditorModule::OnPackageSaved);
@@ -104,6 +109,8 @@ public:
 
 	virtual void ShutdownModule() override
 	{
+		FCoreDelegates::OnPostEngineInit.RemoveAll(this);
+
 		// This function may be called during shutdown to clean up your module.  For modules that support dynamic reloading,
 		// we call this function before unloading the module.
 
@@ -113,7 +120,10 @@ public:
 			SettingsModule->UnregisterSettings("Project", "Project", "GameplayTags Developer");
 		}
 
-		FEditorDelegates::OnAssetPostImport.RemoveAll(this);
+		if (GEditor)
+		{
+			GEditor->GetEditorSubsystem<UImportSubsystem>()->OnAssetPostImport.RemoveAll(this);
+		}
 		FEditorDelegates::OnEditAssetIdentifiers.RemoveAll(this);
 		IGameplayTagsModule::OnTagSettingsChanged.RemoveAll(this);
 		UPackage::PackageSavedEvent.RemoveAll(this);
diff --git a/Engine/Plugins/Editor/SpeedTreeImporter/Source/SpeedTreeImporter/Private/SpeedTreeImportFactory.cpp b/Engine/Plugins/Editor/SpeedTreeImporter/Source/SpeedTreeImporter/Private/SpeedTreeImportFactory.cpp
--- a/Engine/Plugins/Editor/SpeedTreeImporter/Source/SpeedTreeImporter/Private/SpeedTreeImportFactory.cpp
+++ b/Engine/Plugins/Editor/SpeedTreeImporter/Source/SpeedTreeImporter/Private/SpeedTreeImportFactory.cpp
@@ -1492,7 +1492,7 @@ FVertexInstanceID ProcessTriangleCorner(
 
 UObject* USpeedTreeImportFactory::FactoryCreateBinary7(UClass* InClass, UObject* InParent, FName InName, EObjectFlags Flags, UObject* Context, const TCHAR* Type, const uint8*& Buffer, const uint8* BufferEnd, FFeedbackContext* Warn, bool& bOutOperationCanceled)
 {
-	FEditorDelegates::OnAssetPreImport.Broadcast(this, InClass, InParent, InName, Type);
+	GEditor->GetEditorSubsystem<UImportSubsystem>()->BroadcastAssetPreImport(this, InClass, InParent, InName, Type);
 
 	TSharedPtr<SWindow> ParentWindow;
 	// Check if the main frame is loaded.  When using the old main frame it may not be.
@@ -2055,14 +2055,14 @@ UObject* USpeedTreeImportFactory::FactoryCreateBinary7(UClass* InClass, UObject*
 		bOutOperationCanceled = true;
 	}
 
-	FEditorDelegates::OnAssetPostImport.Broadcast(this, StaticMesh);
+	GEditor->GetEditorSubsystem<UImportSubsystem>()->BroadcastAssetPostImport(this, StaticMesh);
 
 	return StaticMesh;
 }
 
 UObject* USpeedTreeImportFactory::FactoryCreateBinary8(UClass* InClass, UObject* InParent, FName InName, EObjectFlags Flags, UObject* Context, const TCHAR* Type, const uint8*& Buffer, const uint8* BufferEnd, FFeedbackContext* Warn, bool& bOutOperationCanceled)
 {
-	FEditorDelegates::OnAssetPreImport.Broadcast(this, InClass, InParent, InName, Type);
+	GEditor->GetEditorSubsystem<UImportSubsystem>()->BroadcastAssetPreImport(this, InClass, InParent, InName, Type);
 
 	TSharedPtr<SWindow> ParentWindow;
 	// Check if the main frame is loaded.  When using the old main frame it may not be.
@@ -2442,7 +2442,7 @@ UObject* USpeedTreeImportFactory::FactoryCreateBinary8(UClass* InClass, UObject*
 		}
 	}
 
-	FEditorDelegates::OnAssetPostImport.Broadcast(this, StaticMesh);
+	GEditor->GetEditorSubsystem<UImportSubsystem>()->BroadcastAssetPostImport(this, StaticMesh);
 
 	return StaticMesh;
 }
diff --git a/Engine/Plugins/Experimental/AlembicImporter/Source/AlembicImporter/Private/AlembicImportFactory.cpp b/Engine/Plugins/Experimental/AlembicImporter/Source/AlembicImporter/Private/AlembicImportFactory.cpp
--- a/Engine/Plugins/Experimental/AlembicImporter/Source/AlembicImporter/Private/AlembicImportFactory.cpp
+++ b/Engine/Plugins/Experimental/AlembicImporter/Source/AlembicImporter/Private/AlembicImportFactory.cpp
@@ -63,7 +63,7 @@ UClass* UAlembicImportFactory::ResolveSupportedClass()
 
 UObject* UAlembicImportFactory::FactoryCreateFile(UClass* InClass, UObject* InParent, FName InName, EObjectFlags Flags, const FString& Filename, const TCHAR* Parms, FFeedbackContext* Warn, bool& bOutOperationCanceled)
 {
-	FEditorDelegates::OnAssetPreImport.Broadcast(this, InClass, InParent, InName, TEXT("ABC"));
+	GEditor->GetEditorSubsystem<UImportSubsystem>()->BroadcastAssetPreImport(this, InClass, InParent, InName, TEXT("ABC"));
 
 	FAbcImporter Importer;
 	EAbcImportError ErrorCode = Importer.OpenAbcFileForImport(Filename);
@@ -72,7 +72,7 @@ UObject* UAlembicImportFactory::FactoryCreateFile(UClass* InClass, UObject* InPa
 	if (ErrorCode != AbcImportError_NoError)
 	{
 		// Failed to read the file info, fail the import
-		FEditorDelegates::OnAssetPostImport.Broadcast(this, nullptr);
+		GEditor->GetEditorSubsystem<UImportSubsystem>()->BroadcastAssetPostImport(this, nullptr);
 		return nullptr;
 	}
 
@@ -96,7 +96,7 @@ UObject* UAlembicImportFactory::FactoryCreateFile(UClass* InClass, UObject* InPa
 	TArray<UObject*> ResultAssets;
 	if (!bOutOperationCanceled)
 	{
-		FEditorDelegates::OnAssetPreImport.Broadcast(this, InClass, InParent, InName, TEXT("ABC"));
+		GEditor->GetEditorSubsystem<UImportSubsystem>()->BroadcastAssetPreImport(this, InClass, InParent, InName, TEXT("ABC"));
 
 			int32 NumThreads = 1;
 			if (FPlatformProcess::SupportsMultithreading())
@@ -110,7 +110,7 @@ UObject* UAlembicImportFactory::FactoryCreateFile(UClass* InClass, UObject* InPa
 			if (ErrorCode != AbcImportError_NoError)
 			{
 				// Failed to read the file info, fail the import
-				FEditorDelegates::OnAssetPostImport.Broadcast(this, nullptr);
+				GEditor->GetEditorSubsystem<UImportSubsystem>()->BroadcastAssetPostImport(this, nullptr);
 				FAbcImportLogger::OutputMessages(PageName);
 				return nullptr;
 			}
@@ -144,7 +144,7 @@ UObject* UAlembicImportFactory::FactoryCreateFile(UClass* InClass, UObject* InPa
 		{
 			if (Object)
 			{
-				FEditorDelegates::OnAssetPostImport.Broadcast(this, Object);
+				GEditor->GetEditorSubsystem<UImportSubsystem>()->BroadcastAssetPostImport(this, Object);
 				Object->MarkPackageDirty();
 				Object->PostEditChange();
 			}
@@ -229,7 +229,7 @@ UObject* UAlembicImportFactory::ImportGeometryCache(FAbcImporter& Importer, UObj
 	else
 	{
 		// Not able to import a static mesh
-		FEditorDelegates::OnAssetPostImport.Broadcast(this, nullptr);
+		GEditor->GetEditorSubsystem<UImportSubsystem>()->BroadcastAssetPostImport(this, nullptr);
 		return nullptr;
 	}
 }
@@ -298,7 +298,7 @@ UObject* UAlembicImportFactory::ImportSkeletalMesh(FAbcImporter& Importer, UObje
 	else
 	{
 		// Not able to import a static mesh
-		FEditorDelegates::OnAssetPostImport.Broadcast(this, nullptr);
+		GEditor->GetEditorSubsystem<UImportSubsystem>()->BroadcastAssetPostImport(this, nullptr);
 		return nullptr;
 	}
 }
diff --git a/Engine/Plugins/Experimental/PythonScriptPlugin/Source/PythonScriptPlugin/Private/PyEditor.cpp b/Engine/Plugins/Experimental/PythonScriptPlugin/Source/PythonScriptPlugin/Private/PyEditor.cpp
--- a/Engine/Plugins/Experimental/PythonScriptPlugin/Source/PythonScriptPlugin/Private/PyEditor.cpp
+++ b/Engine/Plugins/Experimental/PythonScriptPlugin/Source/PythonScriptPlugin/Private/PyEditor.cpp
@@ -147,7 +147,63 @@ PyTypeObject PyScopedEditorTransactionType = InitializePyScopedEditorTransaction
 namespace PyEditor
 {
 
+static PyObject* GetEngineSubsystem(FPyScopedEditorTransaction* InSelf, PyObject* InArgs)
+{
+	PyObject* PyObj = nullptr;
+	if (!PyArg_ParseTuple(InArgs, "O:get_engine_subsystem", &PyObj))
+	{
+		return nullptr;
+	}
+	check(PyObj);
+
+	UClass* Class = nullptr;
+	if (!PyConversion::NativizeClass(PyObj, Class, UEngineSubsystem::StaticClass()))
+	{
+		PyUtil::SetPythonError(PyExc_TypeError, TEXT("get_engine_subsystem"), *FString::Printf(TEXT("Parameter must be a 'Class' not '%s'"), *PyUtil::GetFriendlyTypename(PyObj)));
+		return nullptr;
+	}
+
+	if (GEditor)
+	{
+		if (USubsystem* Subsystem = GEditor->GetEngineSubsystemBase(Class))
+		{
+			return PyConversion::Pythonize(Subsystem);
+		}
+	}
+
+	Py_RETURN_NONE;
+}
+
+static PyObject* GetEditorSubsystem(FPyScopedEditorTransaction* InSelf, PyObject* InArgs)
+{
+	PyObject* PyObj = nullptr;
+	if (!PyArg_ParseTuple(InArgs, "O:get_editor_subsystem", &PyObj))
+	{
+		return nullptr;
+	}
+	check(PyObj);
+
+	UClass* Class = nullptr;
+	if (!PyConversion::NativizeClass(PyObj, Class, UEditorSubsystem::StaticClass()))
+	{
+		PyUtil::SetPythonError(PyExc_TypeError, TEXT("get_editor_subsystem"), *FString::Printf(TEXT("Parameter must be a 'Class' not '%s'"), *PyUtil::GetFriendlyTypename(PyObj)));
+		return nullptr;
+	}
+
+	if (GEditor)
+	{
+		if (USubsystem* Subsystem = GEditor->GetEditorSubsystemBase(Class))
+		{
+			return PyConversion::Pythonize(Subsystem);
+		}
+	}
+
+	Py_RETURN_NONE;
+}
+
 PyMethodDef PyEditorMethods[] = {
+	{ "get_engine_subsystem", PyCFunctionCast(&PyEditor::GetEngineSubsystem), METH_VARARGS, "unreal.get_engine_subsystem() -> subsystem -- returns the requested subsystem could be null" },
+	{ "get_editor_subsystem", PyCFunctionCast(&PyEditor::GetEditorSubsystem), METH_VARARGS, "unreal.get_editor_subsystem() -> subsystem -- returns the requested subsystem could be null" },
 	{ nullptr, nullptr, 0, nullptr }
 };
 
diff --git a/Engine/Plugins/Experimental/PythonScriptPlugin/Source/PythonScriptPlugin/PythonScriptPlugin.Build.cs b/Engine/Plugins/Experimental/PythonScriptPlugin/Source/PythonScriptPlugin/PythonScriptPlugin.Build.cs
--- a/Engine/Plugins/Experimental/PythonScriptPlugin/Source/PythonScriptPlugin/PythonScriptPlugin.Build.cs
+++ b/Engine/Plugins/Experimental/PythonScriptPlugin/Source/PythonScriptPlugin/PythonScriptPlugin.Build.cs
@@ -34,6 +34,7 @@ namespace UnrealBuildTool.Rules
 						"EditorStyle",
 						"LevelEditor",
 						"UnrealEd",
+						"EditorSubsystem",
 					}
 				);
 			}
diff --git a/Engine/Plugins/FX/HoudiniNiagara/Source/HoudiniNiagaraEditor/Private/HoudiniCSVFactory.cpp b/Engine/Plugins/FX/HoudiniNiagara/Source/HoudiniNiagaraEditor/Private/HoudiniCSVFactory.cpp
--- a/Engine/Plugins/FX/HoudiniNiagara/Source/HoudiniNiagaraEditor/Private/HoudiniCSVFactory.cpp
+++ b/Engine/Plugins/FX/HoudiniNiagara/Source/HoudiniNiagaraEditor/Private/HoudiniCSVFactory.cpp
@@ -88,7 +88,7 @@ UHoudiniCSVFactory::FactoryCreateFile(UClass* InClass, UObject* InParent, FName
 	AssetImportData->Update(UFactory::GetCurrentFilename());
 
 	// Broadcast notification that the new asset has been imported.
-	FEditorDelegates::OnAssetPostImport.Broadcast(this, NewHoudiniCSVObject);
+	GEditor->GetEditorSubsystem<UImportSubsystem>()->BroadcastAssetPostImport(this, NewHoudiniCSVObject);
 
 	return NewHoudiniCSVObject;
 }
diff --git a/Engine/Plugins/Importers/USDImporter/Source/USDImporter/Private/USDSceneImportFactory.cpp b/Engine/Plugins/Importers/USDImporter/Source/USDImporter/Private/USDSceneImportFactory.cpp
--- a/Engine/Plugins/Importers/USDImporter/Source/USDImporter/Private/USDSceneImportFactory.cpp
+++ b/Engine/Plugins/Importers/USDImporter/Source/USDImporter/Private/USDSceneImportFactory.cpp
@@ -96,7 +96,7 @@ UObject* UUSDSceneImportFactory::FactoryCreateFile(UClass* InClass, UObject* InP
 			}
 		}
 
-		FEditorDelegates::OnAssetPostImport.Broadcast(this, ImportContext.World);
+		GEditor->GetEditorSubsystem<UImportSubsystem>()->BroadcastAssetPostImport(this, ImportContext.World);
 
 		GEditor->BroadcastLevelActorListChanged();
 
diff --git a/Engine/Plugins/Media/MediaFrameworkUtilities/Source/MediaFrameworkUtilitiesEditor/Private/AssetEditor/MediaBundleEditorToolkit.cpp b/Engine/Plugins/Media/MediaFrameworkUtilities/Source/MediaFrameworkUtilitiesEditor/Private/AssetEditor/MediaBundleEditorToolkit.cpp
--- a/Engine/Plugins/Media/MediaFrameworkUtilities/Source/MediaFrameworkUtilitiesEditor/Private/AssetEditor/MediaBundleEditorToolkit.cpp
+++ b/Engine/Plugins/Media/MediaFrameworkUtilities/Source/MediaFrameworkUtilitiesEditor/Private/AssetEditor/MediaBundleEditorToolkit.cpp
@@ -34,7 +34,7 @@ TSharedRef<FMediaBundleEditorToolkit> FMediaBundleEditorToolkit::CreateEditor(co
 
 void FMediaBundleEditorToolkit::InitMediaBundleEditor(const EToolkitMode::Type Mode, const TSharedPtr< class IToolkitHost >& InitToolkitHost, UMediaBundle* InMediaBundle)
 {
-	FEditorDelegates::OnAssetPostImport.AddRaw(this, &FMediaBundleEditorToolkit::HandleAssetPostImport);
+	GEditor->GetEditorSubsystem<UImportSubsystem>()->OnAssetPostImport.AddRaw(this, &FMediaBundleEditorToolkit::HandleAssetPostImport);
 
 	const bool bIsUpdatable = false;
 	const bool bAllowFavorites = true;
@@ -103,7 +103,7 @@ void FMediaBundleEditorToolkit::InitMediaBundleEditor(const EToolkitMode::Type M
 
 FMediaBundleEditorToolkit::~FMediaBundleEditorToolkit()
 {
-	FEditorDelegates::OnAssetPostImport.RemoveAll(this);
+	GEditor->GetEditorSubsystem<UImportSubsystem>()->OnAssetPostImport.RemoveAll(this);
 }
 
 FName FMediaBundleEditorToolkit::GetToolkitFName() const
diff --git a/Engine/Plugins/Media/TimecodeSynchronizer/Source/TimecodeSynchronizerEditor/Private/AssetEditor/TimecodeSynchronizerEditorToolkit.cpp b/Engine/Plugins/Media/TimecodeSynchronizer/Source/TimecodeSynchronizerEditor/Private/AssetEditor/TimecodeSynchronizerEditorToolkit.cpp
--- a/Engine/Plugins/Media/TimecodeSynchronizer/Source/TimecodeSynchronizerEditor/Private/AssetEditor/TimecodeSynchronizerEditorToolkit.cpp
+++ b/Engine/Plugins/Media/TimecodeSynchronizer/Source/TimecodeSynchronizerEditor/Private/AssetEditor/TimecodeSynchronizerEditorToolkit.cpp
@@ -38,7 +38,7 @@ TSharedRef<FTimecodeSynchronizerEditorToolkit> FTimecodeSynchronizerEditorToolki
 
 void FTimecodeSynchronizerEditorToolkit::InitTimecodeSynchronizerEditor(const EToolkitMode::Type Mode, const TSharedPtr< class IToolkitHost >& InitToolkitHost, UTimecodeSynchronizer* InTimecodeSynchronizer)
 {
-	FEditorDelegates::OnAssetPostImport.AddRaw(this, &FTimecodeSynchronizerEditorToolkit::HandleAssetPostImport);
+	GEditor->GetEditorSubsystem<UImportSubsystem>()->OnAssetPostImport.AddRaw(this, &FTimecodeSynchronizerEditorToolkit::HandleAssetPostImport);
 
 	const bool bIsUpdatable = false;
 	const bool bAllowFavorites = true;
@@ -126,7 +126,7 @@ void FTimecodeSynchronizerEditorToolkit::InitTimecodeSynchronizerEditor(const ET
 
 FTimecodeSynchronizerEditorToolkit::~FTimecodeSynchronizerEditorToolkit()
 {
-	FEditorDelegates::OnAssetPostImport.RemoveAll(this);
+	GEditor->GetEditorSubsystem<UImportSubsystem>()->OnAssetPostImport.RemoveAll(this);
 	if (UTimecodeSynchronizer* Asset = GetTimecodeSynchronizer())
 	{
 		Asset->OnSynchronizationEvent().RemoveAll(this);
diff --git a/Engine/Plugins/Runtime/ApexDestruction/Source/ApexDestructionEditor/Private/ActorFactoryDestructible.cpp b/Engine/Plugins/Runtime/ApexDestruction/Source/ApexDestructionEditor/Private/ActorFactoryDestructible.cpp
--- a/Engine/Plugins/Runtime/ApexDestruction/Source/ApexDestructionEditor/Private/ActorFactoryDestructible.cpp
+++ b/Engine/Plugins/Runtime/ApexDestruction/Source/ApexDestructionEditor/Private/ActorFactoryDestructible.cpp
@@ -160,7 +160,7 @@ UObject* UDestructibleMeshFactory::FactoryCreateBinary
 	FFeedbackContext*	Warn
 )
 {
-	FEditorDelegates::OnAssetPreImport.Broadcast(this, Class, InParent, Name, FileType);
+	GEditor->GetEditorSubsystem<UImportSubsystem>()->BroadcastAssetPreImport(this, Class, InParent, Name, FileType);
 
 	// The return value
 	UDestructibleMesh* DestructibleMesh = nullptr;
@@ -173,7 +173,7 @@ UObject* UDestructibleMeshFactory::FactoryCreateBinary
 		DestructibleMesh = ImportDestructibleMeshFromApexDestructibleAsset(InParent, *ApexDestructibleAsset, Name, Flags, nullptr);
 		if (DestructibleMesh != nullptr)
 		{
-			FEditorDelegates::OnAssetPostImport.Broadcast(this, DestructibleMesh);
+			GEditor->GetEditorSubsystem<UImportSubsystem>()->BroadcastAssetPostImport(this, DestructibleMesh);
 
 			// Success
 			DestructibleMesh->PostEditChange();
diff --git a/Engine/Plugins/Runtime/ApexDestruction/Source/ApexDestructionEditor/Private/DestructibleMeshEditor.cpp b/Engine/Plugins/Runtime/ApexDestruction/Source/ApexDestructionEditor/Private/DestructibleMeshEditor.cpp
--- a/Engine/Plugins/Runtime/ApexDestruction/Source/ApexDestructionEditor/Private/DestructibleMeshEditor.cpp
+++ b/Engine/Plugins/Runtime/ApexDestruction/Source/ApexDestructionEditor/Private/DestructibleMeshEditor.cpp
@@ -43,8 +43,6 @@ void FDestructibleMeshEditorCommands::RegisterCommands()
 FDestructibleMeshEditor::~FDestructibleMeshEditor()
 {
 	FReimportManager::Instance()->OnPostReimport().RemoveAll(this);
-
-	GEditor->OnObjectReimported().RemoveAll(this);
 }
 
 void FDestructibleMeshEditor::RegisterTabSpawners(const TSharedRef<class FTabManager>& InTabManager)
diff --git a/Engine/Plugins/Runtime/AudioCapture/Source/AudioCaptureEditor/Private/AudioRecordingManager.cpp b/Engine/Plugins/Runtime/AudioCapture/Source/AudioCaptureEditor/Private/AudioRecordingManager.cpp
--- a/Engine/Plugins/Runtime/AudioCapture/Source/AudioCaptureEditor/Private/AudioRecordingManager.cpp
+++ b/Engine/Plugins/Runtime/AudioCapture/Source/AudioCaptureEditor/Private/AudioRecordingManager.cpp
@@ -431,7 +431,7 @@ void FAudioRecordingManager::StopRecording(TArray<USoundWave*>& OutSoundWaves)
 
 					if (bCreatedPackage)
 					{
-						//FEditorDelegates::OnAssetPostImport.Broadcast(this, NewSoundWave);
+						//GEditor->GetEditorSubsystem<UImportSubsystem>()->BroadcastAssetPostImport(this, NewSoundWave);
 
 						FAssetRegistryModule::AssetCreated(NewSoundWave);
 						NewSoundWave->MarkPackageDirty();
diff --git a/Engine/Plugins/Runtime/GameplayAbilities/Source/GameplayAbilitiesEditor/Private/GameplayAbilitiesEditor.cpp b/Engine/Plugins/Runtime/GameplayAbilities/Source/GameplayAbilitiesEditor/Private/GameplayAbilitiesEditor.cpp
--- a/Engine/Plugins/Runtime/GameplayAbilities/Source/GameplayAbilitiesEditor/Private/GameplayAbilitiesEditor.cpp
+++ b/Engine/Plugins/Runtime/GameplayAbilities/Source/GameplayAbilitiesEditor/Private/GameplayAbilitiesEditor.cpp
@@ -26,9 +26,6 @@ FGameplayAbilitiesEditor::FGameplayAbilitiesEditor()
 
 FGameplayAbilitiesEditor::~FGameplayAbilitiesEditor()
 {
-	FEditorDelegates::OnAssetPostImport.RemoveAll(this);
-	FReimportManager::Instance()->OnPostReimport().RemoveAll(this);
-	
 	// NOTE: Any tabs that we still have hanging out when destroyed will be cleaned up by FBaseToolkit's destructor
 }
 
diff --git a/Engine/Plugins/Runtime/SoundMod/Source/SoundModImporter/Private/SoundModImporterFactory.cpp b/Engine/Plugins/Runtime/SoundMod/Source/SoundModImporter/Private/SoundModImporterFactory.cpp
--- a/Engine/Plugins/Runtime/SoundMod/Source/SoundModImporter/Private/SoundModImporterFactory.cpp
+++ b/Engine/Plugins/Runtime/SoundMod/Source/SoundModImporter/Private/SoundModImporterFactory.cpp
@@ -77,7 +77,7 @@ UObject* USoundModImporterFactory::FactoryCreateBinary
 			}
 			default:
 			{
-				FEditorDelegates::OnAssetPostImport.Broadcast(this, NULL);
+				GEditor->GetEditorSubsystem<UImportSubsystem>()->BroadcastAssetPostImport(this, NULL);
 				return NULL;
 			}
 		}
@@ -114,7 +114,7 @@ UObject* USoundModImporterFactory::FactoryCreateBinary
 	FMemory::Memcpy(LockedData, Buffer, BufferEnd - Buffer);
 	Sound->RawData.Unlock();
 
-	FEditorDelegates::OnAssetPostImport.Broadcast(this, Sound);
+	GEditor->GetEditorSubsystem<UImportSubsystem>()->BroadcastAssetPostImport(this, Sound);
 
 	for (int32 ComponentIndex = 0; ComponentIndex < ComponentsToRestart.Num(); ++ComponentIndex)
 	{
diff --git a/Engine/Plugins/ScriptPlugin/Source/ScriptEditorPlugin/Private/ScriptFactory.cpp b/Engine/Plugins/ScriptPlugin/Source/ScriptEditorPlugin/Private/ScriptFactory.cpp
--- a/Engine/Plugins/ScriptPlugin/Source/ScriptEditorPlugin/Private/ScriptFactory.cpp
+++ b/Engine/Plugins/ScriptPlugin/Source/ScriptEditorPlugin/Private/ScriptFactory.cpp
@@ -87,7 +87,7 @@ UObject* UScriptFactory::FactoryCreateText(UClass* InClass, UObject* InParent, F
 		// Need to make sure we compile with the new source code
 		FKismetEditorUtilities::CompileBlueprint(NewBlueprint);
 
-		FEditorDelegates::OnAssetPostImport.Broadcast(this, NewBlueprint);
+		GEditor->GetEditorSubsystem<UImportSubsystem>()->BroadcastAssetPostImport(this, NewBlueprint);
 	}
 	else
 	{
diff --git a/Engine/Source/Editor/AnimationBlueprintEditor/Private/AnimationBlueprintEditor.cpp b/Engine/Source/Editor/AnimationBlueprintEditor/Private/AnimationBlueprintEditor.cpp
--- a/Engine/Source/Editor/AnimationBlueprintEditor/Private/AnimationBlueprintEditor.cpp
+++ b/Engine/Source/Editor/AnimationBlueprintEditor/Private/AnimationBlueprintEditor.cpp
@@ -158,7 +158,7 @@ FAnimationBlueprintEditor::~FAnimationBlueprintEditor()
 {
 	GEditor->OnBlueprintPreCompile().RemoveAll(this);
 
-	FEditorDelegates::OnAssetPostImport.RemoveAll(this);
+	GEditor->GetEditorSubsystem<UImportSubsystem>()->OnAssetPostImport.RemoveAll(this);
 	FReimportManager::Instance()->OnPostReimport().RemoveAll(this);
 
 	// NOTE: Any tabs that we still have hanging out when destroyed will be cleaned up by FBaseToolkit's destructor
diff --git a/Engine/Source/Editor/AnimationEditor/Private/AnimationEditor.cpp b/Engine/Source/Editor/AnimationEditor/Private/AnimationEditor.cpp
--- a/Engine/Source/Editor/AnimationEditor/Private/AnimationEditor.cpp
+++ b/Engine/Source/Editor/AnimationEditor/Private/AnimationEditor.cpp
@@ -85,7 +85,7 @@ FAnimationEditor::~FAnimationEditor()
 		Editor->UnregisterForUndo(this);
 	}
 
-	FEditorDelegates::OnAssetPostImport.RemoveAll(this);
+	GEditor->GetEditorSubsystem<UImportSubsystem>()->OnAssetPostImport.RemoveAll(this);
 	FReimportManager::Instance()->OnPostReimport().RemoveAll(this);
 }
 
@@ -107,7 +107,7 @@ void FAnimationEditor::InitAnimationEditor(const EToolkitMode::Type Mode, const
 
 	// Register post import callback to catch animation imports when we have the asset open (we need to reinit)
 	FReimportManager::Instance()->OnPostReimport().AddRaw(this, &FAnimationEditor::HandlePostReimport);
-	FEditorDelegates::OnAssetPostImport.AddRaw(this, &FAnimationEditor::HandlePostImport);
+	GEditor->GetEditorSubsystem<UImportSubsystem>()->OnAssetPostImport.AddRaw(this, &FAnimationEditor::HandlePostImport);
 
 	FPersonaModule& PersonaModule = FModuleManager::LoadModuleChecked<FPersonaModule>("Persona");
 	PersonaToolkit = PersonaModule.CreatePersonaToolkit(InAnimationAsset);
diff --git a/Engine/Source/Editor/AudioEditor/Private/Factories/SoundFactory.cpp b/Engine/Source/Editor/AudioEditor/Private/Factories/SoundFactory.cpp
--- a/Engine/Source/Editor/AudioEditor/Private/Factories/SoundFactory.cpp
+++ b/Engine/Source/Editor/AudioEditor/Private/Factories/SoundFactory.cpp
@@ -118,7 +118,7 @@ UObject* USoundFactory::FactoryCreateBinary
 	FFeedbackContext*	Warn
 	)
 {
-	FEditorDelegates::OnAssetPreImport.Broadcast(this, Class, InParent, Name, FileType);
+	GEditor->GetEditorSubsystem<UImportSubsystem>()->BroadcastAssetPreImport(this, Class, InParent, Name, FileType);
 
 	if (FCString::Stricmp(FileType, TEXT("WAV")) == 0)
 	{
@@ -141,7 +141,7 @@ UObject* USoundFactory::FactoryCreateBinary
 			if (!bCuePathIsValid)
 			{
 				FMessageDialog::Open(EAppMsgType::Ok, FText::Format(NSLOCTEXT("UnrealEd", "Error_ImportFailed_f", "Import failed for {0}: {1}"), FText::FromString(CuePackageName), Reason));
-				FEditorDelegates::OnAssetPostImport.Broadcast(this, nullptr);
+				GEditor->GetEditorSubsystem<UImportSubsystem>()->BroadcastAssetPostImport(this, nullptr);
 				return nullptr;
 			}
 		}
@@ -209,7 +209,7 @@ UObject* USoundFactory::FactoryCreateBinary
 			}
 			default:
 			{
-				FEditorDelegates::OnAssetPostImport.Broadcast(this, nullptr);
+				GEditor->GetEditorSubsystem<UImportSubsystem>()->BroadcastAssetPostImport(this, nullptr);
 				return nullptr;
 			}
 			}
@@ -241,7 +241,7 @@ UObject* USoundFactory::FactoryCreateBinary
 			if ((bIsAmbiX || bIsFuMa) && (int32)*WaveInfo.pChannels != 4)
 			{
 				Warn->Logf(ELogVerbosity::Error, TEXT("Tried to import ambisonics format file but requires exactly 4 channels: '%s'"), *Name.ToString());
-				FEditorDelegates::OnAssetPostImport.Broadcast(this, nullptr);
+				GEditor->GetEditorSubsystem<UImportSubsystem>()->BroadcastAssetPostImport(this, nullptr);
 				return nullptr;
 			}
 
@@ -249,14 +249,14 @@ UObject* USoundFactory::FactoryCreateBinary
 			{
 				WaveInfo.ReportImportFailure();
 				Warn->Logf(ELogVerbosity::Error, TEXT("Currently, only 16 bit WAV files are supported (%s)."), *Name.ToString());
-				FEditorDelegates::OnAssetPostImport.Broadcast(this, nullptr);
+				GEditor->GetEditorSubsystem<UImportSubsystem>()->BroadcastAssetPostImport(this, nullptr);
 				return nullptr;
 			}
 		}
 		else
 		{
 			Warn->Logf(ELogVerbosity::Error, TEXT("Unable to read wave file '%s' - \"%s\""), *Name.ToString(), *ErrorMessage);
-			FEditorDelegates::OnAssetPostImport.Broadcast(this, nullptr);
+			GEditor->GetEditorSubsystem<UImportSubsystem>()->BroadcastAssetPostImport(this, nullptr);
 			return nullptr;
 		}
 
@@ -342,7 +342,7 @@ UObject* USoundFactory::FactoryCreateBinary
 			else
 			{
 				Warn->Logf(ELogVerbosity::Error, TEXT("Wave file '%s' has unsupported number of channels %d"), *Name.ToString(), ChannelCount);
-				FEditorDelegates::OnAssetPostImport.Broadcast(this, nullptr);
+				GEditor->GetEditorSubsystem<UImportSubsystem>()->BroadcastAssetPostImport(this, nullptr);
 				return nullptr;
 			}
 
@@ -412,7 +412,7 @@ UObject* USoundFactory::FactoryCreateBinary
 		Sound->SetSampleRate(*WaveInfo.pSamplesPerSec);
 		Sound->NumChannels = ChannelCount;
 
-		FEditorDelegates::OnAssetPostImport.Broadcast(this, Sound);
+		GEditor->GetEditorSubsystem<UImportSubsystem>()->BroadcastAssetPostImport(this, Sound);
 
 		if (ExistingSound && bUseExistingSettings)
 		{
@@ -439,7 +439,7 @@ UObject* USoundFactory::FactoryCreateBinary
 	{
 		// Unrecognized.
 		Warn->Logf(ELogVerbosity::Error, TEXT("Unrecognized sound format '%s' in %s"), FileType, *Name.ToString());
-		FEditorDelegates::OnAssetPostImport.Broadcast(this, nullptr);
+		GEditor->GetEditorSubsystem<UImportSubsystem>()->BroadcastAssetPostImport(this, nullptr);
 		return nullptr;
 	}
 }
diff --git a/Engine/Source/Editor/AudioEditor/Private/Factories/SoundSurroundFactory.cpp b/Engine/Source/Editor/AudioEditor/Private/Factories/SoundSurroundFactory.cpp
--- a/Engine/Source/Editor/AudioEditor/Private/Factories/SoundSurroundFactory.cpp
+++ b/Engine/Source/Editor/AudioEditor/Private/Factories/SoundSurroundFactory.cpp
@@ -52,7 +52,7 @@ UObject* USoundSurroundFactory::FactoryCreateBinary
 	FFeedbackContext*	Warn
 	)
 {
-	FEditorDelegates::OnAssetPreImport.Broadcast(this, Class, InParent, Name, FileType);
+	GEditor->GetEditorSubsystem<UImportSubsystem>()->BroadcastAssetPreImport(this, Class, InParent, Name, FileType);
 
 	int32		SpeakerIndex, i;
 
@@ -76,7 +76,7 @@ UObject* USoundSurroundFactory::FactoryCreateBinary
 		if (SpeakerIndex == SPEAKER_Count)
 		{
 			Warn->Logf(ELogVerbosity::Error, TEXT("Failed to find speaker location; valid extensions are _fl, _fr, _fc, _lf, _sl, _sr, _bl, _br."));
-			FEditorDelegates::OnAssetPostImport.Broadcast(this, nullptr);
+			GEditor->GetEditorSubsystem<UImportSubsystem>()->BroadcastAssetPostImport(this, nullptr);
 			return(nullptr);
 		}
 
@@ -215,7 +215,7 @@ UObject* USoundSurroundFactory::FactoryCreateBinary
 			}
 		}
 
-		FEditorDelegates::OnAssetPostImport.Broadcast(this, Sound);
+		GEditor->GetEditorSubsystem<UImportSubsystem>()->BroadcastAssetPostImport(this, Sound);
 
 		return(Sound);
 	}
@@ -223,7 +223,7 @@ UObject* USoundSurroundFactory::FactoryCreateBinary
 	{
 		// Unrecognized.
 		Warn->Logf(ELogVerbosity::Error, TEXT("Unrecognized sound extension '%s' in %s"), FileType, *Name.ToString());
-		FEditorDelegates::OnAssetPostImport.Broadcast(this, nullptr);
+		GEditor->GetEditorSubsystem<UImportSubsystem>()->BroadcastAssetPostImport(this, nullptr);
 	}
 
 	return(nullptr);
diff --git a/Engine/Source/Editor/EditorSubsystem/EditorSubsystem.Build.cs b/Engine/Source/Editor/EditorSubsystem/EditorSubsystem.Build.cs
new file mode 100644
--- /dev/null
+++ b/Engine/Source/Editor/EditorSubsystem/EditorSubsystem.Build.cs
@@ -0,0 +1,17 @@
+// Copyright 1998-2018 Epic Games, Inc. All Rights Reserved.
+
+using UnrealBuildTool;
+
+public class EditorSubsystem : ModuleRules
+{
+    public EditorSubsystem(ReadOnlyTargetRules Target)
+         : base(Target)
+    {
+        PrivateDependencyModuleNames.AddRange(
+            new string[] {
+                "Core",
+                "CoreUObject",
+                "Engine",
+            });
+    }
+}
diff --git a/Engine/Source/Editor/EditorSubsystem/Private/EditorSubsystem.cpp b/Engine/Source/Editor/EditorSubsystem/Private/EditorSubsystem.cpp
new file mode 100644
--- /dev/null
+++ b/Engine/Source/Editor/EditorSubsystem/Private/EditorSubsystem.cpp
@@ -0,0 +1,9 @@
+// Copyright 1998-2018 Epic Games, Inc. All Rights Reserved.
+
+#include "EditorSubsystem.h"
+
+UEditorSubsystem::UEditorSubsystem()
+	: UDynamicSubsystem()
+{
+
+}
diff --git a/Engine/Source/Editor/EditorSubsystem/Private/EditorSubsystemModule.cpp b/Engine/Source/Editor/EditorSubsystem/Private/EditorSubsystemModule.cpp
new file mode 100644
--- /dev/null
+++ b/Engine/Source/Editor/EditorSubsystem/Private/EditorSubsystemModule.cpp
@@ -0,0 +1,15 @@
+// Copyright 1998-2018 Epic Games, Inc. All Rights Reserved.
+
+#include "EditorSubsystemModule.h"
+#include "Modules/ModuleManager.h"
+
+IMPLEMENT_MODULE(FEditorSubsystemModule, EditorSubsystem);
+
+void FEditorSubsystemModule::StartupModule()
+{
+
+}
+
+void FEditorSubsystemModule::ShutdownModule()
+{
+}
diff --git a/Engine/Source/Editor/EditorSubsystem/Public/EditorSubsystem.h b/Engine/Source/Editor/EditorSubsystem/Public/EditorSubsystem.h
new file mode 100644
--- /dev/null
+++ b/Engine/Source/Editor/EditorSubsystem/Public/EditorSubsystem.h
@@ -0,0 +1,28 @@
+// Copyright 1998-2018 Epic Games, Inc. All Rights Reserved.
+
+#pragma once
+
+#include "Subsystems/Subsystem.h"
+
+#include "EditorSubsystem.generated.h"
+
+/**
+ * UEditorSubsystem
+ * Base class for auto instanced and initialized systems that share the lifetime of the Editor
+ *
+ * UEditorSubsystems are dynamic and will be initialized when the module is loaded if necessary.
+ * This means that after StartupModule() is called on the module containing a subsystem,
+ * the subsystem collection with instantiate and initialize the subsystem automatically.
+ * If the subsystem collection is created post module load then the instances will be created at
+ * collection initialization time.
+ */
+
+UCLASS(Abstract)
+class EDITORSUBSYSTEM_API UEditorSubsystem : public UDynamicSubsystem
+{
+	GENERATED_BODY()
+
+public:
+	UEditorSubsystem();
+};
+
diff --git a/Engine/Source/Editor/EditorSubsystem/Public/EditorSubsystemModule.h b/Engine/Source/Editor/EditorSubsystem/Public/EditorSubsystemModule.h
new file mode 100644
--- /dev/null
+++ b/Engine/Source/Editor/EditorSubsystem/Public/EditorSubsystemModule.h
@@ -0,0 +1,16 @@
+// Copyright 1998-2018 Epic Games, Inc. All Rights Reserved.
+
+#pragma once
+
+#include "CoreMinimal.h"
+#include "Modules/ModuleInterface.h"
+
+class EDITORSUBSYSTEM_API FEditorSubsystemModule : public IModuleInterface
+{
+public:
+	// IModuleInterface
+	virtual void StartupModule() override;
+	virtual void ShutdownModule() override;
+
+private:
+};
diff --git a/Engine/Source/Editor/FontEditor/Private/FontEditor.cpp b/Engine/Source/Editor/FontEditor/Private/FontEditor.cpp
--- a/Engine/Source/Editor/FontEditor/Private/FontEditor.cpp
+++ b/Engine/Source/Editor/FontEditor/Private/FontEditor.cpp
@@ -155,7 +155,7 @@ FFontEditor::~FFontEditor()
 	if (Editor != NULL)
 	{
 		Editor->UnregisterForUndo(this);
-		Editor->OnObjectReimported().RemoveAll(this);
+		Editor->GetEditorSubsystem<UImportSubsystem>()->OnAssetReimport.RemoveAll(this);
 	}
 }
 
@@ -164,7 +164,7 @@ void FFontEditor::InitFontEditor(const EToolkitMode::Type Mode, const TSharedPtr
 	FReimportManager::Instance()->OnPostReimport().AddRaw(this, &FFontEditor::OnPostReimport);
 
 	// Register to be notified when an object is reimported.
-	GEditor->OnObjectReimported().AddSP(this, &FFontEditor::OnObjectReimported);
+	GEditor->GetEditorSubsystem<UImportSubsystem>()->OnAssetReimport.AddSP(this, &FFontEditor::OnObjectReimported);
 
 	FCoreUObjectDelegates::OnObjectPropertyChanged.AddSP(this, &FFontEditor::OnObjectPropertyChanged);
 
diff --git a/Engine/Source/Editor/MeshPaint/Private/IMeshPaintMode.cpp b/Engine/Source/Editor/MeshPaint/Private/IMeshPaintMode.cpp
--- a/Engine/Source/Editor/MeshPaint/Private/IMeshPaintMode.cpp
+++ b/Engine/Source/Editor/MeshPaint/Private/IMeshPaintMode.cpp
@@ -80,7 +80,7 @@ void IMeshPaintEdMode::Enter()
 	FEditorDelegates::PostSaveWorld.AddSP(this, &IMeshPaintEdMode::OnPostSaveWorld);
 
 	// Catch assets if they are about to be (re)imported
-	FEditorDelegates::OnAssetPostImport.AddSP(this, &IMeshPaintEdMode::OnPostImportAsset);
+	GEditor->GetEditorSubsystem<UImportSubsystem>()->OnAssetPostImport.AddSP(this, &IMeshPaintEdMode::OnPostImportAsset);
 	FReimportManager::Instance()->OnPostReimport().AddSP(this, &IMeshPaintEdMode::OnPostReimportAsset);
 
 	FAssetRegistryModule& AssetRegistryModule = FModuleManager::LoadModuleChecked<FAssetRegistryModule>(TEXT("AssetRegistry"));
@@ -172,7 +172,7 @@ void IMeshPaintEdMode::Exit()
 	FAssetRegistryModule& AssetRegistryModule = FModuleManager::LoadModuleChecked<FAssetRegistryModule>(TEXT("AssetRegistry"));
 	AssetRegistryModule.Get().OnAssetRemoved().RemoveAll(this);
 	FReimportManager::Instance()->OnPostReimport().RemoveAll(this);
-	FEditorDelegates::OnAssetPostImport.RemoveAll(this);
+	GEditor->GetEditorSubsystem<UImportSubsystem>()->OnAssetPostImport.RemoveAll(this);
 	FEditorDelegates::PreSaveWorld.RemoveAll(this);
 	FEditorDelegates::PostSaveWorld.RemoveAll(this);
 	GEditor->OnObjectsReplaced().RemoveAll(this);
diff --git a/Engine/Source/Editor/StaticMeshEditor/Private/StaticMeshEditor.cpp b/Engine/Source/Editor/StaticMeshEditor/Private/StaticMeshEditor.cpp
--- a/Engine/Source/Editor/StaticMeshEditor/Private/StaticMeshEditor.cpp
+++ b/Engine/Source/Editor/StaticMeshEditor/Private/StaticMeshEditor.cpp
@@ -139,7 +139,7 @@ FStaticMeshEditor::~FStaticMeshEditor()
 	FReimportManager::Instance()->OnPostReimport().RemoveAll(this);
 
 	GEditor->UnregisterForUndo( this );
-	GEditor->OnObjectReimported().RemoveAll(this);
+	GEditor->GetEditorSubsystem<UImportSubsystem>()->OnAssetReimport.RemoveAll(this);
 }
 
 void FStaticMeshEditor::InitStaticMeshEditor( const EToolkitMode::Type Mode, const TSharedPtr< class IToolkitHost >& InitToolkitHost, UStaticMesh* ObjectToEdit )
@@ -155,7 +155,7 @@ void FStaticMeshEditor::InitStaticMeshEditor( const EToolkitMode::Type Mode, con
 	FStaticMeshEditorCommands::Register();
 
 	// Register to be notified when an object is reimported.
-	GEditor->OnObjectReimported().AddSP(this, &FStaticMeshEditor::OnObjectReimported);
+	GEditor->GetEditorSubsystem<UImportSubsystem>()->OnAssetReimport.AddSP(this, &FStaticMeshEditor::OnObjectReimported);
 
 	BindCommands();
 
diff --git a/Engine/Source/Editor/TextureEditor/Private/TextureEditorToolkit.cpp b/Engine/Source/Editor/TextureEditor/Private/TextureEditorToolkit.cpp
--- a/Engine/Source/Editor/TextureEditor/Private/TextureEditorToolkit.cpp
+++ b/Engine/Source/Editor/TextureEditor/Private/TextureEditorToolkit.cpp
@@ -67,7 +67,7 @@ FTextureEditorToolkit::~FTextureEditorToolkit( )
 {
 	FReimportManager::Instance()->OnPreReimport().RemoveAll(this);
 	FReimportManager::Instance()->OnPostReimport().RemoveAll(this);
-	FEditorDelegates::OnAssetPostImport.RemoveAll(this);
+	GEditor->GetEditorSubsystem<UImportSubsystem>()->OnAssetPostImport.RemoveAll(this);
 
 	GEditor->UnregisterForUndo(this);
 }
@@ -114,7 +114,7 @@ void FTextureEditorToolkit::InitTextureEditor( const EToolkitMode::Type Mode, co
 {
 	FReimportManager::Instance()->OnPreReimport().AddRaw(this, &FTextureEditorToolkit::HandleReimportManagerPreReimport);
 	FReimportManager::Instance()->OnPostReimport().AddRaw(this, &FTextureEditorToolkit::HandleReimportManagerPostReimport);
-	FEditorDelegates::OnAssetPostImport.AddRaw(this, &FTextureEditorToolkit::HandleAssetPostImport);
+	GEditor->GetEditorSubsystem<UImportSubsystem>()->OnAssetPostImport.AddRaw(this, &FTextureEditorToolkit::HandleAssetPostImport);
 
 	Texture = CastChecked<UTexture>(ObjectToEdit);
 
diff --git a/Engine/Source/Editor/UnrealEd/Classes/Editor/EditorEngine.h b/Engine/Source/Editor/UnrealEd/Classes/Editor/EditorEngine.h
--- a/Engine/Source/Editor/UnrealEd/Classes/Editor/EditorEngine.h
+++ b/Engine/Source/Editor/UnrealEd/Classes/Editor/EditorEngine.h
@@ -23,6 +23,10 @@
 #include "Settings/LevelEditorPlaySettings.h"
 #include "Settings/LevelEditorViewportSettings.h"
 #include "Misc/CompilationResult.h"
+
+#include "EditorSubsystem.h"
+#include "Subsystems/SubsystemCollection.h"
+
 #include "EditorEngine.generated.h"
 
 class AMatineeActor;
@@ -712,6 +716,7 @@ public:
 
 	/** Called when an object is reimported. */
 	DECLARE_EVENT_OneParam( UEditorEngine, FObjectReimported, UObject* );
+	UE_DEPRECATED(4.22, "Use the ImportSubsystem instead. GEditor->GetEditorSubsystem<UImportSubsystem>()")
 	FObjectReimported& OnObjectReimported() { return ObjectReimportedEvent; }
 
 	/** Editor-only event triggered before an actor or component is moved, rotated or scaled by an editor system */
@@ -3034,6 +3039,41 @@ private:
 
 	/** Delegate handle for game viewport close requests in PIE sessions. */
 	FDelegateHandle ViewportCloseRequestedDelegateHandle;
+
+public:
+	/**
+	 * Get a Subsystem of specified type
+	 */
+	UEditorSubsystem* GetEditorSubsystemBase(TSubclassOf<UEditorSubsystem> SubsystemClass) const
+	{
+		checkSlow(this != nullptr);
+		return EditorSubsystemCollection.GetSubsystem<UEditorSubsystem>(SubsystemClass);
+	}
+
+	/**
+	 * Get a Subsystem of specified type
+	 */
+	template <typename TSubsystemClass>
+	TSubsystemClass* GetEditorSubsystem() const
+	{
+		checkSlow(this != nullptr);
+		return EditorSubsystemCollection.GetSubsystem<TSubsystemClass>(TSubsystemClass::StaticClass());
+	}
+
+	/**
+	 * Get all Subsystem of specified type, this is only necessary for interfaces that can have multiple implementations instanced at a time.
+	 *
+	 * Do not hold onto this Array reference unless you are sure the lifetime is less than that of UGameInstance
+	 */
+	template <typename TSubsystemClass>
+	const TArray<TSubsystemClass*>& GetEditorSubsystemArray() const
+	{
+		return EditorSubsystemCollection.GetSubsystemArray<TSubsystemClass>(TSubsystemClass::StaticClass());
+	}
+
+private:
+	FSubsystemCollection<UEditorSubsystem> EditorSubsystemCollection;
+
 };
 
 //////////////////////////////////////////////////////////////////////////
diff --git a/Engine/Source/Editor/UnrealEd/Private/Editor.cpp b/Engine/Source/Editor/UnrealEd/Private/Editor.cpp
--- a/Engine/Source/Editor/UnrealEd/Private/Editor.cpp
+++ b/Engine/Source/Editor/UnrealEd/Private/Editor.cpp
@@ -70,6 +70,8 @@
 
 #define LOCTEXT_NAMESPACE "UnrealEd.Editor"
 
+PRAGMA_DISABLE_DEPRECATION_WARNINGS
+
 FSimpleMulticastDelegate								FEditorDelegates::NewCurrentLevel;
 FEditorDelegates::FOnMapChanged							FEditorDelegates::MapChange;
 FSimpleMulticastDelegate								FEditorDelegates::LayerChange;
@@ -140,6 +142,8 @@ FEditorDelegates::FOnViewAssetIdentifiers				FEditorDelegates::OnOpenSizeMap;
 FEditorDelegates::FOnViewAssetIdentifiers				FEditorDelegates::OnOpenAssetAudit;
 FEditorDelegates::FOnViewAssetIdentifiers				FEditorDelegates::OnEditAssetIdentifiers;
 
+PRAGMA_ENABLE_DEPRECATION_WARNINGS
+
 /*-----------------------------------------------------------------------------
 	Globals.
 -----------------------------------------------------------------------------*/
diff --git a/Engine/Source/Editor/UnrealEd/Private/EditorEngine.cpp b/Engine/Source/Editor/UnrealEd/Private/EditorEngine.cpp
--- a/Engine/Source/Editor/UnrealEd/Private/EditorEngine.cpp
+++ b/Engine/Source/Editor/UnrealEd/Private/EditorEngine.cpp
@@ -305,6 +305,7 @@ UEditorEngine* GEditor = nullptr;
 
 UEditorEngine::UEditorEngine(const FObjectInitializer& ObjectInitializer)
 	: Super(ObjectInitializer)
+	, EditorSubsystemCollection(this)
 {
 	if (!IsRunningCommandlet() && !IsRunningDedicatedServer())
 	{
@@ -582,6 +583,8 @@ void UEditorEngine::InitEditor(IEngineLoop* InEngineLoop)
 	// Call base.
 	UEngine::Init(InEngineLoop);
 
+	EditorSubsystemCollection.Initialize();
+
 	// Specify "-ForceLauncher" on the command-line to always open the launcher, even in unusual cases.  This is useful for debugging the Launcher startup.
 	const bool bForceLauncherToOpen = FParse::Param(FCommandLine::Get(), TEXT("ForceLauncher"));
 
@@ -1253,7 +1256,7 @@ void UEditorEngine::InitBuilderBrush( UWorld* InWorld )
 void UEditorEngine::BroadcastObjectReimported(UObject* InObject)
 {
 	ObjectReimportedEvent.Broadcast(InObject);
-	FEditorDelegates::OnAssetReimport.Broadcast(InObject);
+	GetEditorSubsystem<UImportSubsystem>()->BroadcastAssetReimport(InObject);
 }
 
 void UEditorEngine::FinishDestroy()
diff --git a/Engine/Source/Editor/UnrealEd/Private/Factories/CSVImportFactory.cpp b/Engine/Source/Editor/UnrealEd/Private/Factories/CSVImportFactory.cpp
--- a/Engine/Source/Editor/UnrealEd/Private/Factories/CSVImportFactory.cpp
+++ b/Engine/Source/Editor/UnrealEd/Private/Factories/CSVImportFactory.cpp
@@ -100,7 +100,7 @@ bool UCSVImportFactory::FactoryCanImport(const FString& Filename)
 
 UObject* UCSVImportFactory::FactoryCreateText(UClass* InClass, UObject* InParent, FName InName, EObjectFlags Flags, UObject* Context, const TCHAR* Type, const TCHAR*& Buffer, const TCHAR* BufferEnd, FFeedbackContext* Warn, bool& bOutOperationCanceled)
 {
-	FEditorDelegates::OnAssetPreImport.Broadcast(this, InClass, InParent, InName, Type);
+	GEditor->GetEditorSubsystem<UImportSubsystem>()->BroadcastAssetPreImport(this, InClass, InParent, InName, Type);
 
 	bOutOperationCanceled = false;
 
@@ -300,7 +300,7 @@ UObject* UCSVImportFactory::FactoryCreateText(UClass* InClass, UObject* InParent
 		}
 	}
 
-	FEditorDelegates::OnAssetPostImport.Broadcast(this, NewAsset);
+	GEditor->GetEditorSubsystem<UImportSubsystem>()->BroadcastAssetPostImport(this, NewAsset);
 
 	return NewAsset;
 }
diff --git a/Engine/Source/Editor/UnrealEd/Private/Factories/EditorFactories.cpp b/Engine/Source/Editor/UnrealEd/Private/Factories/EditorFactories.cpp
--- a/Engine/Source/Editor/UnrealEd/Private/Factories/EditorFactories.cpp
+++ b/Engine/Source/Editor/UnrealEd/Private/Factories/EditorFactories.cpp
@@ -257,6 +257,7 @@
 #include "Rendering/SkeletalMeshModel.h"
 
 #include "Misc/App.h"
+#include "Subsystems/ImportSubsystem.h"
 
 #include "IDesktopPlatform.h"
 #include "DesktopPlatformModule.h"
@@ -673,7 +674,7 @@ UObject* ULevelFactory::FactoryCreateText
 	FFeedbackContext*	Warn
 )
 {
-	FEditorDelegates::OnAssetPreImport.Broadcast(this, Class, InParent, Name, Type);
+	GEditor->GetEditorSubsystem<UImportSubsystem>()->BroadcastAssetPreImport(this, Class, InParent, Name, Type);
 
 	UWorld* World = GWorld;
 	//@todo locked levels - if lock state is persistent, do we need to check for whether the level is locked?
@@ -707,7 +708,7 @@ UObject* ULevelFactory::FactoryCreateText
 				else
 				{
 					Warn->Logf(ELogVerbosity::Warning, TEXT("The Root map package name : '%s', conflicts with the existing object : '%s'"), *RootMapPackage->GetFullName(), *MapName);
-					FEditorDelegates::OnAssetPostImport.Broadcast( this, nullptr );
+					GEditor->GetEditorSubsystem<UImportSubsystem>()->BroadcastAssetPostImport(this, nullptr );
 					return nullptr;
 				}
 				
@@ -921,7 +922,7 @@ UObject* ULevelFactory::FactoryCreateText
 				if ( FLevelUtils::IsLevelLocked(World->GetCurrentLevel()) )
 				{
 					UE_LOG(LogEditorFactories, Warning, TEXT("Import actor: The requested operation could not be completed because the level is locked."));
-					FEditorDelegates::OnAssetPostImport.Broadcast( this, nullptr );
+					GEditor->GetEditorSubsystem<UImportSubsystem>()->BroadcastAssetPostImport(this, nullptr );
 					return nullptr;
 				}
 				else if ( !(bShouldSkipImportSpecialActors && ActorIndex < 2) )
@@ -1298,7 +1299,7 @@ UObject* ULevelFactory::FactoryCreateText
 	GEditor->IsImportingT3D = 0;
 	GIsImportingT3D = false;
 
-	FEditorDelegates::OnAssetPostImport.Broadcast( this, World );
+	GEditor->GetEditorSubsystem<UImportSubsystem>()->BroadcastAssetPostImport(this, World );
 
 	return World;
 }
@@ -1319,7 +1320,7 @@ UPackageFactory::UPackageFactory(const FObjectInitializer& ObjectInitializer)
 
 UObject* UPackageFactory::FactoryCreateText( UClass* Class, UObject* InParent, FName Name, EObjectFlags Flags, UObject* Context, const TCHAR* Type, const TCHAR*& Buffer, const TCHAR* BufferEnd, FFeedbackContext* Warn )
 {
-	FEditorDelegates::OnAssetPreImport.Broadcast(this, Class, InParent, Name, Type);
+	GEditor->GetEditorSubsystem<UImportSubsystem>()->BroadcastAssetPreImport(this, Class, InParent, Name, Type);
 
 	bool bSavedImportingT3D = GIsImportingT3D;
 	// Mark us as importing a T3D.
@@ -1482,7 +1483,7 @@ UObject* UPackageFactory::FactoryCreateText( UClass* Class, UObject* InParent, F
 	GEditor->IsImportingT3D = bSavedImportingT3D;
 	GIsImportingT3D = bSavedImportingT3D;
 
-	FEditorDelegates::OnAssetPostImport.Broadcast( this, TopLevelPackage );
+	GEditor->GetEditorSubsystem<UImportSubsystem>()->BroadcastAssetPostImport(this, TopLevelPackage );
 
 	return TopLevelPackage;
 }
@@ -1517,7 +1518,7 @@ UObject* UPolysFactory::FactoryCreateText
 	FVector PointPool[4096];
 	int32 NumPoints = 0;
 
-	FEditorDelegates::OnAssetPreImport.Broadcast(this, Class, InParent, Name, Type);
+	GEditor->GetEditorSubsystem<UImportSubsystem>()->BroadcastAssetPreImport(this, Class, InParent, Name, Type);
 
 	// Create polys.	
 	UPolys* Polys = Context ? CastChecked<UPolys>(Context) : NewObject<UPolys>(InParent, Name, Flags);
@@ -1758,7 +1759,7 @@ UObject* UPolysFactory::FactoryCreateText
 		}
 	}
 
-	FEditorDelegates::OnAssetPostImport.Broadcast( this, Polys );
+	GEditor->GetEditorSubsystem<UImportSubsystem>()->BroadcastAssetPostImport(this, Polys );
 
 	// Success.
 	return Polys;
@@ -1789,7 +1790,7 @@ UObject* UModelFactory::FactoryCreateText
 	FFeedbackContext*	Warn
 )
 {
-	FEditorDelegates::OnAssetPreImport.Broadcast(this, Class, InParent, Name, Type);
+	GEditor->GetEditorSubsystem<UImportSubsystem>()->BroadcastAssetPreImport(this, Class, InParent, Name, Type);
 
 	ABrush* TempOwner = (ABrush*)Context;
 	UModel* Model = NewObject<UModel>(InParent, Name, Flags);
@@ -1845,7 +1846,7 @@ UObject* UModelFactory::FactoryCreateText
 		}
 	}
 
-	FEditorDelegates::OnAssetPostImport.Broadcast( this, Model );
+	GEditor->GetEditorSubsystem<UImportSubsystem>()->BroadcastAssetPostImport(this, Model );
 
 	return Model;
 }
@@ -3809,7 +3810,7 @@ UObject* UTextureFactory::FactoryCreateBinary
 {
 	check(Type);
 
-	FEditorDelegates::OnAssetPreImport.Broadcast(this, Class, InParent, Name, Type);
+	GEditor->GetEditorSubsystem<UImportSubsystem>()->BroadcastAssetPreImport(this, Class, InParent, Name, Type);
 
 	// if the texture already exists, remember the user settings
 	UTexture* ExistingTexture = FindObject<UTexture>( InParent, *Name.ToString() );
@@ -3869,7 +3870,7 @@ UObject* UTextureFactory::FactoryCreateBinary
 		case EAppReturnType::Cancel:
 		default:
 			{
-				FEditorDelegates::OnAssetPostImport.Broadcast( this, nullptr );
+				GEditor->GetEditorSubsystem<UImportSubsystem>()->BroadcastAssetPostImport( this, nullptr );
 				return nullptr;
 			}
 		}
@@ -3930,7 +3931,7 @@ UObject* UTextureFactory::FactoryCreateBinary
 		}
 
 		Warn->Logf(ELogVerbosity::Error, TEXT("Texture import failed") );
-		FEditorDelegates::OnAssetPostImport.Broadcast( this, nullptr );
+		GEditor->GetEditorSubsystem<UImportSubsystem>()->BroadcastAssetPostImport( this, nullptr );
 		return nullptr;
 	}
 
@@ -4057,7 +4058,7 @@ UObject* UTextureFactory::FactoryCreateBinary
 		ApplyAutoImportSettings(Texture);
 	}
 
-	FEditorDelegates::OnAssetPostImport.Broadcast(this, Texture);
+	GEditor->GetEditorSubsystem<UImportSubsystem>()->BroadcastAssetPostImport(this, Texture);
 
 	// Invalidate any materials using the newly imported texture. (occurs if you import over an existing texture)
 	Texture->PostEditChange();
@@ -4765,7 +4766,7 @@ UObject* UFontFileImportFactory::FactoryCreateBinary(UClass* InClass, UObject* I
 		return nullptr;
 	}
 
-	FEditorDelegates::OnAssetPreImport.Broadcast(this, InClass, InParent, InName, InType);
+	GEditor->GetEditorSubsystem<UImportSubsystem>()->BroadcastAssetPreImport(this, InClass, InParent, InName, InType);
 
 	// Create the font face
 	UFontFace* const FontFace = NewObject<UFontFace>(InParent, InClass, InName, InFlags);
@@ -4779,7 +4780,7 @@ UObject* UFontFileImportFactory::FactoryCreateBinary(UClass* InClass, UObject* I
 		FontFace->CacheSubFaces();
 	}
 
-	FEditorDelegates::OnAssetPostImport.Broadcast(this, FontFace);
+	GEditor->GetEditorSubsystem<UImportSubsystem>()->BroadcastAssetPostImport(this, FontFace);
 	
 	// Create the font (if requested)
 	if (FontFace && bCreateFontAsset)
@@ -6535,7 +6536,7 @@ UCurveImportFactory::UCurveImportFactory(const FObjectInitializer& ObjectInitial
 // @note jf: for importing a curve from a text format.  this is experimental code for a prototype feature and not fully fleshed out
 UObject* UCurveImportFactory::FactoryCreateText( UClass* InClass, UObject* InParent, FName InName, EObjectFlags Flags, UObject* Context, const TCHAR* Type, const TCHAR*& Buffer, const TCHAR* BufferEnd, FFeedbackContext* Warn )
 {
-	FEditorDelegates::OnAssetPreImport.Broadcast(this, InClass, InParent, InName, Type);
+	GEditor->GetEditorSubsystem<UImportSubsystem>()->BroadcastAssetPreImport(this, InClass, InParent, InName, Type);
 
 	if(	FCString::Stricmp( Type, TEXT( "AS" ) ) == 0 )
 	{
@@ -6605,13 +6606,13 @@ UObject* UCurveImportFactory::FactoryCreateText( UClass* InClass, UObject* InPar
 				}
 			}
 
-			FEditorDelegates::OnAssetPostImport.Broadcast(this, NewCurve);
+			GEditor->GetEditorSubsystem<UImportSubsystem>()->BroadcastAssetPostImport(this, NewCurve);
 
 			return NewCurve;
 		}
 	}
 
-	FEditorDelegates::OnAssetPostImport.Broadcast( this, nullptr );
+	GEditor->GetEditorSubsystem<UImportSubsystem>()->BroadcastAssetPostImport( this, nullptr );
 	return nullptr;
 }
 
diff --git a/Engine/Source/Editor/UnrealEd/Private/Factories/VectorFieldFactory.cpp b/Engine/Source/Editor/UnrealEd/Private/Factories/VectorFieldFactory.cpp
--- a/Engine/Source/Editor/UnrealEd/Private/Factories/VectorFieldFactory.cpp
+++ b/Engine/Source/Editor/UnrealEd/Private/Factories/VectorFieldFactory.cpp
@@ -202,7 +202,7 @@ UObject* UVectorFieldStaticFactory::FactoryCreateBinary(
 	const uint8* BufferEnd,
 	FFeedbackContext* Warn )
 {
-	FEditorDelegates::OnAssetPreImport.Broadcast(this, InClass, InParent, InName, Type);
+	GEditor->GetEditorSubsystem<UImportSubsystem>()->BroadcastAssetPreImport(this, InClass, InParent, InName, Type);
 
 	check( InClass == UVectorFieldStatic::StaticClass() );
 
@@ -264,7 +264,7 @@ UObject* UVectorFieldStaticFactory::FactoryCreateBinary(
 		VectorField->InitResource();
 	}
 
-	FEditorDelegates::OnAssetPostImport.Broadcast(this, VectorField);
+	GEditor->GetEditorSubsystem<UImportSubsystem>()->BroadcastAssetPostImport(this, VectorField);
 
 	return VectorField;
 }
diff --git a/Engine/Source/Editor/UnrealEd/Private/Fbx/FbxFactory.cpp b/Engine/Source/Editor/UnrealEd/Private/Fbx/FbxFactory.cpp
--- a/Engine/Source/Editor/UnrealEd/Private/Fbx/FbxFactory.cpp
+++ b/Engine/Source/Editor/UnrealEd/Private/Fbx/FbxFactory.cpp
@@ -177,11 +177,11 @@ UObject* UFbxFactory::FactoryCreateFile
 	if( bOperationCanceled )
 	{
 		bOutOperationCanceled = true;
-		FEditorDelegates::OnAssetPostImport.Broadcast(this, NULL);
+		GEditor->GetEditorSubsystem<UImportSubsystem>()->BroadcastAssetPostImport(this, NULL);
 		return NULL;
 	}
 
-	FEditorDelegates::OnAssetPreImport.Broadcast(this, Class, InParent, Name, Type);
+	GEditor->GetEditorSubsystem<UImportSubsystem>()->BroadcastAssetPreImport(this, Class, InParent, Name, Type);
 
 	UObject* CreatedObject = NULL;
 	//Look if its a re-import, in that cazse we must call the re-import factory
@@ -224,7 +224,7 @@ UObject* UFbxFactory::FactoryCreateFile
 		if ( !DetectImportType(UFactory::CurrentFilename) )
 		{
 			// Failed to read the file info, fail the import
-			FEditorDelegates::OnAssetPostImport.Broadcast(this, NULL);
+			GEditor->GetEditorSubsystem<UImportSubsystem>()->BroadcastAssetPostImport(this, NULL);
 			return NULL;
 		}
 	}
@@ -767,7 +767,7 @@ UObject* UFbxFactory::FactoryCreateFile
 		FbxImporter->ReleaseScene();
 	}
 
-	FEditorDelegates::OnAssetPostImport.Broadcast(this, CreatedObject);
+	GEditor->GetEditorSubsystem<UImportSubsystem>()->BroadcastAssetPostImport(this, CreatedObject);
 
 	return CreatedObject;
 }
diff --git a/Engine/Source/Editor/UnrealEd/Private/Fbx/FbxSceneImportFactory.cpp b/Engine/Source/Editor/UnrealEd/Private/Fbx/FbxSceneImportFactory.cpp
--- a/Engine/Source/Editor/UnrealEd/Private/Fbx/FbxSceneImportFactory.cpp
+++ b/Engine/Source/Editor/UnrealEd/Private/Fbx/FbxSceneImportFactory.cpp
@@ -1004,7 +1004,7 @@ FFeedbackContext*	Warn
 	// Unselect all actors.
 	GEditor->SelectNone(false, false);
 
-	FEditorDelegates::OnAssetPreImport.Broadcast(this, Class, InParent, Name, Type);
+	GEditor->GetEditorSubsystem<UImportSubsystem>()->BroadcastAssetPreImport(this, Class, InParent, Name, Type);
 	
 	//TODO verify if we really need this when instancing actor in a level from an import
 	//In that case we should change the variable name.
@@ -1040,7 +1040,7 @@ FFeedbackContext*	Warn
 		GEditor->IsImportingT3D = 0;
 		GIsImportingT3D = false;
 		Warn->EndSlowTask();
-		FEditorDelegates::OnAssetPostImport.Broadcast(this, World);
+		GEditor->GetEditorSubsystem<UImportSubsystem>()->BroadcastAssetPostImport(this, World);
 		return nullptr;
 	}
 
@@ -1084,7 +1084,7 @@ FFeedbackContext*	Warn
 		GEditor->IsImportingT3D = 0;
 		GIsImportingT3D = false;
 		Warn->EndSlowTask();
-		FEditorDelegates::OnAssetPostImport.Broadcast(this, World);
+		GEditor->GetEditorSubsystem<UImportSubsystem>()->BroadcastAssetPostImport(this, World);
 		return nullptr;
 	}
 
@@ -1116,7 +1116,7 @@ FFeedbackContext*	Warn
 			GEditor->IsImportingT3D = 0;
 			GIsImportingT3D = false;
 			Warn->EndSlowTask();
-			FEditorDelegates::OnAssetPostImport.Broadcast(this, World);
+			GEditor->GetEditorSubsystem<UImportSubsystem>()->BroadcastAssetPostImport(this, World);
 			return nullptr;
 		}
 	}
@@ -1236,7 +1236,7 @@ FFeedbackContext*	Warn
 	ReimportData = nullptr;
 
 	Warn->EndSlowTask();
-	FEditorDelegates::OnAssetPostImport.Broadcast(this, World);
+	GEditor->GetEditorSubsystem<UImportSubsystem>()->BroadcastAssetPostImport(this, World);
 
 	return ReturnObject;
 }
diff --git a/Engine/Source/Editor/UnrealEd/Private/Fbx/ReimportFbxSceneFactory.cpp b/Engine/Source/Editor/UnrealEd/Private/Fbx/ReimportFbxSceneFactory.cpp
--- a/Engine/Source/Editor/UnrealEd/Private/Fbx/ReimportFbxSceneFactory.cpp
+++ b/Engine/Source/Editor/UnrealEd/Private/Fbx/ReimportFbxSceneFactory.cpp
@@ -1022,7 +1022,7 @@ EReimportResult::Type UReimportFbxSceneFactory::ImportSkeletalMesh(void* VoidRoo
 {
 	UnFbx::FFbxImporter* FbxImporter = (UnFbx::FFbxImporter*)VoidFbxImporter;
 	FbxNode *RootNodeToImport = (FbxNode *)VoidRootNodeToImport;
-	//FEditorDelegates::OnAssetPreImport.Broadcast(this, UStaticMesh::StaticClass(), GWorld, FName(""), TEXT("fbx"));
+	//GEditor->GetEditorSubsystem<UImportSubsystem>()->BroadcastAssetPreImport(this, UStaticMesh::StaticClass(), GWorld, FName(""), TEXT("fbx"));
 	FbxNode *GeometryParentNode = nullptr;
 	//Get the first parent geometry node
 	for (int idx = 0; idx < FbxImporter->Scene->GetGeometryCount(); ++idx)
@@ -1110,7 +1110,7 @@ EReimportResult::Type UReimportFbxSceneFactory::ImportSkeletalMesh(void* VoidRoo
 EReimportResult::Type UReimportFbxSceneFactory::ImportStaticMesh(void* VoidFbxImporter, TSharedPtr<FFbxMeshInfo> MeshInfo, TSharedPtr<FFbxSceneInfo> SceneInfoPtr)
 {
 	UnFbx::FFbxImporter* FbxImporter = (UnFbx::FFbxImporter*)VoidFbxImporter;
-	//FEditorDelegates::OnAssetPreImport.Broadcast(this, UStaticMesh::StaticClass(), GWorld, FName(""), TEXT("fbx"));
+	//GEditor->GetEditorSubsystem<UImportSubsystem>()->BroadcastAssetPreImport(this, UStaticMesh::StaticClass(), GWorld, FName(""), TEXT("fbx"));
 	FbxNode *GeometryParentNode = nullptr;
 	//Get the first parent geometry node
 	for (int idx = 0; idx < FbxImporter->Scene->GetGeometryCount(); ++idx)
diff --git a/Engine/Source/Editor/UnrealEd/Private/Subsystems/ImportSubsystem.cpp b/Engine/Source/Editor/UnrealEd/Private/Subsystems/ImportSubsystem.cpp
new file mode 100644
--- /dev/null
+++ b/Engine/Source/Editor/UnrealEd/Private/Subsystems/ImportSubsystem.cpp
@@ -0,0 +1,45 @@
+// Copyright 1998-2018 Epic Games, Inc. All Rights Reserved.
+
+#include "Subsystems/ImportSubsystem.h"
+
+#include "Editor.h"
+
+UImportSubsystem::UImportSubsystem()
+	: UEditorSubsystem()
+{
+
+}
+
+void UImportSubsystem::Initialize(FSubsystemCollectionBase& Collection)
+{
+
+}
+
+void UImportSubsystem::Deinitialize()
+{
+
+}
+
+PRAGMA_DISABLE_DEPRECATION_WARNINGS
+
+void UImportSubsystem::BroadcastAssetPreImport(UFactory* InFactory, UClass* InClass, UObject* InParent, const FName& Name, const TCHAR* Type)
+{
+	FEditorDelegates::OnAssetPreImport.Broadcast(InFactory, InClass, InParent, Name, Type);
+	OnAssetPreImport.Broadcast(InFactory, InClass, InParent, Name, Type);
+	OnAssetPreImport_BP.Broadcast(InFactory, InClass, InParent, Name, FString(Type));
+}
+
+void UImportSubsystem::BroadcastAssetPostImport(UFactory* InFactory, UObject* InCreatedObject)
+{
+	FEditorDelegates::OnAssetPostImport.Broadcast(InFactory, InCreatedObject);
+	OnAssetPostImport.Broadcast(InFactory, InCreatedObject);
+	OnAssetPostImport_BP.Broadcast(InFactory, InCreatedObject);
+}
+
+void UImportSubsystem::BroadcastAssetReimport(UObject* InCreatedObject)
+{
+	FEditorDelegates::OnAssetReimport.Broadcast(InCreatedObject);
+	OnAssetReimport.Broadcast(InCreatedObject);
+	OnAssetReimport_BP.Broadcast(InCreatedObject);
+}
+PRAGMA_ENABLE_DEPRECATION_WARNINGS
\ No newline at end of file
diff --git a/Engine/Source/Editor/UnrealEd/Private/Toolkits/SimpleAssetEditor.cpp b/Engine/Source/Editor/UnrealEd/Private/Toolkits/SimpleAssetEditor.cpp
--- a/Engine/Source/Editor/UnrealEd/Private/Toolkits/SimpleAssetEditor.cpp
+++ b/Engine/Source/Editor/UnrealEd/Private/Toolkits/SimpleAssetEditor.cpp
@@ -37,7 +37,7 @@ const FName FSimpleAssetEditor::SimpleEditorAppIdentifier( TEXT( "GenericEditorA
 
 FSimpleAssetEditor::~FSimpleAssetEditor()
 {
-	FEditorDelegates::OnAssetPostImport.RemoveAll(this);
+	GEditor->GetEditorSubsystem<UImportSubsystem>()->OnAssetPostImport.RemoveAll(this);
 
 	DetailsView.Reset();
 	PropertiesTab.Reset();
@@ -51,7 +51,7 @@ void FSimpleAssetEditor::InitEditor( const EToolkitMode::Type Mode, const TShare
 	const bool bIsLockable = false;
 
 	EditingObjects = ObjectsToEdit;
-	FEditorDelegates::OnAssetPostImport.AddRaw(this, &FSimpleAssetEditor::HandleAssetPostImport);
+	GEditor->GetEditorSubsystem<UImportSubsystem>()->OnAssetPostImport.AddRaw(this, &FSimpleAssetEditor::HandleAssetPostImport);
 
 	FPropertyEditorModule& PropertyEditorModule = FModuleManager::GetModuleChecked<FPropertyEditorModule>( "PropertyEditor" );
 	const FDetailsViewArgs DetailsViewArgs( bIsUpdatable, bIsLockable, true, FDetailsViewArgs::ObjectsUseNameArea, false );
diff --git a/Engine/Source/Editor/UnrealEd/Public/Editor.h b/Engine/Source/Editor/UnrealEd/Public/Editor.h
--- a/Engine/Source/Editor/UnrealEd/Public/Editor.h
+++ b/Engine/Source/Editor/UnrealEd/Public/Editor.h
@@ -13,6 +13,8 @@
 #include "Editor/EditorEngine.h"
 #include "Engine/StaticMesh.h"
 
+#include "Subsystems/ImportSubsystem.h"
+
 
 
 #define CAMERA_ZOOM_DAMPEN			200.f
@@ -188,10 +190,13 @@ struct UNREALED_API FEditorDelegates
 	/** Called when finishing configuration of a new asset */
 	static FOnNewAssetCreation OnNewAssetCreated;
 	/** Called when new assets are being (re-)imported. */
+	UE_DEPRECATED(4.22, "Use the ImportSubsystem instead. GEditor->GetEditorSubsystem<UImportSubsystem>()")
 	static FOnAssetPreImport OnAssetPreImport;
 	/** Called when new assets have been (re-)imported. */
+	UE_DEPRECATED(4.22, "Use the ImportSubsystem instead. GEditor->GetEditorSubsystem<UImportSubsystem>()")
 	static FOnAssetPostImport OnAssetPostImport;
 	/** Called after an asset has been reimported */
+	UE_DEPRECATED(4.22, "Use the ImportSubsystem instead. GEditor->GetEditorSubsystem<UImportSubsystem>()")
 	static FOnAssetReimport OnAssetReimport;
 	/** Called when new actors are dropped on to the viewport */
 	static FOnNewActorsDropped OnNewActorsDropped;
diff --git a/Engine/Source/Editor/UnrealEd/Public/Subsystems/ImportSubsystem.h b/Engine/Source/Editor/UnrealEd/Public/Subsystems/ImportSubsystem.h
new file mode 100644
--- /dev/null
+++ b/Engine/Source/Editor/UnrealEd/Public/Subsystems/ImportSubsystem.h
@@ -0,0 +1,63 @@
+// Copyright 1998-2018 Epic Games, Inc. All Rights Reserved.
+
+#pragma once
+
+#include "EditorSubsystem.h"
+
+#include "UObject/ObjectMacros.h"
+#include "UObject/UObjectGlobals.h"
+#include "Templates/SubclassOf.h"
+
+#include "ImportSubsystem.generated.h"
+
+/**
+ * UImportSubsystem
+ * Subsystem for importing assets in the editor, 
+ * Contains utility functions and callbacks for hooking into importing.
+ */
+UCLASS()
+class UNREALED_API UImportSubsystem : public UEditorSubsystem
+{
+	GENERATED_BODY()
+
+public:
+	UImportSubsystem();
+
+	virtual void Initialize(FSubsystemCollectionBase& Collection);
+	virtual void Deinitialize();
+
+	/** delegate type fired when new assets are being (re-)imported. Params: UFactory* InFactory, UClass* InClass, UObject* InParent, const FName& Name, const TCHAR* Type */
+	DECLARE_MULTICAST_DELEGATE_FiveParams(FOnAssetPreImport, UFactory*, UClass*, UObject*, const FName&, const TCHAR*);
+	/** delegate type fired when new assets have been (re-)imported. Note: InCreatedObject can be NULL if import failed. Params: UFactory* InFactory, UObject* InCreatedObject */
+	DECLARE_MULTICAST_DELEGATE_TwoParams(FOnAssetPostImport, UFactory*, UObject*);
+	/** delegate type fired when new assets have been reimported. Note: InCreatedObject can be NULL if import failed. Params: UObject* InCreatedObject */
+	DECLARE_MULTICAST_DELEGATE_OneParam(FOnAssetReimport, UObject*);
+
+	// Broadcast AssetPreImport, do not broadcast with OnAssetPostImport directly.
+	void BroadcastAssetPreImport(UFactory* InFactory, UClass* InClass, UObject* InParent, const FName& Name, const TCHAR* Type);
+	// Broadcast AssetPostImport, do not broadcast with OnAssetPostImport directly.
+	void BroadcastAssetPostImport(UFactory* InFactory, UObject* InCreatedObject);
+	// Broadcast AssetReimport, do not broadcast with OnAssetReimport directly.
+	void BroadcastAssetReimport(UObject* InCreatedObject);
+
+	// Used to register and unregister ONLY use Broadcast functions to execute the delegate
+	FOnAssetPreImport OnAssetPreImport;
+	FOnAssetPostImport OnAssetPostImport;
+	FOnAssetReimport OnAssetReimport;
+
+private:
+
+	/** delegate type fired when new assets are being (re-)imported. Params: UFactory* InFactory, UClass* InClass, UObject* InParent, const FName& Name, const TCHAR* Type */
+	DECLARE_DYNAMIC_MULTICAST_DELEGATE_FiveParams(FOnAssetPreImport_Dyn, UFactory*, InFactory, UClass*, InClass, UObject*, InParent, const FName&, Name, const FString, Type);
+	/** delegate type fired when new assets have been (re-)imported. Note: InCreatedObject can be NULL if import failed. Params: UFactory* InFactory, UObject* InCreatedObject */
+	DECLARE_DYNAMIC_MULTICAST_DELEGATE_TwoParams(FOnAssetPostImport_Dyn, UFactory*, InFactory, UObject*, InCreatedObject);
+	/** delegate type fired when new assets have been reimported. Note: InCreatedObject can be NULL if import failed. Params: UObject* InCreatedObject */
+	DECLARE_DYNAMIC_MULTICAST_DELEGATE_OneParam(FOnAssetReimport_Dyn, UObject*, InCreatedObject);
+
+	UPROPERTY(BlueprintAssignable, DisplayName = "OnAssetPreImport", meta = (ScriptName = "OnAssetPreImport"))
+	FOnAssetPreImport_Dyn OnAssetPreImport_BP;
+	UPROPERTY(BlueprintAssignable, DisplayName = "OnAssetPostImport", meta = (ScriptName = "OnAssetPostImport"))
+	FOnAssetPostImport_Dyn OnAssetPostImport_BP;
+	UPROPERTY(BlueprintAssignable, DisplayName = "OnAssetReimport", meta = (ScriptName = "OnAssetReimport"))
+	FOnAssetReimport_Dyn OnAssetReimport_BP;
+};
diff --git a/Engine/Source/Editor/UnrealEd/UnrealEd.Build.cs b/Engine/Source/Editor/UnrealEd/UnrealEd.Build.cs
--- a/Engine/Source/Editor/UnrealEd/UnrealEd.Build.cs
+++ b/Engine/Source/Editor/UnrealEd/UnrealEd.Build.cs
@@ -98,7 +98,8 @@ public class UnrealEd : ModuleRules
                 "MeshDescription",
                 "MeshDescriptionOperations",
                 "MeshBuilder",
-                "MaterialShaderQualitySettings"
+                "MaterialShaderQualitySettings",
+                "EditorSubsystem",
             }
 		);
 
diff --git a/Engine/Source/Runtime/Engine/Classes/Engine/Engine.h b/Engine/Source/Runtime/Engine/Classes/Engine/Engine.h
--- a/Engine/Source/Runtime/Engine/Classes/Engine/Engine.h
+++ b/Engine/Source/Runtime/Engine/Classes/Engine/Engine.h
@@ -14,6 +14,9 @@
 #include "Engine/World.h"
 #include "Misc/BufferedOutputDevice.h"
 #include "Misc/FrameRate.h"
+#include "Subsystems/SubsystemCollection.h"
+#include "Subsystems/EngineSubsystem.h"
+
 #include "Engine.generated.h"
 
 #define WITH_DYNAMIC_RESOLUTION (!UE_SERVER)
@@ -3170,6 +3173,38 @@ public:
 	virtual void CancelTransaction(int32 Index) { }
 #endif
 
+public:
+	/**
+	 * Get an Engine Subsystem of specified type
+	 */
+	UEngineSubsystem* GetEngineSubsystemBase(TSubclassOf<UEngineSubsystem> SubsystemClass) const
+	{
+		checkSlow(this != nullptr);
+		return EngineSubsystemCollection.GetSubsystem<UEngineSubsystem>(SubsystemClass);
+	}
+
+	/**
+	 * Get an Engine Subsystem of specified type
+	 */
+	template <typename TSubsystemClass>
+	TSubsystemClass* GetEngineSubsystem() const
+	{
+		checkSlow(this != nullptr);
+		return EngineSubsystemCollection.GetSubsystem<TSubsystemClass>(TSubsystemClass::StaticClass());
+	}
+
+	/**
+	 * Get all Engine Subsystem of specified type, this is only necessary for interfaces that can have multiple implementations instanced at a time.
+	 */
+	template <typename TSubsystemClass>
+	const TArray<TSubsystemClass*>& GetEngineSubsystemArray() const
+	{
+		return EngineSubsystemCollection.GetSubsystemArray<TSubsystemClass>(TSubsystemClass::StaticClass());
+	}
+
+private:
+	FSubsystemCollection<UEngineSubsystem> EngineSubsystemCollection;
+
 public:
 	/**
 	 * Delegate we fire every time a new stat has been registered.
diff --git a/Engine/Source/Runtime/Engine/Private/GameEngine.cpp b/Engine/Source/Runtime/Engine/Private/GameEngine.cpp
--- a/Engine/Source/Runtime/Engine/Private/GameEngine.cpp
+++ b/Engine/Source/Runtime/Engine/Private/GameEngine.cpp
@@ -679,6 +679,7 @@ void UGameEngine::OnViewportResized(FViewport* Viewport, uint32 Unused)
 UEngine::UEngine(const FObjectInitializer& ObjectInitializer)
 	: Super(ObjectInitializer)
 	, ViewExtensions( new FSceneViewExtensions() )
+	, EngineSubsystemCollection(this)
 {
 	C_WorldBox = FColor(0, 0, 40, 255);
 	C_BrushWire = FColor(192, 0, 0, 255);
diff --git a/Engine/Source/Runtime/Engine/Private/Subsystems/EngineSubsystem.cpp b/Engine/Source/Runtime/Engine/Private/Subsystems/EngineSubsystem.cpp
new file mode 100644
--- /dev/null
+++ b/Engine/Source/Runtime/Engine/Private/Subsystems/EngineSubsystem.cpp
@@ -0,0 +1,9 @@
+// Copyright 1998-2018 Epic Games, Inc. All Rights Reserved.
+
+#include "Subsystems/EngineSubsystem.h"
+
+UEngineSubsystem::UEngineSubsystem()
+	: UDynamicSubsystem()
+{
+
+}
diff --git a/Engine/Source/Runtime/Engine/Private/Subsystems/Subsystem.cpp b/Engine/Source/Runtime/Engine/Private/Subsystems/Subsystem.cpp
--- a/Engine/Source/Runtime/Engine/Private/Subsystems/Subsystem.cpp
+++ b/Engine/Source/Runtime/Engine/Private/Subsystems/Subsystem.cpp
@@ -2,7 +2,15 @@
 
 #include "Subsystems/Subsystem.h"
 
+#include "Subsystems/SubsystemCollection.h"
+
 USubsystem::USubsystem()
 {
 
 }
+
+UDynamicSubsystem::UDynamicSubsystem()
+	: USubsystem()
+{
+
+}
\ No newline at end of file
diff --git a/Engine/Source/Runtime/Engine/Private/Subsystems/SubsystemCollection.cpp b/Engine/Source/Runtime/Engine/Private/Subsystems/SubsystemCollection.cpp
--- a/Engine/Source/Runtime/Engine/Private/Subsystems/SubsystemCollection.cpp
+++ b/Engine/Source/Runtime/Engine/Private/Subsystems/SubsystemCollection.cpp
@@ -4,6 +4,14 @@
 
 #include "Subsystems/Subsystem.h"
 #include "UObject/UObjectHash.h"
+#include "UObject/Package.h"
+#include "Modules/ModuleManager.h"
+#include "Misc/PackageName.h"
+
+
+FDelegateHandle FSubsystemCollectionBase::ModulesChangedHandle;
+TArray<FSubsystemCollectionBase*> FSubsystemCollectionBase::SubsystemCollections;
+TMap<FName, TArray<TSubclassOf<UDynamicSubsystem>>> FSubsystemCollectionBase::DynamicSystemModuleMap;
 
 FSubsystemCollectionBase::FSubsystemCollectionBase()
 	: Outer(nullptr)
@@ -62,27 +70,98 @@ const TArray<USubsystem*>& FSubsystemCollectionBase::GetSubsystemArrayInternal(T
 	return List;
 }
 
+class FSubsystemModuleWatcher
+{
+public :
+	static void OnModulesChanged(FName ModuleThatChanged, EModuleChangeReason ReasonForChange)
+	{
+
+		switch (ReasonForChange)
+		{
+		case EModuleChangeReason::ModuleLoaded:
+			FSubsystemCollectionBase::AddClassesForModule(ModuleThatChanged);
+			break;
+
+		case EModuleChangeReason::ModuleUnloaded:
+			FSubsystemCollectionBase::RemoveClassesForModule(ModuleThatChanged);
+			break;
+		}
+	}
+};
+
+void FSubsystemCollectionBase::InitializeModuleWatcher()
+{
+	check(!ModulesChangedHandle.IsValid());
+
+	// Add Loaded Modules
+	TArray<UClass*> SubsystemClasses;
+	GetDerivedClasses(UDynamicSubsystem::StaticClass(), SubsystemClasses, true);
+
+	for (UClass* SubsystemClass : SubsystemClasses)
+	{
+		if (!SubsystemClass->HasAllClassFlags(CLASS_Abstract))
+		{
+			UPackage* const ClassPackage = SubsystemClass->GetOuterUPackage();
+			if (ClassPackage)
+			{
+				const FName ModuleName = FPackageName::GetShortFName(ClassPackage->GetFName());
+				if (FModuleManager::Get().IsModuleLoaded(ModuleName))
+				{
+					TArray<TSubclassOf<UDynamicSubsystem>>& ModuleSubsystemClasses = DynamicSystemModuleMap.FindOrAdd(ModuleName);
+					ModuleSubsystemClasses.Add(SubsystemClass);
+				}
+			}
+		}
+	}
+
+	ModulesChangedHandle = FModuleManager::Get().OnModulesChanged().AddStatic(&FSubsystemModuleWatcher::OnModulesChanged);
+}
+
+void FSubsystemCollectionBase::DeinitializeModuleWatcher()
+{
+	if (ModulesChangedHandle.IsValid())
+	{
+		FModuleManager::Get().OnModulesChanged().Remove(ModulesChangedHandle);
+	}
+}
+
 void FSubsystemCollectionBase::Initialize()
 {
 	if (ensure(BaseType) && ensureMsgf(SubsystemMap.Num() == 0, TEXT("Currently don't support repopulation of Subsystem Collections.")))
 	{
 		check(Outer);
 		check(!bPopulating); //Populating collections on multiple threads?
-
+		
+		if (SubsystemCollections.Num() == 0)
+		{
+			InitializeModuleWatcher();
+		}
+		
 		TGuardValue<bool> PopulatingGuard(bPopulating, true);
 
-		TArray<UClass*> SubsystemClasses;
-		GetDerivedClasses(BaseType, SubsystemClasses, true);
-
-		for (UClass* SubsystemClass : SubsystemClasses)
+		for (const TPair<FName, TArray<TSubclassOf<UDynamicSubsystem>>>& SubsystemClasses : DynamicSystemModuleMap)
 		{
-			InteralAddSystem(SubsystemClass);
+			for (const TSubclassOf<UDynamicSubsystem>& SubsystemClass : SubsystemClasses.Value)
+			{
+				AddAndInitializeSubsystem(SubsystemClass);
+			}
 		}
+
+		// Statically track collections
+		SubsystemCollections.Add(this);
 	}
 }
 
 void FSubsystemCollectionBase::Deinitialize()
 {
+	// Remove static tracking 
+	SubsystemCollections.Remove(this);
+	if (SubsystemCollections.Num() == 0)
+	{
+		DeinitializeModuleWatcher();
+	}
+
+	// Deinit and clean up existing systems
 	SubsystemArrayMap.Empty();
 	for (auto Iter = SubsystemMap.CreateIterator(); Iter; ++Iter)
 	{
@@ -91,6 +170,7 @@ void FSubsystemCollectionBase::Deinitialize()
 		if (Subsystem->GetClass() == KeyClass)
 		{
 			Subsystem->Deinitialize();
+			Subsystem->InternalOwningSubsystem = nullptr;
 		}
 	}
 	SubsystemMap.Empty();
@@ -102,7 +182,7 @@ bool FSubsystemCollectionBase::InitializeDependancy(TSubclassOf<USubsystem> Subs
 		&& ensureMsgf(bPopulating, TEXT("InitializeDependancy() should only be called from System USubsystem::Initialization() implementations."))
 		&& ensureMsgf(SubsystemClass->IsChildOf(BaseType), TEXT("ClassType (%s) must be a subclass of BaseType(%s)."), *SubsystemClass->GetName(), *BaseType->GetName()))
 	{
-		return InteralAddSystem(SubsystemClass);
+		return AddAndInitializeSubsystem(SubsystemClass);
 	}
 	return false;
 }
@@ -113,7 +193,7 @@ void FSubsystemCollectionBase::AddReferencedObjects(FReferenceCollector& Collect
 	Collector.AddReferencedObjects(SubsystemMap);
 }
 
-bool FSubsystemCollectionBase::InteralAddSystem(UClass* SubsystemClass)
+bool FSubsystemCollectionBase::AddAndInitializeSubsystem(UClass* SubsystemClass)
 {
 	if (!SubsystemMap.Contains(SubsystemClass))
 	{
@@ -126,12 +206,88 @@ bool FSubsystemCollectionBase::InteralAddSystem(UClass* SubsystemClass)
 				USubsystem*& Subsystem = SubsystemMap.Add(SubsystemClass);
 				Subsystem = NewObject<USubsystem>(Outer, SubsystemClass);
 
+				Subsystem->InternalOwningSubsystem = this;
 				Subsystem->Initialize(*this);
-
+				
 				return true;
 			}
 		}
 		return false;
 	}
 	return true;
-}
\ No newline at end of file
+}
+
+void FSubsystemCollectionBase::RemoveAndDeinitializeSubsystem(USubsystem* Subsystem)
+{
+	check(Subsystem);
+	USubsystem* SubsystemFound = SubsystemMap.FindAndRemoveChecked(Subsystem->GetClass());
+	check(Subsystem == SubsystemFound);
+
+	Subsystem->Deinitialize();
+	Subsystem->InternalOwningSubsystem = nullptr;
+}
+
+void FSubsystemCollectionBase::AddClassesForModule(const FName& InModuleName)
+{
+	check(!DynamicSystemModuleMap.Contains(InModuleName));
+
+	// Find the class package for this module
+	const UPackage* const ClassPackage = FindPackage(nullptr, *(FString("/Script/") + InModuleName.ToString()));
+	if (!ClassPackage)
+	{
+		return;
+	}
+
+	TArray<TSubclassOf<UDynamicSubsystem>> SubsystemClasses;
+	TArray<UObject*> PackageObjects;
+	GetObjectsWithOuter(ClassPackage, PackageObjects, false);
+	for (UObject* Object : PackageObjects)
+	{
+		UClass* const CurrentClass = Cast<UClass>(Object);
+		if (CurrentClass && !CurrentClass->HasAllClassFlags(CLASS_Abstract) && CurrentClass->IsChildOf(UDynamicSubsystem::StaticClass()))
+		{
+			SubsystemClasses.Add(CurrentClass);
+			FSubsystemCollectionBase::AddAllInstances(CurrentClass);
+		}
+	}
+	if (SubsystemClasses.Num() > 0)
+	{
+		DynamicSystemModuleMap.Add(InModuleName, MoveTemp(SubsystemClasses));
+	}
+}
+void FSubsystemCollectionBase::RemoveClassesForModule(const FName& InModuleName)
+{
+	TArray<TSubclassOf<UDynamicSubsystem>>* SubsystemClasses = DynamicSystemModuleMap.Find(InModuleName);
+	if (SubsystemClasses)
+	{
+		for (TSubclassOf<UDynamicSubsystem>& SubsystemClass : *SubsystemClasses)
+		{
+			FSubsystemCollectionBase::RemoveAllInstances(SubsystemClass);
+		}
+		DynamicSystemModuleMap.Remove(InModuleName);
+	}
+}
+
+void FSubsystemCollectionBase::AddAllInstances(UClass* SubsystemClass)
+{
+	for (FSubsystemCollectionBase* SubsystemCollection : SubsystemCollections)
+	{
+		if (SubsystemClass->IsChildOf(SubsystemCollection->BaseType))
+		{
+			SubsystemCollection->AddAndInitializeSubsystem(SubsystemClass);
+		}
+	}
+}
+
+void FSubsystemCollectionBase::RemoveAllInstances(UClass* SubsystemClass)
+{
+	ForEachObjectOfClass(SubsystemClass, [](UObject* SubsystemObj)
+	{
+		USubsystem* Subsystem = CastChecked<USubsystem>(SubsystemObj);
+
+		if (Subsystem->InternalOwningSubsystem)
+		{
+			Subsystem->InternalOwningSubsystem->RemoveAndDeinitializeSubsystem(Subsystem);
+		}
+	});
+}
diff --git a/Engine/Source/Runtime/Engine/Private/UnrealEngine.cpp b/Engine/Source/Runtime/Engine/Private/UnrealEngine.cpp
--- a/Engine/Source/Runtime/Engine/Private/UnrealEngine.cpp
+++ b/Engine/Source/Runtime/Engine/Private/UnrealEngine.cpp
@@ -1386,6 +1386,7 @@ void UEngine::Init(IEngineLoop* InEngineLoop)
 	// Subsystems.
 	FURL::StaticInit();
 	FLinkerLoad::StaticInit(UTexture2D::StaticClass());
+	EngineSubsystemCollection.Initialize();
 
 #if !UE_BUILD_SHIPPING
 	// Check for overrides to the default map on the command line
diff --git a/Engine/Source/Runtime/Engine/Private/UserInterface/Canvas.cpp b/Engine/Source/Runtime/Engine/Private/UserInterface/Canvas.cpp
--- a/Engine/Source/Runtime/Engine/Private/UserInterface/Canvas.cpp
+++ b/Engine/Source/Runtime/Engine/Private/UserInterface/Canvas.cpp
@@ -1233,6 +1233,9 @@ UCanvas::UCanvas(const FObjectInitializer& ObjectInitializer)
 	CachedDisplayWidth = 0;
 	CachedDisplayHeight = 0;
 
+	SceneView = nullptr;
+	Canvas = nullptr;
+
 	// only call once on construction.  Expensive on some platforms (occulus).
 	// Init gets called every frame.	
 	UpdateSafeZoneData();
diff --git a/Engine/Source/Runtime/Engine/Public/Subsystems/EngineSubsystem.h b/Engine/Source/Runtime/Engine/Public/Subsystems/EngineSubsystem.h
new file mode 100644
--- /dev/null
+++ b/Engine/Source/Runtime/Engine/Public/Subsystems/EngineSubsystem.h
@@ -0,0 +1,28 @@
+// Copyright 1998-2018 Epic Games, Inc. All Rights Reserved.
+
+#pragma once
+
+#include "Subsystem.h"
+
+#include "EngineSubsystem.generated.h"
+
+/**
+ * UEngineSubsystem
+ * Base class for auto instanced and initialized systems that share the lifetime of the Engine
+ *
+ * UEngineSubsystems are dynamic and will be initialized when the module is loaded if necessary.
+ * This means that after StartupModule() is called on the module containing a subsystem,
+ * the subsystem collection with instantiate and initialize the subsystem automatically.
+ * If the subsystem collection is created post module load then the instances will be created at
+ * collection initialization time.
+ */
+
+UCLASS(Abstract)
+class ENGINE_API UEngineSubsystem : public UDynamicSubsystem
+{
+	GENERATED_BODY()
+
+public:
+	UEngineSubsystem();
+
+};
diff --git a/Engine/Source/Runtime/Engine/Public/Subsystems/Subsystem.h b/Engine/Source/Runtime/Engine/Public/Subsystems/Subsystem.h
--- a/Engine/Source/Runtime/Engine/Public/Subsystems/Subsystem.h
+++ b/Engine/Source/Runtime/Engine/Public/Subsystems/Subsystem.h
@@ -12,8 +12,10 @@ class FSubsystemCollectionBase;
 /** Subsystems are auto instanced classes that share the lifetime of certain engine constructs
  * 
  *	Currently supported Subsystem lifetimes are:
+ *		Engine		 -> inherit UEngineSubsystem
+ *		Editor		 -> inherit UEditorSubsystem
  *		GameInstance -> inherit UGameInstanceSubsystem
- *		LocalPlayer  -> inherit ULocalPlayerSubsystem
+ *		LocalPlayer	 -> inherit ULocalPlayerSubsystem
  *
  *
  *	Normal Example:
@@ -63,4 +65,25 @@ public:
 	/** Implement this for deinitialization of instances of the system */
 	virtual void Deinitialize() {}
 
+private:
+	friend class FSubsystemCollectionBase;
+	FSubsystemCollectionBase* InternalOwningSubsystem;
+
+};
+
+
+/** Dynamic Subsystems auto populate/depopulate existing collections when modules are loaded and unloaded
+ *
+ * Only UEngineSubsystems and UEditorSubsystems allow for dynamic loading.
+ * 
+ * If instances of your subsystem aren't being created it maybe that the module they are in isn't being explicitly loaded,
+ * make sure there is a LoadModule("ModuleName") to load the module.
+ */
+UCLASS(Abstract)
+class ENGINE_API UDynamicSubsystem : public USubsystem
+{
+	GENERATED_BODY()
+
+public:
+	UDynamicSubsystem();
 };
diff --git a/Engine/Source/Runtime/Engine/Public/Subsystems/SubsystemCollection.h b/Engine/Source/Runtime/Engine/Public/Subsystems/SubsystemCollection.h
--- a/Engine/Source/Runtime/Engine/Public/Subsystems/SubsystemCollection.h
+++ b/Engine/Source/Runtime/Engine/Public/Subsystems/SubsystemCollection.h
@@ -2,12 +2,11 @@
 
 #pragma once
 
-#include "CoreMinimal.h"
-#include "Templates/Casts.h"
 #include "Templates/SubclassOf.h"
 #include "UObject/GCObject.h"
 
 class USubsystem;
+class UDynamicSubsystem;
 
 class ENGINE_API FSubsystemCollectionBase : public FGCObject
 {
@@ -42,9 +41,11 @@ protected:
 
 	/** Get the collection BaseType */
 	const TSubclassOf<USubsystem>& GetBaseType() const { return BaseType; }
+
 private:
+	bool AddAndInitializeSubsystem(UClass* SubsystemClass);
 
-	bool InteralAddSystem(UClass* SubsystemClass);
+	void RemoveAndDeinitializeSubsystem(USubsystem* Subsystem);
 
 	TMap<TSubclassOf<USubsystem>, USubsystem*> SubsystemMap;
 
@@ -55,10 +56,31 @@ private:
 	UObject* Outer;
 
 	bool bPopulating;
+
+public:
+	/** Init / Deinit the Module watcher, this tracks module startup and shutdown to ensure only the appropriate dynamic subsystems are instantiated */
+	static void InitializeModuleWatcher();
+	static void DeinitializeModuleWatcher();
+
+private:
+	friend class FSubsystemModuleWatcher;
+
+	static void AddClassesForModule(const FName& InModuleName);
+	static void RemoveClassesForModule(const FName& InModuleName);
+
+	/** Add Instances of the specified Subsystem class to all existing SubsystemCollections of the correct type */
+	static void AddAllInstances(UClass* SubsystemClass);
+
+	/** Remove Instances of the specified Subsystem class from all existing SubsystemCollections of the correct type */
+	static void RemoveAllInstances(UClass* SubsystemClass);
+
+	static FDelegateHandle ModulesChangedHandle;
+	static TArray<FSubsystemCollectionBase*> SubsystemCollections;
+	static TMap<FName, TArray<TSubclassOf<UDynamicSubsystem>>> DynamicSystemModuleMap;
 };
 
 template<typename TBaseType>
-class ENGINE_API FSubsystemCollection : public FSubsystemCollectionBase
+class FSubsystemCollection : public FSubsystemCollectionBase
 {
 public:
 	/** Get a Subsystem by type */
diff --git a/Engine/Source/Runtime/Launch/Private/LaunchEngineLoop.cpp b/Engine/Source/Runtime/Launch/Private/LaunchEngineLoop.cpp
--- a/Engine/Source/Runtime/Launch/Private/LaunchEngineLoop.cpp
+++ b/Engine/Source/Runtime/Launch/Private/LaunchEngineLoop.cpp
@@ -2944,6 +2944,7 @@ bool FEngineLoop::LoadStartupCoreModules()
 
 	SlowTask.EnterProgressFrame(10);
 #if WITH_EDITOR
+		FModuleManager::Get().LoadModuleChecked("UnrealEd");
 		FModuleManager::LoadModuleChecked<IEditorStyleModule>("EditorStyle");
 #endif //WITH_EDITOR
 
