Fixed potential nullptr derefrence

diff --git a/Engine/Source/Editor/UnrealEd/Private/Factories/SkeletalMeshImport.cpp b/Engine/Source/Editor/UnrealEd/Private/Factories/SkeletalMeshImport.cpp
--- a/Engine/Source/Editor/UnrealEd/Private/Factories/SkeletalMeshImport.cpp
+++ b/Engine/Source/Editor/UnrealEd/Private/Factories/SkeletalMeshImport.cpp
@@ -428,12 +428,14 @@ ExistingSkelMeshData* SaveExistingSkelMeshData(USkeletalMesh* ExistingSkelMesh,
 		{
 			int32 OffsetReductionLODIndex = 0;
 			
-			FSkeletalMeshLODInfo* LODInfo = ExistingSkelMesh->GetLODInfo(ReimportSpecificLOD);
-			ExistingMeshDataPtr->bIsReimportLODReduced = LODInfo->bHasBeenSimplified;
-			if (LODInfo && LODInfo->bHasBeenSimplified)
+			if (FSkeletalMeshLODInfo* LODInfo = ExistingSkelMesh->GetLODInfo(ReimportSpecificLOD))
 			{
-				//Save the imported LOD reduction settings
-				ExistingMeshDataPtr->ExistingReimportLODReductionSettings = LODInfo->ReductionSettings;
+				ExistingMeshDataPtr->bIsReimportLODReduced = LODInfo->bHasBeenSimplified;
+				if (LODInfo->bHasBeenSimplified)
+				{
+					//Save the imported LOD reduction settings
+					ExistingMeshDataPtr->ExistingReimportLODReductionSettings = LODInfo->ReductionSettings;
+				}
 			}
 
 			// Remove the zero'th LOD (ie: the LOD being reimported).
