When EditorEngine initialized a new invalid world it sets it to be an editor world now.
This allows FWorldDelegates::OnPostWorldInitialization() listeners to consitently know what kind of world it is at initialization time.

diff --git a/Engine/Source/Editor/UnrealEd/Private/EditorEngine.cpp b/Engine/Source/Editor/UnrealEd/Private/EditorEngine.cpp
--- a/Engine/Source/Editor/UnrealEd/Private/EditorEngine.cpp
+++ b/Engine/Source/Editor/UnrealEd/Private/EditorEngine.cpp
@@ -6982,6 +6982,9 @@ void UEditorEngine::InitializeNewlyCreatedInactiveWorld(UWorld* World)
 	{
 		const bool bOldDirtyState = World->GetOutermost()->IsDirty();
 
+		// This is an editor world, set it's type prior to initialization
+		World->WorldType = EWorldType::Editor;
+
 		// Create the world without a physics scene because creating too many physics scenes causes deadlock issues in PhysX. The scene will be created when it is opened in the level editor.
 		// Also, don't create an FXSystem because it consumes too much video memory. This is also created when the level editor opens this world.
 		World->InitWorld(GetEditorWorldInitializationValues()
