Drag and drop file import no longer blocks external file explorer for the duration of the import.

diff --git a/Engine/Source/Editor/ContentBrowser/Private/SAssetView.cpp b/Engine/Source/Editor/ContentBrowser/Private/SAssetView.cpp
--- a/Engine/Source/Editor/ContentBrowser/Private/SAssetView.cpp
+++ b/Engine/Source/Editor/ContentBrowser/Private/SAssetView.cpp
@@ -22,11 +22,9 @@
 #include "Widgets/Input/SSlider.h"
 #include "Framework/Docking/TabManager.h"
 #include "EditorStyleSet.h"
-#include "EditorReimportHandler.h"
 #include "Settings/ContentBrowserSettings.h"
 #include "Engine/Blueprint.h"
 #include "Editor.h"
-#include "FileHelpers.h"
 #include "AssetSelection.h"
 #include "AssetRegistryModule.h"
 #include "IAssetTools.h"
@@ -1577,80 +1575,8 @@ FReply SAssetView::OnDrop( const FGeometry& MyGeometry, const FDragDropEvent& Dr
 			{
 				if (ExternalDragDropOp->HasFiles())
 				{
-					TArray<FString> ImportFiles;
-					TMap<FString, UObject*> ReimportFiles;
-					FAssetToolsModule& AssetToolsModule = FModuleManager::Get().LoadModuleChecked<FAssetToolsModule>("AssetTools");
-					FString RootDestinationPath = SourcesData.PackagePaths[0].ToString();
-					TArray<TPair<FString, FString>> FilesAndDestinations;
-					const TArray<FString>& DragFiles = ExternalDragDropOp->GetFiles();
-					AssetToolsModule.Get().ExpandDirectories(DragFiles, RootDestinationPath, FilesAndDestinations);
-
-					TArray<int32> ReImportIndexes;
-					for (int32 FileIdx = 0; FileIdx < FilesAndDestinations.Num(); ++FileIdx)
-					{
-						const FString& Filename = FilesAndDestinations[FileIdx].Key;
-						const FString& DestinationPath = FilesAndDestinations[FileIdx].Value;
-						FString Name = ObjectTools::SanitizeObjectName(FPaths::GetBaseFilename(Filename));
-						FString PackageName = ObjectTools::SanitizeInvalidChars(DestinationPath + TEXT("/") + Name, INVALID_LONGPACKAGE_CHARACTERS);
-
-
-						// We can not create assets that share the name of a map file in the same location
-						if (FEditorFileUtils::IsMapPackageAsset(PackageName))
-						{
-							//The error message will be log in the import process
-							ImportFiles.Add(Filename);
-							continue;
-						}
-						//Check if package exist in memory
-						UPackage* Pkg = FindPackage(nullptr, *PackageName);
-						bool IsPkgExist = Pkg != nullptr;
-						//check if package exist on file
-						if (!IsPkgExist && !FPackageName::DoesPackageExist(PackageName))
-						{
-							ImportFiles.Add(Filename);
-							continue;
-						}
-						if (Pkg == nullptr)
-						{
-							Pkg = CreatePackage(nullptr, *PackageName);
-							if (Pkg == nullptr)
-							{
-								//Cannot create a package that don't exist on disk or in memory!!!
-								//The error message will be log in the import process
-								ImportFiles.Add(Filename);
-								continue;
-							}
-						}
-						// Make sure the destination package is loaded
-						Pkg->FullyLoad();
-
-						// Check for an existing object
-						UObject* ExistingObject = StaticFindObject(UObject::StaticClass(), Pkg, *Name);
-						if (ExistingObject != nullptr)
-						{
-							ReimportFiles.Add(Filename, ExistingObject);
-							ReImportIndexes.Add(FileIdx);
-						}
-						else
-						{
-							ImportFiles.Add(Filename);
-						}
-					}
-					//Reimport
-					for (auto kvp : ReimportFiles)
-					{
-						FReimportManager::Instance()->Reimport(kvp.Value, false, true, kvp.Key);
-					}
-					//Import
-					if (ImportFiles.Num() > 0)
-					{
-						//Remove it in reverse so the smaller index are still valid
-						for (int32 IndexToRemove = ReImportIndexes.Num() - 1; IndexToRemove >= 0; --IndexToRemove)
-						{
-							FilesAndDestinations.RemoveAt(ReImportIndexes[IndexToRemove]);
-						}
-						AssetToolsModule.Get().ImportAssets(ImportFiles, SourcesData.PackagePaths[0].ToString(), nullptr, true, &FilesAndDestinations);
-					}
+					// Delay import until next tick to avoid blocking the process that files were dragged from
+					GEditor->GetEditorSubsystem<UImportSubsystem>()->ImportNextTick(ExternalDragDropOp->GetFiles(), SourcesData.PackagePaths[0].ToString());
 				}
 			}
 
diff --git a/Engine/Source/Editor/UnrealEd/Private/Subsystems/ImportSubsystem.cpp b/Engine/Source/Editor/UnrealEd/Private/Subsystems/ImportSubsystem.cpp
--- a/Engine/Source/Editor/UnrealEd/Private/Subsystems/ImportSubsystem.cpp
+++ b/Engine/Source/Editor/UnrealEd/Private/Subsystems/ImportSubsystem.cpp
@@ -3,6 +3,104 @@
 #include "Subsystems/ImportSubsystem.h"
 
 #include "Editor.h"
+#include "AssetToolsModule.h"
+#include "ObjectTools.h"
+#include "FileHelpers.h"
+#include "EditorReimportHandler.h"
+
+
+class FImportFilesByPath : public IImportSubsystemTask
+{
+public:
+	FImportFilesByPath(const TArray<FString>& InFiles, const FString& InRootDestinationPath) :
+		Files(InFiles),
+		RootDestinationPath(InRootDestinationPath)
+	{
+	}
+
+	virtual void Run() override
+	{
+		TArray<FString> ImportFiles;
+		TMap<FString, UObject*> ReimportFiles;
+		FAssetToolsModule& AssetToolsModule = FModuleManager::Get().LoadModuleChecked<FAssetToolsModule>("AssetTools");
+
+		TArray<TPair<FString, FString>> FilesAndDestinations;
+		AssetToolsModule.Get().ExpandDirectories(Files, RootDestinationPath, FilesAndDestinations);
+
+		TArray<int32> ReImportIndexes;
+		for (int32 FileIdx = 0; FileIdx < FilesAndDestinations.Num(); ++FileIdx)
+		{
+			const FString& Filename = FilesAndDestinations[FileIdx].Key;
+			const FString& DestinationPath = FilesAndDestinations[FileIdx].Value;
+			FString Name = ObjectTools::SanitizeObjectName(FPaths::GetBaseFilename(Filename));
+			FString PackageName = ObjectTools::SanitizeInvalidChars(DestinationPath + TEXT("/") + Name, INVALID_LONGPACKAGE_CHARACTERS);
+
+			// We can not create assets that share the name of a map file in the same location
+			if (FEditorFileUtils::IsMapPackageAsset(PackageName))
+			{
+				//The error message will be log in the import process
+				ImportFiles.Add(Filename);
+				continue;
+			}
+			//Check if package exist in memory
+			UPackage* Pkg = FindPackage(nullptr, *PackageName);
+			bool IsPkgExist = Pkg != nullptr;
+			//check if package exist on file
+			if (!IsPkgExist && !FPackageName::DoesPackageExist(PackageName))
+			{
+				ImportFiles.Add(Filename);
+				continue;
+			}
+			if (Pkg == nullptr)
+			{
+				Pkg = CreatePackage(nullptr, *PackageName);
+				if (Pkg == nullptr)
+				{
+					//Cannot create a package that don't exist on disk or in memory!!!
+					//The error message will be log in the import process
+					ImportFiles.Add(Filename);
+					continue;
+				}
+			}
+			// Make sure the destination package is loaded
+			Pkg->FullyLoad();
+
+			// Check for an existing object
+			UObject* ExistingObject = StaticFindObject(UObject::StaticClass(), Pkg, *Name);
+			if (ExistingObject != nullptr)
+			{
+				ReimportFiles.Add(Filename, ExistingObject);
+				ReImportIndexes.Add(FileIdx);
+			}
+			else
+			{
+				ImportFiles.Add(Filename);
+			}
+		}
+
+		//Reimport
+		for (auto kvp : ReimportFiles)
+		{
+			FReimportManager::Instance()->Reimport(kvp.Value, false, true, kvp.Key);
+		}
+		
+		//Import
+		if (ImportFiles.Num() > 0)
+		{
+			//Remove it in reverse so the smaller index are still valid
+			for (int32 IndexToRemove = ReImportIndexes.Num() - 1; IndexToRemove >= 0; --IndexToRemove)
+			{
+				FilesAndDestinations.RemoveAt(ReImportIndexes[IndexToRemove]);
+			}
+			AssetToolsModule.Get().ImportAssets(ImportFiles, RootDestinationPath, nullptr, true, &FilesAndDestinations);
+		}
+	}
+
+private:
+	TArray<FString> Files;
+	FString RootDestinationPath;
+};
+
 
 UImportSubsystem::UImportSubsystem()
 	: UEditorSubsystem()
@@ -20,6 +118,27 @@ void UImportSubsystem::Deinitialize()
 
 }
 
+void UImportSubsystem::ImportNextTick(const TArray<FString>& Files, const FString& DestinationPath)
+{
+	PendingTasks.Enqueue(MakeShared<FImportFilesByPath>(Files, DestinationPath));
+	GEditor->GetTimerManager()->SetTimerForNextTick(this, &UImportSubsystem::HandleNextTick);
+}
+
+void UImportSubsystem::HandleNextTick()
+{
+	if (!PendingTasks.IsEmpty())
+	{
+		TSharedPtr<IImportSubsystemTask> Task;
+		while (PendingTasks.Dequeue(Task))
+		{
+			if (Task.IsValid())
+			{
+				Task->Run();
+			}
+		}
+	}
+}
+
 PRAGMA_DISABLE_DEPRECATION_WARNINGS
 
 void UImportSubsystem::BroadcastAssetPreImport(UFactory* InFactory, UClass* InClass, UObject* InParent, const FName& Name, const TCHAR* Type)
@@ -42,4 +161,4 @@ void UImportSubsystem::BroadcastAssetReimport(UObject* InCreatedObject)
 	OnAssetReimport.Broadcast(InCreatedObject);
 	OnAssetReimport_BP.Broadcast(InCreatedObject);
 }
-PRAGMA_ENABLE_DEPRECATION_WARNINGS
\ No newline at end of file
+PRAGMA_ENABLE_DEPRECATION_WARNINGS
diff --git a/Engine/Source/Editor/UnrealEd/Public/Subsystems/ImportSubsystem.h b/Engine/Source/Editor/UnrealEd/Public/Subsystems/ImportSubsystem.h
--- a/Engine/Source/Editor/UnrealEd/Public/Subsystems/ImportSubsystem.h
+++ b/Engine/Source/Editor/UnrealEd/Public/Subsystems/ImportSubsystem.h
@@ -7,9 +7,21 @@
 #include "UObject/ObjectMacros.h"
 #include "UObject/UObjectGlobals.h"
 #include "Templates/SubclassOf.h"
+#include "Containers/Queue.h"
 
 #include "ImportSubsystem.generated.h"
 
+/**
+ * Interface for tasks that need delayed execution
+ */
+class IImportSubsystemTask
+{
+public:
+	virtual ~IImportSubsystemTask() { }
+
+	virtual void Run() = 0;
+};
+
 /**
  * UImportSubsystem
  * Subsystem for importing assets in the editor, 
@@ -26,6 +38,9 @@ public:
 	virtual void Initialize(FSubsystemCollectionBase& Collection);
 	virtual void Deinitialize();
 
+	/* Import files next tick */
+	void ImportNextTick(const TArray<FString>& Files, const FString& DestinationPath);
+
 	/** delegate type fired when new assets are being (re-)imported. Params: UFactory* InFactory, UClass* InClass, UObject* InParent, const FName& Name, const TCHAR* Type */
 	DECLARE_MULTICAST_DELEGATE_FiveParams(FOnAssetPreImport, UFactory*, UClass*, UObject*, const FName&, const TCHAR*);
 	/** delegate type fired when new assets have been (re-)imported. Note: InCreatedObject can be NULL if import failed. Params: UFactory* InFactory, UObject* InCreatedObject */
@@ -47,6 +62,9 @@ public:
 
 private:
 
+	/* Run deferred logic waiting to be run next tick */
+	void HandleNextTick();
+
 	/** delegate type fired when new assets are being (re-)imported. Params: UFactory* InFactory, UClass* InClass, UObject* InParent, const FName& Name, const TCHAR* Type */
 	DECLARE_DYNAMIC_MULTICAST_DELEGATE_FiveParams(FOnAssetPreImport_Dyn, UFactory*, InFactory, UClass*, InClass, UObject*, InParent, const FName&, Name, const FString, Type);
 	/** delegate type fired when new assets have been (re-)imported. Note: InCreatedObject can be NULL if import failed. Params: UFactory* InFactory, UObject* InCreatedObject */
@@ -60,4 +78,7 @@ private:
 	FOnAssetPostImport_Dyn OnAssetPostImport_BP;
 	UPROPERTY(BlueprintAssignable, DisplayName = "OnAssetReimport", meta = (ScriptName = "OnAssetReimport"))
 	FOnAssetReimport_Dyn OnAssetReimport_BP;
+
+	/* Tasks waiting to be run next tick */
+	TQueue<TSharedPtr<IImportSubsystemTask>> PendingTasks;
 };
