DockTabs now supports a suffix string it takes precedend in layout when space is tight, currently used by asset editors to show dirty astrix.

diff --git a/Engine/Plugins/2D/Paper2D/Source/Paper2DEditor/Private/FlipbookEditor/FlipbookEditor.cpp b/Engine/Plugins/2D/Paper2D/Source/Paper2DEditor/Private/FlipbookEditor/FlipbookEditor.cpp
--- a/Engine/Plugins/2D/Paper2D/Source/Paper2DEditor/Private/FlipbookEditor/FlipbookEditor.cpp
+++ b/Engine/Plugins/2D/Paper2D/Source/Paper2DEditor/Private/FlipbookEditor/FlipbookEditor.cpp
@@ -414,12 +414,7 @@ FText FFlipbookEditor::GetBaseToolkitName() const
 
 FText FFlipbookEditor::GetToolkitName() const
 {
-	const bool bDirtyState = FlipbookBeingEdited->GetOutermost()->IsDirty();
-
-	FFormatNamedArguments Args;
-	Args.Add(TEXT("FlipbookName"), FText::FromString(FlipbookBeingEdited->GetName()));
-	Args.Add(TEXT("DirtyState"), bDirtyState ? FText::FromString( TEXT( "*" ) ) : FText::GetEmpty());
-	return FText::Format(LOCTEXT("FlipbookEditorToolkitName", "{FlipbookName}{DirtyState}"), Args);
+	return FText::FromString(FlipbookBeingEdited->GetName());
 }
 
 FText FFlipbookEditor::GetToolkitToolTipText() const
diff --git a/Engine/Plugins/2D/Paper2D/Source/Paper2DEditor/Private/SpriteEditor/SpriteEditor.cpp b/Engine/Plugins/2D/Paper2D/Source/Paper2DEditor/Private/SpriteEditor/SpriteEditor.cpp
--- a/Engine/Plugins/2D/Paper2D/Source/Paper2DEditor/Private/SpriteEditor/SpriteEditor.cpp
+++ b/Engine/Plugins/2D/Paper2D/Source/Paper2DEditor/Private/SpriteEditor/SpriteEditor.cpp
@@ -468,12 +468,7 @@ FText FSpriteEditor::GetBaseToolkitName() const
 
 FText FSpriteEditor::GetToolkitName() const
 {
-	const bool bDirtyState = SpriteBeingEdited->GetOutermost()->IsDirty();
-
-	FFormatNamedArguments Args;
-	Args.Add(TEXT("SpriteName"), FText::FromString(SpriteBeingEdited->GetName()));
-	Args.Add(TEXT("DirtyState"), bDirtyState ? FText::FromString( TEXT( "*" ) ) : FText::GetEmpty());
-	return FText::Format(LOCTEXT("SpriteEditorToolkitName", "{SpriteName}{DirtyState}"), Args);
+	return FText::FromString(SpriteBeingEdited->GetName());
 }
 
 FText FSpriteEditor::GetToolkitToolTipText() const
diff --git a/Engine/Plugins/2D/Paper2D/Source/Paper2DEditor/Private/TileMapEditing/TileMapEdModeToolkit.cpp b/Engine/Plugins/2D/Paper2D/Source/Paper2DEditor/Private/TileMapEditing/TileMapEdModeToolkit.cpp
--- a/Engine/Plugins/2D/Paper2D/Source/Paper2DEditor/Private/TileMapEditing/TileMapEdModeToolkit.cpp
+++ b/Engine/Plugins/2D/Paper2D/Source/Paper2DEditor/Private/TileMapEditing/TileMapEdModeToolkit.cpp
@@ -50,12 +50,7 @@ FText FTileMapEdModeToolkit::GetToolkitName() const
 {
 	if (CurrentTileSetPtr.IsValid())
 	{
-		const bool bDirtyState = CurrentTileSetPtr->GetOutermost()->IsDirty();
-
-		FFormatNamedArguments Args;
-		Args.Add(TEXT("TileSetName"), FText::FromString(CurrentTileSetPtr->GetName()));
-		Args.Add(TEXT("DirtyState"), bDirtyState ? FText::FromString( TEXT( "*" ) ) : FText::GetEmpty());
-		return FText::Format(LOCTEXT("TileMapEditAppLabel", "{TileSetName}{DirtyState}"), Args);
+		return FText::FromString(CurrentTileSetPtr->GetName());
 	}
 	return GetBaseToolkitName();
 }
diff --git a/Engine/Plugins/2D/Paper2D/Source/Paper2DEditor/Private/TileSetEditor.cpp b/Engine/Plugins/2D/Paper2D/Source/Paper2DEditor/Private/TileSetEditor.cpp
--- a/Engine/Plugins/2D/Paper2D/Source/Paper2DEditor/Private/TileSetEditor.cpp
+++ b/Engine/Plugins/2D/Paper2D/Source/Paper2DEditor/Private/TileSetEditor.cpp
@@ -424,12 +424,7 @@ FText FTileSetEditor::GetBaseToolkitName() const
 
 FText FTileSetEditor::GetToolkitName() const
 {
-	const bool bDirtyState = TileSetBeingEdited->GetOutermost()->IsDirty();
-
-	FFormatNamedArguments Args;
-	Args.Add(TEXT("TileSetName"), FText::FromString(TileSetBeingEdited->GetName()));
-	Args.Add(TEXT("DirtyState"), bDirtyState ? FText::FromString( TEXT( "*" ) ) : FText::GetEmpty());
-	return FText::Format(LOCTEXT("TileSetAppLabel", "{TileSetName}{DirtyState}"), Args);
+	return FText::FromString(TileSetBeingEdited->GetName());
 }
 
 FText FTileSetEditor::GetToolkitToolTipText() const
diff --git a/Engine/Source/Editor/Kismet/Private/BlueprintEditor.cpp b/Engine/Source/Editor/Kismet/Private/BlueprintEditor.cpp
--- a/Engine/Source/Editor/Kismet/Private/BlueprintEditor.cpp
+++ b/Engine/Source/Editor/Kismet/Private/BlueprintEditor.cpp
@@ -7345,22 +7345,14 @@ FText FBlueprintEditor::GetToolkitName() const
 
 	if( IsEditingSingleBlueprint() )
 	{
-		const bool bDirtyState = GetBlueprintObj()->GetOutermost()->IsDirty();
-
-		FFormatNamedArguments Args;
-		Args.Add( TEXT("DirtyState"), bDirtyState ? FText::FromString( TEXT( "*" ) ) : FText::GetEmpty() );
-
 		if (FBlueprintEditorUtils::IsLevelScriptBlueprint(GetBlueprintObj()))
 		{
 			const FString& LevelName = FPackageName::GetShortFName( GetBlueprintObj()->GetOutermost()->GetFName().GetPlainNameString() ).GetPlainNameString();	
-
-			Args.Add( TEXT("LevelName"), FText::FromString( LevelName ) );
-			return FText::Format( NSLOCTEXT("KismetEditor", "LevelScriptAppLabel", "{LevelName}{DirtyState} - Level Blueprint Editor"), Args );
+			return FText::FromString(LevelName);
 		}
 		else
 		{
-			Args.Add( TEXT("BlueprintName"), FText::FromString( GetBlueprintObj()->GetName() ) );
-			return FText::Format( NSLOCTEXT("KismetEditor", "BlueprintScriptAppLabel", "{BlueprintName}{DirtyState}"), Args );
+			return FText::FromString(GetBlueprintObj()->GetName());
 		}
 	}
 
diff --git a/Engine/Source/Editor/LevelEditor/Private/SLevelEditor.cpp b/Engine/Source/Editor/LevelEditor/Private/SLevelEditor.cpp
--- a/Engine/Source/Editor/LevelEditor/Private/SLevelEditor.cpp
+++ b/Engine/Source/Editor/LevelEditor/Private/SLevelEditor.cpp
@@ -178,6 +178,7 @@ void SLevelEditor::Initialize( const TSharedRef<SDockTab>& OwnerTab, const TShar
 {
 	// Bind the level editor tab's label to the currently loaded level name string in the main frame
 	OwnerTab->SetLabel( TAttribute<FText>( this, &SLevelEditor::GetTabTitle) );
+	OwnerTab->SetTabLabelSuffix(TAttribute<FText>(this, &SLevelEditor::GetTabSuffix));
 
 	FLevelEditorModule& LevelEditorModule = FModuleManager::GetModuleChecked< FLevelEditorModule >(LevelEditorModuleName);
 
@@ -300,14 +301,13 @@ FText SLevelEditor::GetTabTitle() const
 {
 	const IMainFrameModule& MainFrameModule = FModuleManager::GetModuleChecked< IMainFrameModule >( MainFrameModuleName );
 
-	const bool bIncludeGameName = false;
+	return FText::FromString(MainFrameModule.GetLoadedLevelName());
+}
 
+FText SLevelEditor::GetTabSuffix() const
+{
 	const bool bDirtyState = World && World->GetCurrentLevel()->GetOutermost()->IsDirty();
-
-	FFormatNamedArguments Args;
-	Args.Add( TEXT("LevelName"), FText::FromString( MainFrameModule.GetLoadedLevelName() ) );
-	Args.Add( TEXT("DirtyState"), bDirtyState ? FText::FromString( TEXT( "*" ) ) : FText::GetEmpty() );
-	return FText::Format( NSLOCTEXT("LevelEditor", "TabTitleSpacer", "{LevelName}{DirtyState}"), Args );
+	return bDirtyState ? FText::FromString(TEXT("*")) : FText::GetEmpty();
 }
 
 bool SLevelEditor::HasActivePlayInEditorViewport() const
diff --git a/Engine/Source/Editor/LevelEditor/Private/SLevelEditor.h b/Engine/Source/Editor/LevelEditor/Private/SLevelEditor.h
--- a/Engine/Source/Editor/LevelEditor/Private/SLevelEditor.h
+++ b/Engine/Source/Editor/LevelEditor/Private/SLevelEditor.h
@@ -93,6 +93,9 @@ public:
 	/** @return	Returns the title to display in the level editor's tab label */
 	FText GetTabTitle() const;
 
+	/** @return	Returns the suffix to display in the level editor's after tab label */
+	FText GetTabSuffix() const;
+
 	/**
 	 * Processes level editor keybindings using events made in a viewport
 	 * 
diff --git a/Engine/Source/Editor/MaterialEditor/Private/MaterialEditor.cpp b/Engine/Source/Editor/MaterialEditor/Private/MaterialEditor.cpp
--- a/Engine/Source/Editor/MaterialEditor/Private/MaterialEditor.cpp
+++ b/Engine/Source/Editor/MaterialEditor/Private/MaterialEditor.cpp
@@ -1028,16 +1028,9 @@ FText FMaterialEditor::GetBaseToolkitName() const
 FText FMaterialEditor::GetToolkitName() const
 {
 	const UObject* EditingObject = GetEditingObjects()[0];
-
 	check(EditingObject);
 
-	const bool bDirtyState = EditingObject->GetOutermost()->IsDirty();
-
-	// Overridden to accommodate editing of multiple objects (original and preview materials)
-	FFormatNamedArguments Args;
-	Args.Add( TEXT("ObjectName"), FText::FromString( EditingObject->GetName() ) );
-	Args.Add( TEXT("DirtyState"), bDirtyState ? FText::FromString( TEXT( "*" ) ) : FText::GetEmpty() );
-	return FText::Format( LOCTEXT("MaterialEditorAppLabel", "{ObjectName}{DirtyState}"), Args );
+	return FText::FromString(EditingObject->GetName());
 }
 
 FText FMaterialEditor::GetToolkitToolTipText() const
diff --git a/Engine/Source/Editor/PropertyEditor/Private/PropertyEditorToolkit.cpp b/Engine/Source/Editor/PropertyEditor/Private/PropertyEditorToolkit.cpp
--- a/Engine/Source/Editor/PropertyEditor/Private/PropertyEditorToolkit.cpp
+++ b/Engine/Source/Editor/PropertyEditor/Private/PropertyEditorToolkit.cpp
@@ -508,17 +508,13 @@ FText FPropertyEditorToolkit::GetToolkitName() const
 	{
 		const UObject* EditingObject = EditingObjs[ 0 ];
 
-		const bool bDirtyState = EditingObject->GetOutermost()->IsDirty();
-
 		FFormatNamedArguments Args;
 		Args.Add( TEXT("ObjectName"), FText::FromString( EditingObject->GetName() ) );
-		Args.Add( TEXT("DirtyState"), bDirtyState ? FText::FromString( TEXT( "*" ) ) : FText::GetEmpty() );
 		Args.Add( TEXT("ToolkitName"), GetBaseToolkitName() );
-		return FText::Format( LOCTEXT("ToolkitName_SingleObject", "{ObjectName}{DirtyState} - {ToolkitName}"), Args );
+		return FText::Format( LOCTEXT("ToolkitName_SingleObject", "{ObjectName} - {ToolkitName}"), Args );
 	}
 	else
 	{
-		bool bDirtyState = false;
 		UClass* SharedBaseClass = NULL;
 		for( int32 x = 0; x < NumEditingObjects; ++x )
 		{
@@ -544,16 +540,12 @@ FText FPropertyEditorToolkit::GetToolkitName() const
 			{
 				SharedBaseClass = SharedBaseClass->GetSuperClass();
 			}
-
-			// If any of the objects are dirty, flag the label
-			bDirtyState |= Obj->GetOutermost()->IsDirty();
 		}
 
 		FFormatNamedArguments Args;
 		Args.Add( TEXT("NumberOfObjects"), EditingObjs.Num() );
 		Args.Add( TEXT("ClassName"), FText::FromString( SharedBaseClass->GetName() ) );
-		Args.Add( TEXT("DirtyState"), bDirtyState ? FText::FromString( TEXT( "*" ) ) : FText::GetEmpty() );
-		return FText::Format( LOCTEXT("ToolkitName_MultiObject", "{NumberOfObjects} {ClassName}{DirtyState} Objects - Property Matrix Editor"), Args );
+		return FText::Format( LOCTEXT("ToolkitName_MultiObject", "{NumberOfObjects} {ClassName} Objects - Property Matrix Editor"), Args );
 	}
 }
 
diff --git a/Engine/Source/Editor/TranslationEditor/Private/TranslationEditor.cpp b/Engine/Source/Editor/TranslationEditor/Private/TranslationEditor.cpp
--- a/Engine/Source/Editor/TranslationEditor/Private/TranslationEditor.cpp
+++ b/Engine/Source/Editor/TranslationEditor/Private/TranslationEditor.cpp
@@ -213,18 +213,13 @@ FText FTranslationEditor::GetBaseToolkitName() const
 FText FTranslationEditor::GetToolkitName() const
 {
 	const UObject* EditingObject = GetEditingObject();
-
 	check (EditingObject != NULL);
 
-	// This doesn't correctly indicate dirty status for Translation Editor currently...
-	const bool bDirtyState = EditingObject->GetOutermost()->IsDirty();
-
 	FFormatNamedArguments Args;
 	Args.Add(TEXT("Language"), FText::FromString(FPaths::GetBaseFilename(FPaths::GetPath(ArchiveFilePath))));
 	Args.Add(TEXT("ProjectName"), FText::FromString(FPaths::GetBaseFilename(ManifestFilePath)));
-	Args.Add( TEXT("DirtyState"), bDirtyState ? FText::FromString( TEXT( "*" ) ) : FText::GetEmpty() );
 	Args.Add( TEXT("ToolkitName"), GetBaseToolkitName() );
-	return FText::Format( LOCTEXT("TranslationEditorAppLabel", "{Language}{DirtyState} - {ProjectName} - {ToolkitName}"), Args );
+	return FText::Format( LOCTEXT("TranslationEditorAppLabel", "{Language} - {ProjectName} - {ToolkitName}"), Args );
 }
 
 FText FTranslationEditor::GetToolkitToolTipText() const
diff --git a/Engine/Source/Editor/UnrealEd/Private/Toolkits/AssetEditorToolkit.cpp b/Engine/Source/Editor/UnrealEd/Private/Toolkits/AssetEditorToolkit.cpp
--- a/Engine/Source/Editor/UnrealEd/Private/Toolkits/AssetEditorToolkit.cpp
+++ b/Engine/Source/Editor/UnrealEd/Private/Toolkits/AssetEditorToolkit.cpp
@@ -89,6 +89,7 @@ void FAssetEditorToolkit::InitAssetEditor( const EToolkitMode::Type Mode, const
 
 		// Create the label and the link for the toolkit documentation.
 		TAttribute<FText> Label = TAttribute<FText>( this, &FAssetEditorToolkit::GetToolkitName );
+		TAttribute<FText> LabelSuffix = TAttribute<FText>(this, &FAssetEditorToolkit::GetTabSuffix);
 		TAttribute<FText> ToolTipText = TAttribute<FText>( this, &FAssetEditorToolkit::GetToolkitToolTipText );
 		FString DocLink = GetDocumentationLink();
 		if ( !DocLink.StartsWith( "Shared/" ) )
@@ -97,13 +98,14 @@ void FAssetEditorToolkit::InitAssetEditor( const EToolkitMode::Type Mode, const
 		}
 
 		// Create a new SlateToolkitHost
-		NewMajorTab = SNew(SDockTab) 
-			.ContentPadding(0.0f) 
+		NewMajorTab = SNew(SDockTab)
+			.ContentPadding(0.0f)
 			.TabRole(ETabRole::MajorTab)
 			.ToolTip(IDocumentation::Get()->CreateToolTip(ToolTipText, nullptr, DocLink, GetToolkitFName().ToString()))
-			.Icon( this, &FAssetEditorToolkit::GetDefaultTabIcon )
-			.TabColorScale( this, &FAssetEditorToolkit::GetDefaultTabColor )
-			.Label( Label );
+			.Icon(this, &FAssetEditorToolkit::GetDefaultTabIcon)
+			.TabColorScale(this, &FAssetEditorToolkit::GetDefaultTabColor)
+			.Label(Label)
+			.LabelSuffix(LabelSuffix);
 
 		{
 			static_assert(sizeof(EAssetEditorToolkitTabLocation) == sizeof(int32), "EAssetEditorToolkitTabLocation is the incorrect size");
@@ -319,6 +321,20 @@ FText FAssetEditorToolkit::GetToolkitName() const
 	return GetLabelForObject(EditingObject);
 }
 
+FText FAssetEditorToolkit::GetTabSuffix() const
+{
+	bool bDirtyState = false;
+	for (int32 x = 0; x < EditingObjects.Num(); ++x)
+	{
+		if (EditingObjects[x]->GetOutermost()->IsDirty())
+		{
+			bDirtyState = true;
+			break;
+		}
+	}
+	return bDirtyState ? LOCTEXT("TabSuffixAsterix", "*") : FText::GetEmpty();
+}
+
 FText FAssetEditorToolkit::GetToolkitToolTipText() const
 {
 	const UObject* EditingObject = GetEditingObject();
@@ -330,7 +346,6 @@ FText FAssetEditorToolkit::GetToolkitToolTipText() const
 
 FText FAssetEditorToolkit::GetLabelForObject(const UObject* InObject)
 {
-	const bool bDirtyState = InObject->GetOutermost()->IsDirty();
 	FString NameString;
 	if(const AActor* ObjectAsActor = Cast<AActor>(InObject))
 	{
@@ -341,10 +356,7 @@ FText FAssetEditorToolkit::GetLabelForObject(const UObject* InObject)
 		NameString = InObject->GetName();
 	}
 
-	FFormatNamedArguments Args;
-	Args.Add( TEXT("ObjectName"), FText::FromString( NameString ));
-	Args.Add( TEXT("DirtyState"), bDirtyState ? FText::FromString( TEXT( "*" ) ) : FText::GetEmpty() );
-	return FText::Format( LOCTEXT("AssetEditorAppLabel", "{ObjectName}{DirtyState}"), Args );
+	return FText::FromString(NameString);
 }
 
 FText FAssetEditorToolkit::GetToolTipTextForObject(const UObject* InObject)
diff --git a/Engine/Source/Editor/UnrealEd/Private/Toolkits/SimpleAssetEditor.cpp b/Engine/Source/Editor/UnrealEd/Private/Toolkits/SimpleAssetEditor.cpp
--- a/Engine/Source/Editor/UnrealEd/Private/Toolkits/SimpleAssetEditor.cpp
+++ b/Engine/Source/Editor/UnrealEd/Private/Toolkits/SimpleAssetEditor.cpp
@@ -129,16 +129,10 @@ FText FSimpleAssetEditor::GetToolkitName() const
 	if( EditingObjs.Num() == 1 )
 	{
 		const UObject* EditingObject = EditingObjs[ 0 ];
-
-		const bool bDirtyState = EditingObject->GetOutermost()->IsDirty();
-
-		Args.Add( TEXT("ObjectName"), FText::FromString( EditingObject->GetName() ) );
-		Args.Add( TEXT("DirtyState"), bDirtyState ? FText::FromString( TEXT( "*" ) ) : FText::GetEmpty() );
-		return FText::Format( LOCTEXT("ToolkitTitle", "{ObjectName}{DirtyState} - {ToolkitName}"), Args );
+		return FText::FromString(EditingObject->GetName());
 	}
 	else
 	{
-		bool bDirtyState = false;
 		UClass* SharedBaseClass = nullptr;
 		for( int32 x = 0; x < EditingObjs.Num(); ++x )
 		{
@@ -164,17 +158,13 @@ FText FSimpleAssetEditor::GetToolkitName() const
 			{
 				SharedBaseClass = SharedBaseClass->GetSuperClass();
 			}
-
-			// If any of the objects are dirty, flag the label
-			bDirtyState |= Obj->GetOutermost()->IsDirty();
 		}
 
 		check(SharedBaseClass);
 
 		Args.Add( TEXT("NumberOfObjects"), EditingObjs.Num() );
 		Args.Add( TEXT("ClassName"), FText::FromString( SharedBaseClass->GetName() ) );
-		Args.Add( TEXT("DirtyState"), bDirtyState ? FText::FromString( TEXT( "*" ) ) : FText::GetEmpty() );
-		return FText::Format( LOCTEXT("ToolkitTitle_EditingMultiple", "{NumberOfObjects} {ClassName}{DirtyState} - {ToolkitName}"), Args );
+		return FText::Format( LOCTEXT("ToolkitTitle_EditingMultiple", "{NumberOfObjects} {ClassName} - {ToolkitName}"), Args );
 	}
 }
 
diff --git a/Engine/Source/Editor/UnrealEd/Public/Toolkits/AssetEditorToolkit.h b/Engine/Source/Editor/UnrealEd/Public/Toolkits/AssetEditorToolkit.h
--- a/Engine/Source/Editor/UnrealEd/Public/Toolkits/AssetEditorToolkit.h
+++ b/Engine/Source/Editor/UnrealEd/Public/Toolkits/AssetEditorToolkit.h
@@ -120,6 +120,7 @@ public:
 	virtual FName GetToolkitFName() const override = 0;				// Must implement in derived class!
 	virtual FText GetBaseToolkitName() const override = 0;			// Must implement in derived class!
 	virtual FText GetToolkitName() const override;
+	virtual FText GetTabSuffix() const override;
 	virtual FText GetToolkitToolTipText() const override;
 	virtual FString GetWorldCentricTabPrefix() const override = 0;	// Must implement in derived class!
 	virtual class FEdMode* GetEditorMode() const override;
diff --git a/Engine/Source/Editor/UnrealEd/Public/Toolkits/BaseToolkit.h b/Engine/Source/Editor/UnrealEd/Public/Toolkits/BaseToolkit.h
--- a/Engine/Source/Editor/UnrealEd/Public/Toolkits/BaseToolkit.h
+++ b/Engine/Source/Editor/UnrealEd/Public/Toolkits/BaseToolkit.h
@@ -30,6 +30,7 @@ public:
 	// IToolkit interface
 
 	virtual FName GetToolkitContextFName() const override;
+	virtual FText GetTabSuffix() const override { return FText::GetEmpty(); }
 	virtual bool ProcessCommandBindings(const FKeyEvent& InKeyEvent) const override;
 	virtual bool IsHosted() const override;
 	virtual const TSharedRef<IToolkitHost> GetToolkitHost() const override;
diff --git a/Engine/Source/Editor/UnrealEd/Public/Toolkits/IToolkit.h b/Engine/Source/Editor/UnrealEd/Public/Toolkits/IToolkit.h
--- a/Engine/Source/Editor/UnrealEd/Public/Toolkits/IToolkit.h
+++ b/Engine/Source/Editor/UnrealEd/Public/Toolkits/IToolkit.h
@@ -80,6 +80,9 @@ public:
 	/** Returns the localized name of this toolkit */
 	virtual FText GetToolkitName() const = 0;
 
+	/** Returns the localized tab suffix */
+	virtual FText GetTabSuffix() const = 0;
+
 	/** Returns the localized tooltip text of this toolkit */
 	virtual FText GetToolkitToolTipText() const = 0;
 
diff --git a/Engine/Source/Runtime/Slate/Private/Widgets/Docking/SDockTab.cpp b/Engine/Source/Runtime/Slate/Private/Widgets/Docking/SDockTab.cpp
--- a/Engine/Source/Runtime/Slate/Private/Widgets/Docking/SDockTab.cpp
+++ b/Engine/Source/Runtime/Slate/Private/Widgets/Docking/SDockTab.cpp
@@ -354,6 +354,7 @@ SDockTab::SDockTab()
 	, TabRole(ETabRole::PanelTab)
 	, ParentPtr()
 	, TabLabel(NSLOCTEXT("DockTab", "DefaultTabTitle", "UNNAMED"))
+	, TabLabelSuffix(FText::GetEmpty())
 	, OnTabClosed()
 	, OnTabActivated()
 	, OnCanCloseTab()
@@ -481,6 +482,7 @@ void SDockTab::Construct( const FArguments& InArgs )
 	this->OnCanCloseTab = InArgs._OnCanCloseTab;
 	this->OnPersistVisualState = InArgs._OnPersistVisualState;
 	this->TabLabel = InArgs._Label;
+	this->TabLabelSuffix = InArgs._LabelSuffix;
 	this->TabIcon = InArgs._Icon;
 	this->bShouldAutosize = InArgs._ShouldAutosize;
 	this->TabColorScale = InArgs._TabColorScale;
@@ -529,6 +531,7 @@ void SDockTab::Construct( const FArguments& InArgs )
 			+ SOverlay::Slot()
 			.Padding(TAttribute<FMargin>::Create(TAttribute<FMargin>::FGetter::CreateSP(this, &SDockTab::GetTabPadding)))
 			.VAlign(VAlign_Center)
+			.HAlign(HAlign_Left)
 			[
 				SNew(SHorizontalBox)
 				.Visibility(EVisibility::Visible)
@@ -555,7 +558,7 @@ void SDockTab::Construct( const FArguments& InArgs )
 				]
 				// Tab Label
 				+ SHorizontalBox::Slot()
-				.FillWidth(1)
+				.FillWidth(1.0f)
 				.Padding( 0.0f, 1.0f )
 				.VAlign(VAlign_Center)
 				[
@@ -563,6 +566,15 @@ void SDockTab::Construct( const FArguments& InArgs )
 					.TextStyle( FCoreStyle::Get(), "Docking.TabFont" )					
 					.Text( this, &SDockTab::GetTabLabel )
 				]
+				+ SHorizontalBox::Slot()
+				.AutoWidth()
+				.Padding( 0.0f, 1.0f )
+				.VAlign(VAlign_Center)
+				[
+					SAssignNew(LabelWidget, STextBlock)
+					.TextStyle( FCoreStyle::Get(), "Docking.TabFont" )					
+					.Text(this, &SDockTab::GetTabLabelSuffix)
+				]
 				
 				// @todo toolkit major: Could inject inline content here into tab for standalone asset editing dropdown/dirty state, etc.
 
@@ -697,6 +709,16 @@ void SDockTab::SetLabel( const TAttribute<FText>& InTabLabel )
 	TabLabel = InTabLabel;
 }
 
+FText SDockTab::GetTabLabelSuffix() const
+{
+	return TabLabelSuffix.Get();
+}
+
+void SDockTab::SetTabLabelSuffix(const TAttribute<FText>& InTabLabelSuffix)
+{
+	TabLabelSuffix = InTabLabelSuffix;
+}
+
 const FSlateBrush* SDockTab::GetTabIcon() const
 {
 	return TabIcon.Get();
diff --git a/Engine/Source/Runtime/Slate/Public/Widgets/Docking/SDockTab.h b/Engine/Source/Runtime/Slate/Public/Widgets/Docking/SDockTab.h
--- a/Engine/Source/Runtime/Slate/Public/Widgets/Docking/SDockTab.h
+++ b/Engine/Source/Runtime/Slate/Public/Widgets/Docking/SDockTab.h
@@ -69,6 +69,7 @@ public:
 		, _ContentPadding( FMargin( 2 ) )
 		, _TabRole(ETabRole::PanelTab)
 		, _Label()
+		, _LabelSuffix()
 		, _Icon( FStyleDefaults::GetNoBrush() )
 		, _OnTabClosed()
 		, _OnTabActivated()
@@ -85,6 +86,7 @@ public:
 		SLATE_ATTRIBUTE( FMargin, ContentPadding )
 		SLATE_ARGUMENT( ETabRole, TabRole )
 		SLATE_ATTRIBUTE( FText, Label )
+		SLATE_ATTRIBUTE(FText, LabelSuffix)
 		SLATE_ATTRIBUTE( const FSlateBrush*, Icon )
 		SLATE_EVENT( FOnTabClosedCallback, OnTabClosed )
 		SLATE_EVENT( FOnTabActivatedCallback, OnTabActivated )
@@ -174,6 +176,12 @@ public:
 	/** The label that appears on the tab. */
 	void SetLabel( const TAttribute<FText>& InTabLabel );
 
+	/** Get Label Suffix */
+	FText GetTabLabelSuffix() const;
+
+	/** Set Label Suffix.  A second text field at the end of the Label that takes precedence and isn't lost when space is restricted */
+	void SetTabLabelSuffix(const TAttribute<FText>& InTabLabelSuffix);
+
 	/** The tooltip text that appears on the tab. */
 	void SetTabToolTipWidget(TSharedPtr<SToolTip> InTabToolTipWidget);
 
@@ -339,6 +347,9 @@ protected:
 	/** The label on the tab */
 	TAttribute<FText> TabLabel;
 
+	/** A second text field at the end of the Label that takes precedence and isn't lost when space is restricted */
+	TAttribute<FText> TabLabelSuffix;
+
 	/** The icon on the tab */
 	TAttribute<const FSlateBrush*> TabIcon;
 	
