Exposed SetLevelVisibility to editor scripting

diff --git a/Engine/Source/Editor/UnrealEd/Public/EditorLevelUtils.h b/Engine/Source/Editor/UnrealEd/Public/EditorLevelUtils.h
--- a/Engine/Source/Editor/UnrealEd/Public/EditorLevelUtils.h
+++ b/Engine/Source/Editor/UnrealEd/Public/EditorLevelUtils.h
@@ -17,7 +17,8 @@ class ULevelStreaming;
 
 DECLARE_LOG_CATEGORY_EXTERN(LogLevelTools, Warning, All);
 
-enum class ELevelVisibilityDirtyMode
+UENUM(BlueprintType)
+enum class ELevelVisibilityDirtyMode : uint8
 {
 	// Use when the user is causing the visibility change.  Will update transaction state and mark the package dirty.
 	ModifyOnChange,
@@ -167,6 +168,7 @@ public:
 	 * @param	bShouldBeVisible		The level's new visibility state.
 	 * @param	bForceLayersVisible		If true and the level is visible, force the level's layers to be visible.
 	 */
+	UFUNCTION(BlueprintCallable, Category = "Editor Scripting | Level Utility")
 	static UNREALED_API void SetLevelVisibility(ULevel* Level, bool bShouldBeVisible, bool bForceLayersVisible, ELevelVisibilityDirtyMode ModifyMode = ELevelVisibilityDirtyMode::ModifyOnChange);
 	
 	/** 
