Fix a crash when importing static mesh and the parent package is null.

diff --git a/Engine/Source/Editor/UnrealEd/Private/Fbx/FbxStaticMeshImport.cpp b/Engine/Source/Editor/UnrealEd/Private/Fbx/FbxStaticMeshImport.cpp
--- a/Engine/Source/Editor/UnrealEd/Private/Fbx/FbxStaticMeshImport.cpp
+++ b/Engine/Source/Editor/UnrealEd/Private/Fbx/FbxStaticMeshImport.cpp
@@ -1410,7 +1410,15 @@ UStaticMesh* UnFbx::FFbxImporter::ImportStaticMeshAsSingle(UObject* InParent, TA
 		// Create a package for each mesh
 		if (Package == nullptr)
 		{
-			NewPackageName = FPackageName::GetLongPackagePath(Parent->GetOutermost()->GetName()) + TEXT("/") + MeshName;
+			if (Parent != nullptr && Parent->GetOutermost() != nullptr)
+			{
+				NewPackageName = FPackageName::GetLongPackagePath(Parent->GetOutermost()->GetName()) + TEXT("/") + MeshName;
+			}
+			else
+			{
+				AddTokenizedErrorMessage(FTokenizedMessage::Create(EMessageSeverity::Error, FText::Format(LOCTEXT("ImportStaticMeshAsSingle", "Invalid Parent package when importing {0}.\nThe asset will not be imported."), FText::FromString(MeshName))), FFbxErrors::Generic_ImportingNewObjectFailed);
+				return nullptr;
+			}
 			NewPackageName = UPackageTools::SanitizePackageName(NewPackageName);
 			Package = CreatePackage(NULL, *NewPackageName);
 		}
