PR #5501: UE-62796: Brush nullptr check in PostEditChangeProperty (Contributed by [REDACTED_PERSON])

diff --git a/Engine/Source/Editor/UnrealEd/Private/EditorBrushBuilder.cpp b/Engine/Source/Editor/UnrealEd/Private/EditorBrushBuilder.cpp
--- a/Engine/Source/Editor/UnrealEd/Private/EditorBrushBuilder.cpp
+++ b/Engine/Source/Editor/UnrealEd/Private/EditorBrushBuilder.cpp
@@ -254,10 +254,11 @@ void UEditorBrushBuilder::PostEditChangeProperty(struct FPropertyChangedEvent& P
 	{
 		// Rebuild brush on property change
 		ABrush* Brush = Cast<ABrush>(GetOuter());
-		if (Brush)
+		UWorld* BrushWorld = Brush ? Brush->GetWorld() : nullptr;
+		if (BrushWorld)
 		{
 			Brush->bInManipulation = PropertyChangedEvent.ChangeType == EPropertyChangeType::Interactive;
-			Build(Brush->GetWorld(), Brush);
+			Build(BrushWorld, Brush);
 		}
 	}
 }
