Back out changelist 4784301

This caused an issue with GC leaks checks in USD Python automated tests.

diff --git a/Engine/Source/Editor/UnrealEd/Private/EditorEngine.cpp b/Engine/Source/Editor/UnrealEd/Private/EditorEngine.cpp
--- a/Engine/Source/Editor/UnrealEd/Private/EditorEngine.cpp
+++ b/Engine/Source/Editor/UnrealEd/Private/EditorEngine.cpp
@@ -6982,9 +6982,6 @@ void UEditorEngine::InitializeNewlyCreatedInactiveWorld(UWorld* World)
 	{
 		const bool bOldDirtyState = World->GetOutermost()->IsDirty();
 
-		// This is an editor world, set it's type prior to initialization
-		World->WorldType = EWorldType::Editor;
-
 		// Create the world without a physics scene because creating too many physics scenes causes deadlock issues in PhysX. The scene will be created when it is opened in the level editor.
 		// Also, don't create an FXSystem because it consumes too much video memory. This is also created when the level editor opens this world.
 		World->InitWorld(GetEditorWorldInitializationValues()
