Fixed NaN handling when detecting degenerate polygons in FbxStaticMeshImporter

diff --git a/Engine/Source/Editor/UnrealEd/Private/Fbx/FbxStaticMeshImport.cpp b/Engine/Source/Editor/UnrealEd/Private/Fbx/FbxStaticMeshImport.cpp
--- a/Engine/Source/Editor/UnrealEd/Private/Fbx/FbxStaticMeshImport.cpp
+++ b/Engine/Source/Editor/UnrealEd/Private/Fbx/FbxStaticMeshImport.cpp
@@ -633,7 +633,7 @@ bool UnFbx::FFbxImporter::BuildStaticMeshFromGeometry(FbxNode* Node, UStaticMesh
 			check(P.Num() > 2); //triangle is the smallest polygon we can have
 			const FVector Normal = ((P[1] - P[2]) ^ (P[0] - P[2])).GetSafeNormal(ComparisonThreshold);
 			//Check for degenerated polygons, avoid NAN
-			if (Normal.IsNearlyZero(ComparisonThreshold))
+			if (Normal.IsNearlyZero(ComparisonThreshold) || Normal.ContainsNaN())
 			{
 				SkippedVertexInstance += PolygonVertexCount;
 				continue;
