Added logging to ForceDelete operation

diff --git a/Engine/Source/Editor/UnrealEd/Private/ObjectTools.cpp b/Engine/Source/Editor/UnrealEd/Private/ObjectTools.cpp
--- a/Engine/Source/Editor/UnrealEd/Private/ObjectTools.cpp
+++ b/Engine/Source/Editor/UnrealEd/Private/ObjectTools.cpp
@@ -2236,6 +2236,19 @@ namespace ObjectTools
 			return 0;
 		}
 
+		{
+			// Force delete is a dangerous operation, add some fingerprints to the log:
+			FString Msg;
+			Msg.Append(FString::Printf(TEXT("Force Deleting %d Package(s):"), ShownObjectsToDelete.Num()));
+			const int32 MAX_PACKAGES_TO_LOG = 10;
+			for(int32 I = 0; I < ShownObjectsToDelete.Num() && I < MAX_PACKAGES_TO_LOG; ++I)
+			{
+				Msg.Append(TEXT("\n"));
+				Msg.Append(FString::Printf(TEXT("\tAsset Name: %s"), *GetPathNameSafe(ShownObjectsToDelete[I])));
+			}
+			UE_LOG(LogUObjectGlobals, Log, TEXT("%s"), *Msg);
+		}
+
 		GWarn->BeginSlowTask( NSLOCTEXT("UnrealEd", "Deleting", "Deleting"), true );
 
 		struct FSCSNodeToDelete
