Fixed Replace Selected Actors action when replacing Brushes.

diff --git a/Engine/Source/Editor/UnrealEd/Private/EditorEngine.cpp b/Engine/Source/Editor/UnrealEd/Private/EditorEngine.cpp
--- a/Engine/Source/Editor/UnrealEd/Private/EditorEngine.cpp
+++ b/Engine/Source/Editor/UnrealEd/Private/EditorEngine.cpp
@@ -5108,6 +5108,18 @@ void UEditorEngine::ReplaceActors(UActorFactory* Factory, const FAssetData& Asse
 			NewActor->PostEditMove(true);
 			NewActor->MarkPackageDirty();
 
+			TSet<ULevel*> LevelsToRebuildBSP;
+			ABrush* Brush = Cast<ABrush>(OldActor);
+			if (Brush && !FActorEditorUtils::IsABuilderBrush(Brush)) // Track whether or not a brush actor was deleted.
+			{
+				ULevel* BrushLevel = OldActor->GetLevel();
+				if (BrushLevel && !Brush->IsVolumeBrush())
+				{
+					BrushLevel->Model->Modify();
+					LevelsToRebuildBSP.Add(BrushLevel);
+				}
+			}
+
 			// Replace references in the level script Blueprint with the new Actor
 			const bool bDontCreate = true;
 			ULevelScriptBlueprint* LSB = NewActor->GetLevel()->GetLevelScriptBlueprint(bDontCreate);
@@ -5122,6 +5134,17 @@ void UEditorEngine::ReplaceActors(UActorFactory* Factory, const FAssetData& Asse
 				Layers->DisassociateActorFromLayers( OldActor );
 			}
 			World->EditorDestroyActor(OldActor, true);
+
+			// If any brush actors were modified, update the BSP in the appropriate levels
+			if (LevelsToRebuildBSP.Num())
+			{
+				FlushRenderingCommands();
+
+				for (ULevel* LevelToRebuild : LevelsToRebuildBSP)
+				{
+					GEditor->RebuildLevel(*LevelToRebuild);
+				}
+			}
 		}
 		else
 		{
