Fix bug with flood-filling 0-alpha white pixels when importing pngs

diff --git a/Engine/Source/Editor/UnrealEd/Private/Factories/EditorFactories.cpp b/Engine/Source/Editor/UnrealEd/Private/Factories/EditorFactories.cpp
--- a/Engine/Source/Editor/UnrealEd/Private/Factories/EditorFactories.cpp
+++ b/Engine/Source/Editor/UnrealEd/Private/Factories/EditorFactories.cpp
@@ -2944,7 +2944,8 @@ public:
 	/* returns False if requires further processing because entire row is filled with zeroed alpha values */
 	bool ProcessHorizontalRow(int32 Y)
 	{
-		const uint32 White = FColor::White.DWColor();
+		// only wipe out colors that are affected by png turning valid colors white if alpha = 0
+		const uint32 WhiteWithZeroAlpha = FColor(255, 255, 255, 0).DWColor();
 
 		// Left -> Right
 		int32 NumLeftmostZerosToProcess = 0;
@@ -2953,8 +2954,8 @@ public:
 		{
 			PixelDataType* PixelData = SourceData + (Y * TextureWidth + X) * 4;
 			ColorDataType* ColorData = reinterpret_cast<ColorDataType*>(PixelData);
-			// only wipe out colors that are affected by png turning valid colors white if alpha = 0
-			if (PixelData[AIdx] == 0 && *ColorData == White)
+
+			if (*ColorData == WhiteWithZeroAlpha)
 			{
 				if (FillColor)
 				{
