Hide incompatible material preview feature levels and platforms when ray tracing is enabled

diff --git a/Engine/Source/Editor/LevelEditor/Private/LevelEditorToolBar.cpp b/Engine/Source/Editor/LevelEditor/Private/LevelEditorToolBar.cpp
--- a/Engine/Source/Editor/LevelEditor/Private/LevelEditorToolBar.cpp
+++ b/Engine/Source/Editor/LevelEditor/Private/LevelEditorToolBar.cpp
@@ -1714,6 +1714,15 @@ static void MakeShaderModelPreviewMenu(FMenuBuilder& MenuBuilder)
 	{
 		for (int32 i = GMaxRHIFeatureLevel; i >= 0; --i)
 		{
+#if RHI_RAYTRACING
+			if (IsRayTracingSupportedForThisProject())
+			{
+				if (i != ERHIFeatureLevel::SM5)
+				{
+					continue;
+				}
+			}
+#endif
 			switch (i)
 			{
 				case ERHIFeatureLevel::ES2:
diff --git a/Engine/Source/Editor/UnrealEd/Private/EditorEngine.cpp b/Engine/Source/Editor/UnrealEd/Private/EditorEngine.cpp
--- a/Engine/Source/Editor/UnrealEd/Private/EditorEngine.cpp
+++ b/Engine/Source/Editor/UnrealEd/Private/EditorEngine.cpp
@@ -7477,6 +7477,12 @@ void UEditorEngine::AllMaterialsCacheResourceShadersForRendering()
 
 void UEditorEngine::SetPreviewPlatform(const FName MaterialQualityPlatform, const ERHIFeatureLevel::Type PreviewFeatureLevel, const bool bSaveSettings/* = true*/)
 {
+#if RHI_RAYTRACING
+	if (IsRayTracingSupportedForThisProject())
+	{
+		ensure(PreviewFeatureLevel == ERHIFeatureLevel::SM5);
+	}
+#endif
 	// If we have specified a MaterialQualityPlatform ensure its feature level matches the requested feature level.
 	check(MaterialQualityPlatform.IsNone() || GetMaxSupportedFeatureLevel(ShaderFormatToLegacyShaderPlatform(MaterialQualityPlatform)) == PreviewFeatureLevel);
 
