Force preview feature level SM5 when ray tracing is enabled and editor was left in an incompatible feature level

diff --git a/Engine/Source/Editor/UnrealEd/Classes/Editor/EditorEngine.h b/Engine/Source/Editor/UnrealEd/Classes/Editor/EditorEngine.h
--- a/Engine/Source/Editor/UnrealEd/Classes/Editor/EditorEngine.h
+++ b/Engine/Source/Editor/UnrealEd/Classes/Editor/EditorEngine.h
@@ -2988,7 +2988,7 @@ public:
 	void OnSceneMaterialsModified();
 
 	/** Call this function to change the feature level and to override the material quality platform of the editor and PIE worlds */
-	void SetPreviewPlatform(const FName MaterialQualityPlatform, const ERHIFeatureLevel::Type PreviewFeatureLevel, const bool bSaveSettings = true);
+	void SetPreviewPlatform(const FName MaterialQualityPlatform, ERHIFeatureLevel::Type PreviewFeatureLevel, const bool bSaveSettings = true);
 
 protected:
 	/** Call this function to change the feature level of the editor and PIE worlds */
diff --git a/Engine/Source/Editor/UnrealEd/Private/EditorEngine.cpp b/Engine/Source/Editor/UnrealEd/Private/EditorEngine.cpp
--- a/Engine/Source/Editor/UnrealEd/Private/EditorEngine.cpp
+++ b/Engine/Source/Editor/UnrealEd/Private/EditorEngine.cpp
@@ -7475,12 +7475,16 @@ void UEditorEngine::AllMaterialsCacheResourceShadersForRendering()
 	UMaterialInstance::AllMaterialsCacheResourceShadersForRendering();
 }
 
-void UEditorEngine::SetPreviewPlatform(const FName MaterialQualityPlatform, const ERHIFeatureLevel::Type PreviewFeatureLevel, const bool bSaveSettings/* = true*/)
+void UEditorEngine::SetPreviewPlatform(const FName MaterialQualityPlatform, ERHIFeatureLevel::Type PreviewFeatureLevel, const bool bSaveSettings/* = true*/)
 {
 #if RHI_RAYTRACING
 	if (IsRayTracingSupportedForThisProject())
 	{
-		ensure(PreviewFeatureLevel == ERHIFeatureLevel::SM5);
+		if (PreviewFeatureLevel != ERHIFeatureLevel::SM5)
+		{
+			UE_LOG(LogEditor, Warning, TEXT("Preview feature level is incompatible with ray tracing, defaulting to Shader Model 5"));
+			PreviewFeatureLevel = ERHIFeatureLevel::SM5;
+		}
 	}
 #endif
 	// If we have specified a MaterialQualityPlatform ensure its feature level matches the requested feature level.
