DR - Missed file

diff --git a/Engine/Source/Editor/UnrealEd/Private/SEditorViewportViewMenu.cpp b/Engine/Source/Editor/UnrealEd/Private/SEditorViewportViewMenu.cpp
--- a/Engine/Source/Editor/UnrealEd/Private/SEditorViewportViewMenu.cpp
+++ b/Engine/Source/Editor/UnrealEd/Private/SEditorViewportViewMenu.cpp
@@ -301,7 +301,7 @@ TSharedRef<SWidget> SEditorViewportViewMenu::GenerateViewMenuContent() const
 			}
 
 #if RHI_RAYTRACING
-			if (IsRayTracingSupportedForThisProject() > 0)
+			if (IsRayTracingTierSupported(2))
 			{
 				struct Local
 				{
