DR - Monolithic fix

diff --git a/Engine/Source/Editor/UnrealEd/Private/SEditorViewportViewMenu.cpp b/Engine/Source/Editor/UnrealEd/Private/SEditorViewportViewMenu.cpp
--- a/Engine/Source/Editor/UnrealEd/Private/SEditorViewportViewMenu.cpp
+++ b/Engine/Source/Editor/UnrealEd/Private/SEditorViewportViewMenu.cpp
@@ -6,6 +6,7 @@
 #include "EditorStyleSet.h"
 #include "EditorViewportCommands.h"
 #include "RayTracingDebugVisualizationMenuCommands.h"
+#include "RenderResource.h"
 
 #define LOCTEXT_NAMESPACE "EditorViewportViewMenu"
 
