Clamp volumetric lightmap grid size so that it will not overflow int32 and crash.

diff --git a/Engine/Source/Editor/UnrealEd/Private/Lightmass/Lightmass.cpp b/Engine/Source/Editor/UnrealEd/Private/Lightmass/Lightmass.cpp
--- a/Engine/Source/Editor/UnrealEd/Private/Lightmass/Lightmass.cpp
+++ b/Engine/Source/Editor/UnrealEd/Private/Lightmass/Lightmass.cpp
@@ -2087,11 +2087,27 @@ void FLightmassExporter::SetVolumetricLightmapSettings(Lightmass::FVolumetricLig
 
 	const float TargetDetailCellSize = WorldInfoSettings.VolumetricLightmapDetailCellSize;
 
-	const FIntVector FullGridSize(
+	FIntVector FullGridSize(
 		FMath::TruncToInt(RequiredVolumeSize.X / TargetDetailCellSize) + 1,
 		FMath::TruncToInt(RequiredVolumeSize.Y / TargetDetailCellSize) + 1,
 		FMath::TruncToInt(RequiredVolumeSize.Z / TargetDetailCellSize) + 1);
 
+	if (FullGridSize.GetMax() > 800)
+	{
+		UE_LOG(LogLightmassSolver, Warning, 
+			TEXT("Volumetric lightmap grid size is too large which can cause potential crashes or long build time, clamping from (%d, %d, %d) to (%d, %d, %d)"),
+			FullGridSize.X,
+			FullGridSize.Y,
+			FullGridSize.Z,
+			FMath::Min(FullGridSize.X, 800),
+			FMath::Min(FullGridSize.Y, 800),
+			FMath::Min(FullGridSize.Z, 800)
+			);
+		FullGridSize.X = FMath::Min(FullGridSize.X, 800);
+		FullGridSize.Y = FMath::Min(FullGridSize.Y, 800);
+		FullGridSize.Z = FMath::Min(FullGridSize.Z, 800);
+	}
+
 	const int32 BrickSizeLog2 = FMath::FloorLog2(OutSettings.BrickSize);
 	const int32 DetailCellsPerTopLevelBrick = 1 << (OutSettings.MaxRefinementLevels * BrickSizeLog2);
 
