DR - Fix for releases of Uniform buffers after rebuilding lighting

diff --git a/Engine/Source/Editor/UnrealEd/Private/Lightmass/LightmassRender.cpp b/Engine/Source/Editor/UnrealEd/Private/Lightmass/LightmassRender.cpp
--- a/Engine/Source/Editor/UnrealEd/Private/Lightmass/LightmassRender.cpp
+++ b/Engine/Source/Editor/UnrealEd/Private/Lightmass/LightmassRender.cpp
@@ -816,10 +816,19 @@ private:
 
 FMaterialExportDataEntry::~FMaterialExportDataEntry()
 {
-	delete DiffuseMaterialProxy;
-	delete EmissiveMaterialProxy;
-	delete OpacityMaterialProxy;
-	delete NormalMaterialProxy;
+	FLightmassMaterialProxy* DiffuseMaterialProxy_ = DiffuseMaterialProxy;
+	FLightmassMaterialProxy* EmissiveMaterialProxy_ = EmissiveMaterialProxy;
+	FLightmassMaterialProxy* OpacityMaterialProxy_ = OpacityMaterialProxy;
+	FLightmassMaterialProxy* NormalMaterialProxy_ = NormalMaterialProxy;
+	ENQUEUE_RENDER_COMMAND(FMaterialExportDataEntryDelete)(
+		[DiffuseMaterialProxy_, EmissiveMaterialProxy_, OpacityMaterialProxy_, NormalMaterialProxy_](FRHICommandListImmediate& RHICmdList)
+		{
+			delete DiffuseMaterialProxy_;
+			delete EmissiveMaterialProxy_;
+			delete OpacityMaterialProxy_;
+			delete NormalMaterialProxy_;
+		}
+	);
 }
 
 //
