DR - Fixed style (minor)
#rnx

diff --git a/Engine/Source/Editor/UnrealEd/Private/Lightmass/LightmassRender.cpp b/Engine/Source/Editor/UnrealEd/Private/Lightmass/LightmassRender.cpp
--- a/Engine/Source/Editor/UnrealEd/Private/Lightmass/LightmassRender.cpp
+++ b/Engine/Source/Editor/UnrealEd/Private/Lightmass/LightmassRender.cpp
@@ -816,17 +816,17 @@ private:
 
 FMaterialExportDataEntry::~FMaterialExportDataEntry()
 {
-	FLightmassMaterialProxy* DiffuseMaterialProxy_ = DiffuseMaterialProxy;
-	FLightmassMaterialProxy* EmissiveMaterialProxy_ = EmissiveMaterialProxy;
-	FLightmassMaterialProxy* OpacityMaterialProxy_ = OpacityMaterialProxy;
-	FLightmassMaterialProxy* NormalMaterialProxy_ = NormalMaterialProxy;
+	FLightmassMaterialProxy* LocalDiffuseMaterialProxy = DiffuseMaterialProxy;
+	FLightmassMaterialProxy* LocalEmissiveMaterialProxy = EmissiveMaterialProxy;
+	FLightmassMaterialProxy* LocalOpacityMaterialProxy = OpacityMaterialProxy;
+	FLightmassMaterialProxy* LocalNormalMaterialProxy = NormalMaterialProxy;
 	ENQUEUE_RENDER_COMMAND(FMaterialExportDataEntryDelete)(
-		[DiffuseMaterialProxy_, EmissiveMaterialProxy_, OpacityMaterialProxy_, NormalMaterialProxy_](FRHICommandListImmediate& RHICmdList)
+		[LocalDiffuseMaterialProxy, LocalEmissiveMaterialProxy, LocalOpacityMaterialProxy, LocalNormalMaterialProxy](FRHICommandListImmediate& RHICmdList)
 		{
-			delete DiffuseMaterialProxy_;
-			delete EmissiveMaterialProxy_;
-			delete OpacityMaterialProxy_;
-			delete NormalMaterialProxy_;
+			delete LocalDiffuseMaterialProxy;
+			delete LocalEmissiveMaterialProxy;
+			delete LocalOpacityMaterialProxy;
+			delete LocalNormalMaterialProxy;
 		}
 	);
 }
