Fixed feature level switching logic by moving scene recreate inside the FGlobalComponentReregisterContext scope. This way during proxy recreate scene has new feature level. Before this change materials were cached for a new feature level, but proxies were reattached with a previous one causing crash.
#rnx

diff --git a/Engine/Source/Editor/UnrealEd/Classes/Editor/EditorEngine.h b/Engine/Source/Editor/UnrealEd/Classes/Editor/EditorEngine.h
--- a/Engine/Source/Editor/UnrealEd/Classes/Editor/EditorEngine.h
+++ b/Engine/Source/Editor/UnrealEd/Classes/Editor/EditorEngine.h
@@ -3055,10 +3055,10 @@ protected:
 	void SetFeatureLevelPreview(const ERHIFeatureLevel::Type InPreviewFeatureLevel);
 
 	/** call this function to change the feature level for all materials */
-	void SetMaterialsFeatureLevel(const ERHIFeatureLevel::Type InFeatureLevel);
+	void SetMaterialsFeatureLevel(const ERHIFeatureLevel::Type InPreviewFeatureLevel);
 
 	/** call this to recompile the materials */
-	void AllMaterialsCacheResourceShadersForRendering();
+	void AllMaterialsCacheResourceShadersForRendering(ERHIFeatureLevel::Type InPreviewFeatureLevel);
 
 	/** Function pair used to save and restore the global feature level */
 	void LoadEditorFeatureLevel();
diff --git a/Engine/Source/Editor/UnrealEd/Private/EditorEngine.cpp b/Engine/Source/Editor/UnrealEd/Private/EditorEngine.cpp
--- a/Engine/Source/Editor/UnrealEd/Private/EditorEngine.cpp
+++ b/Engine/Source/Editor/UnrealEd/Private/EditorEngine.cpp
@@ -7325,6 +7325,8 @@ void UEditorEngine::SetMaterialsFeatureLevel(const ERHIFeatureLevel::Type InFeat
 
 	SlowTask.EnterProgressFrame(15.0f, NSLOCTEXT("Engine", "SlowTaskFinalizingMessage", "Finalizing"));
 	GShaderCompilingManager->ProcessAsyncResults(false, true);
+
+	PreviewFeatureLevelChanged.Broadcast(InFeatureLevel);
 }
 
 void UEditorEngine::SetFeatureLevelPreview(const ERHIFeatureLevel::Type InPreviewFeatureLevel)
@@ -7339,12 +7341,13 @@ void UEditorEngine::SetFeatureLevelPreview(const ERHIFeatureLevel::Type InPrevie
 	}
 }
 
-void UEditorEngine::AllMaterialsCacheResourceShadersForRendering()
+void UEditorEngine::AllMaterialsCacheResourceShadersForRendering(ERHIFeatureLevel::Type InPreviewFeatureLevel)
 {
 	FGlobalComponentRecreateRenderStateContext Recreate;
 	FlushRenderingCommands();
 	UMaterial::AllMaterialsCacheResourceShadersForRendering(true);
 	UMaterialInstance::AllMaterialsCacheResourceShadersForRendering(true);
+	PreviewFeatureLevelChanged.Broadcast(InPreviewFeatureLevel);
 }
 
 void UEditorEngine::SetPreviewPlatform(const FName MaterialQualityPlatform, ERHIFeatureLevel::Type InPreviewFeatureLevel, const bool bSaveSettings/* = true*/)
@@ -7374,11 +7377,9 @@ void UEditorEngine::SetPreviewPlatform(const FName MaterialQualityPlatform, ERHI
 	else if (InitialPreviewPlatform != MaterialQualityPlatform)
 	{
 		// Rebuild materials if we have the same feature level but a different 'material quality platform'
-		AllMaterialsCacheResourceShadersForRendering();
+		AllMaterialsCacheResourceShadersForRendering(InPreviewFeatureLevel);
 	}
 
-	PreviewFeatureLevelChanged.Broadcast(InPreviewFeatureLevel);
-
 	if (bSaveSettings)
 	{
 		SaveEditorFeatureLevel();
