When replacing an actor ensure that its Owner and Instigator are correctly maintained.

diff --git a/Engine/Source/Editor/UnrealEd/Private/Kismet2/KismetReinstanceUtilities.cpp b/Engine/Source/Editor/UnrealEd/Private/Kismet2/KismetReinstanceUtilities.cpp
--- a/Engine/Source/Editor/UnrealEd/Private/Kismet2/KismetReinstanceUtilities.cpp
+++ b/Engine/Source/Editor/UnrealEd/Private/Kismet2/KismetReinstanceUtilities.cpp
@@ -1783,6 +1783,8 @@ static void ReplaceActorHelper(UObject* OldObject, UClass* OldClass, UObject*& N
 
 	FActorSpawnParameters SpawnInfo;
 	SpawnInfo.OverrideLevel = ActorLevel;
+	SpawnInfo.Owner = OldActor->GetOwner();
+	SpawnInfo.Instigator = OldActor->GetInstigator();
 	SpawnInfo.Template = NewArchetype;
 	SpawnInfo.SpawnCollisionHandlingOverride = ESpawnActorCollisionHandlingMethod::AlwaysSpawn;
 	SpawnInfo.bDeferConstruction = true;
