Merging 4824984 from  Dev-Enterprise to Release-4.22

Fix the convert actors to bp not always preserving the hierarchy and the right transforms.

There was in a issue with the function AddComponentsToBlueprint were the hierarchy wasn't always properly remade.

The convert selected actor to blueprint and the convert components to blueprint didn't adjust the scale of the new actor created.

The convert selected components to blueprint did only work properly when the
root components relative transform were the same as their world transform.

diff --git a/Engine/Source/Editor/UnrealEd/Private/Kismet2/Kismet2.cpp b/Engine/Source/Editor/UnrealEd/Private/Kismet2/Kismet2.cpp
--- a/Engine/Source/Editor/UnrealEd/Private/Kismet2/Kismet2.cpp
+++ b/Engine/Source/Editor/UnrealEd/Private/Kismet2/Kismet2.cpp
@@ -1381,29 +1381,90 @@ void FKismetEditorUtilities::AddComponentsToBlueprint(UBlueprint* Blueprint, TAr
 		}
 	};
 
-	// Clear out nulls and components we won't be able to create, can swap since we're about to sort anyways
-	Components.RemoveAllSwap([](const UActorComponent* ActorComponent)
+	// Associate a scene node to its first attached parent if there is one.
+	struct FSceneComponentAndFirstParent
 	{
-		return (ActorComponent == nullptr || !ActorComponent->GetClass()->HasMetaData(FBlueprintMetadata::MD_BlueprintSpawnableComponent));
-	});
+		FSceneComponentAndFirstParent(USceneComponent* InSceneComponent, USceneComponent* InFirstAttachParent)
+			: SceneComponent(InSceneComponent)
+			, FirstAttachParent(InFirstAttachParent)
+		{}
 
-	// Sort scene components to the start of the array and those that are children after their parents
-	Components.Sort([](const UActorComponent& A, const UActorComponent& B)
+		USceneComponent* SceneComponent;
+
+		// The first attach parent is valid when we find a parent that is part of the components array.
+		USceneComponent* FirstAttachParent;
+	};
+
+	// Array of FSceneComponentAndFirstParent that will be filled in a way that the parents are always before their children
+	TArray<FSceneComponentAndFirstParent> SceneComponentNodes;
+	SceneComponentNodes.Reserve(Components.Num());
+
+	// The boolean indicate if the scene component was added to the SceneComponentNodes array
+	TMap<USceneComponent*, bool> SceneComponentsMap;
+	SceneComponentsMap.Reserve(Components.Num());
+
+	for (int32 CompIndex = 0; CompIndex < Components.Num(); ++CompIndex)
 	{
-		const USceneComponent* SceneA = Cast<const USceneComponent>(&A);
-		const USceneComponent* SceneB = Cast<const USceneComponent>(&B);
+		UActorComponent* ActorComponent = Components[CompIndex];
 
-		if (SceneA)
+		// Clear out nulls and the components we won't be able to create.
+		if (ActorComponent == nullptr || !ActorComponent->GetClass()->HasMetaData(FBlueprintMetadata::MD_BlueprintSpawnableComponent))
 		{
-			if (SceneB)
+			Components.RemoveAtSwap(CompIndex, 1, false);
+			// We don't want to skip the component swapped in
+			CompIndex--;
+		}
+		else
+		{
+			USceneComponent* SceneComponent = Cast<USceneComponent>(Components[CompIndex]);
+			// Remove the scene components from the array as they will be added in the array SceneComponentNodes
+			if (SceneComponent)
 			{
-				return SceneB->IsAttachedTo(SceneA);
+				SceneComponentsMap.Add(SceneComponent, false);
+				Components.RemoveAtSwap(CompIndex, 1, false);
+				// We don't want to skip the component swapped in
+				CompIndex--;
 			}
-			return true;
 		}
-		return (SceneB == nullptr);
-	});
+	}
+
+	TFunction<void (USceneComponent*)> ConstructSceneComponentNodes;
+	// Fill the SceneComponentNodes array
+	ConstructSceneComponentNodes = [&SceneComponentsMap, &SceneComponentNodes, &ConstructSceneComponentNodes] (USceneComponent* SceneComponent)
+	{
+		// The scene component should always be present in the map
+		if (*SceneComponentsMap.Find(SceneComponent))
+		{
+			// If the node for a scene component is already part of the array just return
+			return;
+		}
+		else
+		{
+			USceneComponent* Parent = SceneComponent->GetAttachParent();
+			while (Parent)
+			{
+				if (SceneComponentsMap.Contains(Parent))
+				{
+					// always add the parent first
+					ConstructSceneComponentNodes(Parent);
+					break;
+				}
+				Parent = Parent->GetAttachParent();
+			}
+
+			SceneComponentNodes.Emplace(SceneComponent, Parent);
+			SceneComponentsMap.Add(SceneComponent, true);
+		}
+	};
+
+	TArray<USceneComponent*> SceneComponents;
+	SceneComponentsMap.GetKeys(SceneComponents);
+	for (USceneComponent* SceneComponent : SceneComponents)
+	{
+		ConstructSceneComponentNodes(SceneComponent);
+	}
 
+	// The easy part to add the non-scene components.
 	for (int32 CompIndex = 0; CompIndex < Components.Num(); ++CompIndex)
 	{
 		UActorComponent* ActorComponent = Components[CompIndex];
@@ -1412,121 +1473,114 @@ void FKismetEditorUtilities::AddComponentsToBlueprint(UBlueprint* Blueprint, TAr
 
 		USCS_Node* SCSNode = FAddComponentsToBlueprintImpl::MakeComponentCopy(ActorComponent, SCS, InstanceComponentToNodeMap, bKeepMobility);
 
-		USceneComponent* SceneComponent = Cast<USceneComponent>(ActorComponent);
-		// The easy part is non-scene component or the Root simply add it
-		if (SceneComponent == nullptr)
-		{
-			SCS->AddNode(SCSNode);
-		}
-		else
+		SCS->AddNode(SCSNode);
+	}
+
+	// The loop to add the scene components
+	for (int32 CompIndex = 0; CompIndex < SceneComponentNodes.Num(); ++CompIndex)
+	{
+		const FSceneComponentAndFirstParent& ComponentNode = SceneComponentNodes[CompIndex];
+		USceneComponent* SceneComponent = ComponentNode.SceneComponent;
+		AActor* Actor = SceneComponent->GetOwner();
+		check(Actor);
+
+		USCS_Node* SCSNode = FAddComponentsToBlueprintImpl::MakeComponentCopy(SceneComponent, SCS, InstanceComponentToNodeMap, bKeepMobility);
+
+		USceneComponent* FirstAttachParent = ComponentNode.FirstAttachParent;
+		
+		if (FirstAttachParent)
 		{
-			// Since there is no guarantee that the attach parent is in the component set, but there could still be attachment hierarchy
-			// which is best to maintain, seek out the closest attached parent in the component set
-			USceneComponent* FirstAttachParent = nullptr;
-			if (SceneComponent->GetAttachParent())
+			// If the parent we're going to be attached to isn't the original attach parent, then adjust the
+			// relative transform so that the end result will be consistent with current relationship
+			if (FirstAttachParent != SceneComponent->GetAttachParent())
 			{
-				const int32 FirstAttachParentIndex = Components.FindLastByPredicate([SceneComponent](const UActorComponent* Other)
-				{
-					return (SceneComponent->IsAttachedTo(CastChecked<const USceneComponent>(Other)));
-				}, CompIndex);
+				USceneComponent* SceneComponentTemplate = CastChecked<USceneComponent>(SCSNode->ComponentTemplate);
 
-				if (FirstAttachParentIndex != INDEX_NONE)
+				const FTransform ComponentToWorld = SceneComponent->GetComponentTransform();
+				const FTransform RelativeTransform = ComponentToWorld.GetRelativeTransform(FirstAttachParent->GetComponentTransform());
+				if (!SceneComponent->bAbsoluteLocation)
 				{
-					FirstAttachParent = CastChecked<USceneComponent>(Components[FirstAttachParentIndex]);
-
-					// If the parent we're going to be attached to isn't the original attach parent, then adjust the
-					// relative transform so that the end result will be consistent with current relationship
-					if (FirstAttachParent != SceneComponent->GetAttachParent())
-					{
-						USceneComponent* SceneComponentTemplate = CastChecked<USceneComponent>(SCSNode->ComponentTemplate);
-
-						const FTransform ComponentToWorld = SceneComponent->GetComponentTransform();
-						const FTransform RelativeTransform = ComponentToWorld.GetRelativeTransform(FirstAttachParent->GetComponentTransform());
-						if (!SceneComponent->bAbsoluteLocation)
-						{
-							SceneComponentTemplate->RelativeLocation = RelativeTransform.GetLocation();
-						}
-						if (!SceneComponent->bAbsoluteRotation)
-						{
-							SceneComponentTemplate->RelativeRotation = RelativeTransform.GetRotation().Rotator();
-						}
-						if (!SceneComponent->bAbsoluteScale)
-						{
-							SceneComponentTemplate->RelativeScale3D = RelativeTransform.GetScale3D();
-						}
-					}
+					SceneComponentTemplate->RelativeLocation = RelativeTransform.GetLocation();
 				}
-			}
-
-			if (SceneComponent == Actor->GetRootComponent() && (SceneComponent->GetAttachParent() == nullptr || FirstAttachParent == nullptr))
-			{
-				if (OptionalNewRootNode != nullptr)
+				if (!SceneComponent->bAbsoluteRotation)
 				{
-					OptionalNewRootNode->AddChildNode(SCSNode);
+					SceneComponentTemplate->RelativeRotation = RelativeTransform.GetRotation().Rotator();
 				}
-				else
+				if (!SceneComponent->bAbsoluteScale)
 				{
-					SCS->AddNode(SCSNode);
+					SceneComponentTemplate->RelativeScale3D = RelativeTransform.GetScale3D();
 				}
 			}
-			// If we're not attached to a blueprint component, add ourself to the root node or the SCS root component:
-			else if (SceneComponent->GetAttachParent() == nullptr)
+		}
+
+		if (SceneComponent == Actor->GetRootComponent() && (SceneComponent->GetAttachParent() == nullptr || FirstAttachParent == nullptr))
+		{
+			if (OptionalNewRootNode != nullptr)
 			{
-					AddChildToSCSRootNodeLambda(SCSNode);
+				OptionalNewRootNode->AddChildNode(SCSNode);
 			}
-			// If we're attached to a blueprint component look it up as the variable name is the component name
-			else if (SceneComponent->GetAttachParent()->IsCreatedByConstructionScript())
+			else
 			{
-				USCS_Node* ParentSCSNode = nullptr;
-				if (FirstAttachParent)
-				{
-					ParentSCSNode = InstanceComponentToNodeMap.FindChecked(FirstAttachParent);
-				}
-				else
+				SCS->AddNode(SCSNode);
+			}
+		}
+		// If we're not attached to a blueprint component, add ourself to the root node or the SCS root component:
+		else if (SceneComponent->GetAttachParent() == nullptr)
+		{
+				AddChildToSCSRootNodeLambda(SCSNode);
+		}
+		// If we're attached to a blueprint component look it up as the variable name is the component name
+		else if (SceneComponent->GetAttachParent()->IsCreatedByConstructionScript())
+		{
+			USCS_Node* ParentSCSNode = nullptr;
+			if (FirstAttachParent)
+			{
+				ParentSCSNode = InstanceComponentToNodeMap.FindChecked(FirstAttachParent);
+			}
+			else
+			{
+				for (UBlueprint* ParentBlueprint : ParentBPStack)
 				{
-					for (UBlueprint* ParentBlueprint : ParentBPStack)
+					if (ParentBlueprint->SimpleConstructionScript)
 					{
-						if (ParentBlueprint->SimpleConstructionScript)
+						ParentSCSNode = ParentBlueprint->SimpleConstructionScript->FindSCSNode(SceneComponent->GetAttachParent()->GetFName());
+						if (ParentSCSNode)
 						{
-							ParentSCSNode = ParentBlueprint->SimpleConstructionScript->FindSCSNode(SceneComponent->GetAttachParent()->GetFName());
-							if (ParentSCSNode)
-							{
-								break;
-							}
+							break;
 						}
 					}
 				}
-				check(ParentSCSNode);
-
-				if (ParentSCSNode->GetSCS() != SCS)
-				{
-					SCS->AddNode(SCSNode);
-					SCSNode->SetParent(ParentSCSNode);
-				}
-				else
-				{
-					ParentSCSNode->AddChildNode(SCSNode);
-				}
 			}
-			else if ((SceneComponent->GetAttachParent()->CreationMethod == EComponentCreationMethod::Native) && !bHarvesting)
+			check(ParentSCSNode);
+
+			if (ParentSCSNode->GetSCS() != SCS)
 			{
-				// If we're attached to a component that will be native in the new blueprint
 				SCS->AddNode(SCSNode);
-				SCSNode->SetParent(SceneComponent->GetAttachParent());
+				SCSNode->SetParent(ParentSCSNode);
 			}
 			else
 			{
-				// Otherwise we will already have created the parents' new SCS node, so attach to that
-				if (USCS_Node** ParentSCSNode = InstanceComponentToNodeMap.Find(FirstAttachParent))
-				{
-					(*ParentSCSNode)->AddChildNode(SCSNode);
-				}
-				else
-				{
-					// Unsure what this case is really for when harvesting, but to be consistent with previous behavior still do this
-					ensure(bHarvesting);
-					AddChildToSCSRootNodeLambda(SCSNode);
-				}
+				ParentSCSNode->AddChildNode(SCSNode);
+			}
+		}
+		else if ((SceneComponent->GetAttachParent()->CreationMethod == EComponentCreationMethod::Native) && !bHarvesting)
+		{
+			// If we're attached to a component that will be native in the new blueprint
+			SCS->AddNode(SCSNode);
+			SCSNode->SetParent(SceneComponent->GetAttachParent());
+		}
+		else
+		{
+			// Otherwise we will already have created the parents' new SCS node, so attach to that
+			if (USCS_Node** ParentSCSNode = InstanceComponentToNodeMap.Find(FirstAttachParent))
+			{
+				(*ParentSCSNode)->AddChildNode(SCSNode);
+			}
+			else
+			{
+				// Unsure what this case is really for when harvesting, but to be consistent with previous behavior still do this
+				ensure(bHarvesting);
+				AddChildToSCSRootNodeLambda(SCSNode);
 			}
 		}
 	}
@@ -1605,7 +1659,11 @@ UBlueprint* FKismetEditorUtilities::CreateBlueprintFromActor(const FName Bluepri
 				FVector Location = Actor->GetActorLocation();
 				FRotator Rotator = Actor->GetActorRotation();
 
-				CreateBlueprintInstanceFromSelection(NewBlueprint, Actors, Location, Rotator);
+				AActor* NewActor = CreateBlueprintInstanceFromSelection(NewBlueprint, Actors, Location, Rotator);
+				if (NewActor)
+				{
+					NewActor->SetActorScale3D(Actor->GetActorScale3D());
+				}
 			}
 		}
 	}
@@ -1703,8 +1761,28 @@ public:
 					}
 				}
 			}
+
+			TArray<TPair<USceneComponent*, FTransform>> SceneComponentOldRelativeTransforms;
+			if (RootActors.Num() > 1)
+			{
+				SceneComponentOldRelativeTransforms.Reserve(RootActors.Num());
+				// Convert the components relative transform to world 
+				for (AActor* Actor : RootActors)
+				{
+					USceneComponent* SceneComponent = Actor->GetRootComponent();
+					SceneComponentOldRelativeTransforms.Emplace(SceneComponent, SceneComponent->GetRelativeTransform());
+					SceneComponent->SetRelativeTransform(SceneComponent->GetComponentTransform());
+				}
+			}
+
 			FKismetEditorUtilities::AddComponentsToBlueprint(Blueprint, MoveTemp(AllSelectedComponents), /*bHarvesting=*/ true, RootNodeOverride);
 
+			// Replace the modified components to their relative transform
+			for (const TPair<USceneComponent*, FTransform >& Pair : SceneComponentOldRelativeTransforms)
+			{
+				Pair.Key->SetRelativeTransform(Pair.Value);
+			}
+
 			FTransform NewActorTransform = FTransform::Identity;
 			if (RootActors.Num() == 1)
 			{
@@ -1738,10 +1816,9 @@ public:
 						{
 							if (USceneComponent* ChildComponent = Cast<USceneComponent>(ChildNode->ComponentTemplate))
 							{
-								const FTransform OldChildToWorld(ChildComponent->RelativeRotation, ChildComponent->RelativeLocation);
-								const FTransform NewRelativeTransform = OldChildToWorld.GetRelativeTransform(NewActorTransform);
-								ChildComponent->RelativeLocation = NewRelativeTransform.GetLocation();
-								ChildComponent->RelativeRotation = NewRelativeTransform.GetRotation().Rotator();
+								//The relative transform for those component was converted into the world space
+								const FTransform NewRelativeTransform = ChildComponent->GetRelativeTransform().GetRelativeTransform(NewActorTransform);
+								ChildComponent->SetRelativeTransform(NewRelativeTransform);
 							}
 						}
 					}
@@ -1763,7 +1840,11 @@ public:
 				FVector Location = NewActorTransform.GetLocation();
 				FRotator Rotator = NewActorTransform.Rotator();
 
-				FKismetEditorUtilities::CreateBlueprintInstanceFromSelection(Blueprint, SelectedActors, Location, Rotator);
+				AActor* NewActor = FKismetEditorUtilities::CreateBlueprintInstanceFromSelection(Blueprint, SelectedActors, Location, Rotator);
+				if (NewActor)
+				{
+					NewActor->SetActorScale3D(NewActorTransform.GetScale3D());
+				}
 			}
 
 			// Open the editor for the new blueprint
@@ -1813,7 +1894,6 @@ AActor* FKismetEditorUtilities::CreateBlueprintInstanceFromSelection(UBlueprint*
 	AActor* NewActor = World->SpawnActor(Blueprint->GeneratedClass, &Location, &Rotator);
 	GEditor->Layers->InitializeNewActorLayers(NewActor);
 	FActorLabelUtilities::SetActorLabelUnique(NewActor, Blueprint->GetName());
-
 	// Quietly ensure that no components are selected
 	USelection* ComponentSelection = GEditor->GetSelectedComponents();
 	ComponentSelection->BeginBatchSelectOperation();
