Prevented server from compressing pak files.
Enabled compression on windows for some pak files.

diff --git a/Engine/Source/Editor/UnrealEd/Private/Commandlets/AssetRegistryGenerator.cpp b/Engine/Source/Editor/UnrealEd/Private/Commandlets/AssetRegistryGenerator.cpp
--- a/Engine/Source/Editor/UnrealEd/Private/Commandlets/AssetRegistryGenerator.cpp
+++ b/Engine/Source/Editor/UnrealEd/Private/Commandlets/AssetRegistryGenerator.cpp
@@ -210,10 +210,14 @@ bool FAssetRegistryGenerator::GenerateStreamingInstallManifest(int64 InExtraFlav
 
 	TArray<FString> CompressedChunkWildcards;
 	{
-		FConfigFile PlatformIniFile;
-		FConfigCacheIni::LoadLocalIniFile(PlatformIniFile, TEXT("Game"), true, *TargetPlatform->IniPlatformName());
-		FString ConfigString;
-		PlatformIniFile.GetArray(TEXT("/Script/UnrealEd.ProjectPackagingSettings"), TEXT("CompressedChunkWildcard"), CompressedChunkWildcards);
+		// never screw with server pak files.  This hack only cares about client platforms
+		if (!TargetPlatform->IsServerOnly())
+		{
+			FConfigFile PlatformIniFile;
+			FConfigCacheIni::LoadLocalIniFile(PlatformIniFile, TEXT("Game"), true, *TargetPlatform->IniPlatformName());
+			FString ConfigString;
+			PlatformIniFile.GetArray(TEXT("/Script/UnrealEd.ProjectPackagingSettings"), TEXT("CompressedChunkWildcard"), CompressedChunkWildcards);
+		}
 	}
 
 	// Add manifests for any staging-time only groups
diff --git a/Engine/Source/Programs/AutomationTool/Win/WinPlatform.Automation.cs b/Engine/Source/Programs/AutomationTool/Win/WinPlatform.Automation.cs
--- a/Engine/Source/Programs/AutomationTool/Win/WinPlatform.Automation.cs
+++ b/Engine/Source/Programs/AutomationTool/Win/WinPlatform.Automation.cs
@@ -178,13 +178,14 @@ public abstract class BaseWinPlatform : Platform
 	public override string GetPlatformPakCommandLine(ProjectParams Params, DeploymentContext SC)
 	{
 		string PakParams = " -patchpaddingalign=2048";
-		/*
-		string OodleDllPath = DirectoryReference.Combine(SC.ProjectRoot, "Binaries/ThirdParty/Oodle/Win64/UnrealPakPlugin.dll").FullName;
-		if (File.Exists(OodleDllPath))
+		if (!SC.DedicatedServer)
 		{
-			PakParams += String.Format(" -customcompressor=\"{0}\"", OodleDllPath);
+			string OodleDllPath = DirectoryReference.Combine(SC.ProjectRoot, "Binaries/ThirdParty/Oodle/Win64/UnrealPakPlugin.dll").FullName;
+			if (File.Exists(OodleDllPath))
+			{
+				PakParams += String.Format(" -customcompressor=\"{0}\"", OodleDllPath);
+			}
 		}
-		*/
 		return PakParams;
 	}
 
