#rnx

diff --git a/Engine/Source/Editor/UnrealEd/Classes/Editor/EditorEngine.h b/Engine/Source/Editor/UnrealEd/Classes/Editor/EditorEngine.h
--- a/Engine/Source/Editor/UnrealEd/Classes/Editor/EditorEngine.h
+++ b/Engine/Source/Editor/UnrealEd/Classes/Editor/EditorEngine.h
@@ -2666,6 +2666,14 @@ public:
 	 */
 	void OnViewportCloseRequested(FViewport* InViewport);
 
+	/*
+	 * Fills level viewport client transform used in simulation mode 
+	 *
+	 * @param OutViewTransform filled view transform used in SIE
+	 * @return true if function succeeded, false if failed
+	 */
+	bool GetSimulateInEditorViewTransform(FTransform& OutViewTransform) const;
+
 private:
 	/** Gets the scene viewport for a viewport client */
 	FSceneViewport* GetGameSceneViewport(UGameViewportClient* ViewportClient) const;
diff --git a/Engine/Source/Editor/UnrealEd/Private/PlayLevel.cpp b/Engine/Source/Editor/UnrealEd/Private/PlayLevel.cpp
--- a/Engine/Source/Editor/UnrealEd/Private/PlayLevel.cpp
+++ b/Engine/Source/Editor/UnrealEd/Private/PlayLevel.cpp
@@ -3429,6 +3429,34 @@ FViewport* UEditorEngine::GetPIEViewport()
 	return nullptr;
 }
 
+bool UEditorEngine::GetSimulateInEditorViewTransform(FTransform& OutViewTransform) const
+{
+	if (bIsSimulatingInEditor)
+	{
+		// The first PIE world context is the one that can toggle between PIE and SIE
+		for (const FWorldContext& WorldContext : WorldList)
+		{
+			if (WorldContext.WorldType == EWorldType::PIE && !WorldContext.RunAsDedicated)
+			{
+				const FSlatePlayInEditorInfo* SlateInfoPtr = SlatePlayInEditorMap.Find(WorldContext.ContextHandle);
+				if (SlateInfoPtr)
+				{
+					// This is only supported inside SLevelEditor viewports currently
+					TSharedPtr<ILevelViewport> LevelViewport = SlateInfoPtr->DestinationSlateViewport.Pin();
+					if (LevelViewport.IsValid())
+					{
+						FLevelEditorViewportClient& EditorViewportClient = LevelViewport->GetLevelViewportClient();
+						OutViewTransform = FTransform(EditorViewportClient.GetViewRotation(), EditorViewportClient.GetViewLocation());
+						return true;
+					}
+				}
+				break;
+			}
+		}
+	}
+	return false;
+}
+
 void UEditorEngine::ToggleBetweenPIEandSIE( bool bNewSession )
 {
 	bIsToggleBetweenPIEandSIEQueued = false;
