Reverting CL# 4703649, as it doesn't handle cases where a blueprint class may be dropable even though it's native parent isn't.

diff --git a/Engine/Source/Editor/UnrealEd/Private/LevelEditorViewport.cpp b/Engine/Source/Editor/UnrealEd/Private/LevelEditorViewport.cpp
--- a/Engine/Source/Editor/UnrealEd/Private/LevelEditorViewport.cpp
+++ b/Engine/Source/Editor/UnrealEd/Private/LevelEditorViewport.cpp
@@ -1201,12 +1201,7 @@ FDropQuery FLevelEditorViewportClient::CanDropObjectsAtCoordinates(int32 MouseX,
 	}
 
 	UObject* AssetObj = AssetData.GetAsset();
-	UClass* ClassObj = Cast<UClass>(AssetObj);
-	if (!ClassObj)
-	{
-		UBlueprint* BlueprintObj = Cast<UBlueprint>(AssetObj);
-		ClassObj = BlueprintObj ? BlueprintObj->ParentClass : nullptr;
-	}
+	UClass* ClassObj = Cast<UClass>( AssetObj );
 
 	if ( ClassObj )
 	{
