- Adding config entry checkup to support additional files not part of package copies when generating distillation

diff --git a/Engine/Source/Editor/UnrealEd/Private/Commandlets/GenerateDistillFileSetsCommandlet.cpp b/Engine/Source/Editor/UnrealEd/Private/Commandlets/GenerateDistillFileSetsCommandlet.cpp
--- a/Engine/Source/Editor/UnrealEd/Private/Commandlets/GenerateDistillFileSetsCommandlet.cpp
+++ b/Engine/Source/Editor/UnrealEd/Private/Commandlets/GenerateDistillFileSetsCommandlet.cpp
@@ -341,6 +341,55 @@ int32 UGenerateDistillFileSetsCommandlet::Main( const FString& InParams )
 		}
 	}
 
+	// Add additional files marked for distillation
+	TArray<FString> AdditionalFilesToDistill;
+	GConfig->GetArray(TEXT("DistillSettings"), TEXT("FilesToAlwaysDistill"), AdditionalFilesToDistill, GEngineIni);
+
+	TArray<FString> AdditionalFilesToDistillFullPath;
+	for (const FString& File : AdditionalFilesToDistill)
+	{
+		//Only support path relative to content directory
+		if (!FPaths::IsRelative(File))
+		{
+			continue;
+		}
+
+		FString FileAbsolutePath = FPaths::ConvertRelativePathToFull(FPaths::Combine(AbsoluteGameContentDir, File));
+
+		FString Path;
+		FString Filename;
+		FString Extension;
+		FPaths::Split(FileAbsolutePath, Path, Filename, Extension);
+		if (!FPaths::DirectoryExists(Path)
+			|| Filename.IsEmpty())
+		{
+			continue;
+		}
+
+		//Verify if we're looking at an arbitrary amount of desired files
+		if (Filename == TEXT("*"))
+		{
+			TArray<FString> SubDirectoryFiles;
+			IFileManager::Get().FindFilesRecursive(SubDirectoryFiles, *Path, *(Filename + (!Extension.IsEmpty() ? TEXT(".") + Extension : TEXT(""))), true, false);
+			AdditionalFilesToDistillFullPath.Append(SubDirectoryFiles);
+		}
+		else
+		{
+			AdditionalFilesToDistillFullPath.Add(FileAbsolutePath);
+		}
+	}
+
+	for (FString& AdditionalFile : AdditionalFilesToDistillFullPath)
+	{
+		FPaths::MakeStandardFilename(AdditionalFile);
+		if (FPaths::FileExists(AdditionalFile))
+		{
+			AdditionalFile = IFileManager::Get().ConvertToAbsolutePathForExternalAppForRead(*AdditionalFile);
+			AllFileSets += FString::Printf(TEXT("%s") LINE_TERMINATOR, *AdditionalFile);
+			UE_LOG(LogGenerateDistillFileSetsCommandlet, Log, TEXT("Additional file: %s"), *AdditionalFile);
+		}
+	}
+
 	// Write the output file
 	FString OutputFileContents;
 	if (bSimpleTxtOutput)
