Editor: better wording for the options.

diff --git a/Engine/Source/Editor/UnrealEd/Private/Analytics/AnalyticsPrivacySettings.cpp b/Engine/Source/Editor/UnrealEd/Private/Analytics/AnalyticsPrivacySettings.cpp
--- a/Engine/Source/Editor/UnrealEd/Private/Analytics/AnalyticsPrivacySettings.cpp
+++ b/Engine/Source/Editor/UnrealEd/Private/Analytics/AnalyticsPrivacySettings.cpp
@@ -56,7 +56,7 @@ FString UAnalyticsPrivacySettings::GetAdditionalInfoUrl() const
 
 FText UAnalyticsPrivacySettings::GetAdditionalInfoUrlLabel() const
 {
-	return LOCTEXT("HyperlinkLabel", "Epic Games Privacy Notice");
+	return LOCTEXT("HyperlinkLabel", "Epic Games Privacy Policy");
 };
 
 #if WITH_EDITOR
diff --git a/Engine/Source/Editor/UnrealEd/Private/Analytics/CrashReportsPrivacySettings.cpp b/Engine/Source/Editor/UnrealEd/Private/Analytics/CrashReportsPrivacySettings.cpp
--- a/Engine/Source/Editor/UnrealEd/Private/Analytics/CrashReportsPrivacySettings.cpp
+++ b/Engine/Source/Editor/UnrealEd/Private/Analytics/CrashReportsPrivacySettings.cpp
@@ -31,7 +31,7 @@ FText UCrashReportsPrivacySettings::GetFalseStateTooltip() const
 
 FText UCrashReportsPrivacySettings::GetFalseStateDescription() const
 {
-	return LOCTEXT("FalseStateDescription", "Turn on this setting to allow the editor send non-critical bug reports by default to Epic Games to help us improve Unreal Engine. Please note that if you turn off this setting, we may ask you to share information about critical bugs, but sharing such information is optional. You can find out more at our Privacy Policy.");
+	return LOCTEXT("FalseStateDescription", "By opting out you have chosen to not send Editor non-critical bug reports by default to Epic Games. Please consider opting in to help improve Unreal Engine. Please note that while this setting is disabled, we may ask you to share information about critical bugs, but sharing such information is optional. You can find out more at our Privacy Policy.");
 };
 
 FText UCrashReportsPrivacySettings::GetTrueStateLabel() const
@@ -46,7 +46,7 @@ FText UCrashReportsPrivacySettings::GetTrueStateTooltip() const
 
 FText UCrashReportsPrivacySettings::GetTrueStateDescription() const
 {
-	return LOCTEXT("TrueStateDescription", "Turn off this setting if you prefer to not have the editor send non-critical bug reports by default to Epic Games to help us improve Unreal Engine. Please note that if you turn off this setting, we may ask you to share information about critical bugs, but sharing such information is optional. You can find out more at our Privacy Policy.");
+	return LOCTEXT("TrueStateDescription", "By opting in you have chosen to send Editor non-critical bug reports by default to Epic Games to help us improve Unreal Engine. Please note that if you turn off this setting, we may ask you to share information about critical bugs, but sharing such information is optional. You can find out more at our Privacy Policy.");
 };
 
 FString UCrashReportsPrivacySettings::GetAdditionalInfoUrl() const
@@ -56,7 +56,7 @@ FString UCrashReportsPrivacySettings::GetAdditionalInfoUrl() const
 
 FText UCrashReportsPrivacySettings::GetAdditionalInfoUrlLabel() const
 {
-	return LOCTEXT("HyperlinkLabel", "Epic Games Privacy Notice");
+	return LOCTEXT("HyperlinkLabel", "Epic Games Privacy Policy");
 };
 
 #undef LOCTEXT_NAMESPACE
