Put in missing null check from CL# 4279612

diff --git a/Engine/Source/Editor/UnrealEd/Private/AssetSelection.cpp b/Engine/Source/Editor/UnrealEd/Private/AssetSelection.cpp
--- a/Engine/Source/Editor/UnrealEd/Private/AssetSelection.cpp
+++ b/Engine/Source/Editor/UnrealEd/Private/AssetSelection.cpp
@@ -212,20 +212,23 @@ namespace AssetSelectionUtils
 
 					for (UActorComponent* Component : CurrentActor->GetComponents())
 					{
-						if( UStaticMeshComponent* SMComp = Cast<UStaticMeshComponent>(Component) )
+						if (Component)
 						{
-							if( SMComp->IsRegistered() )
+							if( UStaticMeshComponent* SMComp = Cast<UStaticMeshComponent>(Component) )
 							{
-								ActorInfo.bHaveStaticMeshComponent = true;
+								if( SMComp->IsRegistered() )
+								{
+									ActorInfo.bHaveStaticMeshComponent = true;
+								}
 							}
-						}
 
-						// Check for experimental/early-access classes in the component hierarchy
-						bool bIsExperimental, bIsEarlyAccess;
-						FObjectEditorUtils::GetClassDevelopmentStatus(Component->GetClass(), bIsExperimental, bIsEarlyAccess);
+							// Check for experimental/early-access classes in the component hierarchy
+							bool bIsExperimental, bIsEarlyAccess;
+							FObjectEditorUtils::GetClassDevelopmentStatus(Component->GetClass(), bIsExperimental, bIsEarlyAccess);
 
-						ActorInfo.bHaveExperimentalClass |= bIsExperimental;
-						ActorInfo.bHaveEarlyAccessClass |= bIsEarlyAccess;
+							ActorInfo.bHaveExperimentalClass |= bIsExperimental;
+							ActorInfo.bHaveEarlyAccessClass |= bIsEarlyAccess;
+						}
 					}
 
 					// Check for experimental/early-access classes in the actor hierarchy
