Copy of 4963994 from //UE4/Dev-Enterprise:

DS DataPrep - Remove operation - fix asset deletion issue

diff --git a/Engine/Source/Editor/UnrealEd/Public/AssetDeleteModel.h b/Engine/Source/Editor/UnrealEd/Public/AssetDeleteModel.h
--- a/Engine/Source/Editor/UnrealEd/Public/AssetDeleteModel.h
+++ b/Engine/Source/Editor/UnrealEd/Public/AssetDeleteModel.h
@@ -82,7 +82,7 @@ private:
  * that must be performed against the GC for UObjects, and looking up references for assets through
  * the asset registry.
  */
-class FAssetDeleteModel
+class UNREALED_API FAssetDeleteModel
 {
 public:
 	/** States used to manage the async deletion process. */
