[NO NON-PII COMMIT MESSAGE]

diff --git a/Engine/Source/Editor/Kismet/Private/SSCSEditor.cpp b/Engine/Source/Editor/Kismet/Private/SSCSEditor.cpp
--- a/Engine/Source/Editor/Kismet/Private/SSCSEditor.cpp
+++ b/Engine/Source/Editor/Kismet/Private/SSCSEditor.cpp
@@ -3094,7 +3094,8 @@ bool SSCS_RowWidget::OnNameTextVerifyChanged(const FText& InNewText, FText& OutE
 				OutErrorMessage = FText::Format(LOCTEXT("ComponentRenameFailed_TooLong", "Component name must be less than {CharCount} characters long."), Arguments);
 				return false;
 			}
-			else if (!FComponentEditorUtils::IsComponentNameAvailable(InNewText.ToString(), ExistingNameSearchScope, ComponentInstance))
+			else if (!FComponentEditorUtils::IsComponentNameAvailable(InNewText.ToString(), ExistingNameSearchScope, ComponentInstance) 
+					|| !FComponentEditorUtils::IsComponentNameAvailable(InNewText.ToString(), ComponentInstance->GetOuter(), ComponentInstance ))
 			{
 				OutErrorMessage = LOCTEXT("RenameFailed_ExistingName", "Another component already has the same name.");
 				return false;
diff --git a/Engine/Source/Editor/UnrealEd/Private/Kismet2/ComponentEditorUtils.cpp b/Engine/Source/Editor/UnrealEd/Private/Kismet2/ComponentEditorUtils.cpp
--- a/Engine/Source/Editor/UnrealEd/Private/Kismet2/ComponentEditorUtils.cpp
+++ b/Engine/Source/Editor/UnrealEd/Private/Kismet2/ComponentEditorUtils.cpp
@@ -206,7 +206,7 @@ bool FComponentEditorUtils::IsValidVariableNameString(const UActorComponent* InC
 	return bIsValid;
 }
 
-bool FComponentEditorUtils::IsComponentNameAvailable(const FString& InString, AActor* ComponentOwner, const UActorComponent* ComponentToIgnore)
+bool FComponentEditorUtils::IsComponentNameAvailable(const FString& InString, UObject* ComponentOwner, const UActorComponent* ComponentToIgnore)
 {
 	UObject* Object = FindObjectFast<UObject>(ComponentOwner, *InString);
 
diff --git a/Engine/Source/Editor/UnrealEd/Public/Kismet2/ComponentEditorUtils.h b/Engine/Source/Editor/UnrealEd/Public/Kismet2/ComponentEditorUtils.h
--- a/Engine/Source/Editor/UnrealEd/Public/Kismet2/ComponentEditorUtils.h
+++ b/Engine/Source/Editor/UnrealEd/Public/Kismet2/ComponentEditorUtils.h
@@ -21,11 +21,11 @@ public:
 	static bool IsValidVariableNameString(const UActorComponent* InComponent, const FString& InString);
 
 	/** 
-	 * Test whether or not the given string is already the name string of a component on the the actor
+	 * Test whether or not the given string is already the name string of a component on the the owner
 	 * Optionally excludes an existing component from the check (ex. a component currently being renamed)
 	 * @return True if the InString is an available name for a component of ComponentOwner
 	 */
-	static bool IsComponentNameAvailable(const FString& InString, AActor* ComponentOwner, const UActorComponent* ComponentToIgnore = nullptr);
+	static bool IsComponentNameAvailable(const FString& InString, UObject* ComponentOwner, const UActorComponent* ComponentToIgnore = nullptr);
 		
 	/** Generate a valid variable name string for the given component instance */
 	static FString GenerateValidVariableName(TSubclassOf<UActorComponent> InComponentClass, AActor* ComponentOwner);
