Fix landscape Lightmass material UV being calculated wrong. Swarm Agent cache needs to be cleared for the fix to be effective.

diff --git a/Engine/Source/Editor/UnrealEd/Private/Lightmass/LightmassLandscapeRender.cpp b/Engine/Source/Editor/UnrealEd/Private/Lightmass/LightmassLandscapeRender.cpp
--- a/Engine/Source/Editor/UnrealEd/Private/Lightmass/LightmassLandscapeRender.cpp
+++ b/Engine/Source/Editor/UnrealEd/Private/Lightmass/LightmassLandscapeRender.cpp
@@ -70,7 +70,7 @@ void RenderLandscapeMaterialForLightmass(const FLandscapeStaticLightingMesh* Lan
 			                                    (SubsectionY >= 0 && SubsectionY < NumSubsections ? 1 : 0));
 
 			const FVector2D BasePosition = PatchExpandOffset + FVector2D(SubsectionX, SubsectionY) * PositionScale;
-			const FVector2D BaseLayerCoords = FVector2D(UVSubsection) * LayerScale;
+			const FVector2D BaseLayerCoords = FVector2D(LandscapeComponent->SectionBaseX, LandscapeComponent->SectionBaseY) + FVector2D(UVSubsection) * LayerScale;
 			const FVector2D BaseWeightmapCoords = WeightmapBias + FVector2D(UVSubsection) * WeightmapSubsection;
 
 			int32 Index = Vertices.Add(FDynamicMeshVertex(FVector(BasePosition /*FVector2D(0, 0) * PositionScale*/, 0), FVector(BaseLayerCoords /*FVector2D(0, 0) * UVScale * LayerScale*/, 0), BaseWeightmapCoords /*FVector2D(0, 0) * UVScale * WeightmapScale*/));
