Cooker settings will now automatically update console variables associated with them so that they take immediate effect

diff --git a/Engine/Source/Editor/UnrealEd/Private/CookerSettings.cpp b/Engine/Source/Editor/UnrealEd/Private/CookerSettings.cpp
--- a/Engine/Source/Editor/UnrealEd/Private/CookerSettings.cpp
+++ b/Engine/Source/Editor/UnrealEd/Private/CookerSettings.cpp
@@ -75,5 +75,9 @@ void UCookerSettings::PostEditChangeProperty(FPropertyChangedEvent& PropertyChan
 				UpdateSinglePropertyInConfigFile(OldProperty, GetDefaultConfigFilename());
 			}
 		}
+		else
+		{
+			ExportValuesToConsoleVariables(PropertyChangedEvent.Property);
+		}
 	}
 }
