Fixed issue where updating a preview window would cause a "Compiling Missing ViewMode Shaders" dialog to always appear.

diff --git a/Engine/Source/Editor/UnrealEd/Private/EditorBuildUtils.cpp b/Engine/Source/Editor/UnrealEd/Private/EditorBuildUtils.cpp
--- a/Engine/Source/Editor/UnrealEd/Private/EditorBuildUtils.cpp
+++ b/Engine/Source/Editor/UnrealEd/Private/EditorBuildUtils.cpp
@@ -1076,12 +1076,11 @@ EDebugViewShaderMode ViewModeIndexToDebugViewShaderMode(EViewModeIndex SelectedV
 		return DVSM_PrimitiveDistanceAccuracy;
 	case VMI_MeshUVDensityAccuracy:
 		return DVSM_MeshUVDensityAccuracy;
-	case VMI_Unknown:
-		return DVSM_OutputMaterialTextureScales;
 	case VMI_MaterialTextureScaleAccuracy:
 		return DVSM_MaterialTextureScaleAccuracy;
 	case VMI_RequiredTextureResolution:
 		return DVSM_RequiredTextureResolution;
+	case VMI_Unknown:
 	default :
 		return DVSM_None;
 	}
