Allow Data Table factory to be used by Python

diff --git a/Engine/Source/Editor/UnrealEd/Classes/Factories/DataTableFactory.h b/Engine/Source/Editor/UnrealEd/Classes/Factories/DataTableFactory.h
--- a/Engine/Source/Editor/UnrealEd/Classes/Factories/DataTableFactory.h
+++ b/Engine/Source/Editor/UnrealEd/Classes/Factories/DataTableFactory.h
@@ -14,7 +14,7 @@ class UNREALED_API UDataTableFactory : public UFactory
 {
 	GENERATED_UCLASS_BODY()
 
-	UPROPERTY()
+	UPROPERTY(BlueprintReadWrite, Category = "Data Table Factory")
 	class UScriptStruct* Struct;
 
 	//~ Begin UFactory Interface
