Stop PIE windows from auto-centering with DPI scaling.
Undo DPI scaling from PIE window position before saving to config file.

diff --git a/Engine/Source/Editor/UnrealEd/Private/PlayLevel.cpp b/Engine/Source/Editor/UnrealEd/Private/PlayLevel.cpp
--- a/Engine/Source/Editor/UnrealEd/Private/PlayLevel.cpp
+++ b/Engine/Source/Editor/UnrealEd/Private/PlayLevel.cpp
@@ -90,6 +90,7 @@
 #include "Widgets/Notifications/SNotificationList.h"
 #include "Engine/LocalPlayer.h"
 #include "Slate/SGameLayerManager.h"
+#include "HAL/PlatformApplicationMisc.h"
 
 #include "IHeadMountedDisplay.h"
 #include "IXRTrackingSystem.h"
@@ -3236,7 +3237,9 @@ UGameInstance* UEditorEngine::CreatePIEGameInstance(int32 InPIEInstance, bool bI
 						static void OnPIEWindowClosed( const TSharedRef< SWindow >& WindowBeingClosed, TWeakPtr< SViewport > PIEViewportWidget, int32 index, bool bRestoreRootWindow )
 						{
 							// Save off the window position
-							const FVector2D PIEWindowPos = WindowBeingClosed->GetPositionInScreen();
+							FVector2D PIEWindowPos = WindowBeingClosed->GetPositionInScreen();
+							const float DPIScale = FPlatformApplicationMisc::GetDPIScaleFactorAtPoint(PIEWindowPos.X, PIEWindowPos.Y);
+							PIEWindowPos /= DPIScale;
 
 							ULevelEditorPlaySettings* LevelEditorPlaySettings = ULevelEditorPlaySettings::StaticClass()->GetDefaultObject<ULevelEditorPlaySettings>();
 
@@ -3296,7 +3299,7 @@ UGameInstance* UEditorEngine::CreatePIEGameInstance(int32 InPIEInstance, bool bI
 				ViewportClient->Viewport->SetPlayInEditorViewport( ViewportClient->bIsPlayInEditorViewport );
 
 				// Ensure the window has a valid size before calling BeginPlay
-				SlatePlayInEditorSession.SlatePlayInEditorWindowViewport->ResizeFrame( NewWindowWidth, NewWindowHeight, EWindowMode::Windowed );
+				PieWindow->ReshapeWindow(PieWindow->GetPositionInScreen(), FVector2D(NewWindowWidth, NewWindowHeight));
 
 				// Change the system resolution to match our window, to make sure game and slate window are kept syncronised
 				FSystemResolution::RequestResolutionChange(NewWindowWidth, NewWindowHeight, EWindowMode::Windowed);
