[NO NON-PII COMMIT MESSAGE]

diff --git a/Engine/Source/Editor/UnrealEd/Private/Editor.cpp b/Engine/Source/Editor/UnrealEd/Private/Editor.cpp
--- a/Engine/Source/Editor/UnrealEd/Private/Editor.cpp
+++ b/Engine/Source/Editor/UnrealEd/Private/Editor.cpp
@@ -429,7 +429,12 @@ void FReimportManager::ValidateAllSourceFileAndReimport(TArray<UObject*> &ToImpo
 		{
 			if (SourceFilenames.Num() == 0)
 			{
-				MissingFileSelectedAssets.FindOrAdd(Asset);
+				TArray<int32>& SourceIndexArray = MissingFileSelectedAssets.FindOrAdd(Asset);
+				if (SourceIndexArray.Num() == 0)
+				{
+					// Insert an invalid index to indicate no file
+					SourceIndexArray.Add(INDEX_NONE);
+				}
 			}
 			else
 			{
@@ -484,12 +489,10 @@ void FReimportManager::ValidateAllSourceFileAndReimport(TArray<UObject*> &ToImpo
 					MaxListFile--;
 					for (int32 FileIndex : SourceIndexArray)
 					{
-						int32 RemapFileIndex = 0;
 						if (SourceFilenames.IsValidIndex(FileIndex))
 						{
-							RemapFileIndex = FileIndex;
+							AssetToFileListString += FString::Printf(TEXT("Asset %s -> Missing file %s"), *(Asset->GetName()), *(SourceFilenames[FileIndex]));
 						}
-						AssetToFileListString += FString::Printf(TEXT("Asset %s -> Missing file %s"), *(Asset->GetName()), *(SourceFilenames[RemapFileIndex]));
 					}
 				}
 			}
diff --git a/Engine/Source/Editor/UnrealEd/Private/Factories/CSVImportFactory.cpp b/Engine/Source/Editor/UnrealEd/Private/Factories/CSVImportFactory.cpp
--- a/Engine/Source/Editor/UnrealEd/Private/Factories/CSVImportFactory.cpp
+++ b/Engine/Source/Editor/UnrealEd/Private/Factories/CSVImportFactory.cpp
@@ -409,14 +409,9 @@ bool UReimportDataTableFactory::CanReimport( UObject* Obj, TArray<FString>& OutF
 	if (DataTable)
 	{
 		DataTable->AssetImportData->ExtractFilenames(OutFilenames);
-		if (OutFilenames.Num() > 0)
-		{
-			return true;
-		}
-		else
-		{
-			return false;
-		}
+		
+		// Always allow reimporting a data table as it is common to convert from manual to CSV-driven tables
+		return true;
 	}
 	return false;
 }
