Fix false ensure triggering

diff --git a/Engine/Source/Editor/UnrealEd/Private/Commandlets/PackageUtilities.cpp b/Engine/Source/Editor/UnrealEd/Private/Commandlets/PackageUtilities.cpp
--- a/Engine/Source/Editor/UnrealEd/Private/Commandlets/PackageUtilities.cpp
+++ b/Engine/Source/Editor/UnrealEd/Private/Commandlets/PackageUtilities.cpp
@@ -2255,7 +2255,8 @@ struct CompressAnimationsFunctor
 			AnimSeq->RequestAnimCompression(FRequestAnimCompressionParams(false, true, false));
 
 			// Automatic compression should have picked a suitable compressor that is not UAnimCompress_Automatic
-			if (ensure(!AnimSeq->CompressionScheme->IsA(UAnimCompress_Automatic::StaticClass())))
+			// May still be automatic if we read the data from the DDC
+			if (!AnimSeq->CompressionScheme->IsA(UAnimCompress_Automatic::StaticClass())))
 			{
 				// Update CompressCommandletVersion in that case, and create a proper DDC entry
 				// (with actual compressor)
