Fix typo

diff --git a/Engine/Source/Editor/UnrealEd/Private/Commandlets/PackageUtilities.cpp b/Engine/Source/Editor/UnrealEd/Private/Commandlets/PackageUtilities.cpp
--- a/Engine/Source/Editor/UnrealEd/Private/Commandlets/PackageUtilities.cpp
+++ b/Engine/Source/Editor/UnrealEd/Private/Commandlets/PackageUtilities.cpp
@@ -2256,7 +2256,7 @@ struct CompressAnimationsFunctor
 
 			// Automatic compression should have picked a suitable compressor that is not UAnimCompress_Automatic
 			// May still be automatic if we read the data from the DDC
-			if (!AnimSeq->CompressionScheme->IsA(UAnimCompress_Automatic::StaticClass())))
+			if (!AnimSeq->CompressionScheme->IsA(UAnimCompress_Automatic::StaticClass()))
 			{
 				// Update CompressCommandletVersion in that case, and create a proper DDC entry
 				// (with actual compressor)
