Fix Crash in skeletal mesh editor after LOD import
Fix from Dev-Enterprise CL:5294522

diff --git a/Engine/Source/Editor/Persona/Private/PersonaMeshDetails.cpp b/Engine/Source/Editor/Persona/Private/PersonaMeshDetails.cpp
--- a/Engine/Source/Editor/Persona/Private/PersonaMeshDetails.cpp
+++ b/Engine/Source/Editor/Persona/Private/PersonaMeshDetails.cpp
@@ -326,6 +326,17 @@ void SSkeletalLODActions::Construct(const FArguments& InArgs)
 	RegisterActiveTimer(1.0f + LODTimeOffset * LODIndex, FWidgetActiveTimerDelegate::CreateSP(this, &SSkeletalLODActions::RefreshExistFlag));
 }
 
+/**
+* FPersonaMeshDetails
+*/
+FPersonaMeshDetails::FPersonaMeshDetails(TSharedRef<class IPersonaToolkit> InPersonaToolkit) : PersonaToolkitPtr(InPersonaToolkit), MeshDetailLayout(nullptr)
+{
+	CustomLODEditMode = false;
+	bDeleteWarningConsumed = false;
+
+	GEditor->GetEditorSubsystem<UImportSubsystem>()->OnAssetPostLODImport.AddRaw(this, &FPersonaMeshDetails::OnAssetPostLODImported);
+}
+
 /**
 * FPersonaMeshDetails
 */
@@ -336,6 +347,8 @@ FPersonaMeshDetails::~FPersonaMeshDetails()
 		TSharedRef<IPersonaPreviewScene> PreviewScene = GetPersonaToolkit()->GetPreviewScene();
 		PreviewScene->UnregisterOnPreviewMeshChanged(this);
 	}
+
+	GEditor->GetEditorSubsystem<UImportSubsystem>()->OnAssetPostLODImport.RemoveAll(this);
 }
 
 TSharedRef<IDetailCustomization> FPersonaMeshDetails::MakeInstance(TWeakPtr<class IPersonaToolkit> InPersonaToolkit)
@@ -1390,6 +1403,14 @@ bool FPersonaMeshDetails::IsLODInfoEditingEnabled(int32 LODIndex) const
 	return true;
 }
 
+void FPersonaMeshDetails::OnAssetPostLODImported(UObject* InObject, int32 InLODIndex)
+{
+	if (InObject == GetPersonaToolkit()->GetMesh())
+	{
+		MeshDetailLayout->ForceRefreshDetails();
+	}
+}
+
 void FPersonaMeshDetails::OnImportLOD(TSharedPtr<FString> NewValue, ESelectInfo::Type SelectInfo, IDetailLayoutBuilder* DetailLayout)
 {
 	int32 LODIndex = 0;
@@ -1399,8 +1420,6 @@ void FPersonaMeshDetails::OnImportLOD(TSharedPtr<FString> NewValue, ESelectInfo:
 		check(SkelMesh);
 
 		FbxMeshUtils::ImportMeshLODDialog(SkelMesh, LODIndex);
-
-		DetailLayout->ForceRefreshDetails();
 	}
 }
 
diff --git a/Engine/Source/Editor/Persona/Private/PersonaMeshDetails.h b/Engine/Source/Editor/Persona/Private/PersonaMeshDetails.h
--- a/Engine/Source/Editor/Persona/Private/PersonaMeshDetails.h
+++ b/Engine/Source/Editor/Persona/Private/PersonaMeshDetails.h
@@ -81,11 +81,7 @@ struct FSectionLocalizer
 class FPersonaMeshDetails : public IDetailCustomization
 {
 public:
-	FPersonaMeshDetails(TSharedRef<class IPersonaToolkit> InPersonaToolkit) : PersonaToolkitPtr(InPersonaToolkit), MeshDetailLayout(nullptr)
-	{
-		CustomLODEditMode = false;
-		bDeleteWarningConsumed = false;
-	}
+	FPersonaMeshDetails(TSharedRef<class IPersonaToolkit> InPersonaToolkit);
 	~FPersonaMeshDetails();
 
 	/** Makes a new instance of this detail layout class for a specific detail view requesting it */
@@ -352,6 +348,9 @@ private:
 	/** for LOD settings category */
 	void CustomizeLODSettingsCategories(IDetailLayoutBuilder& DetailLayout);
 
+	/** Called when a LOD is imported. Refreshes the UI. */
+	void OnAssetPostLODImported(UObject* InObject, int32 InLODIndex);
+	/** Called from the PersonalMeshDetails UI to import a LOD. */
 	void OnImportLOD(TSharedPtr<FString> NewValue, ESelectInfo::Type SelectInfo, IDetailLayoutBuilder* DetailLayout);
 	void UpdateLODNames();
 	int32 GetLODCount() const;
diff --git a/Engine/Source/Editor/UnrealEd/Private/FbxMeshUtils.cpp b/Engine/Source/Editor/UnrealEd/Private/FbxMeshUtils.cpp
--- a/Engine/Source/Editor/UnrealEd/Private/FbxMeshUtils.cpp
+++ b/Engine/Source/Editor/UnrealEd/Private/FbxMeshUtils.cpp
@@ -22,6 +22,7 @@
 #include "Rendering/SkeletalMeshModel.h"
 #include "EditorFramework/AssetImportData.h"
 #include "DesktopPlatformModule.h"
+#include "Editor.h"
 
 #if WITH_APEX_CLOTHING
 	#include "ApexClothingUtils.h"
@@ -714,6 +715,18 @@ namespace FbxMeshUtils
 			FMessageDialog::Open(EAppMsgType::Ok, FText::Format(LOCTEXT("LODImport_Failure", "Failed to import LOD{0}"), FText::AsNumber(LODLevel)));
 		}
 
+		if (bImportSuccess)
+		{
+			if (SkeletalMesh)
+			{
+				GEditor->GetEditorSubsystem<UImportSubsystem>()->BroadcastAssetPostLODImport(SkeletalMesh, LODLevel);
+			}				
+			else if(StaticMesh)
+			{
+				GEditor->GetEditorSubsystem<UImportSubsystem>()->BroadcastAssetPostLODImport(StaticMesh, LODLevel);
+			}
+		}
+
 		return bImportSuccess;
 	}
 
diff --git a/Engine/Source/Editor/UnrealEd/Private/Subsystems/ImportSubsystem.cpp b/Engine/Source/Editor/UnrealEd/Private/Subsystems/ImportSubsystem.cpp
--- a/Engine/Source/Editor/UnrealEd/Private/Subsystems/ImportSubsystem.cpp
+++ b/Engine/Source/Editor/UnrealEd/Private/Subsystems/ImportSubsystem.cpp
@@ -161,4 +161,10 @@ void UImportSubsystem::BroadcastAssetReimport(UObject* InCreatedObject)
 	OnAssetReimport.Broadcast(InCreatedObject);
 	OnAssetReimport_BP.Broadcast(InCreatedObject);
 }
+
+void UImportSubsystem::BroadcastAssetPostLODImport(UObject* InObject, int32 inLODIndex)
+{
+	OnAssetPostLODImport.Broadcast(InObject, inLODIndex);
+	OnAssetPostLODImport_BP.Broadcast(InObject, inLODIndex);
+}
 PRAGMA_ENABLE_DEPRECATION_WARNINGS
diff --git a/Engine/Source/Editor/UnrealEd/Public/Subsystems/ImportSubsystem.h b/Engine/Source/Editor/UnrealEd/Public/Subsystems/ImportSubsystem.h
--- a/Engine/Source/Editor/UnrealEd/Public/Subsystems/ImportSubsystem.h
+++ b/Engine/Source/Editor/UnrealEd/Public/Subsystems/ImportSubsystem.h
@@ -47,6 +47,8 @@ public:
 	DECLARE_MULTICAST_DELEGATE_TwoParams(FOnAssetPostImport, UFactory*, UObject*);
 	/** delegate type fired when new assets have been reimported. Note: InCreatedObject can be NULL if import failed. Params: UObject* InCreatedObject */
 	DECLARE_MULTICAST_DELEGATE_OneParam(FOnAssetReimport, UObject*);
+	/** delegate type fired when new LOD have been imported to an asset. */
+	DECLARE_MULTICAST_DELEGATE_TwoParams(FOnAssetPostLODImport, UObject*, int32);
 
 	// Broadcast AssetPreImport, do not broadcast with OnAssetPostImport directly.
 	void BroadcastAssetPreImport(UFactory* InFactory, UClass* InClass, UObject* InParent, const FName& Name, const TCHAR* Type);
@@ -54,11 +56,14 @@ public:
 	void BroadcastAssetPostImport(UFactory* InFactory, UObject* InCreatedObject);
 	// Broadcast AssetReimport, do not broadcast with OnAssetReimport directly.
 	void BroadcastAssetReimport(UObject* InCreatedObject);
+	// Broadcast AssetPostLODImport, do not broadcast with OnAssetPostLODImport directly.
+	void BroadcastAssetPostLODImport(UObject* InObject, int32 inLODIndex);
 
 	// Used to register and unregister ONLY use Broadcast functions to execute the delegate
 	FOnAssetPreImport OnAssetPreImport;
 	FOnAssetPostImport OnAssetPostImport;
 	FOnAssetReimport OnAssetReimport;
+	FOnAssetPostLODImport OnAssetPostLODImport;
 
 private:
 
@@ -71,6 +76,8 @@ private:
 	DECLARE_DYNAMIC_MULTICAST_DELEGATE_TwoParams(FOnAssetPostImport_Dyn, UFactory*, InFactory, UObject*, InCreatedObject);
 	/** delegate type fired when new assets have been reimported. Note: InCreatedObject can be NULL if import failed. Params: UObject* InCreatedObject */
 	DECLARE_DYNAMIC_MULTICAST_DELEGATE_OneParam(FOnAssetReimport_Dyn, UObject*, InCreatedObject);
+	/** delegate type fired when new LOD have been imported to an asset. */
+	DECLARE_DYNAMIC_MULTICAST_DELEGATE_TwoParams(FOnAssetPostLODImport_Dyn, UObject*, InObject, int32, InLODIndex);
 
 	UPROPERTY(BlueprintAssignable, DisplayName = "OnAssetPreImport", meta = (ScriptName = "OnAssetPreImport"))
 	FOnAssetPreImport_Dyn OnAssetPreImport_BP;
@@ -78,6 +85,8 @@ private:
 	FOnAssetPostImport_Dyn OnAssetPostImport_BP;
 	UPROPERTY(BlueprintAssignable, DisplayName = "OnAssetReimport", meta = (ScriptName = "OnAssetReimport"))
 	FOnAssetReimport_Dyn OnAssetReimport_BP;
+	UPROPERTY(BlueprintAssignable, DisplayName = "OnAssetPostLODImport", meta = (ScriptName = "OnAssetPostLODImport"))
+	FOnAssetPostLODImport_Dyn OnAssetPostLODImport_BP;
 
 	/* Tasks waiting to be run next tick */
 	TQueue<TSharedPtr<IImportSubsystemTask>> PendingTasks;
