Filtered optimization viewport menu to filter out invalid options.

diff --git a/Engine/Source/Editor/UnrealEd/Private/SEditorViewportViewMenu.cpp b/Engine/Source/Editor/UnrealEd/Private/SEditorViewportViewMenu.cpp
--- a/Engine/Source/Editor/UnrealEd/Private/SEditorViewportViewMenu.cpp
+++ b/Engine/Source/Editor/UnrealEd/Private/SEditorViewportViewMenu.cpp
@@ -317,18 +317,27 @@ TSharedRef<SWidget> SEditorViewportViewMenu::GenerateViewMenuContent() const
 					{
 						const FEditorViewportCommands& BaseViewportCommands = FEditorViewportCommands::Get();
 
+						const ERHIFeatureLevel::Type FeatureLevel = (IsInGameThread() && GWorld) ? GWorld->FeatureLevel : GMaxRHIFeatureLevel;
+
 						Menu.BeginSection("OptimizationViewmodes", LOCTEXT("OptimizationSubMenuHeader", "Optimization Viewmodes"));
 						{
-							Menu.AddMenuEntry(BaseViewportCommands.LightComplexityMode, NAME_None, LOCTEXT("LightComplexityViewModeDisplayName", "Light Complexity"));
-							Menu.AddMenuEntry(BaseViewportCommands.LightmapDensityMode, NAME_None, LOCTEXT("LightmapDensityViewModeDisplayName", "Lightmap Density"));
-							Menu.AddMenuEntry(BaseViewportCommands.StationaryLightOverlapMode, NAME_None, LOCTEXT("StationaryLightOverlapViewModeDisplayName", "Stationary Light Overlap"));
-							Menu.AddMenuEntry(BaseViewportCommands.ShaderComplexityMode, NAME_None, LOCTEXT("ShaderComplexityViewModeDisplayName", "Shader Complexity"));
+							if (FeatureLevel == ERHIFeatureLevel::SM5)
+							{
+								Menu.AddMenuEntry(BaseViewportCommands.LightComplexityMode, NAME_None, LOCTEXT("LightComplexityViewModeDisplayName", "Light Complexity"));
+								Menu.AddMenuEntry(BaseViewportCommands.LightmapDensityMode, NAME_None, LOCTEXT("LightmapDensityViewModeDisplayName", "Lightmap Density"));
+								Menu.AddMenuEntry(BaseViewportCommands.StationaryLightOverlapMode, NAME_None, LOCTEXT("StationaryLightOverlapViewModeDisplayName", "Stationary Light Overlap"));
+							}
+
+							if (AllowDebugViewShaderMode(DVSM_ShaderComplexityContainedQuadOverhead, GMaxRHIShaderPlatform, FeatureLevel))
+							{
+								Menu.AddMenuEntry(BaseViewportCommands.ShaderComplexityMode, NAME_None, LOCTEXT("ShaderComplexityViewModeDisplayName", "Shader Complexity"));
+							}
 
-							if (AllowDebugViewShaderMode(DVSM_ShaderComplexityContainedQuadOverhead))
+							if (AllowDebugViewShaderMode(DVSM_ShaderComplexityContainedQuadOverhead, GMaxRHIShaderPlatform, FeatureLevel))
 							{
 								Menu.AddMenuEntry(BaseViewportCommands.ShaderComplexityWithQuadOverdrawMode, NAME_None, LOCTEXT("ShaderComplexityWithQuadOverdrawViewModeDisplayName", "Shader Complexity & Quads"));
 							}
-							if (AllowDebugViewShaderMode(DVSM_QuadComplexity))
+							if (AllowDebugViewShaderMode(DVSM_QuadComplexity, GMaxRHIShaderPlatform, FeatureLevel))
 							{
 								Menu.AddMenuEntry(BaseViewportCommands.QuadOverdrawMode, NAME_None, LOCTEXT("QuadOverdrawViewModeDisplayName", "Quad Overdraw"));
 							}
@@ -336,20 +345,20 @@ TSharedRef<SWidget> SEditorViewportViewMenu::GenerateViewMenuContent() const
 						Menu.EndSection();
 
 						Menu.BeginSection("TextureStreaming", LOCTEXT("TextureStreamingHeader", "Texture Streaming Accuracy") );
-						if ( AllowDebugViewShaderMode(DVSM_PrimitiveDistanceAccuracy) && (!InParentToolBar.IsValid() || InParentToolBar.Pin()->IsViewModeSupported(VMI_PrimitiveDistanceAccuracy)) )
+						if ( AllowDebugViewShaderMode(DVSM_PrimitiveDistanceAccuracy, GMaxRHIShaderPlatform, FeatureLevel) && (!InParentToolBar.IsValid() || InParentToolBar.Pin()->IsViewModeSupported(VMI_PrimitiveDistanceAccuracy)) )
 						{
 							Menu.AddMenuEntry( BaseViewportCommands.TexStreamAccPrimitiveDistanceMode, NAME_None, LOCTEXT("TexStreamAccPrimitiveDistanceViewModeDisplayName", "Primitive Distance") );
 						}
-						if ( AllowDebugViewShaderMode(DVSM_MeshUVDensityAccuracy) && (!InParentToolBar.IsValid() || InParentToolBar.Pin()->IsViewModeSupported(VMI_MeshUVDensityAccuracy)) )
+						if ( AllowDebugViewShaderMode(DVSM_MeshUVDensityAccuracy, GMaxRHIShaderPlatform, FeatureLevel) && (!InParentToolBar.IsValid() || InParentToolBar.Pin()->IsViewModeSupported(VMI_MeshUVDensityAccuracy)) )
 						{
 							Menu.AddMenuEntry(BaseViewportCommands.TexStreamAccMeshUVDensityMode, NAME_None, LOCTEXT("TexStreamAccMeshUVDensityViewModeDisplayName", "Mesh UV Densities"));
 						}
 						// TexCoordScale accuracy viewmode requires shaders that are only built in the TextureStreamingBuild, which requires the new metrics to be enabled.
-						if ( AllowDebugViewShaderMode(DVSM_MaterialTextureScaleAccuracy) && CVarStreamingUseNewMetrics.GetValueOnAnyThread() != 0 && (!InParentToolBar.IsValid() || InParentToolBar.Pin()->IsViewModeSupported(VMI_MaterialTextureScaleAccuracy)) )
+						if ( AllowDebugViewShaderMode(DVSM_MaterialTextureScaleAccuracy, GMaxRHIShaderPlatform, FeatureLevel) && CVarStreamingUseNewMetrics.GetValueOnAnyThread() != 0 && (!InParentToolBar.IsValid() || InParentToolBar.Pin()->IsViewModeSupported(VMI_MaterialTextureScaleAccuracy)) )
 						{
 							Menu.AddMenuEntry(BaseViewportCommands.TexStreamAccMaterialTextureScaleMode, NAME_None, LOCTEXT("TexStreamAccMaterialTextureScaleViewModeDisplayName", "Material Texture Scales"));
 						}
-						if ( AllowDebugViewShaderMode(DVSM_RequiredTextureResolution) && (!InParentToolBar.IsValid() || InParentToolBar.Pin()->IsViewModeSupported(VMI_MaterialTextureScaleAccuracy)) )
+						if ( AllowDebugViewShaderMode(DVSM_RequiredTextureResolution, GMaxRHIShaderPlatform, FeatureLevel) && (!InParentToolBar.IsValid() || InParentToolBar.Pin()->IsViewModeSupported(VMI_MaterialTextureScaleAccuracy)) )
 						{
 							Menu.AddMenuEntry(BaseViewportCommands.RequiredTextureResolutionMode, NAME_None, LOCTEXT("RequiredTextureResolutionModeDisplayName", "Required Texture Resolution"));
 						}
diff --git a/Engine/Source/Runtime/Engine/Private/DebugViewModeHelpers.cpp b/Engine/Source/Runtime/Engine/Private/DebugViewModeHelpers.cpp
--- a/Engine/Source/Runtime/Engine/Private/DebugViewModeHelpers.cpp
+++ b/Engine/Source/Runtime/Engine/Private/DebugViewModeHelpers.cpp
@@ -21,8 +21,17 @@ static bool PlatformSupportsDebugViewShaders(EShaderPlatform Platform)
 	// List of platforms that have been tested and proved functional.
 	return Platform == SP_PCD3D_SM4 || Platform == SP_PCD3D_SM5 || Platform == SP_OPENGL_SM4;
 }
- 
-bool AllowDebugViewPS(EDebugViewShaderMode ShaderMode, EShaderPlatform Platform)
+
+bool AllowDebugViewVSDSHS(EShaderPlatform Platform)
+{
+#if WITH_EDITOR
+	return true; 
+#else
+	return false;
+#endif
+}
+
+bool AllowDebugViewShaderMode(EDebugViewShaderMode ShaderMode, EShaderPlatform Platform, ERHIFeatureLevel::Type FeatureLevel)
 {
 #if WITH_EDITOR
 	// Those options are used to test compilation on specific platforms
@@ -35,40 +44,26 @@ bool AllowDebugViewPS(EDebugViewShaderMode ShaderMode, EShaderPlatform Platform)
 	case DVSM_None:
 		return false;
 	case DVSM_ShaderComplexity:
-		return true;
+		return FeatureLevel != ERHIFeatureLevel::SM4; // SM4 doesn't apply post-processes correctly in the deferred path.
 	case DVSM_ShaderComplexityContainedQuadOverhead:
 	case DVSM_ShaderComplexityBleedingQuadOverhead:
 	case DVSM_QuadComplexity:
-		return IsFeatureLevelSupported(Platform, ERHIFeatureLevel::SM5) && (bForceQuadOverdraw || PlatformSupportsDebugViewShaders(Platform));
+		return FeatureLevel >= ERHIFeatureLevel::SM5 && (bForceQuadOverdraw || PlatformSupportsDebugViewShaders(Platform));
 	case DVSM_PrimitiveDistanceAccuracy:
 	case DVSM_MeshUVDensityAccuracy:
-		return IsFeatureLevelSupported(Platform, ERHIFeatureLevel::SM4) && (bForceStreamingAccuracy || PlatformSupportsDebugViewShaders(Platform));
+		return FeatureLevel >= ERHIFeatureLevel::SM5 && (bForceStreamingAccuracy || PlatformSupportsDebugViewShaders(Platform));
 	case DVSM_MaterialTextureScaleAccuracy:
 	case DVSM_RequiredTextureResolution:
 	case DVSM_OutputMaterialTextureScales:
-		return IsFeatureLevelSupported(Platform, ERHIFeatureLevel::SM4) && (bForceTextureStreamingBuild || PlatformSupportsDebugViewShaders(Platform));
+		return FeatureLevel >= ERHIFeatureLevel::SM5 && (bForceTextureStreamingBuild || PlatformSupportsDebugViewShaders(Platform));
 	default:
 		return false;
 	}
 #else
-	return ShaderMode == DVSM_ShaderComplexity;
+	return ShaderMode == DVSM_ShaderComplexity && FeatureLevel != ERHIFeatureLevel::SM4;
 #endif
 }
 
-bool AllowDebugViewVSDSHS(EShaderPlatform Platform)
-{
-#if WITH_EDITOR
-	return true; 
-#else
-	return false;
-#endif
-}
-
-bool AllowDebugViewShaderMode(EDebugViewShaderMode ShaderMode)
-{
-	return AllowDebugViewPS(ShaderMode, GMaxRHIShaderPlatform);
-}
-
 #endif // !(UE_BUILD_SHIPPING || UE_BUILD_TEST)
 
 
@@ -204,7 +199,7 @@ bool GetUsedMaterialsInWorld(UWorld* InWorld, OUT TSet<UMaterialInterface*>& Out
 			}
 		}
 	}
-	return true;
+	return OutMaterials.Num() != 0;
 #else
 	return false;
 #endif
@@ -223,13 +218,11 @@ bool GetUsedMaterialsInWorld(UWorld* InWorld, OUT TSet<UMaterialInterface*>& Out
 bool CompileDebugViewModeShaders(EDebugViewShaderMode ShaderMode, EMaterialQualityLevel::Type QualityLevel, ERHIFeatureLevel::Type FeatureLevel, bool bFullRebuild, bool bWaitForPreviousShaders, TSet<UMaterialInterface*>& Materials, FSlowTask& ProgressTask)
 {
 #if WITH_EDITORONLY_DATA
-	if (!GShaderCompilingManager)
+	if (!GShaderCompilingManager || !Materials.Num())
 	{
 		return false;
 	}
 
-	check(Materials.Num())
-
 	// Finish compiling pending shaders first.
 	if (!bWaitForPreviousShaders)
 	{
diff --git a/Engine/Source/Runtime/Engine/Private/SceneView.cpp b/Engine/Source/Runtime/Engine/Private/SceneView.cpp
--- a/Engine/Source/Runtime/Engine/Private/SceneView.cpp
+++ b/Engine/Source/Runtime/Engine/Private/SceneView.cpp
@@ -2373,7 +2373,7 @@ FSceneViewFamily::FSceneViewFamily(const ConstructionValues& CVS)
 	ViewModeParam = CVS.ViewModeParam;
 	ViewModeParamName = CVS.ViewModeParamName;
 
-	if (!AllowDebugViewPS(DebugViewShaderMode, GetShaderPlatform()))
+	if (!AllowDebugViewShaderMode(DebugViewShaderMode, GetShaderPlatform(), GetFeatureLevel()))
 	{
 		DebugViewShaderMode = DVSM_None;
 	}
diff --git a/Engine/Source/Runtime/Engine/Public/DebugViewModeHelpers.h b/Engine/Source/Runtime/Engine/Public/DebugViewModeHelpers.h
--- a/Engine/Source/Runtime/Engine/Public/DebugViewModeHelpers.h
+++ b/Engine/Source/Runtime/Engine/Public/DebugViewModeHelpers.h
@@ -31,16 +31,13 @@ enum EDebugViewShaderMode
 };
 
 #if !(UE_BUILD_SHIPPING || UE_BUILD_TEST)
-/** Returns true if the specified shadermode is available for the given shader platform. Called for shader compilation shader compilation. */
-ENGINE_API bool AllowDebugViewPS(EDebugViewShaderMode ShaderMode, EShaderPlatform Platform);
 /** Returns true if the vertex shader (and potential hull and domain) should be compiled on the given platform. */
 ENGINE_API bool AllowDebugViewVSDSHS(EShaderPlatform Platform);
 /** Returns true if the shader mode can be enabled. This is only for UI elements as no shader platform is actually passed. */
-ENGINE_API bool AllowDebugViewShaderMode(EDebugViewShaderMode ShaderMode);
+ENGINE_API bool AllowDebugViewShaderMode(EDebugViewShaderMode ShaderMode, EShaderPlatform Platform, ERHIFeatureLevel::Type FeatureLevel);
 #else
-FORCEINLINE bool AllowDebugViewPS(EDebugViewShaderMode ShaderMode, EShaderPlatform Platform) { return false; }
 FORCEINLINE bool AllowDebugViewVSDSHS(EShaderPlatform Platform)  { return false; }
-FORCEINLINE bool AllowDebugViewShaderMode(EDebugViewShaderMode ShaderMode) { return false; }
+FORCEINLINE bool AllowDebugViewShaderMode(EDebugViewShaderMode ShaderMode, EShaderPlatform Platform, ERHIFeatureLevel::Type FeatureLevel) { return false; }
 #endif
 
 ENGINE_API int32 GetNumActorsInWorld(UWorld* InWorld);
diff --git a/Engine/Source/Runtime/Renderer/Private/MaterialTexCoordScalesRendering.h b/Engine/Source/Runtime/Renderer/Private/MaterialTexCoordScalesRendering.h
--- a/Engine/Source/Runtime/Renderer/Private/MaterialTexCoordScalesRendering.h
+++ b/Engine/Source/Runtime/Renderer/Private/MaterialTexCoordScalesRendering.h
@@ -34,7 +34,7 @@ public:
 	static bool ShouldCompilePermutation(EShaderPlatform Platform, const FMaterial* Material, const FVertexFactoryType* VertexFactoryType)
 	{
 		// See FDebugViewModeMaterialProxy::GetFriendlyName()
-		return AllowDebugViewPS(DVSM_OutputMaterialTextureScales, Platform) && Material->GetFriendlyName().Contains(TEXT("MaterialTexCoordScale"));
+		return AllowDebugViewShaderMode(DVSM_OutputMaterialTextureScales, Platform, GetMaxSupportedFeatureLevel(Platform)) && Material->GetFriendlyName().Contains(TEXT("MaterialTexCoordScale"));
 	}
 
 	FMaterialTexCoordScalePS(const ShaderMetaType::CompiledShaderInitializerType& Initializer):
diff --git a/Engine/Source/Runtime/Renderer/Private/MeshTexCoordSizeAccuracyRendering.h b/Engine/Source/Runtime/Renderer/Private/MeshTexCoordSizeAccuracyRendering.h
--- a/Engine/Source/Runtime/Renderer/Private/MeshTexCoordSizeAccuracyRendering.h
+++ b/Engine/Source/Runtime/Renderer/Private/MeshTexCoordSizeAccuracyRendering.h
@@ -31,7 +31,7 @@ public:
 	static bool ShouldCompilePermutation(EShaderPlatform Platform, const FMaterial* Material, const FVertexFactoryType* VertexFactoryType)
 	{
 		// See FDebugViewModeMaterialProxy::GetFriendlyName()
-		return AllowDebugViewPS(DVSM_MeshUVDensityAccuracy, Platform) && Material->GetFriendlyName().Contains(TEXT("MeshTexCoordSizeAccuracy"));
+		return AllowDebugViewShaderMode(DVSM_MeshUVDensityAccuracy, Platform, GetMaxSupportedFeatureLevel(Platform)) && Material->GetFriendlyName().Contains(TEXT("MeshTexCoordSizeAccuracy"));
 	}
 
 	FMeshTexCoordSizeAccuracyPS(const ShaderMetaType::CompiledShaderInitializerType& Initializer):
diff --git a/Engine/Source/Runtime/Renderer/Private/PostProcess/PostProcessVisualizeComplexity.h b/Engine/Source/Runtime/Renderer/Private/PostProcess/PostProcessVisualizeComplexity.h
--- a/Engine/Source/Runtime/Renderer/Private/PostProcess/PostProcessVisualizeComplexity.h
+++ b/Engine/Source/Runtime/Renderer/Private/PostProcess/PostProcessVisualizeComplexity.h
@@ -61,7 +61,7 @@ public:
 	static void ModifyCompilationEnvironment(const FGlobalShaderPermutationParameters& Parameters, FShaderCompilerEnvironment& OutEnvironment)
 	{
 		FGlobalShader::ModifyCompilationEnvironment(Parameters, OutEnvironment);
-		OutEnvironment.SetDefine(TEXT("READ_QUAD_OVERDRAW"), AllowDebugViewPS(DVSM_QuadComplexity, Parameters.Platform));
+		OutEnvironment.SetDefine(TEXT("READ_QUAD_OVERDRAW"), AllowDebugViewShaderMode(DVSM_QuadComplexity, Parameters.Platform, GetMaxSupportedFeatureLevel(Parameters.Platform)));
 		OutEnvironment.SetDefine(TEXT("MAX_NUM_COMPLEXITY_COLORS"), MaxNumShaderComplexityColors);
 		// EColorSampling values
 		OutEnvironment.SetDefine(TEXT("CS_RAMP"), (uint32)CS_RAMP);
diff --git a/Engine/Source/Runtime/Renderer/Private/PostProcess/SceneRenderTargets.cpp b/Engine/Source/Runtime/Renderer/Private/PostProcess/SceneRenderTargets.cpp
--- a/Engine/Source/Runtime/Renderer/Private/PostProcess/SceneRenderTargets.cpp
+++ b/Engine/Source/Runtime/Renderer/Private/PostProcess/SceneRenderTargets.cpp
@@ -752,7 +752,7 @@ int32 FSceneRenderTargets::GetGBufferRenderTargets(ERenderTargetLoadAction Color
 
 void FSceneRenderTargets::SetQuadOverdrawUAV(FRHICommandList& RHICmdList, bool bBindQuadOverdrawBuffers, bool bClearQuadOverdrawBuffers, FRHIRenderPassInfo& OutInfo)
 {
-	if (bBindQuadOverdrawBuffers && AllowDebugViewPS(DVSM_QuadComplexity, GetFeatureLevelShaderPlatform(CurrentFeatureLevel)))
+	if (bBindQuadOverdrawBuffers && AllowDebugViewShaderMode(DVSM_QuadComplexity, GetFeatureLevelShaderPlatform(CurrentFeatureLevel), CurrentFeatureLevel))
 	{
 		if (QuadOverdrawBuffer.IsValid() && QuadOverdrawBuffer->GetRenderTargetItem().UAV.IsValid())
 		{
@@ -2074,7 +2074,7 @@ void FSceneRenderTargets::AllocateReflectionTargets(FRHICommandList& RHICmdList,
 void FSceneRenderTargets::AllocateDebugViewModeTargets(FRHICommandList& RHICmdList)
 {
 	// If the shader/quad complexity shader need a quad overdraw buffer to be bind, allocate it.
-	if (AllowDebugViewPS(DVSM_QuadComplexity, GetFeatureLevelShaderPlatform(CurrentFeatureLevel)))
+	if (AllowDebugViewShaderMode(DVSM_QuadComplexity, GetFeatureLevelShaderPlatform(CurrentFeatureLevel), CurrentFeatureLevel))
 	{
 		FIntPoint QuadOverdrawSize;
 		QuadOverdrawSize.X = 2 * FMath::Max<uint32>((BufferSize.X + 1) / 2, 1); // The size is time 2 since left side is QuadDescriptor, and right side QuadComplexity.
diff --git a/Engine/Source/Runtime/Renderer/Private/PrimitiveDistanceAccuracyRendering.h b/Engine/Source/Runtime/Renderer/Private/PrimitiveDistanceAccuracyRendering.h
--- a/Engine/Source/Runtime/Renderer/Private/PrimitiveDistanceAccuracyRendering.h
+++ b/Engine/Source/Runtime/Renderer/Private/PrimitiveDistanceAccuracyRendering.h
@@ -31,7 +31,7 @@ public:
 	static bool ShouldCompilePermutation(EShaderPlatform Platform, const FMaterial* Material, const FVertexFactoryType* VertexFactoryType)
 	{
 		// See FDebugViewModeMaterialProxy::GetFriendlyName()
-		return AllowDebugViewPS(DVSM_PrimitiveDistanceAccuracy, Platform) && Material->GetFriendlyName().Contains(TEXT("PrimitiveDistanceAccuracy"));
+		return AllowDebugViewShaderMode(DVSM_PrimitiveDistanceAccuracy, Platform, GetMaxSupportedFeatureLevel(Platform)) && Material->GetFriendlyName().Contains(TEXT("PrimitiveDistanceAccuracy"));
 	}
 
 	FPrimitiveDistanceAccuracyPS(const ShaderMetaType::CompiledShaderInitializerType& Initializer):
diff --git a/Engine/Source/Runtime/Renderer/Private/RequiredTextureResolutionRendering.h b/Engine/Source/Runtime/Renderer/Private/RequiredTextureResolutionRendering.h
--- a/Engine/Source/Runtime/Renderer/Private/RequiredTextureResolutionRendering.h
+++ b/Engine/Source/Runtime/Renderer/Private/RequiredTextureResolutionRendering.h
@@ -27,7 +27,7 @@ public:
 	static bool ShouldCompilePermutation(EShaderPlatform Platform, const FMaterial* Material, const FVertexFactoryType* VertexFactoryType)
 	{
 		// See FDebugViewModeMaterialProxy::GetFriendlyName()
-		return AllowDebugViewPS(DVSM_RequiredTextureResolution, Platform) && Material->GetFriendlyName().Contains(TEXT("RequiredTextureResolution"));
+		return AllowDebugViewShaderMode(DVSM_RequiredTextureResolution, Platform, GetMaxSupportedFeatureLevel(Platform)) && Material->GetFriendlyName().Contains(TEXT("RequiredTextureResolution"));
 	}
 
 	FRequiredTextureResolutionPS(const ShaderMetaType::CompiledShaderInitializerType& Initializer):
diff --git a/Engine/Source/Runtime/Renderer/Private/ShaderComplexityRendering.h b/Engine/Source/Runtime/Renderer/Private/ShaderComplexityRendering.h
--- a/Engine/Source/Runtime/Renderer/Private/ShaderComplexityRendering.h
+++ b/Engine/Source/Runtime/Renderer/Private/ShaderComplexityRendering.h
@@ -28,7 +28,7 @@ public:
 	static bool ShouldCompilePermutation(EShaderPlatform Platform, const FMaterial* Material, const FVertexFactoryType* VertexFactoryType)
 	{
 		// See FDebugViewModeMaterialProxy::GetFriendlyName()
-		return AllowDebugViewPS(bQuadComplexity ? DVSM_QuadComplexity : DVSM_ShaderComplexity, Platform) && Material->GetFriendlyName().Contains(TEXT("ComplexityAccumulate"));
+		return AllowDebugViewShaderMode(bQuadComplexity ? DVSM_QuadComplexity : DVSM_ShaderComplexity, Platform, GetMaxSupportedFeatureLevel(Platform)) && Material->GetFriendlyName().Contains(TEXT("ComplexityAccumulate"));
 	}
 
 	TComplexityAccumulatePS(const ShaderMetaType::CompiledShaderInitializerType& Initializer):
@@ -52,7 +52,7 @@ public:
 
 	static void ModifyCompilationEnvironment(EShaderPlatform Platform, const FMaterial* Material, FShaderCompilerEnvironment& OutEnvironment)
 	{
-		OutEnvironment.SetDefine(TEXT("OUTPUT_QUAD_OVERDRAW"), AllowDebugViewPS(DVSM_QuadComplexity, Platform));
+		OutEnvironment.SetDefine(TEXT("OUTPUT_QUAD_OVERDRAW"), AllowDebugViewShaderMode(DVSM_QuadComplexity, Platform, GetMaxSupportedFeatureLevel(Platform)));
 	}
 
 	virtual void GetDebugViewModeShaderBindings(
