Fixed ambiguous conditional expression.

diff --git a/Engine/Source/Editor/UnrealEd/Private/SEditorViewportViewMenu.cpp b/Engine/Source/Editor/UnrealEd/Private/SEditorViewportViewMenu.cpp
--- a/Engine/Source/Editor/UnrealEd/Private/SEditorViewportViewMenu.cpp
+++ b/Engine/Source/Editor/UnrealEd/Private/SEditorViewportViewMenu.cpp
@@ -317,7 +317,7 @@ TSharedRef<SWidget> SEditorViewportViewMenu::GenerateViewMenuContent() const
 					{
 						const FEditorViewportCommands& BaseViewportCommands = FEditorViewportCommands::Get();
 
-						const ERHIFeatureLevel::Type FeatureLevel = (IsInGameThread() && GWorld) ? GWorld->FeatureLevel : GMaxRHIFeatureLevel;
+						const ERHIFeatureLevel::Type FeatureLevel = (IsInGameThread() && GWorld) ? (ERHIFeatureLevel::Type)GWorld->FeatureLevel : GMaxRHIFeatureLevel;
 
 						Menu.BeginSection("OptimizationViewmodes", LOCTEXT("OptimizationSubMenuHeader", "Optimization Viewmodes"));
 						{
