Adding an overridable call to the editor engine before creating single process PIE server

diff --git a/Engine/Source/Editor/UnrealEd/Classes/Editor/EditorEngine.h b/Engine/Source/Editor/UnrealEd/Classes/Editor/EditorEngine.h
--- a/Engine/Source/Editor/UnrealEd/Classes/Editor/EditorEngine.h
+++ b/Engine/Source/Editor/UnrealEd/Classes/Editor/EditorEngine.h
@@ -339,7 +339,7 @@ struct FSelectionStateOfLevel
 	TArray<FString> SelectedComponents;
 };
 
-/** Overrides you can pass when starting PIE to temporarily supress the user's options for dedicated server, number of clients..etc. */
+/** Overrides you can pass when starting PIE to temporarily suppress the user's options for dedicated server, number of clients..etc. */
 struct FPlayInEditorOverrides
 {
 	FPlayInEditorOverrides()
@@ -2659,6 +2659,9 @@ public:
 	 */
 	bool CreatePIEWorldFromLogin(FWorldContext& PieWorldContext, EPlayNetMode PlayNetMode, FPieLoginStruct& DataStruct);
 
+	/** Called before creating a PIE server instance. */
+	virtual FGameInstancePIEResult PreCreatePIEServerInstance(const bool bAnyBlueprintErrors, const bool bStartInSpectatorMode, const float PIEStartTime, const bool bSupportsOnlinePIE, int32& InNumOnlinePIEInstances);
+
 	/*
 	 * Handler for when viewport close request is made. 
 	 *
diff --git a/Engine/Source/Editor/UnrealEd/Private/PlayLevel.cpp b/Engine/Source/Editor/UnrealEd/Private/PlayLevel.cpp
--- a/Engine/Source/Editor/UnrealEd/Private/PlayLevel.cpp
+++ b/Engine/Source/Editor/UnrealEd/Private/PlayLevel.cpp
@@ -2548,6 +2548,11 @@ void UEditorEngine::PlayInEditor( UWorld* InWorld, bool bInSimulateInEditor, FPl
 	FEditorDelegates::PostPIEStarted.Broadcast( bInSimulateInEditor );
 }
 
+FGameInstancePIEResult UEditorEngine::PreCreatePIEServerInstance(const bool bAnyBlueprintErrors, const bool bStartInSpectatorMode, const float PIEStartTime, const bool bSupportsOnlinePIE, int32& InNumOnlinePIEInstances)
+{
+	return FGameInstancePIEResult::Success();
+}
+
 void UEditorEngine::SpawnIntraProcessPIEWorlds(bool bAnyBlueprintErrors, bool bStartInSpectatorMode)
 {
 	double PIEStartTime = FPlatformTime::Seconds();
@@ -2576,6 +2581,12 @@ void UEditorEngine::SpawnIntraProcessPIEWorlds(bool bAnyBlueprintErrors, bool bS
 	// Server
 	if (CanPlayNetDedicated || WillAutoConnectToServer)
 	{
+		FGameInstancePIEResult PreCreateResult = PreCreatePIEServerInstance(bAnyBlueprintErrors, bStartInSpectatorMode, PIEStartTime, false, NumOnlinePIEInstances);
+		if (!PreCreateResult.IsSuccess())
+		{
+			return;
+		}
+
 		PlayInSettings->SetPlayNetMode(EPlayNetMode::PIE_ListenServer);
 
 		if (!CanPlayNetDedicated)
@@ -2707,6 +2718,12 @@ void UEditorEngine::LoginPIEInstances(bool bAnyBlueprintErrors, bool bStartInSpe
 	// Server
 	if (WillAutoConnectToServer || CanPlayNetDedicated)
 	{
+		FGameInstancePIEResult PreCreateResult = PreCreatePIEServerInstance(bAnyBlueprintErrors, bStartInSpectatorMode, PIEStartTime, true, NumOnlinePIEInstances);
+		if (!PreCreateResult.IsSuccess())
+		{
+			return;
+		}
+
 		FWorldContext &PieWorldContext = CreateNewWorldContext(EWorldType::PIE);
 		PieWorldContext.PIEInstance = PIEInstance++;
 		PieWorldContext.RunAsDedicated = CanPlayNetDedicated;
