Used the defered method of updating the notification, to avoid ensuring about proper thread utilization.

diff --git a/Engine/Source/Editor/UnrealEd/Private/PlayLevel.cpp b/Engine/Source/Editor/UnrealEd/Private/PlayLevel.cpp
--- a/Engine/Source/Editor/UnrealEd/Private/PlayLevel.cpp
+++ b/Engine/Source/Editor/UnrealEd/Private/PlayLevel.cpp
@@ -1627,7 +1627,11 @@ void UEditorEngine::HandleStageStarted(const FString& InStage, TWeakPtr<SNotific
 
 	if (bSetNotification)
 	{
-		NotificationItemPtr.Pin()->SetText(NotificationText);
+		TGraphTask<FLauncherNotificationTask>::CreateTask().ConstructAndDispatchWhenReady(
+			NotificationItemPtr,
+			SNotificationItem::CS_Fail,
+			NotificationText
+		);
 	}
 }
 
@@ -1669,8 +1673,6 @@ void UEditorEngine::HandleLaunchCompleted(bool Succeeded, double TotalTime, int3
 			(PlayUsingLauncherDeviceId.Left(PlayUsingLauncherDeviceId.Find(TEXT("@"))) == TEXT("TVOS") && PlayUsingLauncherDeviceName.Contains(DummyTVOSDeviceName)))
 		{
 			CompletionMsg = LOCTEXT("DeploymentTaskCompleted", "Deployment complete! Open the app on your device to launch.");
-			TSharedPtr<SNotificationItem> NotificationItem = NotificationItemPtr.Pin();
-//			NotificationItem->SetExpireDuration(30.0f);
 		}
 		else
 		{
