Used the defered method of updating the notification, to avoid ensuring about proper thread utilization. (Set Pending correctly)

diff --git a/Engine/Source/Editor/UnrealEd/Private/PlayLevel.cpp b/Engine/Source/Editor/UnrealEd/Private/PlayLevel.cpp
--- a/Engine/Source/Editor/UnrealEd/Private/PlayLevel.cpp
+++ b/Engine/Source/Editor/UnrealEd/Private/PlayLevel.cpp
@@ -1629,7 +1629,7 @@ void UEditorEngine::HandleStageStarted(const FString& InStage, TWeakPtr<SNotific
 	{
 		TGraphTask<FLauncherNotificationTask>::CreateTask().ConstructAndDispatchWhenReady(
 			NotificationItemPtr,
-			SNotificationItem::CS_Fail,
+			SNotificationItem::CS_Pending,
 			NotificationText
 		);
 	}
