Removed cancel dialog options when compiling debug view mode shaders.

diff --git a/Engine/Source/Editor/UnrealEd/Private/EditorBuildUtils.cpp b/Engine/Source/Editor/UnrealEd/Private/EditorBuildUtils.cpp
--- a/Engine/Source/Editor/UnrealEd/Private/EditorBuildUtils.cpp
+++ b/Engine/Source/Editor/UnrealEd/Private/EditorBuildUtils.cpp
@@ -1299,7 +1299,7 @@ bool FEditorBuildUtils::CompileViewModeShaders(UWorld* InWorld, EViewModeIndex S
 	const EMaterialQualityLevel::Type QualityLevel = GetCachedScalabilityCVars().MaterialQualityLevel;;
 
 	FScopedSlowTask CompileShaderTask(3.f, LOCTEXT("CompileMissingViewModeShaders", "Compiling Missing ViewMode Shaders")); // { Get Used Materials, Sync Pending Shader, Wait for Compilation }
-	CompileShaderTask.MakeDialog(true);
+	CompileShaderTask.MakeDialog(false); // Can't cancel because get compiled anyway if shaders are missing.
 
 	CollectGarbage(GARBAGE_COLLECTION_KEEPFLAGS);
 
