Harden another access of GetSerializeContext, fix another crash compiling blueprints in preview builds

diff --git a/Engine/Source/Editor/UnrealEd/Private/Kismet2/BlueprintEditorUtils.cpp b/Engine/Source/Editor/UnrealEd/Private/Kismet2/BlueprintEditorUtils.cpp
--- a/Engine/Source/Editor/UnrealEd/Private/Kismet2/BlueprintEditorUtils.cpp
+++ b/Engine/Source/Editor/UnrealEd/Private/Kismet2/BlueprintEditorUtils.cpp
@@ -734,7 +734,7 @@ void FBlueprintEditorUtils::PatchNewCDOIntoLinker(UObject* CDO, FLinkerLoad* Lin
 			FUObjectSerializeContext* LoadContext = InLoadContext ? InLoadContext : Linker->GetSerializeContext();
 
 			// Make sure the new CDO gets PostLoad called on it, so either add it to ObjLoaded list, or replace it if already present.
-			if (!LoadContext->PRIVATE_PatchNewObjectIntoExport(OldCDO, CDO))
+			if (LoadContext && !LoadContext->PRIVATE_PatchNewObjectIntoExport(OldCDO, CDO))
 			{
 				if (OldObjectFlags & RF_NeedPostLoad)
 				{
