Changing blutility checks to also check for editor utility widgets

diff --git a/Engine/Source/Editor/BlueprintGraph/Private/K2Node_GetSubsystem.cpp b/Engine/Source/Editor/BlueprintGraph/Private/K2Node_GetSubsystem.cpp
--- a/Engine/Source/Editor/BlueprintGraph/Private/K2Node_GetSubsystem.cpp
+++ b/Engine/Source/Editor/BlueprintGraph/Private/K2Node_GetSubsystem.cpp
@@ -686,7 +686,7 @@ bool UK2Node_GetEditorSubsystem::IsActionFilteredOut(class FBlueprintActionFilte
 {
 	for(UBlueprint* BP : Filter.Context.Blueprints)
 	{
-		if (!FBlueprintEditorUtils::IsBlutility(BP))
+		if (!FBlueprintEditorUtils::IsEditorUtilityBlueprint(BP))
 		{
 			return true;
 		}
@@ -700,7 +700,7 @@ void UK2Node_GetEditorSubsystem::ValidateNodeDuringCompilation(class FCompilerRe
 
 	const UBlueprint* BP = FBlueprintEditorUtils::FindBlueprintForNodeChecked(this);
 
-	if (!FBlueprintEditorUtils::IsBlutility(BP))
+	if (!FBlueprintEditorUtils::IsEditorUtilityBlueprint(BP))
 	{
 		const FText ErrorText = NSLOCTEXT("K2Node", "GetEditorSubsystem_Error", "Editor Subsystems can only be used in Editor Utilities / Blutilities");
 		MessageLog.Error(*ErrorText.ToString(), this);
diff --git a/Engine/Source/Editor/Blutility/Private/BlutilityModule.cpp b/Engine/Source/Editor/Blutility/Private/BlutilityModule.cpp
--- a/Engine/Source/Editor/Blutility/Private/BlutilityModule.cpp
+++ b/Engine/Source/Editor/Blutility/Private/BlutilityModule.cpp
@@ -201,11 +201,13 @@ public:
 		FEditorSupportDelegates::PrepareToCleanseEditorObject.RemoveAll(this);
 	}
 
-	virtual bool IsBlutility( const UBlueprint* Blueprint ) const override
+	virtual bool IsEditorUtilityBlueprint( const UBlueprint* Blueprint ) const override
 	{
 		const UClass* BPClass = Blueprint ? Blueprint->GetClass() : nullptr;
 
-		if( BPClass && BPClass->IsChildOf( UEditorUtilityBlueprint::StaticClass() ))
+		if( BPClass && 
+			(BPClass->IsChildOf( UEditorUtilityBlueprint::StaticClass())
+			|| BPClass->IsChildOf(UEditorUtilityWidgetBlueprint::StaticClass())))
 		{
 			return true;
 		}
diff --git a/Engine/Source/Editor/Blutility/Public/IBlutilityModule.h b/Engine/Source/Editor/Blutility/Public/IBlutilityModule.h
--- a/Engine/Source/Editor/Blutility/Public/IBlutilityModule.h
+++ b/Engine/Source/Editor/Blutility/Public/IBlutilityModule.h
@@ -16,7 +16,7 @@ class IBlutilityModule : public IModuleInterface
 public:
 
 	/** Returns if the blueprint is blutility based */
-	virtual bool IsBlutility( const UBlueprint* Blueprint ) const = 0;
+	virtual bool IsEditorUtilityBlueprint( const UBlueprint* Blueprint ) const = 0;
 
 	/** Global Find Results workspace menu item */
 	virtual TSharedPtr<class FWorkspaceItem> GetMenuGroup() const = 0;
diff --git a/Engine/Source/Editor/DetailCustomizations/Private/ObjectDetails.cpp b/Engine/Source/Editor/DetailCustomizations/Private/ObjectDetails.cpp
--- a/Engine/Source/Editor/DetailCustomizations/Private/ObjectDetails.cpp
+++ b/Engine/Source/Editor/DetailCustomizations/Private/ObjectDetails.cpp
@@ -102,7 +102,7 @@ void FObjectDetails::AddCallInEditorMethods(IDetailLayoutBuilder& DetailBuilder)
 			{
 				if (UBlueprint* Blueprint = Cast<UBlueprint>(TestFunctionOwnerClass->ClassGeneratedBy))
 				{
-					if (FBlueprintEditorUtils::IsBlutility(Blueprint))
+					if (FBlueprintEditorUtils::IsEditorUtilityBlueprint(Blueprint))
 					{
 						// Skip Blutilities as these are handled by FEditorUtilityInstanceDetails
 						continue;
diff --git a/Engine/Source/Editor/Kismet/Private/BlueprintDetailsCustomization.cpp b/Engine/Source/Editor/Kismet/Private/BlueprintDetailsCustomization.cpp
--- a/Engine/Source/Editor/Kismet/Private/BlueprintDetailsCustomization.cpp
+++ b/Engine/Source/Editor/Kismet/Private/BlueprintDetailsCustomization.cpp
@@ -3418,7 +3418,7 @@ void FBlueprintGraphActionDetails::CustomizeDetails( IDetailLayoutBuilder& Detai
 				]
 			];
 		}
-		const bool bShowCallInEditor = IsCustomEvent() || FBlueprintEditorUtils::IsBlutility( GetBlueprintObj() ) || (FunctionEntryNode && FunctionEntryNode->IsEditable());
+		const bool bShowCallInEditor = IsCustomEvent() || FBlueprintEditorUtils::IsEditorUtilityBlueprint( GetBlueprintObj() ) || (FunctionEntryNode && FunctionEntryNode->IsEditable());
 		if( bShowCallInEditor )
 		{
 			Category.AddCustomRow( LOCTEXT( "EditorCallable", "Call In Editor" ) )
diff --git a/Engine/Source/Editor/UnrealEd/Private/Kismet2/BlueprintEditorUtils.cpp b/Engine/Source/Editor/UnrealEd/Private/Kismet2/BlueprintEditorUtils.cpp
--- a/Engine/Source/Editor/UnrealEd/Private/Kismet2/BlueprintEditorUtils.cpp
+++ b/Engine/Source/Editor/UnrealEd/Private/Kismet2/BlueprintEditorUtils.cpp
@@ -3137,13 +3137,13 @@ bool FBlueprintEditorUtils::IsBlueprintConst(const UBlueprint* Blueprint)
 	return Blueprint && Blueprint->BlueprintType == BPTYPE_Const;
 }
 
-bool FBlueprintEditorUtils::IsBlutility(const UBlueprint* Blueprint)
+bool FBlueprintEditorUtils::IsEditorUtilityBlueprint(const UBlueprint* Blueprint)
 {
 	IBlutilityModule* BlutilityModule = FModuleManager::GetModulePtr<IBlutilityModule>("Blutility");
 
 	if (BlutilityModule)
 	{
-		return BlutilityModule->IsBlutility( Blueprint );
+		return BlutilityModule->IsEditorUtilityBlueprint( Blueprint );
 	}
 	return false;
 }
diff --git a/Engine/Source/Editor/UnrealEd/Public/Kismet2/BlueprintEditorUtils.h b/Engine/Source/Editor/UnrealEd/Public/Kismet2/BlueprintEditorUtils.h
--- a/Engine/Source/Editor/UnrealEd/Public/Kismet2/BlueprintEditorUtils.h
+++ b/Engine/Source/Editor/UnrealEd/Public/Kismet2/BlueprintEditorUtils.h
@@ -350,7 +350,7 @@ public:
 				// We need to mark the function entry as editable so that we can
 				// set metadata on it if it is a blutility:
 				K2Schema->MarkFunctionEntryAsEditable(Graph, true);
-				if( IsBlutility( Blueprint ))
+				if( IsEditorUtilityBlueprint( Blueprint ))
 				{
 					if( FKismetUserDeclaredFunctionMetadata* MetaData = GetGraphFunctionMetaData( Graph ))
 					{
@@ -524,7 +524,7 @@ public:
 	static bool IsBlueprintConst(const UBlueprint* Blueprint);
 
 	/** Returns whether or not the blueprint is a blutility */
-	static bool IsBlutility(const UBlueprint* Blueprint);
+	static bool IsEditorUtilityBlueprint(const UBlueprint* Blueprint);
 
 	/**
 	 * Whether or not this is an actor-based blueprint, and supports features like the uber-graph, components, etc
