Disable editor heartbeat warning when trying to send too many heartbeats in one minute. A long term fix is discussed in #jira UE-62403

diff --git a/Engine/Source/Editor/UnrealEd/Private/UnrealEdMisc.cpp b/Engine/Source/Editor/UnrealEd/Private/UnrealEdMisc.cpp
--- a/Engine/Source/Editor/UnrealEd/Private/UnrealEdMisc.cpp
+++ b/Engine/Source/Editor/UnrealEd/Private/UnrealEdMisc.cpp
@@ -635,7 +635,9 @@ void FUnrealEdMisc::EditorAnalyticsHeartbeat()
 	// allow a little bit of slop in the timing in case the event is only firing slightly too soon, even though that is technically impossible.
 	if (Now <= LastHeartbeatTime + HeartbeatInvervalSec*0.9)
 	{
-		UE_LOG(LogUnrealEdMisc, Warning, TEXT("Heartbeat event firing too frequently (%.3f sec). This should never happen. Something is wrong with the timer delegate!"), (float)(Now - LastHeartbeatTime) );
+		// This can happen too often with editor hitches and long-running tasks, which cause the timer manager to flush many events at once.
+		// UE-62403 is tracking a deeeper fix.
+		//UE_LOG(LogUnrealEdMisc, Warning, TEXT("Heartbeat event firing too frequently (%.3f sec). This should never happen. Something is wrong with the timer delegate!"), (float)(Now - LastHeartbeatTime) );
 		return;
 	}
 
