Make sure that the blueprint diff window is closed if the level blueprint is being diffed when the level is closed so as to avoid crash.

diff --git a/Engine/Source/Editor/Kismet/Private/SBlueprintDiff.cpp b/Engine/Source/Editor/Kismet/Private/SBlueprintDiff.cpp
--- a/Engine/Source/Editor/Kismet/Private/SBlueprintDiff.cpp
+++ b/Engine/Source/Editor/Kismet/Private/SBlueprintDiff.cpp
@@ -836,7 +836,7 @@ void SBlueprintDiff::OnCloseAssetEditor(UObject* Asset, EAssetEditorCloseReason
 
 		if (WeakParentWindow.IsValid())
 		{
-			WeakParentWindow.Pin()->RequestDestroyWindow();
+			WeakParentWindow.Pin()->DestroyWindowImmediately();
 		}
 	}
 }
diff --git a/Engine/Source/Editor/UnrealEd/Private/EditorEngine.cpp b/Engine/Source/Editor/UnrealEd/Private/EditorEngine.cpp
--- a/Engine/Source/Editor/UnrealEd/Private/EditorEngine.cpp
+++ b/Engine/Source/Editor/UnrealEd/Private/EditorEngine.cpp
@@ -2346,14 +2346,7 @@ void UEditorEngine::CloseEditedWorldAssets(UWorld* InWorld)
 
 		if (AssetWorld && ClosingWorlds.Contains(AssetWorld))
 		{
-			const TArray<IAssetEditorInstance*> AssetEditors = EditorManager.FindEditorsForAsset(Asset);
-			for (IAssetEditorInstance* EditorInstance : AssetEditors )
-			{
-				if (EditorInstance != NULL)
-				{
-					EditorInstance->CloseWindow();
-				}
-			}
+			EditorManager.CloseAllEditorsForAsset(Asset);
 		}
 	}
 }
diff --git a/Engine/Source/Editor/UnrealEd/Private/Toolkits/AssetEditorManager.cpp b/Engine/Source/Editor/UnrealEd/Private/Toolkits/AssetEditorManager.cpp
--- a/Engine/Source/Editor/UnrealEd/Private/Toolkits/AssetEditorManager.cpp
+++ b/Engine/Source/Editor/UnrealEd/Private/Toolkits/AssetEditorManager.cpp
@@ -108,7 +108,7 @@ void FAssetEditorManager::AddReferencedObjects( FReferenceCollector& Collector )
 	}
 
 	// If the pointer got deleted or swapped out due to hot reload.
-	for ( auto& Entry : ModifiedKeys )
+	for ( TPair<UObject*, UObject*>& Entry : ModifiedKeys )
 	{
 		if ( Entry.Value )
 		{
@@ -167,9 +167,12 @@ int32 FAssetEditorManager::CloseAllEditorsForAsset(UObject* Asset)
 {
 	TArray<IAssetEditorInstance*> EditorInstances = FindEditorsForAsset(Asset);
 
-	for( auto EditorIter : EditorInstances )
+	for (IAssetEditorInstance* EditorInstance : EditorInstances)
 	{
-		EditorIter->CloseWindow();
+		if (EditorInstance)
+		{
+			EditorInstance->CloseWindow();
+		}
 	}
 
 	AssetEditorRequestCloseEvent.Broadcast(Asset, EAssetEditorCloseReason::CloseAllEditorsForAsset);
@@ -181,9 +184,12 @@ void FAssetEditorManager::RemoveAssetFromAllEditors(UObject* Asset)
 {
 	TArray<IAssetEditorInstance*> EditorInstances = FindEditorsForAsset(Asset);
 
-	for (auto EditorIter : EditorInstances)
+	for (IAssetEditorInstance* EditorIter : EditorInstances)
 	{
-		EditorIter->RemoveEditingAsset(Asset);
+		if (EditorInstance)
+		{
+			EditorIter->RemoveEditingAsset(Asset);
+		}
 	}
 
 	AssetEditorRequestCloseEvent.Broadcast(Asset, EAssetEditorCloseReason::RemoveAssetFromAllEditors);
@@ -342,7 +348,7 @@ bool FAssetEditorManager::OpenEditorForAsset(UObject* Asset, const EToolkitMode:
 
 	TWeakPtr<IAssetTypeActions> AssetTypeActions = AssetToolsModule.Get().GetAssetTypeActionsForClass( Asset->GetClass() );
 	
-	auto ActualToolkitMode = ToolkitMode;
+	EToolkitMode::Type ActualToolkitMode = ToolkitMode;
 	if( AssetTypeActions.IsValid() )
 	{
 		if( AssetTypeActions.Pin()->ShouldForceWorldCentric() )
@@ -413,7 +419,7 @@ bool FAssetEditorManager::OpenEditorForAsset(UObject* Asset, const EToolkitMode:
 }
 
 
-bool FAssetEditorManager::OpenEditorForAssets( const TArray< UObject* >& Assets, const EToolkitMode::Type ToolkitMode, TSharedPtr< IToolkitHost > OpenedFromLevelEditor )
+bool FAssetEditorManager::OpenEditorForAssets( const TArray <UObject* >& Assets, const EToolkitMode::Type ToolkitMode, TSharedPtr< IToolkitHost > OpenedFromLevelEditor )
 {
 	if( Assets.Num() == 1 )
 	{
@@ -422,7 +428,7 @@ bool FAssetEditorManager::OpenEditorForAssets( const TArray< UObject* >& Assets,
 	else if (Assets.Num() > 0)
 	{
 		TArray<UObject*> SkipOpenAssets;
-		for (auto Asset : Assets)
+		for (UObject* Asset : Assets)
 		{
 			// If any of the assets are already open or the package is cooked,
 			// remove them from the list of assets to open an editor for
@@ -435,7 +441,7 @@ bool FAssetEditorManager::OpenEditorForAssets( const TArray< UObject* >& Assets,
 
 		// Verify that all the assets are of the same class
 		bool bAssetClassesMatch = true;
-		auto AssetClass = Assets[0]->GetClass();
+		UClass* AssetClass = Assets[0]->GetClass();
 		for (int32 i = 1; i < Assets.Num(); i++)
 		{
 			if (Assets[i]->GetClass() != AssetClass)
@@ -456,7 +462,7 @@ bool FAssetEditorManager::OpenEditorForAssets( const TArray< UObject* >& Assets,
 				GWarn->BeginSlowTask(LOCTEXT("OpenEditors", "Opening Editor(s)..."), true);
 
 				// Determine the appropriate toolkit mode for the asset type
-				auto ActualToolkitMode = ToolkitMode;
+				EToolkitMode::Type ActualToolkitMode = ToolkitMode;
 				if (AssetTypeActions.Pin()->ShouldForceWorldCentric())
 				{
 					// This asset type prefers a specific toolkit mode
@@ -488,7 +494,7 @@ bool FAssetEditorManager::OpenEditorForAssets( const TArray< UObject* >& Assets,
 
 				TArray<FLocalAssetInfo> AssetInfoList;
 				AssetInfoList.Reserve(Assets.Num());
-				for (auto Asset : Assets)
+				for (UObject* Asset : Assets)
 				{
 					AssetInfoList.Add(FLocalAssetInfo(Asset, Asset->GetPathName()));
 				}
@@ -499,8 +505,8 @@ bool FAssetEditorManager::OpenEditorForAssets( const TArray< UObject* >& Assets,
 				// If any assets were destroyed, attempt to find them if they were recreated
 				for (int32 i = 0; i < Assets.Num(); i++)
 				{
-					auto& AssetInfo = AssetInfoList[i];
-					auto Asset = Assets[i];
+					const FLocalAssetInfo& AssetInfo = AssetInfoList[i];
+					UObject* Asset = Assets[i];
 
 					if (!AssetInfo.WeakAsset.IsValid() && !AssetInfo.AssetPath.IsEmpty())
 					{
@@ -517,7 +523,7 @@ bool FAssetEditorManager::OpenEditorForAssets( const TArray< UObject* >& Assets,
 		else
 		{
 			// Asset types don't match or some are already open, so just open individual editors for the unopened ones
-			for (auto Asset : Assets)
+			for (UObject* Asset : Assets)
 			{
 				if (!SkipOpenAssets.Contains(Asset))
 				{
@@ -734,7 +740,7 @@ void FAssetEditorManager::SaveOpenAssetEditors(bool bOnShutdown)
 		// Don't save a list of assets to restore if we are running under a debugger
 		if(!FPlatformMisc::IsDebuggerPresent())
 		{
-			for (auto EditorPair : OpenedEditors)
+			for (const TPair<IAssetEditorInstance*, UObject*>& EditorPair : OpenedEditors)
 			{
 				IAssetEditorInstance* Editor = EditorPair.Key;
 				if (Editor != NULL)
