Fix copy paste mistake in review
#rnx

diff --git a/Engine/Source/Editor/UnrealEd/Private/Toolkits/AssetEditorManager.cpp b/Engine/Source/Editor/UnrealEd/Private/Toolkits/AssetEditorManager.cpp
--- a/Engine/Source/Editor/UnrealEd/Private/Toolkits/AssetEditorManager.cpp
+++ b/Engine/Source/Editor/UnrealEd/Private/Toolkits/AssetEditorManager.cpp
@@ -186,7 +186,7 @@ void FAssetEditorManager::RemoveAssetFromAllEditors(UObject* Asset)
 
 	for (IAssetEditorInstance* EditorIter : EditorInstances)
 	{
-		if (EditorInstance)
+		if (EditorIter)
 		{
 			EditorIter->RemoveEditingAsset(Asset);
 		}
