Keep a list of all of the previously seen objects when debugging blueprints in FKismetDebugUtilities::FindDebuggingData and check against the list before recursing.
Potential crash fix.

diff --git a/Engine/Source/Editor/UnrealEd/Private/Kismet2/KismetDebugUtilities.cpp b/Engine/Source/Editor/UnrealEd/Private/Kismet2/KismetDebugUtilities.cpp
--- a/Engine/Source/Editor/UnrealEd/Private/Kismet2/KismetDebugUtilities.cpp
+++ b/Engine/Source/Editor/UnrealEd/Private/Kismet2/KismetDebugUtilities.cpp
@@ -1113,7 +1113,8 @@ FKismetDebugUtilities::EWatchTextResult FKismetDebugUtilities::GetWatchText(FStr
 	void* DataPtr = nullptr;
 	void* DeltaPtr = nullptr;
 	UObject* ParentObj = nullptr;
-	FKismetDebugUtilities::EWatchTextResult Result = FindDebuggingData(Blueprint, ActiveObject, WatchPin, PropertyToDebug, DataPtr, DeltaPtr, ParentObj);
+	TArray<UObject*> SeenObjects;
+	FKismetDebugUtilities::EWatchTextResult Result = FindDebuggingData(Blueprint, ActiveObject, WatchPin, PropertyToDebug, DataPtr, DeltaPtr, ParentObj, SeenObjects);
 
 	if (Result == FKismetDebugUtilities::EWatchTextResult::EWTR_Valid)
 	{
@@ -1129,7 +1130,8 @@ FKismetDebugUtilities::EWatchTextResult FKismetDebugUtilities::GetDebugInfo(FDeb
 	void* DataPtr = nullptr;
 	void* DeltaPtr = nullptr;
 	UObject* ParentObj = nullptr;
-	FKismetDebugUtilities::EWatchTextResult Result = FindDebuggingData(Blueprint, ActiveObject, WatchPin, PropertyToDebug, DataPtr, DeltaPtr, ParentObj);
+	TArray<UObject*> SeenObjects;
+	FKismetDebugUtilities::EWatchTextResult Result = FindDebuggingData(Blueprint, ActiveObject, WatchPin, PropertyToDebug, DataPtr, DeltaPtr, ParentObj, SeenObjects);
 
 	if (Result == FKismetDebugUtilities::EWatchTextResult::EWTR_Valid)
 	{
@@ -1139,7 +1141,7 @@ FKismetDebugUtilities::EWatchTextResult FKismetDebugUtilities::GetDebugInfo(FDeb
 	return Result;
 }
 
-FKismetDebugUtilities::EWatchTextResult FKismetDebugUtilities::FindDebuggingData(UBlueprint* Blueprint, UObject* ActiveObject, const UEdGraphPin* WatchPin, UProperty*& OutProperty, void*& OutData, void*& OutDelta, UObject*& OutParent)
+FKismetDebugUtilities::EWatchTextResult FKismetDebugUtilities::FindDebuggingData(UBlueprint* Blueprint, UObject* ActiveObject, const UEdGraphPin* WatchPin, UProperty*& OutProperty, void*& OutData, void*& OutDelta, UObject*& OutParent, TArray<UObject*>& SeenObjects)
 {
 	FKismetDebugUtilitiesData& Data = FKismetDebugUtilitiesData::Get();
 
@@ -1267,7 +1269,8 @@ FKismetDebugUtilities::EWatchTextResult FKismetDebugUtilities::FindDebuggingData
 						void* SelfPinData = nullptr;
 						void* SelfPinDelta = nullptr;
 						UObject* SelfPinParent = nullptr;
-						FKismetDebugUtilities::EWatchTextResult Result = FindDebuggingData(Blueprint, ActiveObject, SelfPin, SelfPinProperty, SelfPinData, SelfPinDelta, SelfPinParent);
+						SeenObjects.AddUnique(ActiveObject);
+						FKismetDebugUtilities::EWatchTextResult Result = FindDebuggingData(Blueprint, ActiveObject, SelfPin, SelfPinProperty, SelfPinData, SelfPinDelta, SelfPinParent, SeenObjects);
 						UObjectPropertyBase* SelfPinPropertyBase = Cast<UObjectPropertyBase>(SelfPinProperty);
 						if (Result == EWTR_Valid && SelfPinPropertyBase != nullptr)
 						{
@@ -1275,7 +1278,10 @@ FKismetDebugUtilities::EWatchTextResult FKismetDebugUtilities::FindDebuggingData
 							UObject* TempActiveObject = SelfPinPropertyBase->GetObjectPropertyValue(PropertyValue);
 							if (TempActiveObject && TempActiveObject != ActiveObject)
 							{
-								return FindDebuggingData(Blueprint, TempActiveObject, WatchPin, OutProperty, OutData, OutDelta, OutParent);
+								if (!SeenObjects.Contains(TempActiveObject))
+								{
+									return FindDebuggingData(Blueprint, TempActiveObject, WatchPin, OutProperty, OutData, OutDelta, OutParent, SeenObjects);
+								}
 							}
 						}
 					}
diff --git a/Engine/Source/Editor/UnrealEd/Public/Kismet2/KismetDebugUtilities.h b/Engine/Source/Editor/UnrealEd/Public/Kismet2/KismetDebugUtilities.h
--- a/Engine/Source/Editor/UnrealEd/Public/Kismet2/KismetDebugUtilities.h
+++ b/Engine/Source/Editor/UnrealEd/Public/Kismet2/KismetDebugUtilities.h
@@ -251,7 +251,7 @@ protected:
 
 	// helper function for converting between blueprint and debuggable data
 	// output params are only valid if the return result is EWatchTextResult::EWTR_Valid
-	static EWatchTextResult FindDebuggingData(UBlueprint* Blueprint, UObject* ActiveObject, const UEdGraphPin* WatchPin, UProperty*& OutProperty, void*& OutData, void*& OutDelta, UObject*& OutParent);
+	static EWatchTextResult FindDebuggingData(UBlueprint* Blueprint, UObject* ActiveObject, const UEdGraphPin* WatchPin, UProperty*& OutProperty, void*& OutData, void*& OutDelta, UObject*& OutParent, TArray<UObject*>& SeenObjects);
 
 private:
 	FKismetDebugUtilities() {}
