Make the skeletal mesh split workflow setup more polish.

diff --git a/Engine/Plugins/Experimental/SkeletalReduction/Source/Private/SkeletalMeshReductionPlugin.cpp b/Engine/Plugins/Experimental/SkeletalReduction/Source/Private/SkeletalMeshReductionPlugin.cpp
--- a/Engine/Plugins/Experimental/SkeletalReduction/Source/Private/SkeletalMeshReductionPlugin.cpp
+++ b/Engine/Plugins/Experimental/SkeletalReduction/Source/Private/SkeletalMeshReductionPlugin.cpp
@@ -1900,7 +1900,7 @@ void FQuadricSkeletalMeshReduction::ReduceSkeletalMesh(USkeletalMesh& SkeletalMe
 	FSkeletalMeshLODModel* NewModel = new FSkeletalMeshLODModel();
 
 	// Swap out the old model.  
-
+	FSkeletalMeshImportData RawMesh;
 	{
 		FSkeletalMeshLODModel* Old = LODModels[LODIndex];
 
@@ -1913,9 +1913,7 @@ void FQuadricSkeletalMeshReduction::ReduceSkeletalMesh(USkeletalMesh& SkeletalMe
 		else if(bReducingSourceModel)
 		{
 			//In case we reduce the source model we want to keep the original import data
-			FSkeletalMeshImportData RawMesh;
 			SrcModel->RawSkeletalMeshBulkData.LoadRawMesh(RawMesh);
-			LODModels[LODIndex]->RawSkeletalMeshBulkData.SaveRawMesh(RawMesh);
 		}
 	}
 
@@ -2047,6 +2045,12 @@ void FQuadricSkeletalMeshReduction::ReduceSkeletalMesh(USkeletalMesh& SkeletalMe
 		}
 	}
 
+	if (bReducingSourceModel && RawMesh.Points.Num() > 0)
+	{
+		//Put back the original import data, we need it to allow inline reduction and skeletal mesh split workflow
+		SkeletalMeshResource.LODModels[LODIndex].RawSkeletalMeshBulkData.SaveRawMesh(RawMesh);
+	}
+
 	SkeletalMesh.CalculateRequiredBones(SkeletalMeshResource.LODModels[LODIndex], SkeletalMesh.RefSkeleton, &BonesToRemove);
 }
 
diff --git a/Engine/Source/Editor/UnrealEd/Private/Editor.cpp b/Engine/Source/Editor/UnrealEd/Private/Editor.cpp
--- a/Engine/Source/Editor/UnrealEd/Private/Editor.cpp
+++ b/Engine/Source/Editor/UnrealEd/Private/Editor.cpp
@@ -329,16 +329,32 @@ bool FReimportManager::Reimport( UObject* Obj, bool bAskForNewFileIfMissing, boo
 						GetNewReimportPath(Obj, SourceFilenames, FileIndex);
 					}
 				}
-				if ( SourceFilenames.Num() == 0 )
+				bool bAllSourceFileEmpty = true;
+				for (int32 SourceIndex = 0; SourceIndex < SourceFilenames.Num(); ++SourceIndex)
+				{
+					if (!SourceFilenames[SourceIndex].IsEmpty())
+					{
+						bAllSourceFileEmpty = false;
+						break;
+					}
+				}
+				if ( SourceFilenames.Num() == 0 || bAllSourceFileEmpty)
 				{
 					// Failed to specify a new filename. Don't show a notification of the failure since the user exited on his own
 					bValidSourceFilename = false;
 					bShowNotification = false;
+					SourceFilenames.Empty();
 				}
 				else
 				{
 					// A new filename was supplied, update the path
-					CanReimportHandler->SetReimportPaths(Obj, SourceFilenames[0], SourceFileIndex);
+					for (int32 SourceIndex = 0; SourceIndex < SourceFilenames.Num(); ++SourceIndex)
+					{
+						if (!SourceFilenames[SourceIndex].IsEmpty())
+						{
+							CanReimportHandler->SetReimportPaths(Obj, SourceFilenames[SourceIndex], SourceIndex);
+						}
+					}
 				}
 			}
 			else if (!PreferredReimportFile.IsEmpty() && !SourceFilenames.Contains(PreferredReimportFile))
@@ -594,6 +610,7 @@ bool FReimportManager::ReimportMultiple(TArrayView<UObject*> Objects, bool bAskF
 
 void FReimportManager::GetNewReimportPath(UObject* Obj, TArray<FString>& InOutFilenames, int32 SourceFileIndex /*= INDEX_NONE*/)
 {
+	int32 RealSourceFileIndex = SourceFileIndex == INDEX_NONE ? 0 : SourceFileIndex;
 	TArray<UObject*> ReturnObjects;
 	FString FileTypes;
 	FString AllExtensions;
@@ -609,8 +626,19 @@ void FReimportManager::GetNewReimportPath(UObject* Obj, TArray<FString>& InOutFi
 	}
 
 	// Determine whether we will allow multi select and clear old filenames
-	bool bAllowMultiSelect = InOutFilenames.Num() > 1;
-	InOutFilenames.Empty();
+	bool bAllowMultiSelect = SourceFileIndex == INDEX_NONE && InOutFilenames.Num() > 1;
+	if (bAllowMultiSelect)
+	{
+		InOutFilenames.Empty();
+	}
+	else
+	{
+		if (!InOutFilenames.IsValidIndex(RealSourceFileIndex))
+		{
+			InOutFilenames.AddZeroed(RealSourceFileIndex - InOutFilenames.Num() + 1);
+		}
+		InOutFilenames[RealSourceFileIndex].Empty();
+	}
 
 	// Get the list of valid factories
 	for( TObjectIterator<UClass> It ; It ; ++It )
@@ -700,9 +728,20 @@ void FReimportManager::GetNewReimportPath(UObject* Obj, TArray<FString>& InOutFi
 
 	if ( bOpened )
 	{
-		for (int32 FileIndex = 0; FileIndex < OpenFilenames.Num(); ++FileIndex)
+		if (bAllowMultiSelect)
 		{
-			InOutFilenames.Add(OpenFilenames[FileIndex]);
+			for (int32 FileIndex = 0; FileIndex < OpenFilenames.Num(); ++FileIndex)
+			{
+				InOutFilenames.Add(OpenFilenames[FileIndex]);
+			}
+		}
+		else
+		{
+			//Use the first valid entry
+			if(OpenFilenames.Num() > 0)
+			{
+				InOutFilenames[RealSourceFileIndex] = OpenFilenames[0];
+			}
 		}
 	}
 }
