Fix for Material Instance editor not properly loading the override state of the subsurface profile.

diff --git a/Engine/Source/Editor/UnrealEd/Private/PreviewMaterial.cpp b/Engine/Source/Editor/UnrealEd/Private/PreviewMaterial.cpp
--- a/Engine/Source/Editor/UnrealEd/Private/PreviewMaterial.cpp
+++ b/Engine/Source/Editor/UnrealEd/Private/PreviewMaterial.cpp
@@ -1553,10 +1553,9 @@ void UMaterialEditorInstanceConstant::CopyBasePropertiesFromParent()
 	//Copy refraction settings
 	SourceInstance->GetRefractionSettings(RefractionDepthBias);
 
-	if (!bOverrideSubsurfaceProfile)
-	{
-		SubsurfaceProfile = SourceInstance->SubsurfaceProfile;
-	}
+	bOverrideSubsurfaceProfile = SourceInstance->bOverrideSubsurfaceProfile;
+	// Copy the subsurface profile. GetSubsurfaceProfile_Internal() will return either the overridden profile or one from a parent
+	SubsurfaceProfile = SourceInstance->GetSubsurfaceProfile_Internal();
 }
 
 #if WITH_EDITOR
