Optimization to asset gathering by changing ResolveChunkDependencyGraph to use a set of FNames.

FAssetRegistryGenerator::ResolveChunkDependencyGraph was using a Map of FNames and asset strings for tracking which base assets each graph nodes had. The algorithm didn't actually use the asset name string. Now uses a FNames set instead. On a full cook of a large dataset walltime went from 39.5s to 1.23s.

diff --git a/Engine/Source/Editor/UnrealEd/Private/Commandlets/AssetRegistryGenerator.cpp b/Engine/Source/Editor/UnrealEd/Private/Commandlets/AssetRegistryGenerator.cpp
--- a/Engine/Source/Editor/UnrealEd/Private/Commandlets/AssetRegistryGenerator.cpp
+++ b/Engine/Source/Editor/UnrealEd/Private/Commandlets/AssetRegistryGenerator.cpp
@@ -1278,23 +1278,23 @@ void FAssetRegistryGenerator::RemovePackageFromManifest(FName PackageName, int32
 	}
 }
 
-void FAssetRegistryGenerator::ResolveChunkDependencyGraph(const FChunkDependencyTreeNode& Node, FChunkPackageSet BaseAssetSet, TArray<TArray<FName>>& OutPackagesMovedBetweenChunks)
+void FAssetRegistryGenerator::ResolveChunkDependencyGraph(const FChunkDependencyTreeNode& Node, TSet<FName> BaseAssetSet, TArray<TArray<FName>>& OutPackagesMovedBetweenChunks)
 {
 	if (FinalChunkManifests.Num() > Node.ChunkID && FinalChunkManifests[Node.ChunkID])
 	{
 		for (auto It = BaseAssetSet.CreateConstIterator(); It; ++It)
 		{
 			// Remove any assets belonging to our parents.			
-			if (FinalChunkManifests[Node.ChunkID]->Remove(It.Key()) > 0)
+			if (FinalChunkManifests[Node.ChunkID]->Remove(*It) > 0)
 			{
-				OutPackagesMovedBetweenChunks[Node.ChunkID].Add(It.Key());
-				UE_LOG(LogAssetRegistryGenerator, Verbose, TEXT("Removed %s from chunk %i because it is duplicated in another chunk."), *It.Key().ToString(), Node.ChunkID);
+				OutPackagesMovedBetweenChunks[Node.ChunkID].Add(*It);
+				UE_LOG(LogAssetRegistryGenerator, Verbose, TEXT("Removed %s from chunk %i because it is duplicated in another chunk."), *It->ToString(), Node.ChunkID);
 			}
 		}
 		// Add the current Chunk's assets
 		for (auto It = FinalChunkManifests[Node.ChunkID]->CreateConstIterator(); It; ++It)//for (const auto It : *(FinalChunkManifests[Node.ChunkID]))
 		{
-			BaseAssetSet.Add(It.Key(), It.Value());
+			BaseAssetSet.Add(It.Key());
 		}
 		for (const auto It : Node.ChildNodes)
 		{
@@ -1450,7 +1450,7 @@ void FAssetRegistryGenerator::FixupPackageDependenciesForChunks(FSandboxPlatform
 	PackagesRemovedFromChunks.AddDefaulted(ChunkManifests.Num());
 
 	//Finally, if the previous step may added any extra packages to the 0 chunk. Pull them out of other chunks and save space
-	ResolveChunkDependencyGraph(*ChunkDepGraph, FChunkPackageSet(), PackagesRemovedFromChunks);
+	ResolveChunkDependencyGraph(*ChunkDepGraph, TSet<FName>(), PackagesRemovedFromChunks);
 
 	for (int32 i = 0; i < ChunkManifests.Num(); ++i)
 	{
diff --git a/Engine/Source/Editor/UnrealEd/Public/Commandlets/AssetRegistryGenerator.h b/Engine/Source/Editor/UnrealEd/Public/Commandlets/AssetRegistryGenerator.h
--- a/Engine/Source/Editor/UnrealEd/Public/Commandlets/AssetRegistryGenerator.h
+++ b/Engine/Source/Editor/UnrealEd/Public/Commandlets/AssetRegistryGenerator.h
@@ -322,9 +322,9 @@ private:
 	FString	GetShortestReferenceChain(FName PackageName, int32 ChunkID);
 
 	/** Deprecated method to remove redundant chunks */
-	void ResolveChunkDependencyGraph(const FChunkDependencyTreeNode& Node, FChunkPackageSet BaseAssetSet, TArray<TArray<FName>>& OutPackagesMovedBetweenChunks);
+	void ResolveChunkDependencyGraph(const FChunkDependencyTreeNode& Node, TSet<FName> BaseAssetSet, TArray<TArray<FName>>& OutPackagesMovedBetweenChunks);
 
-	/** Helper function to verify Chunk asset assigment is valid */
+	/** Helper function to verify Chunk asset assignment is valid */
 	bool CheckChunkAssetsAreNotInChild(const FChunkDependencyTreeNode& Node);
 
 	/** Helper function to create a given collection. */
