The cooker will no longer perform EDL checks when doing iterative cooks

diff --git a/Engine/Source/Editor/UnrealEd/Private/Commandlets/CookCommandlet.cpp b/Engine/Source/Editor/UnrealEd/Private/Commandlets/CookCommandlet.cpp
--- a/Engine/Source/Editor/UnrealEd/Private/Commandlets/CookCommandlet.cpp
+++ b/Engine/Source/Editor/UnrealEd/Private/Commandlets/CookCommandlet.cpp
@@ -1034,7 +1034,10 @@ bool UCookCommandlet::CookByTheBook( const TArray<ITargetPlatform*>& Platforms,
 		}
 	} while (bTestCook);
 
-	VerifyEDLCookInfo();
+	if (!bIterativeCooking)
+	{
+		VerifyEDLCookInfo();
+	}
 
 	return true;
 }
