Optimize cook -iterate by reducing OS file system calls in FAssetRegistryGenerator::GenerateStreamingInstallManifest.
Replace individual package file lookups with one full directory scan of the cooked directory.

Test Case:
A data set of ~175 000 cooked packages in ~15 000 cooked sub directories containing ~370 000 files.
A fast SSD on Windows with antivirus installed.

Results:
Overall SaveManifests time reduced from varying numbers of 16-26 seconds to stable times around 4 seconds.
~1 million Win32 GetFileAttribute calls replaced with ~15 000 Win32 FindFirstFile calls and ~415 000 Win32 FindNextFile calls.

diff --git a/Engine/Source/Editor/UnrealEd/Private/Commandlets/AssetRegistryGenerator.cpp b/Engine/Source/Editor/UnrealEd/Private/Commandlets/AssetRegistryGenerator.cpp
--- a/Engine/Source/Editor/UnrealEd/Private/Commandlets/AssetRegistryGenerator.cpp
+++ b/Engine/Source/Editor/UnrealEd/Private/Commandlets/AssetRegistryGenerator.cpp
@@ -174,6 +174,51 @@ int64 FAssetRegistryGenerator::GetMaxChunkSizePerPlatform(const ITargetPlatform*
 	return -1;
 }
 
+class FPackageFileSizeVisitor : public IPlatformFile::FDirectoryStatVisitor
+{
+	TMap<FString, int64>& PackageFileSizes;
+public:
+	FPackageFileSizeVisitor(TMap<FString, int64>& InFileSizes)
+		: PackageFileSizes(InFileSizes)
+	{}
+	virtual bool Visit(const TCHAR* FilenameOrDirectory, const FFileStatData& StatData)
+	{
+		const TCHAR* Extensions[] = { TEXT(".uexp"), TEXT(".uasset"), TEXT(".ubulk"), TEXT(".ufont"), TEXT(".umap"), TEXT(".uptnl") };
+
+		if (StatData.bIsDirectory)
+			return true;
+
+		const TCHAR* Extension = FCString::Strrchr(FilenameOrDirectory, '.');
+		if (!Extension)
+			return true;
+
+		int32 ExtIndex = 0;
+		for (; ExtIndex < ARRAY_COUNT(Extensions); ++ExtIndex)
+		{
+			if (0 == FCString::Stricmp(Extension, Extensions[ExtIndex]))
+				break;
+		}
+
+		if (ExtIndex >= ARRAY_COUNT(Extensions))
+			return true;
+
+		int32 LengthWithoutExtension = Extension - FilenameOrDirectory;
+		FString FilenameWithoutExtension(LengthWithoutExtension, FilenameOrDirectory);
+
+		if (int64* CurrentPackageSize = PackageFileSizes.Find(FilenameWithoutExtension))
+		{
+			int64& TotalPackageSize = *CurrentPackageSize;
+			TotalPackageSize += StatData.FileSize;
+		}
+		else
+		{
+			PackageFileSizes.Add(FilenameWithoutExtension, StatData.FileSize);
+		}
+
+		return true;
+	}
+};
+
 bool FAssetRegistryGenerator::GenerateStreamingInstallManifest(int64 InExtraFlavorChunkSize, FSandboxPlatformFile* InSandboxFile)
 {
 	const FString Platform = TargetPlatform->PlatformName();
@@ -255,6 +300,15 @@ bool FAssetRegistryGenerator::GenerateStreamingInstallManifest(int64 InExtraFlav
 		}
 	}
 
+	TMap<FString, int64> PackageFileSizes;
+	if (MaxChunkSize > 0)
+	{
+		FString SandboxPath = InSandboxFile->GetSandboxDirectory();
+		SandboxPath.ReplaceInline(TEXT("[Platform]"), *Platform);
+		FPackageFileSizeVisitor PackageSearch(PackageFileSizes);
+		IFileManager::Get().IterateDirectoryStatRecursively(*SandboxPath, PackageSearch);
+	}
+
 	// generate per-chunk pak list files
 	for (int32 Index = 0; Index < FinalChunkManifests.Num(); ++Index)
 	{
@@ -267,19 +321,26 @@ bool FAssetRegistryGenerator::GenerateStreamingInstallManifest(int64 InExtraFlav
 		int32 FilenameIndex = 0;
 		TArray<FString> ChunkFilenames;
 		FinalChunkManifests[Index]->GenerateValueArray(ChunkFilenames);
-		int32 SubChunkIndex = 0;
-		while ( true )
+		bool bFinishedAllFiles = false;
+		for (int32 SubChunkIndex = 0; !bFinishedAllFiles; ++SubChunkIndex)
 		{
-			FString PakChunkFilename = FString::Printf(TEXT("pakchunk%d.txt"), Index);
-			if ( SubChunkIndex > 0 )
+			const FString PakChunkFilename = (SubChunkIndex > 0)
+				? FString::Printf(TEXT("pakchunk%d_s%d.txt"), Index, SubChunkIndex)
+				: FString::Printf(TEXT("pakchunk%d.txt"), Index);
+
+			const FString PakListFilename = FString::Printf(TEXT("%s/%s"), *TmpPackagingDir, *PakChunkFilename, Index);
+			TUniquePtr<FArchive> PakListFile(IFileManager::Get().CreateFileWriter(*PakListFilename));
+
+			if (!PakListFile)
 			{
-				PakChunkFilename = FString::Printf(TEXT("pakchunk%d_s%d.txt"), Index, SubChunkIndex);
+				UE_LOG(LogAssetRegistryGenerator, Error, TEXT("Failed to open output paklist file %s"), *PakListFilename);
+				return false;
 			}
 
 			FString PakChunkOptions;
-			for ( FString CompressedChunkWildcard : CompressedChunkWildcards)
+			for (const FString& CompressedChunkWildcard : CompressedChunkWildcards)
 			{
-				if ( PakChunkFilename.MatchesWildcard(CompressedChunkWildcard) )
+				if (PakChunkFilename.MatchesWildcard(CompressedChunkWildcard))
 				{
 					PakChunkOptions += " compressed";
 					break;
@@ -306,37 +367,22 @@ bool FAssetRegistryGenerator::GenerateStreamingInstallManifest(int64 InExtraFlav
 				}
 			}
 
-			++SubChunkIndex;
-			FString PakListFilename = FString::Printf(TEXT("%s/%s"), *TmpPackagingDir, *PakChunkFilename, Index);
-			TUniquePtr<FArchive> PakListFile(IFileManager::Get().CreateFileWriter(*PakListFilename));
-
-			if (!PakListFile)
-			{
-				UE_LOG(LogAssetRegistryGenerator, Error, TEXT("Failed to open output paklist file %s"), *PakListFilename);
-				return false;
-			}
-
-			int64 CurrentPakSize = 0;
-			bool bFinishedAllFiles = true;
-
-			if (bUseAssetManager)
+			if (bUseAssetManager && SubChunkIndex == 0)
 			{
 				// Sort so the order is consistent. If load order is important then it should be specified as a load order file to UnrealPak
 				ChunkFilenames.Sort();
 			}
 
+			int64 CurrentPakSize = 0;
+			bFinishedAllFiles = true;
 			for (; FilenameIndex < ChunkFilenames.Num(); ++FilenameIndex)
 			{
 				FString Filename = ChunkFilenames[FilenameIndex];
 				FString PakListLine = FPaths::ConvertRelativePathToFull(Filename.Replace(TEXT("[Platform]"), *Platform));
 				if (MaxChunkSize > 0)
 				{
-					const TCHAR* Extensions[] = { TEXT(".uexp"), TEXT(".uasset"), TEXT(".ubulk"), TEXT(".ufont"), TEXT(".umap"), TEXT(".uptnl") };
-					for (int32 ExtIndex = 0; ExtIndex < ARRAY_COUNT(Extensions); ++ExtIndex)
-					{
-						int64 FileSize = IFileManager::Get().FileSize(*(PakListLine + Extensions[ExtIndex]));
-						CurrentPakSize += FileSize > 0 ? FileSize : 0;
-					}
+					const int64* PackageFileSize = PackageFileSizes.Find(PakListLine);
+					CurrentPakSize += PackageFileSize ? *PackageFileSize : 0;
 					if (MaxChunkSize < CurrentPakSize)
 					{
 						// early out if we are over memory limit
@@ -362,11 +408,6 @@ bool FAssetRegistryGenerator::GenerateStreamingInstallManifest(int64 InExtraFlav
 			FString LayerString = FString::Printf(TEXT("%d\r\n"), TargetLayer);
 
 			ChunkLayerFile->Serialize(TCHAR_TO_ANSI(*LayerString), LayerString.Len());
-			
-			if (bFinishedAllFiles)
-			{
-				break;
-			}
 		}
 
 	}
