Optimize cook -iterate by fast forwarding through already cooked packages instead of ticking them.
Reduce total time in tick with ~10 seconds (measured with -stdout and default log settings) for ~23000 already cooked packages

New output example:
[[REDACTED_PHONE]:858][  0]LogCook: Display: Cooked packages 175399 Packages Remain 23552 Total 198951
[[REDACTED_PHONE]:912][  0]LogCook: Display: Cooked packages 175399 Packages Remain 0 Total 175399
[[REDACTED_PHONE]:916][  0]LogCook: Display: Finishing up...

Old output example:
[[REDACTED_PHONE]:470][  0]LogCook: Display: Cooked packages 175399 Packages Remain 23552 Total 198951
[[REDACTED_PHONE]:472][  0]LogCook: Display: Cooked packages 175399 Packages Remain 23551 Total 198950
...
[[REDACTED_PHONE]:962][  0]LogCook: Display: Cooked packages 175399 Packages Remain 1 Total 175400

diff --git a/Engine/Source/Editor/UnrealEd/Private/CookOnTheFlyServer.cpp b/Engine/Source/Editor/UnrealEd/Private/CookOnTheFlyServer.cpp
--- a/Engine/Source/Editor/UnrealEd/Private/CookOnTheFlyServer.cpp
+++ b/Engine/Source/Editor/UnrealEd/Private/CookOnTheFlyServer.cpp
@@ -2070,12 +2070,8 @@ uint32 UCookOnTheFlyServer::TickCookOnTheSide( const float TimeSlice, uint32 &Co
 	{
 		uint32 Result = FullLoadAndSave(CookedPackageCount);
 
-		UE_LOG(LogCook, Display, TEXT("Finishing up..."));
-
 		CookByTheBookFinished();
 
-		UE_LOG(LogCook, Display, TEXT("Done!"));
-
 		return Result;
 	}
 
@@ -2135,15 +2131,26 @@ uint32 UCookOnTheFlyServer::TickCookOnTheSide( const float TimeSlice, uint32 &Co
 			}
 		}
 
-		FFilePlatformRequest ToBuild;
-		
-		if (HasCookRequests())
+		auto GetNextUncookedRequestLambda = [this](FFilePlatformRequest& OutToBuild) -> bool
 		{
-			PackageTracker->CookRequests.Dequeue(/* out */ ToBuild);
-		}
-		else
+			while (this->PackageTracker->CookRequests.Dequeue(OutToBuild))
+			{
+				if (this->PackageTracker->CookedPackages.Exists(OutToBuild))
+				{
+#if DEBUG_COOKONTHEFLY
+					UE_LOG(LogCook, Display, TEXT("Package for platform already cooked %s, discarding request"), *ToBuild.GetFilename().ToString());
+#endif
+					continue;
+				}
+				return true;
+			}
+			return false;
+		};
+
+		FFilePlatformRequest ToBuild;
+		if (!GetNextUncookedRequestLambda(ToBuild))
 		{
-			// no more to do this tick break out and do some other stuff
+			// no more real work to do this tick, break out and do some other stuff
 			break;
 		}
 
@@ -2399,6 +2406,9 @@ uint32 UCookOnTheFlyServer::TickCookOnTheSide( const float TimeSlice, uint32 &Co
 		check(IsCookByTheBookMode());
 
 		// if we are out of stuff and we are in cook by the book from the editor mode then we finish up
+		UE_CLOG(!(TickFlags & ECookTickFlags::HideProgressDisplay) && (GCookProgressDisplay & (int32)ECookProgressDisplayMode::RemainingPackages),
+			LogCook, Display, TEXT("Cooked packages %d Packages Remain %d Total %d"),
+			PackageTracker->CookedPackages.Num(), 0, PackageTracker->CookedPackages.Num());
 		CookByTheBookFinished();
 	}
 
@@ -6231,6 +6241,7 @@ void UCookOnTheFlyServer::CookByTheBookFinished()
 	check( IsCookByTheBookMode() );
 	check( CookByTheBookOptions->bRunning == true );
 
+	UE_LOG(LogCook, Display, TEXT("Finishing up..."));
 
 	UPackage::WaitForAsyncFileWrites();
 
@@ -6420,6 +6431,8 @@ void UCookOnTheFlyServer::CookByTheBookFinished()
 
 	OutputHierarchyTimers();
 	ClearHierarchyTimers();
+
+	UE_LOG(LogCook, Display, TEXT("Done!"));
 }
 
 void UCookOnTheFlyServer::BuildMapDependencyGraph(const FName& PlatformName)
