Optimize cook -iterate by reducing OS file system calls during engine and cook startup.
Remove PackageNameCache and PackageTracker from the UCookOnTheFlyServer default object, only instantiate them for the actual cooker object.
Remove FileExists checks for files just returned by the OS directory scan.
Use AssetRegistry instead of FileSystem for package existence lookup.

Test Case:
A data set of ~175 000 cooked packages in ~15 000 cooked sub directories containing ~370 000 files.
A fast SSD on Windows with antivirus installed.

Results:
Removes ~360 000 out of ~450 000 Win32 GetFileAttribute calls (~80% less)
PopulateCookedPackagesFromDisk reduced from ~11 seconds to ~4 seconds

diff --git a/Engine/Source/Editor/UnrealEd/Private/CookOnTheFlyServer.cpp b/Engine/Source/Editor/UnrealEd/Private/CookOnTheFlyServer.cpp
--- a/Engine/Source/Editor/UnrealEd/Private/CookOnTheFlyServer.cpp
+++ b/Engine/Source/Editor/UnrealEd/Private/CookOnTheFlyServer.cpp
@@ -1427,8 +1427,6 @@ UCookOnTheFlyServer::UCookOnTheFlyServer(const FObjectInitializer& ObjectInitial
 	bIsSavingPackage(false),
 	AssetRegistry(nullptr)
 {
-	PackageNameCache = new FPackageNameCache;
-	PackageTracker	 = new FPackageTracker(PackageNameCache);
 }
 
 UCookOnTheFlyServer::~UCookOnTheFlyServer()
@@ -1549,6 +1547,8 @@ bool UCookOnTheFlyServer::StartNetworkFileServer( const bool BindAnyPort )
 #endif
 
 	GenerateAssetRegistry();
+	PackageNameCache = new FPackageNameCache(AssetRegistry);
+	PackageTracker = new FPackageTracker(PackageNameCache);
 
 	InitializeSandbox();
 	InitializeTargetPlatforms();
@@ -4055,11 +4055,6 @@ void UCookOnTheFlyServer::Initialize( ECookMode::Type DesiredCookMode, ECookInit
 
 	FCoreUObjectDelegates::GetPreGarbageCollectDelegate().AddUObject(this, &UCookOnTheFlyServer::PreGarbageCollect);
 
-	if (IsCookByTheBookMode() && !IsCookingInEditor())
-	{
-		FCoreUObjectDelegates::PackageCreatedForLoad.AddUObject(this, &UCookOnTheFlyServer::MaybeMarkPackageAsAlreadyLoaded);
-	}
-
 	if (IsCookingInEditor())
 	{
 		FCoreUObjectDelegates::OnObjectPropertyChanged.AddUObject(this, &UCookOnTheFlyServer::OnObjectPropertyChanged);
@@ -5151,6 +5146,8 @@ void UCookOnTheFlyServer::PopulateCookedPackagesFromDisk(const TArray<ITargetPla
 		const static FName NAME_DummyCookedFilename(TEXT("DummyCookedFilename")); // pls never name a package dummycookedfilename otherwise shit might go wonky
 		if (bIsIterateSharedBuild)
 		{
+			check(IFileManager::Get().FileExists(*NAME_DummyCookedFilename.ToString()) == false);
+
 			TSet<FName> ExistingPackages = ModifiedPackages;
 			ExistingPackages.Append(RemovedPackages);
 			ExistingPackages.Append(IdenticalCookedPackages);
@@ -5206,11 +5203,6 @@ void UCookOnTheFlyServer::PopulateCookedPackagesFromDisk(const TArray<ITargetPla
 					bShouldKeep = false;
 				}
 			}
-				
-			if (CookedFile == NAME_DummyCookedFilename)
-			{
-				check(IFileManager::Get().FileExists(*CookedFile.ToString()) == false);
-			}
 
 			TArray<FName> PlatformNames;
 			PlatformNames.Add(PlatformFName);
@@ -5218,10 +5210,6 @@ void UCookOnTheFlyServer::PopulateCookedPackagesFromDisk(const TArray<ITargetPla
 			if (bShouldKeep)
 			{
 				// Add this package to the CookedPackages list so that we don't try cook it again
-				if (CookedFile != NAME_DummyCookedFilename)
-				{
-					check(IFileManager::Get().FileExists(*CookedFile.ToString()));
-				}
 				TArray<bool> Succeeded;
 				Succeeded.Add(true);
 
@@ -6672,6 +6660,12 @@ void UCookOnTheFlyServer::StartCookByTheBook( const FCookByTheBookStartupOptions
 	CookByTheBookOptions->bErrorOnEngineContentUse = CookByTheBookStartupOptions.bErrorOnEngineContentUse;
 
 	GenerateAssetRegistry();
+	PackageNameCache = new FPackageNameCache(AssetRegistry);
+	PackageTracker = new FPackageTracker(PackageNameCache);
+	if (!IsCookingInEditor())
+	{
+		FCoreUObjectDelegates::PackageCreatedForLoad.AddUObject(this, &UCookOnTheFlyServer::MaybeMarkPackageAsAlreadyLoaded);
+	}
 
 	// Find all the localized packages and map them back to their source package
 	{
diff --git a/Engine/Source/Editor/UnrealEd/Private/Cooker/PackageNameCache.h b/Engine/Source/Editor/UnrealEd/Private/Cooker/PackageNameCache.h
--- a/Engine/Source/Editor/UnrealEd/Private/Cooker/PackageNameCache.h
+++ b/Engine/Source/Editor/UnrealEd/Private/Cooker/PackageNameCache.h
@@ -3,6 +3,8 @@
 #pragma once
 
 #include "CoreMinimal.h"
+#include "IAssetRegistry.h"
+#include "Algo/Find.h"
 #include "Misc/PackageName.h"
 
 struct FCachedPackageFilename
@@ -31,6 +33,7 @@ struct FCachedPackageFilename
 struct FPackageNameCache
 {
 	// Package name cache
+	FPackageNameCache(IAssetRegistry* InAssetRegistry) : AssetRegistry(InAssetRegistry) {}
 
 	FString			GetCachedPackageFilename(const UPackage* Package) const;
 
@@ -45,8 +48,11 @@ struct FPackageNameCache
 	bool			ClearPackageFilenameCacheForPackage(const UPackage* Package) const;
 
 private:
+	bool DoesPackageExist(const FName& PackageName, FString* OutFilename) const;
 	const FCachedPackageFilename& Cache(const FName& PackageName) const;
 
+	IAssetRegistry* AssetRegistry;
+
 	mutable TMap<FName, FCachedPackageFilename> PackageFilenameCache; // filename cache (only process the string operations once)
 	mutable TMap<FName, FName>					PackageFilenameToPackageFNameCache;
 
@@ -84,6 +90,25 @@ bool FPackageNameCache::ClearPackageFilenameCacheForPackage(const UPackage* Pack
 	return PackageFilenameCache.Remove(Package->GetFName()) >= 1;
 }
 
+bool FPackageNameCache::DoesPackageExist(const FName& PackageName, FString* OutFilename) const
+{
+	TArray<FAssetData> Assets;
+	AssetRegistry->GetAssetsByPackageName(PackageName, Assets, /*bIncludeOnlyDiskAssets*/ true);
+	if (Assets.Num() <= 0)
+	{
+		return false;
+	}
+
+	if (OutFilename)
+	{
+		const bool ContainsMap = Algo::FindByPredicate(Assets, [](const FAssetData& Asset) { return Asset.PackageFlags & PKG_ContainsMap; }) != nullptr;
+		const FString& PackageExtension = ContainsMap ? FPackageName::GetMapPackageExtension() : FPackageName::GetAssetPackageExtension();
+		*OutFilename = FPackageName::LongPackageNameToFilename(PackageName.ToString(), PackageExtension);
+	}
+
+	return true;
+}
+
 const FCachedPackageFilename& FPackageNameCache::Cache(const FName& PackageName) const
 {
 	check(IsInGameThread());
@@ -102,7 +127,7 @@ const FCachedPackageFilename& FPackageNameCache::Cache(const FName& PackageName)
 	FString StandardFilename;
 	FName StandardFileFName = NAME_None;
 
-	if (FPackageName::DoesPackageExist(PackageName.ToString(), NULL, &Filename, false))
+	if (DoesPackageExist(PackageName, &Filename))
 	{
 		StandardFilename = PackageFilename = FPaths::ConvertRelativePathToFull(Filename);
 
