AssetRegistryGenerator compile fixes.
Fix bad merge in 5282591.
Bring back optimizations from 4822103.

diff --git a/Engine/Source/Editor/UnrealEd/Private/Commandlets/AssetRegistryGenerator.cpp b/Engine/Source/Editor/UnrealEd/Private/Commandlets/AssetRegistryGenerator.cpp
--- a/Engine/Source/Editor/UnrealEd/Private/Commandlets/AssetRegistryGenerator.cpp
+++ b/Engine/Source/Editor/UnrealEd/Private/Commandlets/AssetRegistryGenerator.cpp
@@ -1344,7 +1344,7 @@ void FAssetRegistryGenerator::RemovePackageFromManifest(FName PackageName, int32
 	}
 }
 
-void FAssetRegistryGenerator::ResolveChunkDependencyGraph(const FChunkDependencyTreeNode& Node, const FChunkPackageSet& BaseAssetSet, TArray<TArray<FName>>& OutPackagesMovedBetweenChunks)
+void FAssetRegistryGenerator::ResolveChunkDependencyGraph(const FChunkDependencyTreeNode& Node, const TSet<FName>& BaseAssetSet, TArray<TArray<FName>>& OutPackagesMovedBetweenChunks)
 {
 	if (FinalChunkManifests.Num() > Node.ChunkID && FinalChunkManifests[Node.ChunkID])
 	{
@@ -1358,7 +1358,7 @@ void FAssetRegistryGenerator::ResolveChunkDependencyGraph(const FChunkDependency
 			}
 		}
 		
-		FChunkPackageSet ModifiedAssetSet;
+		TSet<FName> ModifiedAssetSet;
 		
 		// Add the current Chunk's assets
 		for (auto It = FinalChunkManifests[Node.ChunkID]->CreateConstIterator(); It; ++It)//for (const auto It : *(FinalChunkManifests[Node.ChunkID]))
@@ -1368,7 +1368,7 @@ void FAssetRegistryGenerator::ResolveChunkDependencyGraph(const FChunkDependency
 				ModifiedAssetSet = BaseAssetSet;
 			}
 			
-			ModifiedAssetSet.Add(It.Key(), It.Value());
+			ModifiedAssetSet.Add(It.Key());
 		}
 		
 		const auto& AssetSet = ModifiedAssetSet.Num() ? ModifiedAssetSet : BaseAssetSet;
diff --git a/Engine/Source/Editor/UnrealEd/Public/Commandlets/AssetRegistryGenerator.h b/Engine/Source/Editor/UnrealEd/Public/Commandlets/AssetRegistryGenerator.h
--- a/Engine/Source/Editor/UnrealEd/Public/Commandlets/AssetRegistryGenerator.h
+++ b/Engine/Source/Editor/UnrealEd/Public/Commandlets/AssetRegistryGenerator.h
@@ -322,7 +322,7 @@ private:
 	FString	GetShortestReferenceChain(FName PackageName, int32 ChunkID);
 
 	/** Deprecated method to remove redundant chunks */
-	void ResolveChunkDependencyGraph(const FChunkDependencyTreeNode& Node, const FChunkPackageSet& BaseAssetSet, TArray<TArray<FName>>& OutPackagesMovedBetweenChunks);
+	void ResolveChunkDependencyGraph(const FChunkDependencyTreeNode& Node, const TSet<FName>& BaseAssetSet, TArray<TArray<FName>>& OutPackagesMovedBetweenChunks);
 
 	/** Helper function to verify Chunk asset assignment is valid */
 	bool CheckChunkAssetsAreNotInChild(const FChunkDependencyTreeNode& Node);
