Optimize cook -iterate by replacing individual localization AssetRegistry queries with one big query.

Test Case:
~ 15 cultures and ~20 root paths, i.e. ~300 queries
~ 375000 packages

Results:
DiscoverLocalizedAssets reduced from ~16 seconds to ~1 second

diff --git a/Engine/Source/Editor/UnrealEd/Private/CookOnTheFlyServer.cpp b/Engine/Source/Editor/UnrealEd/Private/CookOnTheFlyServer.cpp
--- a/Engine/Source/Editor/UnrealEd/Private/CookOnTheFlyServer.cpp
+++ b/Engine/Source/Editor/UnrealEd/Private/CookOnTheFlyServer.cpp
@@ -6693,22 +6693,30 @@ void UCookOnTheFlyServer::StartCookByTheBook( const FCookByTheBookStartupOptions
 
 		TArray<FString> RootPaths;
 		FPackageName::QueryRootContentPaths(RootPaths);
+
+		FARFilter Filter;
+		Filter.bRecursivePaths = true;
+		Filter.bIncludeOnlyOnDiskAssets = false;
+		Filter.PackagePaths.Reserve(AllCulturesToCook.Num() * RootPaths.Num());
 		for (const FString& RootPath : RootPaths)
 		{
 			for (const FString& CultureName : AllCulturesToCook)
 			{
-				TArray<FAssetData> AssetDataForCulture;
-				AssetRegistry->GetAssetsByPath(*(RootPath / TEXT("L10N") / CultureName), AssetDataForCulture, true);
+				FString LocalizedPackagePath = RootPath / TEXT("L10N") / CultureName;
+				Filter.PackagePaths.Add(*LocalizedPackagePath);
+			}
+		}
 
-				for (const FAssetData& AssetData : AssetDataForCulture)
-				{
-					const FName LocalizedPackageName = AssetData.PackageName;
-					const FName SourcePackageName = *FPackageName::GetSourcePackagePath(LocalizedPackageName.ToString());
+		TArray<FAssetData> AssetDataForCultures;
+		AssetRegistry->GetAssets(Filter, AssetDataForCultures);
 
-					TArray<FName>& LocalizedPackageNames = CookByTheBookOptions->SourceToLocalizedPackageVariants.FindOrAdd(SourcePackageName);
-					LocalizedPackageNames.AddUnique(LocalizedPackageName);
-				}
-			}
+		for (const FAssetData& AssetData : AssetDataForCultures)
+		{
+			const FName LocalizedPackageName = AssetData.PackageName;
+			const FName SourcePackageName = *FPackageName::GetSourcePackagePath(LocalizedPackageName.ToString());
+
+			TArray<FName>& LocalizedPackageNames = CookByTheBookOptions->SourceToLocalizedPackageVariants.FindOrAdd(SourcePackageName);
+			LocalizedPackageNames.AddUnique(LocalizedPackageName);
 		}
 	}
 
