Optimize cook -iterate overhead by reducing allocations for ConvertCookedPathToUncookedPath used in GetAllCookedFiles

Test Case:
~175 000 cooked packages

Results:
ConvertCookedPathToUncookedPath ~2.8 seconds to ~1.8 seconds
Remaining time is ~90% FName construction

diff --git a/Engine/Source/Editor/UnrealEd/Classes/CookOnTheSide/CookOnTheFlyServer.h b/Engine/Source/Editor/UnrealEd/Classes/CookOnTheSide/CookOnTheFlyServer.h
--- a/Engine/Source/Editor/UnrealEd/Classes/CookOnTheSide/CookOnTheFlyServer.h
+++ b/Engine/Source/Editor/UnrealEd/Classes/CookOnTheSide/CookOnTheFlyServer.h
@@ -197,7 +197,7 @@ private:
 	int32 LastCookedPackagesCount = 0;
 	double LastProgressDisplayTime = 0;
 
-	FString ConvertCookedPathToUncookedPath(const FString& CookedPackageName) const;
+	FName ConvertCookedPathToUncookedPath(const FString& SandboxPath, const FString& CookedPath) const;
 
 	/** Get dependencies for package */
 	const TArray<FName>& GetFullPackageDependencies(const FName& PackageName) const;
diff --git a/Engine/Source/Editor/UnrealEd/Private/CookOnTheFlyServer.cpp b/Engine/Source/Editor/UnrealEd/Private/CookOnTheFlyServer.cpp
--- a/Engine/Source/Editor/UnrealEd/Private/CookOnTheFlyServer.cpp
+++ b/Engine/Source/Editor/UnrealEd/Private/CookOnTheFlyServer.cpp
@@ -4968,16 +4968,16 @@ bool UCookOnTheFlyServer::SaveCurrentIniSettings(const ITargetPlatform* TargetPl
 
 }
 
-FString UCookOnTheFlyServer::ConvertCookedPathToUncookedPath(const FString& CookedRelativeFilename) const 
+FName UCookOnTheFlyServer::ConvertCookedPathToUncookedPath(const FString& SandboxPath, const FString& CookedPath) const
 {
 	// Check for remapped plugins' cooked content
-	if (PluginsToRemap.Num() > 0 && CookedRelativeFilename.Contains(REMAPPED_PLUGGINS))
+	if (PluginsToRemap.Num() > 0 && CookedPath.Contains(REMAPPED_PLUGGINS))
 	{
-		int32 RemappedIndex = CookedRelativeFilename.Find(REMAPPED_PLUGGINS);
+		int32 RemappedIndex = CookedPath.Find(REMAPPED_PLUGGINS);
 		check(RemappedIndex >= 0);
 		static uint32 RemappedPluginStrLen = FCString::Strlen(REMAPPED_PLUGGINS);
 		// Snip everything up through the RemappedPlugins/ off so we can find the plugin it corresponds to
-		FString PluginPath = CookedRelativeFilename.RightChop(RemappedIndex + RemappedPluginStrLen + 1);
+		FString PluginPath = CookedPath.RightChop(RemappedIndex + RemappedPluginStrLen + 1);
 		FString FullUncookedPath;
 		// Find the plugin that owns this content
 		for (TSharedRef<IPlugin> Plugin : PluginsToRemap)
@@ -4994,40 +4994,24 @@ FString UCookOnTheFlyServer::ConvertCookedPathToUncookedPath(const FString& Cook
 
 		if (FullUncookedPath.Len() > 0)
 		{
-			return FullUncookedPath;
+			return FName(*FullUncookedPath);
 		}
 		// Otherwise fall through to sandbox handling
 	}
 
-	const FString CookedFilename = FPaths::ConvertRelativePathToFull(CookedRelativeFilename);
+	FString UncookedPath;
 
-
-	FString SandboxDirectory = SandboxFile->GetSandboxDirectory();
-
-	SandboxDirectory.ReplaceInline(TEXT("[PLATFORM]"), TEXT(""));
-	SandboxDirectory.ReplaceInline(TEXT("//"), TEXT("/"));
-
-	FString CookedFilenameNoSandbox = CookedFilename;
-	CookedFilenameNoSandbox.RemoveFromStart(SandboxDirectory);
-	int32 EndOfPlatformIndex = 0;
-
-	// assume at this point the cook platform is the next thing on the path
-	FString CookedFilenameNoPlatform = CookedFilename;
-
-	if (CookedFilenameNoSandbox.FindChar(TEXT('/'), EndOfPlatformIndex))
+	if (CookedPath.StartsWith(SandboxPath))
 	{
-		CookedFilenameNoPlatform = SandboxFile->GetSandboxDirectory() / CookedFilenameNoSandbox.Mid(EndOfPlatformIndex);
-		CookedFilenameNoPlatform.ReplaceInline(TEXT("//"), TEXT("/"));
+		UncookedPath = CookedPath.Replace(*SandboxPath, *FPaths::GetRelativePathToRoot());
+	}
+	else
+	{
+		FString FullCookedFilename = FPaths::ConvertRelativePathToFull(CookedPath);
+		UncookedPath = FullCookedFilename.Replace(*SandboxPath, *FPaths::GetRelativePathToRoot());
 	}
-	
-	// after removing the cooked platform we can use the sandbox file to convert back to a uncooked path
-	FString FullUncookedPath = SandboxFile->ConvertFromSandboxPath(*CookedFilenameNoPlatform);
-
-	// make the result a standard filename (relative)
-	FPaths::MakeStandardFilename(FullUncookedPath);
-	return FullUncookedPath;
-
 
+	return FName(*UncookedPath);
 }
 
 void UCookOnTheFlyServer::GetAllCookedFiles(TMap<FName, FName>& UncookedPathToCookedPath, const FString& SandboxPath)
@@ -5039,11 +5023,7 @@ void UCookOnTheFlyServer::GetAllCookedFiles(TMap<FName, FName>& UncookedPathToCo
 	for (const FString& CookedFile : CookedFiles)
 	{
 		const FName CookedFName(*CookedFile);
-
-		const FString CookedFullPath = FPaths::ConvertRelativePathToFull(CookedFile);
-		const FString UncookedFilename = ConvertCookedPathToUncookedPath(CookedFullPath);
-
-		const FName UncookedFName(*UncookedFilename);
+		const FName UncookedFName = ConvertCookedPathToUncookedPath(SandboxPath, CookedFile);
 		UncookedPathToCookedPath.Add(UncookedFName, CookedFName);
 	}
 }
