Optimize cook -iterate GenerateLongPackageNames by using a temporary Set with Add instead of TArray with AddUnique.

Test Case:
~23 000 files too cook, no duplicates, no remap

Results:
GenerateLongPackageName ~0.26 seconds to ~0.04 seconds

diff --git a/Engine/Source/Editor/UnrealEd/Private/CookOnTheFlyServer.cpp b/Engine/Source/Editor/UnrealEd/Private/CookOnTheFlyServer.cpp
--- a/Engine/Source/Editor/UnrealEd/Private/CookOnTheFlyServer.cpp
+++ b/Engine/Source/Editor/UnrealEd/Private/CookOnTheFlyServer.cpp
@@ -5437,16 +5437,23 @@ void UCookOnTheFlyServer::GenerateAssetRegistry()
 
 void UCookOnTheFlyServer::GenerateLongPackageNames(TArray<FName>& FilesInPath)
 {
+	TSet<FName> FilesInPathSet;
 	TArray<FName> FilesInPathReverse;
+	FilesInPathSet.Reserve(FilesInPath.Num());
 	FilesInPathReverse.Reserve(FilesInPath.Num());
 
-	for( int32 FileIndex = 0; FileIndex < FilesInPath.Num(); FileIndex++ )
+	for (int32 FileIndex = 0; FileIndex < FilesInPath.Num(); FileIndex++)
 	{
-		const FString& FileInPath = FilesInPath[FilesInPath.Num() - FileIndex - 1].ToString();
+		const FName& FileInPathFName = FilesInPath[FilesInPath.Num() - FileIndex - 1];
+		const FString& FileInPath = FileInPathFName.ToString();
 		if (FPackageName::IsValidLongPackageName(FileInPath))
 		{
-			const FName FileInPathFName(*FileInPath);
-			FilesInPathReverse.AddUnique(FileInPathFName);
+			bool bIsAlreadyAdded;
+			FilesInPathSet.Add(FileInPathFName, &bIsAlreadyAdded);
+			if (!bIsAlreadyAdded)
+			{
+				FilesInPathReverse.Add(FileInPathFName);
+			}
 		}
 		else
 		{
@@ -5455,7 +5462,12 @@ void UCookOnTheFlyServer::GenerateLongPackageNames(TArray<FName>& FilesInPath)
 			if (FPackageName::TryConvertFilenameToLongPackageName(FileInPath, LongPackageName, &FailureReason))
 			{
 				const FName LongPackageFName(*LongPackageName);
-				FilesInPathReverse.AddUnique(LongPackageFName);
+				bool bIsAlreadyAdded;
+				FilesInPathSet.Add(LongPackageFName, &bIsAlreadyAdded);
+				if (!bIsAlreadyAdded)
+				{
+					FilesInPathReverse.Add(LongPackageFName);
+				}
 			}
 			else
 			{
@@ -5464,12 +5476,8 @@ void UCookOnTheFlyServer::GenerateLongPackageNames(TArray<FName>& FilesInPath)
 			}
 		}
 	}
-	// Exchange(FilesInPathReverse, FilesInPath);
 	FilesInPath.Empty(FilesInPathReverse.Num());
-	for ( const FName& Files : FilesInPathReverse )
-	{
-		FilesInPath.Add(Files);
-	}
+	FilesInPath.Append(FilesInPathReverse);
 }
 
 void UCookOnTheFlyServer::AddFileToCook( TArray<FName>& InOutFilesToCook, const FString &InFilename ) const
