Fix a small reported memory leak in CookOnTheFlyServer by making sure to always clear hierarchy timers.

Verified test case:
Set bEnableCookOnTheSide=True in section [/Script/UnrealEd.CookerSettings] in BaseEngine.ini.
Start editor.
Close editor.

diff --git a/Engine/Source/Editor/UnrealEd/Private/CookOnTheFlyServer.cpp b/Engine/Source/Editor/UnrealEd/Private/CookOnTheFlyServer.cpp
--- a/Engine/Source/Editor/UnrealEd/Private/CookOnTheFlyServer.cpp
+++ b/Engine/Source/Editor/UnrealEd/Private/CookOnTheFlyServer.cpp
@@ -1434,7 +1434,7 @@ UCookOnTheFlyServer::~UCookOnTheFlyServer()
 	FCoreDelegates::OnFConfigCreated.RemoveAll(this);
 	FCoreDelegates::OnFConfigDeleted.RemoveAll(this);
 
-		delete CookByTheBookOptions;
+	delete CookByTheBookOptions;
 	CookByTheBookOptions = nullptr;
 
 	delete PackageTracker;
@@ -1442,6 +1442,8 @@ UCookOnTheFlyServer::~UCookOnTheFlyServer()
 
 	delete PackageNameCache;
 	PackageNameCache = nullptr;
+
+	ClearHierarchyTimers();
 }
 
 // this tick only happens in the editor cook commandlet directly calls tick on the side
