CookOnTheFlyServer compile fix when DEBUG_COOKONTHEFLY is set
Clean up after change 5088206.

diff --git a/Engine/Source/Editor/UnrealEd/Private/CookOnTheFlyServer.cpp b/Engine/Source/Editor/UnrealEd/Private/CookOnTheFlyServer.cpp
--- a/Engine/Source/Editor/UnrealEd/Private/CookOnTheFlyServer.cpp
+++ b/Engine/Source/Editor/UnrealEd/Private/CookOnTheFlyServer.cpp
@@ -2138,7 +2138,7 @@ uint32 UCookOnTheFlyServer::TickCookOnTheSide( const float TimeSlice, uint32 &Co
 				if (this->PackageTracker->CookedPackages.Exists(OutToBuild))
 				{
 #if DEBUG_COOKONTHEFLY
-					UE_LOG(LogCook, Display, TEXT("Package for platform already cooked %s, discarding request"), *ToBuild.GetFilename().ToString());
+					UE_LOG(LogCook, Display, TEXT("Package for platform already cooked %s, discarding request"), *OutToBuild.GetFilename().ToString());
 #endif
 					continue;
 				}
@@ -2169,14 +2169,6 @@ uint32 UCookOnTheFlyServer::TickCookOnTheSide( const float TimeSlice, uint32 &Co
 			GUnrealEd->GetPackageAutoSaver().ForceMinimumTimeTillAutoSave(SecondsWarningTillAutosave);
 		}
 
-		if (PackageTracker->CookedPackages.Exists(ToBuild))
-		{
-#if DEBUG_COOKONTHEFLY
-			UE_LOG(LogCook, Display, TEXT("Package for platform already cooked %s, discarding request"), *ToBuild.GetFilename().ToString());
-#endif
-			continue;
-		}
-
 #if DEBUG_COOKONTHEFLY
 		UE_LOG(LogCook, Display, TEXT("Processing package %s"), *ToBuild.GetFilename().ToString());
 #endif
