Small fixup after the last merge conflict from Main

diff --git a/Engine/Source/Editor/UnrealEd/Private/Commandlets/AssetRegistryGenerator.cpp b/Engine/Source/Editor/UnrealEd/Private/Commandlets/AssetRegistryGenerator.cpp
--- a/Engine/Source/Editor/UnrealEd/Private/Commandlets/AssetRegistryGenerator.cpp
+++ b/Engine/Source/Editor/UnrealEd/Private/Commandlets/AssetRegistryGenerator.cpp
@@ -332,7 +332,7 @@ bool FAssetRegistryGenerator::GenerateStreamingInstallManifest(int64 InExtraFlav
 				? FString::Printf(TEXT("pakchunk%d_s%d.txt"), PakchunkIndex, SubChunkIndex)
 				: FString::Printf(TEXT("pakchunk%d.txt"), PakchunkIndex);
 
-			const FString PakListFilename = FString::Printf(TEXT("%s/%s"), *TmpPackagingDir, *PakChunkFilename, PakchunkIndex);
+			const FString PakListFilename = FString::Printf(TEXT("%s/%s"), *TmpPackagingDir, *PakChunkFilename);
 			TUniquePtr<FArchive> PakListFile(IFileManager::Get().CreateFileWriter(*PakListFilename));
 
 			if (!PakListFile)
