PkgInfo commandlet will now be automatically able to open cooked packages without the need to modify ini files. Added -AllPackagesIn=Path command line param to recursively find all files under the specified folder (useful for dumping info for all cooked packages)

diff --git a/Engine/Source/Editor/UnrealEd/Private/Commandlets/PackageUtilities.cpp b/Engine/Source/Editor/UnrealEd/Private/Commandlets/PackageUtilities.cpp
--- a/Engine/Source/Editor/UnrealEd/Private/Commandlets/PackageUtilities.cpp
+++ b/Engine/Source/Editor/UnrealEd/Private/Commandlets/PackageUtilities.cpp
@@ -1473,7 +1473,13 @@ int32 UPkgInfoCommandlet::Main( const FString& Params )
 	FPkgInfoReporter* Reporter = new FPkgInfoReporter_Log(InfoFlags, bHideOffsets);
 
 	TArray<FString> FilesInPath;
-	if( Switches.Contains(TEXT("AllPackages")) )
+
+	FString PathWithPackages;
+	if (FParse::Value(*Params, TEXT("AllPackagesIn="), PathWithPackages))
+	{
+		FPackageName::FindPackagesInDirectory(FilesInPath, *PathWithPackages);
+	}
+	else if( Switches.Contains(TEXT("AllPackages")) )
 	{
 		FEditorFileUtils::FindAllPackageFiles(FilesInPath);
 	}
@@ -1553,6 +1559,7 @@ int32 UPkgInfoCommandlet::Main( const FString& Params )
 		if (!bDumpProperties)
 		{
 			TGuardValue<bool> GuardAllowUnversionedContentInEditor(GAllowUnversionedContentInEditor, true);
+			TGuardValue<int32> GuardAllowCookedContentInEditor(GAllowCookedDataInEditorBuilds, 1);
 			TRefCountPtr<FUObjectSerializeContext> LoadContext(FUObjectThreadContext::Get().GetSerializeContext());
 			BeginLoad(LoadContext);
 			Linker = CreateLinkerForFilename(LoadContext, Filename);
