Lots of implicit operator bool fixes, some constness fixes, some other fixes to fallout.

diff --git a/Engine/Plugins/Runtime/GameplayAbilities/Source/GameplayAbilities/Public/ActiveGameplayEffectIterator.h b/Engine/Plugins/Runtime/GameplayAbilities/Source/GameplayAbilities/Public/ActiveGameplayEffectIterator.h
--- a/Engine/Plugins/Runtime/GameplayAbilities/Source/GameplayAbilities/Public/ActiveGameplayEffectIterator.h
+++ b/Engine/Plugins/Runtime/GameplayAbilities/Source/GameplayAbilities/Public/ActiveGameplayEffectIterator.h
@@ -48,19 +48,19 @@ public:
 		Next();
 	}
 
-	ElementType& operator*()
+	ElementType& operator*() const
 	{
 		check(Current);
 		return *Current;
 	}
 
-	ElementType& operator->()
+	ElementType& operator->() const
 	{
 		check(Current);
 		return *Current;
 	}
 
-	operator bool()
+	explicit operator bool() const
 	{
 		return (Current != nullptr);
 	}
diff --git a/Engine/Source/Developer/DesktopPlatform/Public/PlatformInfo.h b/Engine/Source/Developer/DesktopPlatform/Public/PlatformInfo.h
--- a/Engine/Source/Developer/DesktopPlatform/Public/PlatformInfo.h
+++ b/Engine/Source/Developer/DesktopPlatform/Public/PlatformInfo.h
@@ -290,7 +290,7 @@ namespace PlatformInfo
 			return Copy;
 		}
 
-		FORCEINLINE operator bool() const
+		FORCEINLINE explicit operator bool() const
 		{
 			return CurrentPlatform < end();
 		}
diff --git a/Engine/Source/Editor/DetailCustomizations/Private/ComponentMaterialCategory.cpp b/Engine/Source/Editor/DetailCustomizations/Private/ComponentMaterialCategory.cpp
--- a/Engine/Source/Editor/DetailCustomizations/Private/ComponentMaterialCategory.cpp
+++ b/Engine/Source/Editor/DetailCustomizations/Private/ComponentMaterialCategory.cpp
@@ -116,7 +116,7 @@ public:
 	/**
 	 * @return Whether or not the iterator is valid
 	 */
-	operator bool()
+	explicit operator bool() const
 	{
 		return !bReachedEnd;	
 	}
diff --git a/Engine/Source/Editor/FoliageEdit/Private/FoliageEdMode.cpp b/Engine/Source/Editor/FoliageEdit/Private/FoliageEdMode.cpp
--- a/Engine/Source/Editor/FoliageEdit/Private/FoliageEdMode.cpp
+++ b/Engine/Source/Editor/FoliageEdit/Private/FoliageEdMode.cpp
@@ -157,7 +157,7 @@ public:
 		return CurrentInfo;
 	}
 
-	FORCEINLINE operator bool() const
+	FORCEINLINE explicit operator bool() const
 	{
 		return CurrentInfo != nullptr;
 	}
diff --git a/Engine/Source/Editor/Matinee/Public/MatineeGroupData.h b/Engine/Source/Editor/Matinee/Public/MatineeGroupData.h
--- a/Engine/Source/Editor/Matinee/Public/MatineeGroupData.h
+++ b/Engine/Source/Editor/Matinee/Public/MatineeGroupData.h
@@ -176,7 +176,7 @@ public:
 	/**
 	 * @return	true if the iterator has not reached the end; false, otherwise.
 	 */
-	FORCEINLINE operator bool()
+	FORCEINLINE explicit operator bool() const
 	{
 		return IsCurrentGroupValid();
 	}
diff --git a/Engine/Source/Editor/Matinee/Public/MatineeTrackData.h b/Engine/Source/Editor/Matinee/Public/MatineeTrackData.h
--- a/Engine/Source/Editor/Matinee/Public/MatineeTrackData.h
+++ b/Engine/Source/Editor/Matinee/Public/MatineeTrackData.h
@@ -148,7 +148,7 @@ public:
 	/**
 	 * @return	true if the iterator has not reached the end. 
 	 */
-	FORCEINLINE operator bool()
+	FORCEINLINE explicit operator bool() const
 	{
 		// The iterator must be pointing to a valid track, otherwise the iterator is considered invalid. 
 		return GroupIt && IsCurrentTrackValid();
diff --git a/Engine/Source/Editor/UnrealEd/Private/EditorCsg.cpp b/Engine/Source/Editor/UnrealEd/Private/EditorCsg.cpp
--- a/Engine/Source/Editor/UnrealEd/Private/EditorCsg.cpp
+++ b/Engine/Source/Editor/UnrealEd/Private/EditorCsg.cpp
@@ -98,7 +98,7 @@ public:
 	 *
 	 * @return true if iterator points to a suitable actor, false if it has reached the end
 	 */
-	operator bool()
+	explicit operator bool() const
 	{
 		return !ReachedEnd;
 	}
diff --git a/Engine/Source/Editor/UnrealEd/Private/VertexSnapping.cpp b/Engine/Source/Editor/UnrealEd/Private/VertexSnapping.cpp
--- a/Engine/Source/Editor/UnrealEd/Private/VertexSnapping.cpp
+++ b/Engine/Source/Editor/UnrealEd/Private/VertexSnapping.cpp
@@ -41,7 +41,7 @@ public:
 	}
 	
 	/** @return True if there are more vertices on the component */
-	operator bool()
+	explicit operator bool() const
 	{
 		return HasMoreVertices();
 	}
@@ -49,18 +49,18 @@ public:
 	/**
 	 * @return The position in world space of the current vertex
 	 */
-	virtual FVector Position() = 0;
+	virtual FVector Position() const = 0;
 
 	/**
 	 * @return The position in world space of the current vertex normal
 	 */
-	virtual FVector Normal() = 0;
+	virtual FVector Normal() const = 0;
 
 protected:
 	/**
 	 * @return True if there are more vertices on the component
 	 */
-	virtual bool HasMoreVertices() = 0;
+	virtual bool HasMoreVertices() const = 0;
 
 	/**
 	 * Advances to the next vertex
@@ -85,12 +85,12 @@ public:
 	}
 
 	/** FVertexIterator interface */
-	virtual FVector Position() override
+	virtual FVector Position() const override
 	{
 		return StaticMeshComponent->GetComponentTransform().TransformPosition( PositionBuffer.VertexPosition( CurrentVertexIndex ) );	
 	}
 
-	virtual FVector Normal() override
+	virtual FVector Normal() const override
 	{
 		return ComponentToWorldIT.TransformVector( VertexBuffer.VertexTangentZ( CurrentVertexIndex ) );
 	}
@@ -101,7 +101,7 @@ protected:
 		++CurrentVertexIndex;
 	}
 
-	virtual bool HasMoreVertices() override
+	virtual bool HasMoreVertices() const override
 	{
 		return CurrentVertexIndex < PositionBuffer.GetNumVertices();
 	}
@@ -141,13 +141,13 @@ public:
 	}
 
 	/** FVertexIterator interface */
-	virtual FVector Position() override
+	virtual FVector Position() const override
 	{
 		return BrushComponent->GetComponentTransform().TransformPosition( Vertices[CurrentVertexIndex] );
 	}
 
 	/** FVertexIterator interface */
-	virtual FVector Normal() override
+	virtual FVector Normal() const override
 	{
 		return FVector::ZeroVector;
 	}
@@ -158,7 +158,7 @@ protected:
 		++CurrentVertexIndex;
 	}
 
-	virtual bool HasMoreVertices() override
+	virtual bool HasMoreVertices() const override
 	{
 		return Vertices.IsValidIndex( CurrentVertexIndex );
 	}
@@ -188,13 +188,13 @@ public:
 	}
 	
 	/** FVertexIterator interface */
-	virtual FVector Position() override
+	virtual FVector Position() const override
 	{
 		const FVector VertPos = LODData.StaticVertexBuffers.PositionVertexBuffer.VertexPosition(VertexIndex);
 		return SkinnedMeshComponent->GetComponentTransform().TransformPosition(VertPos);
 	}
 
-	virtual FVector Normal() override
+	virtual FVector Normal() const override
 	{
 		FPackedNormal TangentZ = LODData.StaticVertexBuffers.StaticMeshVertexBuffer.VertexTangentZ(VertexIndex);
 		FVector VertNormal = TangentZ.ToFVector();
@@ -207,7 +207,7 @@ protected:
 		VertexIndex++;
 	}
 
-	virtual bool HasMoreVertices() override
+	virtual bool HasMoreVertices() const override
 	{
 		return VertexIndex < LODData.StaticVertexBuffers.PositionVertexBuffer.GetNumVertices();
 	}
diff --git a/Engine/Source/Editor/UnrealEd/Public/Commandlets/EditorCommandlets.h b/Engine/Source/Editor/UnrealEd/Public/Commandlets/EditorCommandlets.h
--- a/Engine/Source/Editor/UnrealEd/Public/Commandlets/EditorCommandlets.h
+++ b/Engine/Source/Editor/UnrealEd/Public/Commandlets/EditorCommandlets.h
@@ -263,7 +263,7 @@ struct FNativePropertyData
 	}
 
 	/** bool operator */
-	inline operator bool() const
+	inline explicit operator bool() const
 	{
 		return PropertyData.Num() || PropertyText.Num();
 	}
diff --git a/Engine/Source/Runtime/CoreUObject/Public/UObject/ScriptInterface.h b/Engine/Source/Runtime/CoreUObject/Public/UObject/ScriptInterface.h
--- a/Engine/Source/Runtime/CoreUObject/Public/UObject/ScriptInterface.h
+++ b/Engine/Source/Runtime/CoreUObject/Public/UObject/ScriptInterface.h
@@ -220,7 +220,7 @@ public:
 	 *
 	 * @return	true if InterfacePointer is non-NULL.
 	 */
-	FORCEINLINE operator bool() const
+	FORCEINLINE explicit operator bool() const
 	{
 		return GetInterface() != NULL;
 	}
diff --git a/Engine/Source/Runtime/CoreUObject/Public/UObject/StructOnScope.h b/Engine/Source/Runtime/CoreUObject/Public/UObject/StructOnScope.h
--- a/Engine/Source/Runtime/CoreUObject/Public/UObject/StructOnScope.h
+++ b/Engine/Source/Runtime/CoreUObject/Public/UObject/StructOnScope.h
@@ -226,7 +226,7 @@ public:
 		return Get();
 	}
 
-	operator bool() const
+	explicit operator bool() const
 	{
 		return IsValid();
 	}
diff --git a/Engine/Source/Runtime/CoreUObject/Public/UObject/UObjectGlobals.h b/Engine/Source/Runtime/CoreUObject/Public/UObject/UObjectGlobals.h
--- a/Engine/Source/Runtime/CoreUObject/Public/UObject/UObjectGlobals.h
+++ b/Engine/Source/Runtime/CoreUObject/Public/UObject/UObjectGlobals.h
@@ -687,7 +687,7 @@ public:
 		return !!Object && Object != (UObject*)InvalidPtrValue;
 	}
 	/** Convenience operator. Does the same thing as IsValid(). */
-	FORCEINLINE operator bool() const
+	FORCEINLINE explicit operator bool() const
 	{
 		return IsValid();
 	}
diff --git a/Engine/Source/Runtime/D3D12RHI/Private/D3D12CommandContext.h b/Engine/Source/Runtime/D3D12RHI/Private/D3D12CommandContext.h
--- a/Engine/Source/Runtime/D3D12RHI/Private/D3D12CommandContext.h
+++ b/Engine/Source/Runtime/D3D12RHI/Private/D3D12CommandContext.h
@@ -71,7 +71,7 @@ protected:
 		uint32 IndexDataStride;
 		void* IndexData;
 
-		operator bool() const { return NumVertices != 0; }
+		explicit operator bool() const { return NumVertices != 0; }
 		void Reset() { FMemory::Memzero(*this); }
 	};
 
diff --git a/Engine/Source/Runtime/Engine/Classes/Engine/Selection.h b/Engine/Source/Runtime/Engine/Classes/Engine/Selection.h
--- a/Engine/Source/Runtime/Engine/Classes/Engine/Selection.h
+++ b/Engine/Source/Runtime/Engine/Classes/Engine/Selection.h
@@ -464,7 +464,7 @@ public:
 	}
 
 	/** Returns true if the iterator has not yet reached the end of the selection set. */
-	operator bool() const
+	explicit operator bool() const
 	{
 		return IsIndexValid();
 	}
diff --git a/Engine/Source/Runtime/Engine/Private/Collision/CollisionConversions.cpp b/Engine/Source/Runtime/Engine/Private/Collision/CollisionConversions.cpp
--- a/Engine/Source/Runtime/Engine/Private/Collision/CollisionConversions.cpp
+++ b/Engine/Source/Runtime/Engine/Private/Collision/CollisionConversions.cpp
@@ -404,7 +404,7 @@ static bool ConvertOverlappedShapeToImpactHit(const UWorld* World, const FHitLoc
 	// Return start location as 'safe location'
 	OutResult.Location = QueryTM.GetLocation();
 
-	const bool bValidPosition = (GetFlags(Hit) & EHitFlags::Position);
+	const bool bValidPosition = !!(GetFlags(Hit) & EHitFlags::Position);
 	if (bValidPosition)
 	{
 		const FVector HitPosition = GetPosition(Hit);
diff --git a/Engine/Source/Runtime/Engine/Private/RepLayout.cpp b/Engine/Source/Runtime/Engine/Private/RepLayout.cpp
--- a/Engine/Source/Runtime/Engine/Private/RepLayout.cpp
+++ b/Engine/Source/Runtime/Engine/Private/RepLayout.cpp
@@ -5231,7 +5231,7 @@ void FRepLayout::SerializeProperties_r(
 
 		if ((GNetSharedSerializedData != 0) && Ar.IsSaving() && ((Cmd.Flags & ERepLayoutFlags::IsSharedSerialization) != ERepLayoutFlags::None))
 		{
-			FGuid PropertyGuid(CmdIndex, ArrayIndex, ArrayDepth, (int32)((PTRINT)(Data + Cmd) & 0xFFFFFFFF));
+			FGuid PropertyGuid(CmdIndex, ArrayIndex, ArrayDepth, (int32)((PTRINT)(const uint8*)(Data + Cmd) & 0xFFFFFFFF));
 
 			SharedPropInfo = SharedInfo.SharedPropertyInfo.FindByPredicate([&](const FRepSerializedPropertyInfo& Info) 
 			{ 
@@ -5472,7 +5472,7 @@ void FRepLayout::BuildSharedSerializationForRPC_r(
 
 		if (!Parents[Cmd.ParentIndex].Property->HasAnyPropertyFlags(CPF_OutParm) && ((Cmd.Flags & ERepLayoutFlags::IsSharedSerialization) != ERepLayoutFlags::None))
 		{
-			FGuid PropertyGuid(CmdIndex, ArrayIndex, ArrayDepth, (int32)((PTRINT)(Data + Cmd) & 0xFFFFFFFF));
+			FGuid PropertyGuid(CmdIndex, ArrayIndex, ArrayDepth, (int32)((PTRINT)(const uint8*)(Data + Cmd) & 0xFFFFFFFF));
 
 			SharedInfo.WriteSharedProperty(Cmd, PropertyGuid, CmdIndex, 0, (Data + Cmd).Data, false, false);
 		}
diff --git a/Engine/Source/Runtime/Engine/Private/SkeletalMesh.cpp b/Engine/Source/Runtime/Engine/Private/SkeletalMesh.cpp
--- a/Engine/Source/Runtime/Engine/Private/SkeletalMesh.cpp
+++ b/Engine/Source/Runtime/Engine/Private/SkeletalMesh.cpp
@@ -3634,7 +3634,7 @@ public:
 		} while (NotValidPreviewSection());
 		return *this;
 	}
-	FORCEINLINE operator bool() const
+	FORCEINLINE explicit operator bool() const
 	{
 		return ((SectionIndex < Sections.Num()) && LODSectionElements.SectionElements.IsValidIndex(GetSectionElementIndex()));
 	}
diff --git a/Engine/Source/Runtime/Engine/Private/Streaming/TextureInstanceView.h b/Engine/Source/Runtime/Engine/Private/Streaming/TextureInstanceView.h
--- a/Engine/Source/Runtime/Engine/Private/Streaming/TextureInstanceView.h
+++ b/Engine/Source/Runtime/Engine/Private/Streaming/TextureInstanceView.h
@@ -132,7 +132,7 @@ public:
 	public:
 		FRenderAssetLinkConstIterator(const FRenderAssetInstanceView& InState, const UStreamableRenderAsset* InAsset);
 
-		FORCEINLINE operator bool() const { return CurrElementIndex != INDEX_NONE; }
+		FORCEINLINE explicit operator bool() const { return CurrElementIndex != INDEX_NONE; }
 		FORCEINLINE void operator++() { CurrElementIndex = State.Elements[CurrElementIndex].NextRenderAssetLink; }
 
 		void OutputToLog(float MaxNormalizedSize, float MaxNormalizedSize_VisibleOnly, const TCHAR* Prefix) const;
@@ -170,7 +170,7 @@ public:
 	public:
 		FRenderAssetIterator(const FRenderAssetInstanceView& InState) : MapIt(InState.RenderAssetMap) {}
 
-		operator bool() const { return (bool)MapIt; }
+		explicit operator bool() const { return (bool)MapIt; }
 		void operator++() { ++MapIt; }
 
 		const UStreamableRenderAsset* operator*() const { return MapIt.Key(); }
diff --git a/Engine/Source/Runtime/Engine/Public/Net/RepLayout.h b/Engine/Source/Runtime/Engine/Public/Net/RepLayout.h
--- a/Engine/Source/Runtime/Engine/Public/Net/RepLayout.h
+++ b/Engine/Source/Runtime/Engine/Public/Net/RepLayout.h
@@ -83,21 +83,16 @@ namespace UE4_RepLayout_Private
 			return InBuffer.Data + Offset;
 		}
 
-		operator bool()
+		explicit operator bool() const
 		{
 			return Data != nullptr;
 		}
 
-		operator ConstOrNotType* ()
+		operator ConstOrNotType* () const
 		{
 			return Data;
 		}
 
-		operator ConstOrNotVoid* ()
-		{
-			return (ConstOrNotVoid*)Data;
-		}
-
 		ConstOrNotType* RESTRICT Data;
 	};
 
diff --git a/Engine/Source/Runtime/Engine/Public/PhysicsInterfaceWrapperShared.h b/Engine/Source/Runtime/Engine/Public/PhysicsInterfaceWrapperShared.h
--- a/Engine/Source/Runtime/Engine/Public/PhysicsInterfaceWrapperShared.h
+++ b/Engine/Source/Runtime/Engine/Public/PhysicsInterfaceWrapperShared.h
@@ -23,7 +23,7 @@ inline EQueryFlags operator&(EQueryFlags lhs, EQueryFlags rhs)
 struct FQueryFlags
 {
 	FQueryFlags(EQueryFlags InFlags) : QueryFlags(InFlags) {}
-	operator bool() const { return !!static_cast<uint16>(QueryFlags); }
+	explicit operator bool() const { return !!static_cast<uint16>(QueryFlags); }
 	FQueryFlags operator |(EQueryFlags Rhs) const
 	{
 		return FQueryFlags(QueryFlags | Rhs);
@@ -74,7 +74,7 @@ inline EHitFlags operator&(EHitFlags lhs, EHitFlags rhs)
 struct FHitFlags
 {
 	FHitFlags(EHitFlags InFlags) : HitFlags(InFlags) {}
-	operator bool() const { return !!static_cast<uint16>(HitFlags); }
+	explicit operator bool() const { return !!static_cast<uint16>(HitFlags); }
 	FHitFlags operator |(EHitFlags Rhs) const
 	{
 		return FHitFlags(HitFlags | Rhs);
diff --git a/Engine/Source/Runtime/Engine/Public/SurfaceIterators.h b/Engine/Source/Runtime/Engine/Public/SurfaceIterators.h
--- a/Engine/Source/Runtime/Engine/Public/SurfaceIterators.h
+++ b/Engine/Source/Runtime/Engine/Public/SurfaceIterators.h
@@ -70,7 +70,7 @@ public:
 		check( CurrentSurface );
 		return CurrentSurface;
 	}
-	FORCEINLINE operator bool()
+	FORCEINLINE explicit operator bool() const
 	{
 		return !bReachedEnd;
 	}
diff --git a/Engine/Source/Runtime/Engine/Public/UnrealEngine.h b/Engine/Source/Runtime/Engine/Public/UnrealEngine.h
--- a/Engine/Source/Runtime/Engine/Public/UnrealEngine.h
+++ b/Engine/Source/Runtime/Engine/Public/UnrealEngine.h
@@ -55,15 +55,15 @@ public:
 		++Iter;
 		GetCurrent();
 	}
-	ULocalPlayer* operator*()
+	ULocalPlayer* operator*() const
 	{
 		return *Iter;
 	}
-	ULocalPlayer* operator->()
+	ULocalPlayer* operator->() const
 	{
 		return *Iter;
 	}
-	operator bool()
+	explicit operator bool() const
 	{
 		return (bool)Iter;
 	}
@@ -112,15 +112,15 @@ public:
 	{
 		Next();
 	}
-	T* operator*()
+	T* operator*() const
 	{
 		return Current;
 	}
-	T* operator->()
+	T* operator->() const
 	{
 		return Current;
 	}
-	operator bool()
+	explicit operator bool() const
 	{
 		return (Current!=NULL);
 	}
diff --git a/Engine/Source/Runtime/OpenGLDrv/Private/Android/AndroidESDeferredOpenGL.h b/Engine/Source/Runtime/OpenGLDrv/Private/Android/AndroidESDeferredOpenGL.h
--- a/Engine/Source/Runtime/OpenGLDrv/Private/Android/AndroidESDeferredOpenGL.h
+++ b/Engine/Source/Runtime/OpenGLDrv/Private/Android/AndroidESDeferredOpenGL.h
@@ -32,7 +32,7 @@ union FEGLGLSync
 
 	}
 
-	operator bool() const
+	explicit operator bool() const
 	{
 		return GL;
 	}
diff --git a/Engine/Source/Runtime/OpenGLDrv/Private/Lumin/LuminESDeferredOpenGL.h b/Engine/Source/Runtime/OpenGLDrv/Private/Lumin/LuminESDeferredOpenGL.h
--- a/Engine/Source/Runtime/OpenGLDrv/Private/Lumin/LuminESDeferredOpenGL.h
+++ b/Engine/Source/Runtime/OpenGLDrv/Private/Lumin/LuminESDeferredOpenGL.h
@@ -26,7 +26,7 @@ union FEGLGLSync
 
 	}
 
-	operator bool() const
+	explicit operator bool() const
 	{
 		return GL;
 	}
diff --git a/Engine/Source/Runtime/RuntimeAssetCache/Public/RuntimeAssetCacheInterface.h b/Engine/Source/Runtime/RuntimeAssetCache/Public/RuntimeAssetCacheInterface.h
--- a/Engine/Source/Runtime/RuntimeAssetCache/Public/RuntimeAssetCacheInterface.h
+++ b/Engine/Source/Runtime/RuntimeAssetCache/Public/RuntimeAssetCacheInterface.h
@@ -43,7 +43,7 @@ struct FVoidPtrParam
 		return (!Data || DataSize <= 0);
 	}
 
-	FORCEINLINE operator bool() const
+	FORCEINLINE explicit operator bool() const
 	{
 		return (Data && DataSize > 0);
 	}
