Change DataTable import UI to show a propertyview for import options

diff --git a/Engine/Source/Editor/UnrealEd/Classes/Factories/CSVImportFactory.h b/Engine/Source/Editor/UnrealEd/Classes/Factories/CSVImportFactory.h
--- a/Engine/Source/Editor/UnrealEd/Classes/Factories/CSVImportFactory.h
+++ b/Engine/Source/Editor/UnrealEd/Classes/Factories/CSVImportFactory.h
@@ -81,5 +81,9 @@ private:
 
 	UPROPERTY()
 	FCSVImportSettings AutomatedImportSettings;
+
+	/** Temporary data table to use to display import options */
+	UPROPERTY()
+	UDataTable* TempImportDataTable;
 };
 
diff --git a/Engine/Source/Editor/UnrealEd/Private/Factories/CSVImportFactory.cpp b/Engine/Source/Editor/UnrealEd/Private/Factories/CSVImportFactory.cpp
--- a/Engine/Source/Editor/UnrealEd/Private/Factories/CSVImportFactory.cpp
+++ b/Engine/Source/Editor/UnrealEd/Private/Factories/CSVImportFactory.cpp
@@ -111,22 +111,30 @@ UObject* UCSVImportFactory::FactoryCreateText(UClass* InClass, UObject* InParent
 
 	// Save off information if so
 	bool bHaveInfo = false;
-	UScriptStruct* ImportRowStruct = nullptr;
 	ERichCurveInterpMode ImportCurveInterpMode = RCIM_Linear;
 	ECSVImportType ImportType = ECSVImportType::ECSV_DataTable;
+	UClass* DataTableClass = UDataTable::StaticClass();
+
+	// Clear our temp table
+	TempImportDataTable = nullptr;
 
 	if (IsAutomatedImport())
 	{
-		ImportRowStruct = AutomatedImportSettings.ImportRowStruct;
 		ImportCurveInterpMode = AutomatedImportSettings.ImportCurveInterpMode;
 		ImportType = AutomatedImportSettings.ImportType;
 
+		TempImportDataTable = NewObject<UDataTable>(this, UDataTable::StaticClass(), InName, Flags);
+		TempImportDataTable->RowStruct = AutomatedImportSettings.ImportRowStruct;
+
 		// For automated import to work a row struct must be specified for a datatable type or a curve type must be specified
-		bHaveInfo = ImportRowStruct != nullptr || ImportType != ECSVImportType::ECSV_DataTable;
+		bHaveInfo = TempImportDataTable->RowStruct != nullptr || ImportType != ECSVImportType::ECSV_DataTable;
 	}
 	else if (ExistingTable != nullptr)
 	{
-		ImportRowStruct = ExistingTable->RowStruct;
+		ImportType = ECSVImportType::ECSV_DataTable;
+		TempImportDataTable = NewObject<UDataTable>(this, ExistingTable->GetClass(), InName, Flags);
+		TempImportDataTable->CopyImportOptions(ExistingTable);
+		
 		bHaveInfo = true;
 	}
 	else if (ExistingCurveTable != nullptr)
@@ -142,6 +150,12 @@ UObject* UCSVImportFactory::FactoryCreateText(UClass* InClass, UObject* InParent
 
 	bool bDoImport = true;
 
+	if (!TempImportDataTable)
+	{
+		// Create an empty one
+		TempImportDataTable = NewObject<UDataTable>(this, UDataTable::StaticClass(), InName, Flags);
+	}
+
 	// If we do not have the info we need, pop up window to ask for things
 	if (!bHaveInfo && !IsAutomatedImport())
 	{
@@ -171,19 +185,20 @@ UObject* UCSVImportFactory::FactoryCreateText(UClass* InClass, UObject* InParent
 			SAssignNew(ImportOptionsWindow, SCSVImportOptions)
 				.WidgetWindow(Window)
 				.FullPath(FText::FromString(ParentFullPath))
+				.TempImportDataTable(TempImportDataTable)
 		);
 
 		FSlateApplication::Get().AddModalWindow(Window, ParentWindow, false);
 
 		ImportType = ImportOptionsWindow->GetSelectedImportType();
-		ImportRowStruct = ImportOptionsWindow->GetSelectedRowStruct();
+		TempImportDataTable->RowStruct = ImportOptionsWindow->GetSelectedRowStruct();
 		ImportCurveInterpMode = ImportOptionsWindow->GetSelectedCurveIterpMode();
 		bDoImport = ImportOptionsWindow->ShouldImport();
 		bOutOperationCanceled = !bDoImport;
 	}
 	else if (!bHaveInfo && IsAutomatedImport())
 	{
-		if (ImportType == ECSVImportType::ECSV_DataTable && !ImportRowStruct)
+		if (ImportType == ECSVImportType::ECSV_DataTable && !TempImportDataTable->RowStruct)
 		{
 			UE_LOG(LogCSVImportFactory, Error, TEXT("A Data table row type must be specified in the import settings json file for automated import"));
 		}
@@ -207,28 +222,22 @@ UObject* UCSVImportFactory::FactoryCreateText(UClass* InClass, UObject* InParent
 
 		if (ImportType == ECSVImportType::ECSV_DataTable)
 		{
-			UClass* DataTableClass = UDataTable::StaticClass();
-
 			// If there is an existing table, need to call this to free data memory before recreating object
-			bool bStripFromClientBuilds = false;
 			UDataTable::FOnDataTableChanged OldOnDataTableChanged;
-			FString ImportKeyField;
 			if (ExistingTable != nullptr)
 			{
 				OldOnDataTableChanged = MoveTemp(ExistingTable->OnDataTableChanged());
 				ExistingTable->OnDataTableChanged().Clear();
 				DataTableClass = ExistingTable->GetClass();
 				ExistingTable->EmptyTable();
-				bStripFromClientBuilds = ExistingTable->bStripFromClientBuilds;
-				ImportKeyField = ExistingTable->ImportKeyField;
 			}
 
 			// Create/reset table
 			UDataTable* NewTable = NewObject<UDataTable>(InParent, DataTableClass, InName, Flags);
-			NewTable->RowStruct = ImportRowStruct;
+			
+			NewTable->CopyImportOptions(TempImportDataTable);
 			NewTable->AssetImportData->Update(CurrentFilename);
-			NewTable->bStripFromClientBuilds = bStripFromClientBuilds;
-			NewTable->ImportKeyField = ImportKeyField;
+
 			// Go ahead and create table from string
 			Problems = DoImportDataTable(NewTable, String);
 
diff --git a/Engine/Source/Editor/UnrealEd/Private/SCSVImportOptions.cpp b/Engine/Source/Editor/UnrealEd/Private/SCSVImportOptions.cpp
--- a/Engine/Source/Editor/UnrealEd/Private/SCSVImportOptions.cpp
+++ b/Engine/Source/Editor/UnrealEd/Private/SCSVImportOptions.cpp
@@ -5,16 +5,20 @@
 #include "UObject/UObjectIterator.h"
 #include "UObject/Package.h"
 #include "Widgets/Input/SButton.h"
+#include "Widgets/Layout/SExpandableArea.h"
 #include "EditorStyleSet.h"
 #include "Engine/UserDefinedStruct.h"
-
-
+#include "Engine/DataTable.h"
+#include "Modules/ModuleManager.h"
+#include "PropertyEditorModule.h"
+#include "ObjectEditorUtils.h"
 
 #define LOCTEXT_NAMESPACE "CSVImportFactory"
 
 void SCSVImportOptions::Construct(const FArguments& InArgs)
 {
 	WidgetWindow = InArgs._WidgetWindow;
+	TempImportDataTable = InArgs._TempImportDataTable;
 
 	// Make array of enum pointers
 	TSharedPtr<ECSVImportType> DataTableTypePtr = MakeShareable(new ECSVImportType(ECSVImportType::ECSV_DataTable));
@@ -45,12 +49,32 @@ void SCSVImportOptions::Construct(const FArguments& InArgs)
 		RowStructs.Sort([](const UScriptStruct& ElementA, const UScriptStruct& ElementB) { return (ElementA.GetName() < ElementB.GetName()); } );
 	}
 
+	// Create properties view
+	FPropertyEditorModule & EditModule = FModuleManager::Get().GetModuleChecked<FPropertyEditorModule>("PropertyEditor");
+	FDetailsViewArgs DetailsViewArgs(/*bUpdateFromSelection=*/ false, /*bLockable=*/ false, /*bAllowSearch=*/ false, /*InNameAreaSettings=*/ FDetailsViewArgs::HideNameArea, /*bHideSelectionTip=*/ true);
+	PropertyView = EditModule.CreateDetailView(DetailsViewArgs);
+
+	PropertyView->SetIsPropertyVisibleDelegate(FIsPropertyVisible::CreateLambda([](const FPropertyAndParent& InPropertyAndParent)
+	{
+		static FName ImportOptions = FName(TEXT("ImportOptions"));
+
+		// Only show import options
+		FName CategoryName = FObjectEditorUtils::GetCategoryFName(&InPropertyAndParent.Property);
+
+		if (CategoryName == ImportOptions)
+		{
+			return true;
+		}
+
+		return false;
+	}));
+
 	// Create widget
 	this->ChildSlot
 	[
 		SNew(SBorder)
-		. BorderImage(FEditorStyle::GetBrush(TEXT("Menu.Background")))
-		. Padding(10)
+		.BorderImage(FEditorStyle::GetBrush(TEXT("Menu.Background")))
+		.Padding(10)
 		[
 			SNew(SVerticalBox)
 			+SVerticalBox::Slot()
@@ -95,6 +119,7 @@ void SCSVImportOptions::Construct(const FArguments& InArgs)
 				SAssignNew(ImportTypeCombo, SComboBox< TSharedPtr<ECSVImportType> >)
 				.OptionsSource( &ImportTypes )
 				.OnGenerateWidget( this, &SCSVImportOptions::MakeImportTypeItemWidget )
+				.OnSelectionChanged( this, &SCSVImportOptions::OnImportTypeSelected)
 				[
 					SNew(STextBlock)
 					.Text(this, &SCSVImportOptions::GetSelectedItemText)
@@ -103,6 +128,7 @@ void SCSVImportOptions::Construct(const FArguments& InArgs)
 			// Data row struct
 			+SVerticalBox::Slot()
 			.AutoHeight()
+			.Padding(2)
 			[
 				SNew(STextBlock)
 				.Text( LOCTEXT("ChooseRowType", "Choose DataTable Row Type:") )
@@ -123,6 +149,7 @@ void SCSVImportOptions::Construct(const FArguments& InArgs)
 			// Curve interpolation
 			+SVerticalBox::Slot()
 			.AutoHeight()
+			.Padding(2)
 			[
 				SNew(STextBlock)
 				.Text( LOCTEXT("ChooseCurveType", "Choose Curve Interpolation Type:") )
@@ -140,6 +167,18 @@ void SCSVImportOptions::Construct(const FArguments& InArgs)
 					.Text(this, &SCSVImportOptions::GetSelectedCurveTypeText)
 				]
 			]
+			// Details panel
+			+ SVerticalBox::Slot()
+			.AutoHeight()
+			.Padding(2)
+			[
+				SNew(SBox)
+				.WidthOverride(400)
+				.Visibility(this, &SCSVImportOptions::GetDetailsPanelVis)
+				[
+					PropertyView.ToSharedRef()
+				]
+			]
 			// Ok/Cancel
 			+SVerticalBox::Slot()
 			.AutoHeight()
@@ -147,6 +186,7 @@ void SCSVImportOptions::Construct(const FArguments& InArgs)
 				SNew(SHorizontalBox)
 				+SHorizontalBox::Slot()
 				.AutoWidth()
+				.Padding(2)
 				[
 					SNew(SButton)
 					.Text(LOCTEXT("OK", "OK"))
@@ -155,6 +195,7 @@ void SCSVImportOptions::Construct(const FArguments& InArgs)
 				]
 				+SHorizontalBox::Slot()
 				.AutoWidth()
+				.Padding(2)
 				[
 					SNew(SButton)
 					.Text(LOCTEXT("Cancel", "Cancel"))
@@ -166,8 +207,9 @@ void SCSVImportOptions::Construct(const FArguments& InArgs)
 
 	// set-up selection
 	ImportTypeCombo->SetSelectedItem(DataTableTypePtr);
+	PropertyView->SetObject(TempImportDataTable.Get());
 
-	// Populate the valid interploation modes
+	// Populate the valid interpolation modes
 	{
 		CurveInterpModes.Add( MakeShareable( new ERichCurveInterpMode(ERichCurveInterpMode::RCIM_Constant) ) );
 		CurveInterpModes.Add( MakeShareable( new ERichCurveInterpMode(ERichCurveInterpMode::RCIM_Linear) ) );
@@ -213,6 +255,11 @@ EVisibility SCSVImportOptions::GetCurveTypeVis() const
 	return (ImportTypeCombo.IsValid() && *ImportTypeCombo->GetSelectedItem() == ECSVImportType::ECSV_CurveTable) ? EVisibility::Visible : EVisibility::Collapsed;
 }
 
+EVisibility SCSVImportOptions::GetDetailsPanelVis() const
+{
+	return (ImportTypeCombo.IsValid() && *ImportTypeCombo->GetSelectedItem() == ECSVImportType::ECSV_DataTable) ? EVisibility::Visible : EVisibility::Hidden;
+}
+
 FString SCSVImportOptions::GetImportTypeText(TSharedPtr<ECSVImportType> Type) const
 {
 	FString EnumString;
@@ -242,6 +289,18 @@ TSharedRef<SWidget> SCSVImportOptions::MakeImportTypeItemWidget(TSharedPtr<ECSVI
 			.Text(FText::FromString(GetImportTypeText(Type)));
 }
 
+void SCSVImportOptions::OnImportTypeSelected(TSharedPtr<ECSVImportType> Selection, ESelectInfo::Type SelectionType)
+{
+	if (Selection.IsValid() && *Selection == ECSVImportType::ECSV_DataTable)
+	{
+		PropertyView->SetObject(TempImportDataTable.Get());
+	}
+	else
+	{
+		PropertyView->SetObject(nullptr);
+	}
+}
+
 /** Called to create a widget for each struct */
 TSharedRef<SWidget> SCSVImportOptions::MakeRowStructItemWidget(UScriptStruct* Struct)
 {
diff --git a/Engine/Source/Editor/UnrealEd/Public/SCSVImportOptions.h b/Engine/Source/Editor/UnrealEd/Public/SCSVImportOptions.h
--- a/Engine/Source/Editor/UnrealEd/Public/SCSVImportOptions.h
+++ b/Engine/Source/Editor/UnrealEd/Public/SCSVImportOptions.h
@@ -25,16 +25,19 @@ public:
 	SLATE_BEGIN_ARGS(SCSVImportOptions)
 		: _WidgetWindow()
 		, _FullPath()
+		, _TempImportDataTable()
 		{}
 
 		SLATE_ARGUMENT(TSharedPtr<SWindow>, WidgetWindow)
 		SLATE_ARGUMENT(FText, FullPath)
+		SLATE_ARGUMENT(UDataTable*, TempImportDataTable)
 	SLATE_END_ARGS()
 
 	SCSVImportOptions()
 		: bImport(false)
 		, SelectedImportType(ECSVImportType::ECSV_DataTable)
 		, SelectedStruct(nullptr)
+		, TempImportDataTable(nullptr)
 		{}
 
 	void Construct(const FArguments& InArgs);
@@ -54,14 +57,20 @@ public:
 	/** Whether to show table row options */
 	EVisibility GetTableRowOptionVis() const;
 
-	/** Whether to show table row options */
+	/** Whether to show curve type options */
 	EVisibility GetCurveTypeVis() const;
 
+	/** Whether to show details panel */
+	EVisibility GetDetailsPanelVis() const;
+
 	FString GetImportTypeText(TSharedPtr<ECSVImportType> Type) const;
 
 	/** Called to create a widget for each struct */
 	TSharedRef<SWidget> MakeImportTypeItemWidget(TSharedPtr<ECSVImportType> Type);
 
+	/** Called when import type changes */
+	void OnImportTypeSelected(TSharedPtr<ECSVImportType> Selection, ESelectInfo::Type SelectionType);
+
 	/** Called to create a widget for each struct */
 	TSharedRef<SWidget> MakeRowStructItemWidget(UScriptStruct* Struct);
 
@@ -115,9 +124,15 @@ private:
 	/** The selected row struct */
 	UScriptStruct*									SelectedStruct;
 
+	/** Temp DataTable to hold import options */
+	TWeakObjectPtr< UDataTable >					TempImportDataTable;
+
 	/** The curve interpolation combo box */
 	TSharedPtr< SComboBox<CurveInterpModePtr> >		CurveInterpCombo;
 
+	/** A property view to edit advanced options */
+	TSharedPtr< class IDetailsView >				PropertyView;
+
 	/** All available curve interpolation modes */
 	TArray< CurveInterpModePtr >					CurveInterpModes;
 
diff --git a/Engine/Source/Runtime/Engine/Classes/Engine/DataTable.h b/Engine/Source/Runtime/Engine/Classes/Engine/DataTable.h
--- a/Engine/Source/Runtime/Engine/Classes/Engine/DataTable.h
+++ b/Engine/Source/Runtime/Engine/Classes/Engine/DataTable.h
@@ -52,7 +52,7 @@ struct FTableRowBase
 /**
  * Imported spreadsheet table.
  */
-UCLASS(MinimalAPI, BlueprintType, AutoExpandCategories = DataTable)
+UCLASS(MinimalAPI, BlueprintType, AutoExpandCategories = "DataTable,ImportOptions")
 class UDataTable
 	: public UObject
 {
@@ -66,8 +66,9 @@ class UDataTable
 	friend FDataTableImporterJSON;
 
 	/** Structure to use for each row of the table, must inherit from FTableRowBase */
-	UPROPERTY()
+	UPROPERTY(VisibleAnywhere, Category=DataTable, meta=(DisplayThumbnail="false"))
 	UScriptStruct*			RowStruct;
+
 protected:
 	/** Map of name of row to row data structure. */
 	TMap<FName, uint8*>		RowMap;
@@ -78,6 +79,7 @@ protected:
 	/** Called to add rows to the data table */
 	ENGINE_API virtual void AddRowInternal(FName RowName, uint8* RowDataPtr);
 public:
+
 	virtual const TMap<FName, uint8*>& GetRowMap() const { return RowMap; }
 	virtual const TMap<FName, uint8*>& GetRowMap() { return RowMap; }
 
@@ -90,8 +92,16 @@ public:
 	UPROPERTY(EditAnywhere, Category=DataTable)
 	uint8 bStripFromClientBuilds : 1;
 
+	/** Set to true to ignore extra fields in the import data, if false it will warn about them */
+	UPROPERTY(EditAnywhere, Category=ImportOptions)
+	uint8 bIgnoreExtraFields : 1;
+
+	/** Set to true to ignore any fields that are expected but missing, if false it will warn about them */
+	UPROPERTY(EditAnywhere, Category = ImportOptions)
+	uint8 bIgnoreMissingFields : 1;
+
 	/** Explicit field in import data to use as key. If this is empty it uses Name for JSON and the first field found for CSV */
-	UPROPERTY(EditAnywhere, Category=DataTable)
+	UPROPERTY(EditAnywhere, Category=ImportOptions)
 	FString ImportKeyField;
 	
 #if WITH_EDITOR
@@ -113,7 +123,8 @@ public:
 	ENGINE_API virtual void PostLoad() override;
 	//~ End UObject Interface
 
-	UPROPERTY(VisibleAnywhere, Instanced, Category=ImportSettings)
+	/** The file this data table was imported from, may be empty */
+	UPROPERTY(VisibleAnywhere, Instanced, Category=ImportSource)
 	class UAssetImportData* AssetImportData;
 
 	/** The filename imported to create this object. Relative to this object's package, BaseDir() or absolute */
@@ -131,10 +142,9 @@ protected:
 
 	UPROPERTY(Transient)
 	TSet<UObject*> TemporarilyReferencedObjects;
-
 #endif	// WITH_EDITORONLY_DATA
-private:
 
+private:
 	/** A multicast delegate that is called any time the data table changes. */
 	FOnDataTableChanged OnDataTableChangedDelegate;
 
@@ -306,6 +316,9 @@ public:
 
 	/** Output the fields from a particular row (use RowMap to get RowData) to an existing JsonWriter */
 	ENGINE_API bool WriteRowAsJSON(const TSharedRef< TJsonWriter<TCHAR, TPrettyJsonPrintPolicy<TCHAR> > >& JsonWriter, const void* RowData, const EDataTableExportFlags InDTExportFlags = EDataTableExportFlags::None) const;
+
+	/** Copies all the import options from another table, this does not copy row dawta */
+	ENGINE_API bool CopyImportOptions(UDataTable* SourceTable);
 #endif
 	/** 
 	 *	Create table from CSV style comma-separated string. 
diff --git a/Engine/Source/Runtime/Engine/Private/DataTable.cpp b/Engine/Source/Runtime/Engine/Private/DataTable.cpp
--- a/Engine/Source/Runtime/Engine/Private/DataTable.cpp
+++ b/Engine/Source/Runtime/Engine/Private/DataTable.cpp
@@ -68,6 +68,10 @@ namespace
 UDataTable::UDataTable(const FObjectInitializer& ObjectInitializer)
 	: Super(ObjectInitializer)
 {
+	bIgnoreExtraFields = false;
+	bIgnoreMissingFields = false;
+	bStripFromClientBuilds = false;
+
 #if WITH_EDITORONLY_DATA
 	{ static const FAutoRegisterLocalizationDataGatheringCallback AutomaticRegistrationOfLocalizationGatherer(UDataTable::StaticClass(), &GatherDataTableForLocalization); }
 #endif
@@ -547,6 +551,33 @@ bool UDataTable::WriteRowAsJSON(const TSharedRef< TJsonWriter<TCHAR, TPrettyJson
 	return FDataTableExporterJSON(InDTExportFlags, JsonWriter).WriteRow(RowStruct, RowData);
 }
 
+bool UDataTable::CopyImportOptions(UDataTable* SourceTable)
+{
+	// Only safe to call on an empty table
+	if (!SourceTable || !ensure(RowMap.Num() == 0))
+	{
+		return false;
+	}
+
+	bStripFromClientBuilds = SourceTable->bStripFromClientBuilds;
+	bIgnoreExtraFields = SourceTable->bIgnoreExtraFields;
+	bIgnoreMissingFields = SourceTable->bIgnoreMissingFields;
+	ImportKeyField = SourceTable->ImportKeyField;
+	RowStruct = SourceTable->RowStruct;
+
+	if (RowStruct)
+	{
+		RowStructName = RowStruct->GetFName();
+	}
+
+	if (SourceTable->AssetImportData)
+	{
+		AssetImportData->SourceData = SourceTable->AssetImportData->SourceData;
+	}
+
+	return true;
+}
+
 bool UDataTable::WriteTableAsJSON(const TSharedRef< TJsonWriter<TCHAR, TPrettyJsonPrintPolicy<TCHAR> > >& JsonWriter, const EDataTableExportFlags InDTExportFlags) const
 {
 	return FDataTableExporterJSON(InDTExportFlags, JsonWriter).WriteTable(*this);
@@ -592,7 +623,10 @@ TArray<UProperty*> UDataTable::GetTablePropertyArray(const TArray<const TCHAR*>&
 				// Didn't find a property with this name, problem..
 				if(ColumnProp == nullptr)
 				{
-					OutProblems.Add(FString::Printf(TEXT("Cannot find Property for column '%s' in struct '%s'."), *PropName.ToString(), *InRowStruct->GetName()));
+					if (!bIgnoreExtraFields)
+					{
+						OutProblems.Add(FString::Printf(TEXT("Cannot find Property for column '%s' in struct '%s'."), *PropName.ToString(), *InRowStruct->GetName()));
+					}
 				}
 				// Found one!
 				else
@@ -620,23 +654,26 @@ TArray<UProperty*> UDataTable::GetTablePropertyArray(const TArray<const TCHAR*>&
 		}
 	}
 
-	// Generate warning for any properties in struct we are not filling in
-	for (int32 PropIdx=0; PropIdx < ExpectedPropNames.Num(); PropIdx++)
+	if (!bIgnoreMissingFields)
 	{
-		const UProperty* const ColumnProp = FindField<UProperty>(InRowStruct, ExpectedPropNames[PropIdx]);
+		// Generate warning for any properties in struct we are not filling in
+		for (int32 PropIdx = 0; PropIdx < ExpectedPropNames.Num(); PropIdx++)
+		{
+			const UProperty* const ColumnProp = FindField<UProperty>(InRowStruct, ExpectedPropNames[PropIdx]);
 
 #if WITH_EDITOR
-		// If the structure has specified the property as optional for import (gameplay code likely doing a custom fix-up or parse of that property),
-		// then avoid warning about it
-		static const FName DataTableImportOptionalMetadataKey(TEXT("DataTableImportOptional"));
-		if (ColumnProp->HasMetaData(DataTableImportOptionalMetadataKey))
-		{
-			continue;
-		}
+			// If the structure has specified the property as optional for import (gameplay code likely doing a custom fix-up or parse of that property),
+			// then avoid warning about it
+			static const FName DataTableImportOptionalMetadataKey(TEXT("DataTableImportOptional"));
+			if (ColumnProp->HasMetaData(DataTableImportOptionalMetadataKey))
+			{
+				continue;
+			}
 #endif // WITH_EDITOR
 
-		const FString DisplayName = DataTableUtils::GetPropertyDisplayName(ColumnProp, ExpectedPropNames[PropIdx].ToString());
-		OutProblems.Add(FString::Printf(TEXT("Expected column '%s' not found in input."), *DisplayName));
+			const FString DisplayName = DataTableUtils::GetPropertyDisplayName(ColumnProp, ExpectedPropNames[PropIdx].ToString());
+			OutProblems.Add(FString::Printf(TEXT("Expected column '%s' not found in input."), *DisplayName));
+		}
 	}
 
 	return ColumnProps;
diff --git a/Engine/Source/Runtime/Engine/Private/DataTableJSON.cpp b/Engine/Source/Runtime/Engine/Private/DataTableJSON.cpp
--- a/Engine/Source/Runtime/Engine/Private/DataTableJSON.cpp
+++ b/Engine/Source/Runtime/Engine/Private/DataTableJSON.cpp
@@ -493,7 +493,21 @@ bool FDataTableImporterJSON::ReadStruct(const TSharedRef<FJsonObject>& InParsedO
 
 		if (!ParsedPropertyValue.IsValid())
 		{
-			ImportProblems.Add(FString::Printf(TEXT("Row '%s' is missing an entry for '%s'."), *InRowName.ToString(), *ColumnName));
+#if WITH_EDITOR
+			// If the structure has specified the property as optional for import (gameplay code likely doing a custom fix-up or parse of that property),
+			// then avoid warning about it
+			static const FName DataTableImportOptionalMetadataKey(TEXT("DataTableImportOptional"));
+			if (BaseProp->HasMetaData(DataTableImportOptionalMetadataKey))
+			{
+				continue;
+			}
+#endif // WITH_EDITOR
+
+			if (!DataTable->bIgnoreMissingFields)
+			{
+				ImportProblems.Add(FString::Printf(TEXT("Row '%s' is missing an entry for '%s'."), *InRowName.ToString(), *ColumnName));
+			}
+
 			continue;
 		}
 
