Fix an ensure() that fired on attempt to paste one or more nodes with missing links into another graph that doesn't allow the node type(s).

diff --git a/Engine/Source/Editor/UnrealEd/Private/EdGraphUtilities.cpp b/Engine/Source/Editor/UnrealEd/Private/EdGraphUtilities.cpp
--- a/Engine/Source/Editor/UnrealEd/Private/EdGraphUtilities.cpp
+++ b/Engine/Source/Editor/UnrealEd/Private/EdGraphUtilities.cpp
@@ -15,6 +15,68 @@
 #include "Widgets/Notifications/SNotificationList.h"
 #include "K2Node_Composite.h"
 
+/////////////////////////////////////////////////////
+// Local namespace
+
+namespace
+{
+	// Reconcile other pin links:
+	//   - Links between nodes within the copied set are fine
+	//   - Links to nodes that were not copied need to be fixed up if the copy-paste was in the same graph or broken completely
+	void PostProcessPastedNodePinLinks(TSet<UEdGraphNode*>& InNodes)
+	{
+		for (TSet<UEdGraphNode*>::TIterator It(InNodes); It; ++It)
+		{
+			UEdGraphNode* Node = *It;
+			UEdGraph* CurrentGraph = Node->GetGraph();
+
+			for (int32 PinIndex = 0; PinIndex < Node->Pins.Num(); ++PinIndex)
+			{
+				UEdGraphPin* ThisPin = Node->Pins[PinIndex];
+
+				// Ensure on any NULL entry, as it means there was a problem importing the pin from text, and we should be alerted to that.
+				if (ensure(ThisPin))
+				{
+					for (int32 LinkIndex = 0; LinkIndex < ThisPin->LinkedTo.Num(); )
+					{
+						UEdGraphPin* OtherPin = ThisPin->LinkedTo[LinkIndex];
+
+						if (OtherPin == nullptr)
+						{
+							// Totally bogus link
+							ThisPin->LinkedTo.RemoveAtSwap(LinkIndex);
+						}
+						else if (!InNodes.Contains(OtherPin->GetOwningNode()))
+						{
+							// It's a link across the selection set, so it should be broken
+							OtherPin->LinkedTo.RemoveSwap(ThisPin);
+							ThisPin->LinkedTo.RemoveAtSwap(LinkIndex);
+						}
+						else if (!OtherPin->LinkedTo.Contains(ThisPin))
+						{
+							// The link needs to be reciprocal
+							check(OtherPin->GetOwningNode()->GetGraph() == CurrentGraph);
+							OtherPin->LinkedTo.Add(ThisPin);
+							++LinkIndex;
+						}
+						else
+						{
+							// Everything seems fine but sanity check the graph
+							check(OtherPin->GetOwningNode()->GetGraph() == CurrentGraph);
+							++LinkIndex;
+						}
+					}
+				}
+				else
+				{
+					// Remove NULL entries; these will be replaced with a default value when the node is reconstructed below.
+					Node->Pins.RemoveAt(PinIndex--);
+				}
+			}
+		}
+	}
+}
+
 /////////////////////////////////////////////////////
 // FGraphObjectTextFactory
 
@@ -26,7 +88,7 @@ public:
 	TSet<UEdGraphNode*> SubstituteNodes;
 	const UEdGraph* DestinationGraph;
 	TSet<FName> ExtraNamesInUse;
-	TArray<UEdGraphNode*> NodesToDestroy;
+	TSet<UEdGraphNode*> NodesToDestroy;
 public:
 	FGraphObjectTextFactory(const UEdGraph* InDestinationGraph)
 		: FCustomizableTextObjectFactory(GWarn)
@@ -92,8 +154,25 @@ protected:
 	{
 		if (SubstituteNodes.Num() > 0)
 		{
-			// Display a notification to inform the user that the variable type was invalid (likely due to corruption), it should no longer appear in the list.
-			FNotificationInfo Info(NSLOCTEXT("EdGraphUtilities", "SubstituteNodesWarning", "Conflicting nodes substituted during paste!"));
+			FText NotificationText;
+			if (SubstituteNodes.Contains(nullptr))
+			{
+				if (SubstituteNodes.Num() > 1)
+				{
+					NotificationText = NSLOCTEXT("EdGraphUtilities", "SubstituteAndSkippedNodesWarning", "One or more copied nodes were substituted and/or could not be pasted into this graph!");
+				}
+				else
+				{
+					NotificationText = NSLOCTEXT("EdGraphUtilities", "SkippedNodesWarning", "One or more copied nodes could not be pasted into this graph!");
+				}
+			}
+			else
+			{
+				NotificationText = NSLOCTEXT("EdGraphUtilities", "SubstituteNodesWarning", "One or more copied nodes were substituted during paste!");
+			}
+
+			// Display a notification to inform the user that one or more nodes were substituted rather than pasted into the new graph.
+			FNotificationInfo Info(NotificationText);
 			Info.ExpireDuration = 3.0f;
 			Info.bUseLargeFont = false;
 			Info.Image = FCoreStyle::Get().GetBrush(TEXT("MessageLog.Warning"));
@@ -104,6 +183,9 @@ protected:
 			}
 		}
 
+		// Fix up pin cross-links, etc. before removing incompatible nodes below. Otherwise, BreakAllNodeLinks() will complain about non-reciprocating pin links.
+		PostProcessPastedNodePinLinks(NodesToDestroy);
+
 		for (UEdGraphNode* Node : NodesToDestroy)
 		{
 			// Move the old node into the transient package so that it is GC'd
@@ -121,62 +203,11 @@ TArray< TSharedPtr<FGraphPanelNodeFactory> > FEdGraphUtilities::VisualNodeFactor
 TArray< TSharedPtr<FGraphPanelPinFactory> > FEdGraphUtilities::VisualPinFactories;
 TArray< TSharedPtr<FGraphPanelPinConnectionFactory> > FEdGraphUtilities::VisualPinConnectionFactories;
 
-// Reconcile other pin links:
-//   - Links between nodes within the copied set are fine
-//   - Links to nodes that were not copied need to be fixed up if the copy-paste was in the same graph or broken completely
 // Call PostPasteNode on each node
 void FEdGraphUtilities::PostProcessPastedNodes(TSet<UEdGraphNode*>& SpawnedNodes)
 {
 	// Run thru and fix up the node's pin links; they may point to invalid pins if the paste was to another graph
-	for (TSet<UEdGraphNode*>::TIterator It(SpawnedNodes); It; ++It)
-	{
-		UEdGraphNode* Node = *It;
-		UEdGraph* CurrentGraph = Node->GetGraph();
-
-		for (int32 PinIndex = 0; PinIndex < Node->Pins.Num(); ++PinIndex)
-		{
-			UEdGraphPin* ThisPin = Node->Pins[PinIndex];
-
-			// Ensure on any NULL entry, as it means there was a problem importing the pin from text, and we should be alerted to that.
-			if (ensure(ThisPin))
-			{
-				for (int32 LinkIndex = 0; LinkIndex < ThisPin->LinkedTo.Num(); )
-				{
-					UEdGraphPin* OtherPin = ThisPin->LinkedTo[LinkIndex];
-
-					if (OtherPin == nullptr)
-					{
-						// Totally bogus link
-						ThisPin->LinkedTo.RemoveAtSwap(LinkIndex);
-					}
-					else if (!SpawnedNodes.Contains(OtherPin->GetOwningNode()))
-					{
-						// It's a link across the selection set, so it should be broken
-						OtherPin->LinkedTo.RemoveSwap(ThisPin);
-						ThisPin->LinkedTo.RemoveAtSwap(LinkIndex);
-					}
-					else if (!OtherPin->LinkedTo.Contains(ThisPin))
-					{
-						// The link needs to be reciprocal
-						check(OtherPin->GetOwningNode()->GetGraph() == CurrentGraph);
-						OtherPin->LinkedTo.Add(ThisPin);
-						++LinkIndex;
-					}
-					else
-					{
-						// Everything seems fine but sanity check the graph
-						check(OtherPin->GetOwningNode()->GetGraph() == CurrentGraph);
-						++LinkIndex;
-					}
-				}
-			}
-			else
-			{
-				// Remove NULL entries; these will be replaced with a default value when the node is reconstructed below.
-				Node->Pins.RemoveAt(PinIndex--);
-			}
-		}
-	}
+	PostProcessPastedNodePinLinks(SpawnedNodes);
 
 	// Give every node a chance to deep copy associated resources, etc...
 	for (TSet<UEdGraphNode*>::TIterator It(SpawnedNodes); It; ++It)
