Improve obj refs to print out detailed information about FGCObjects if they implement GetReferencerName.
Change StreamableManager to register a debug name and remove some deprecated functions.

diff --git a/Engine/Source/Editor/UnrealEd/Private/EditorEngine.cpp b/Engine/Source/Editor/UnrealEd/Private/EditorEngine.cpp
--- a/Engine/Source/Editor/UnrealEd/Private/EditorEngine.cpp
+++ b/Engine/Source/Editor/UnrealEd/Private/EditorEngine.cpp
@@ -10,7 +10,6 @@
 #include "Misc/App.h"
 #include "Modules/ModuleManager.h"
 #include "UObject/MetaData.h"
-#include "Serialization/ArchiveTraceRoute.h"
 #include "UObject/ConstructorHelpers.h"
 #include "Application/ThrottleManager.h"
 #include "Framework/Application/SlateApplication.h"
@@ -87,6 +86,7 @@
 #include "Net/NetworkProfiler.h"
 #include "Interfaces/IPluginManager.h"
 #include "UObject/PackageReload.h"
+#include "UObject/ReferenceChainSearch.h"
 #include "HAL/PlatformApplicationMisc.h"
 #include "IMediaModule.h"
 
@@ -6899,14 +6899,11 @@ void UEditorEngine::VerifyLoadMapWorldCleanup()
 
 			if (!ValidWorld)
 			{
-				// Print some debug information...
-				UE_LOG(LogLoad, Log, TEXT("%s not cleaned up by garbage collection! "), *World->GetFullName());
-				StaticExec(World, *FString::Printf(TEXT("OBJ REFS CLASS=WORLD NAME=%s"), *World->GetPathName()));
-				TMap<UObject*,UProperty*>	Route		= FArchiveTraceRoute::FindShortestRootPath( World, true, GARBAGE_COLLECTION_KEEPFLAGS );
-				FString						ErrorString	= FArchiveTraceRoute::PrintRootPath( Route, World );
-				UE_LOG(LogLoad, Log, TEXT("%s"),*ErrorString);
-				// before asserting.
-				UE_LOG(LogLoad, Fatal, TEXT("%s not cleaned up by garbage collection!") LINE_TERMINATOR TEXT("%s") , *World->GetFullName(), *ErrorString );
+				UE_LOG(LogLoad, Error, TEXT("Previously active world %s not cleaned up by garbage collection!"), *World->GetPathName());
+				UE_LOG(LogLoad, Error, TEXT("Once a world has become active, it cannot be reused and must be destroyed and reloaded. World referenced by:"));
+			
+				FReferenceChainSearch RefChainSearch(World, EReferenceChainSearchMode::Shortest | EReferenceChainSearchMode::PrintResults);
+				UE_LOG(LogLoad, Fatal, TEXT("Previously active world %s not cleaned up by garbage collection! Referenced by:") LINE_TERMINATOR TEXT("%s"), *World->GetPathName(), *RefChainSearch.GetRootPath());
 			}
 		}
 	}
diff --git a/Engine/Source/Editor/UnrealEd/Private/EditorLevelUtils.cpp b/Engine/Source/Editor/UnrealEd/Private/EditorLevelUtils.cpp
--- a/Engine/Source/Editor/UnrealEd/Private/EditorLevelUtils.cpp
+++ b/Engine/Source/Editor/UnrealEd/Private/EditorLevelUtils.cpp
@@ -10,7 +10,6 @@ EditorLevelUtils.cpp: Editor-specific level management routines
 #include "UObject/GarbageCollection.h"
 #include "UObject/Class.h"
 #include "UObject/Package.h"
-#include "Serialization/ArchiveTraceRoute.h"
 #include "Engine/EngineTypes.h"
 #include "GameFramework/Actor.h"
 #include "Engine/World.h"
@@ -23,6 +22,7 @@ EditorLevelUtils.cpp: Editor-specific level management routines
 #include "EngineGlobals.h"
 #include "UObject/UObjectHash.h"
 #include "UObject/UObjectIterator.h"
+#include "UObject/ReferenceChainSearch.h"
 #include "GameFramework/WorldSettings.h"
 #include "Engine/LevelStreaming.h"
 #include "Engine/Selection.h"
@@ -713,10 +713,8 @@ bool UEditorLevelUtils::RemoveLevelFromWorld(ULevel* InLevel)
 			UWorld* TheWorld = UWorld::FindWorldInPackage(LevelPackage->GetOutermost());
 			if (TheWorld != nullptr)
 			{
-				StaticExec(NULL, *FString::Printf(TEXT("OBJ REFS CLASS=%s NAME=%s shortest"), *TheWorld->GetClass()->GetName(), *TheWorld->GetPathName()));
-				TMap<UObject*, UProperty*>	Route = FArchiveTraceRoute::FindShortestRootPath(TheWorld, true, GARBAGE_COLLECTION_KEEPFLAGS);
-				FString						ErrorString = FArchiveTraceRoute::PrintRootPath(Route, TheWorld);
-				UE_LOG(LogStreaming, Fatal, TEXT("%s didn't get garbage collected!") LINE_TERMINATOR TEXT("%s"), *TheWorld->GetFullName(), *ErrorString);
+				FReferenceChainSearch RefChainSearch(TheWorld, EReferenceChainSearchMode::Shortest | EReferenceChainSearchMode::PrintResults);
+				UE_LOG(LogStreaming, Fatal, TEXT("Removed world %s not cleaned up by garbage collection! Referenced by:") LINE_TERMINATOR TEXT("%s"), *TheWorld->GetPathName(), *RefChainSearch.GetRootPath());
 			}
 		}
 	}
diff --git a/Engine/Source/Editor/UnrealEd/Private/EditorServer.cpp b/Engine/Source/Editor/UnrealEd/Private/EditorServer.cpp
--- a/Engine/Source/Editor/UnrealEd/Private/EditorServer.cpp
+++ b/Engine/Source/Editor/UnrealEd/Private/EditorServer.cpp
@@ -23,9 +23,9 @@
 #include "UObject/UnrealType.h"
 #include "UObject/UObjectAnnotation.h"
 #include "Serialization/ArchiveCountMem.h"
-#include "Serialization/ArchiveTraceRoute.h"
 #include "Misc/PackageName.h"
 #include "UObject/PackageFileSummary.h"
+#include "UObject/ReferenceChainSearch.h"
 #include "Widgets/DeclarativeSyntaxSupport.h"
 #include "Widgets/SWindow.h"
 #include "Framework/Application/SlateApplication.h"
@@ -1963,17 +1963,12 @@ void UEditorEngine::CheckForWorldGCLeaks( UWorld* NewWorld, UPackage* WorldPacka
 		const bool bIsPersistantWorldType = (RemainingWorld->WorldType == EWorldType::Inactive) || (RemainingWorld->WorldType == EWorldType::EditorPreview) || (RemainingWorld->WorldType == EWorldType::GamePreview);
 		if(!bIsNewWorld && !bIsPersistantWorldType && !WorldHasValidContext(RemainingWorld))
 		{
-			StaticExec(RemainingWorld, *FString::Printf(TEXT("OBJ REFS CLASS=WORLD NAME=%s"), *RemainingWorld->GetPathName()));
-
-			TMap<UObject*, UProperty*>	Route		= FArchiveTraceRoute::FindShortestRootPath(RemainingWorld, true, GARBAGE_COLLECTION_KEEPFLAGS);
-			FString						ErrorString	= FArchiveTraceRoute::PrintRootPath(Route, RemainingWorld);
-
-			UE_LOG(LogEditorServer, Fatal, TEXT("%s still around while trying to load new map") LINE_TERMINATOR TEXT("%s"), *RemainingWorld->GetPathName(), *ErrorString);
+			FReferenceChainSearch RefChainSearch(RemainingWorld, EReferenceChainSearchMode::Shortest | EReferenceChainSearchMode::PrintResults);
+			UE_LOG(LogEditorServer, Fatal, TEXT("Old world %s not cleaned up by garbage collection while loading new map! Referenced by:") LINE_TERMINATOR TEXT("%s"), *RemainingWorld->GetPathName(), *RefChainSearch.GetRootPath());
 		}
 	}
 
-
-	if(WorldPackage != NULL)
+	if(WorldPackage != nullptr)
 	{
 		UPackage* NewWorldPackage = NewWorld ? NewWorld->GetOutermost() : nullptr;
 		for(TObjectIterator<UPackage> It; It; ++It)
@@ -1982,12 +1977,8 @@ void UEditorEngine::CheckForWorldGCLeaks( UWorld* NewWorld, UPackage* WorldPacka
 			const bool bIsNewWorldPackage = (NewWorldPackage && RemainingPackage == NewWorldPackage);
 			if(!bIsNewWorldPackage && RemainingPackage == WorldPackage)
 			{
-				StaticExec(NULL, *FString::Printf(TEXT("OBJ REFS CLASS=PACKAGE NAME=%s"), *RemainingPackage->GetPathName()));
-
-				TMap<UObject*, UProperty*>	Route		= FArchiveTraceRoute::FindShortestRootPath(RemainingPackage, true, GARBAGE_COLLECTION_KEEPFLAGS);
-				FString						ErrorString	= FArchiveTraceRoute::PrintRootPath(Route, RemainingPackage);
-
-				UE_LOG(LogEditorServer, Fatal, TEXT("%s still around while trying to load new map") LINE_TERMINATOR TEXT("%s"), *RemainingPackage->GetPathName(), *ErrorString);
+				FReferenceChainSearch RefChainSearch(RemainingPackage, EReferenceChainSearchMode::Shortest | EReferenceChainSearchMode::PrintResults);
+				UE_LOG(LogEditorServer, Fatal, TEXT("Old level package %s not cleaned up by garbage collection while loading new map! Referenced by:") LINE_TERMINATOR TEXT("%s"), *RemainingPackage->GetPathName(), *RefChainSearch.GetRootPath());
 			}
 		}
 	}
@@ -2554,7 +2545,7 @@ bool UEditorEngine::Map_Load(const TCHAR* Str, FOutputDevice& Ar)
 					}
 				}
 
-				if (WorldPackage == NULL)
+				if (WorldPackage == nullptr)
 				{
 					FMessageDialog::Open( EAppMsgType::Ok, NSLOCTEXT("UnrealEd", "MapPackageLoadFailed", "Failed to open map file. This is most likely because the map was saved with a newer version of the engine."));
 					GIsEditorLoadingPackage = false;
@@ -2564,17 +2555,12 @@ bool UEditorEngine::Map_Load(const TCHAR* Str, FOutputDevice& Ar)
 				// Reset the opened package to nothing.
 				UserOpenedFile				= FString();
 
-
 				UWorld* World = UWorld::FindWorldInPackage(WorldPackage);
 				
-				if (World == NULL)
+				if (World == nullptr)
 				{
-					StaticExec(NULL, *FString::Printf(TEXT("OBJ REFS CLASS=PACKAGE NAME=%s"), *WorldPackage->GetPathName()));
-
-					TMap<UObject*,UProperty*>	Route		= FArchiveTraceRoute::FindShortestRootPath( WorldPackage, true, GARBAGE_COLLECTION_KEEPFLAGS );
-					FString						ErrorString	= FArchiveTraceRoute::PrintRootPath( Route, WorldPackage );
-
-					UE_LOG(LogEditorServer, Fatal,TEXT("Failed to find the world in %s.") LINE_TERMINATOR TEXT("%s"),*WorldPackage->GetPathName(),*TempFname,*ErrorString);
+					FReferenceChainSearch RefChainSearch(WorldPackage, EReferenceChainSearchMode::Shortest | EReferenceChainSearchMode::PrintResults);
+					UE_LOG(LogEditorServer, Fatal, TEXT("Failed to find the world in already loaded world package %s! Referenced by:") LINE_TERMINATOR TEXT("%s"), *WorldPackage->GetPathName(), *RefChainSearch.GetRootPath());
 				}
 				Context.SetCurrentWorld(World);
 				GWorld = World;
diff --git a/Engine/Source/Editor/UnrealEd/Private/PlayLevel.cpp b/Engine/Source/Editor/UnrealEd/Private/PlayLevel.cpp
--- a/Engine/Source/Editor/UnrealEd/Private/PlayLevel.cpp
+++ b/Engine/Source/Editor/UnrealEd/Private/PlayLevel.cpp
@@ -16,7 +16,7 @@
 #include "UObject/Package.h"
 #include "UObject/LazyObjectPtr.h"
 #include "UObject/SoftObjectPtr.h"
-#include "Serialization/ArchiveTraceRoute.h"
+#include "UObject/ReferenceChainSearch.h"
 #include "Misc/PackageName.h"
 #include "InputCoreTypes.h"
 #include "Layout/Margin.h"
@@ -486,11 +486,10 @@ void UEditorEngine::EndPlayMap()
 				UE_LOG(LogPlayLevel, Error, TEXT("No PIE world was found when attempting to gather references after GC."));
 			}
 
-			TMap<UObject*,UProperty*>	Route		= FArchiveTraceRoute::FindShortestRootPath( Object, true, GARBAGE_COLLECTION_KEEPFLAGS );
-			FString						ErrorString	= FArchiveTraceRoute::PrintRootPath( Route, Object );
+			FReferenceChainSearch RefChainSearch(Object, EReferenceChainSearchMode::Shortest);
 
 			FFormatNamedArguments Arguments;
-			Arguments.Add(TEXT("Path"), FText::FromString(ErrorString));
+			Arguments.Add(TEXT("Path"), FText::FromString(RefChainSearch.GetRootPath()));
 				
 			// We cannot safely recover from this.
 			FMessageLog(NAME_CategoryPIE).CriticalError()
diff --git a/Engine/Source/Runtime/CoreUObject/Private/Misc/GCObjectReferencer.cpp b/Engine/Source/Runtime/CoreUObject/Private/Misc/GCObjectReferencer.cpp
--- a/Engine/Source/Runtime/CoreUObject/Private/Misc/GCObjectReferencer.cpp
+++ b/Engine/Source/Runtime/CoreUObject/Private/Misc/GCObjectReferencer.cpp
@@ -16,8 +16,11 @@ extern bool GObjIncrementalPurgeIsInProgress;
 extern bool GObjUnhashUnreachableIsInProgress;
 
 void UGCObjectReferencer::AddReferencedObjects(UObject* InThis, FReferenceCollector& Collector)
-{	
+{
 	UGCObjectReferencer* This = CastChecked<UGCObjectReferencer>(InThis);
+
+	checkSlow(!This->bIsAddingReferencedObjects);
+	This->bIsAddingReferencedObjects = true;
 	// Note we're not locking ReferencedObjectsCritical here because we guard
 	// against adding new references during GC in AddObject and RemoveObject.
 	// Let each registered object handle its AddReferencedObjects call
@@ -27,6 +30,7 @@ void UGCObjectReferencer::AddReferencedObjects(UObject* InThis, FReferenceCollec
 		Object->AddReferencedObjects(Collector);
 	}
 	Super::AddReferencedObjects( This, Collector );
+	This->bIsAddingReferencedObjects = false;
 }
 
 void UGCObjectReferencer::AddObject(FGCObject* Object)
@@ -46,8 +50,13 @@ void UGCObjectReferencer::RemoveObject(FGCObject* Object)
 	check(NumRemoved == 1);
 }
 
-bool UGCObjectReferencer::GetReferencerName(UObject* Object, FString& OutName) const
+bool UGCObjectReferencer::GetReferencerName(UObject* Object, FString& OutName, bool bOnlyIfAddingReferenced) const
 {
+	if (bOnlyIfAddingReferenced && !bIsAddingReferencedObjects)
+	{
+		return false;
+	}
+
 	// Let each registered object handle its AddReferencedObjects call
 	for (int32 i = 0; i < ReferencedObjects.Num(); i++)
 	{
diff --git a/Engine/Source/Runtime/CoreUObject/Private/Serialization/TraceReferences.cpp b/Engine/Source/Runtime/CoreUObject/Private/Serialization/TraceReferences.cpp
--- a/Engine/Source/Runtime/CoreUObject/Private/Serialization/TraceReferences.cpp
+++ b/Engine/Source/Runtime/CoreUObject/Private/Serialization/TraceReferences.cpp
@@ -3,6 +3,8 @@
 #include "Serialization/TraceReferences.h"
 #include "UObject/UnrealType.h"
 
+PRAGMA_DISABLE_DEPRECATION_WARNINGS
+
 // This traces referenced/referencer of an object using FArchiveObjectGraph 
 FTraceReferences::FTraceReferences( bool bIncludeTransients, EObjectFlags KeepFlags )
 : ArchiveObjectGraph ( bIncludeTransients, KeepFlags )
@@ -185,3 +187,4 @@ void FTraceReferences::GetReferencedInternal( UObject* CurrentObject, TArray<FOb
 	}
 }
 
+PRAGMA_ENABLE_DEPRECATION_WARNINGS
\ No newline at end of file
diff --git a/Engine/Source/Runtime/CoreUObject/Private/UObject/Obj.cpp b/Engine/Source/Runtime/CoreUObject/Private/UObject/Obj.cpp
--- a/Engine/Source/Runtime/CoreUObject/Private/UObject/Obj.cpp
+++ b/Engine/Source/Runtime/CoreUObject/Private/UObject/Obj.cpp
@@ -38,7 +38,6 @@
 #include "Serialization/ArchiveCountMem.h"
 #include "Serialization/ArchiveShowReferences.h"
 #include "Serialization/ArchiveFindCulprit.h"
-#include "Serialization/ArchiveTraceRoute.h"
 #include "Misc/PackageName.h"
 #include "Serialization/BulkData.h"
 #include "UObject/LinkerLoad.h"
@@ -2820,16 +2819,6 @@ void UObject::OutputReferencers( FOutputDevice& Ar, FReferencerInformationList*
 
 	Ar.Logf(TEXT("\r\n") );
 
-#if !(UE_BUILD_SHIPPING || UE_BUILD_TEST)
-	Ar.Logf(TEXT("Shortest reachability from root to %s:\r\n"), *GetFullName() );
-	TMap<UObject*,UProperty*> Rt = FArchiveTraceRoute::FindShortestRootPath(this,true,GARBAGE_COLLECTION_KEEPFLAGS);
-
-	FString RootPath = FArchiveTraceRoute::PrintRootPath(Rt, this);
-	Ar.Log(*RootPath);
-
-	Ar.Logf(TEXT("\r\n") );
-#endif
-
 	if (bTempReferencers)
 	{
 		delete Referencers;
diff --git a/Engine/Source/Runtime/CoreUObject/Private/UObject/ReferenceChainSearch.cpp b/Engine/Source/Runtime/CoreUObject/Private/UObject/ReferenceChainSearch.cpp
--- a/Engine/Source/Runtime/CoreUObject/Private/UObject/ReferenceChainSearch.cpp
+++ b/Engine/Source/Runtime/CoreUObject/Private/UObject/ReferenceChainSearch.cpp
@@ -5,6 +5,7 @@
 #include "UObject/UObjectIterator.h"
 #include "UObject/UnrealType.h"
 #include "UObject/FastReferenceCollector.h"
+#include "UObject/GCObject.h"
 #include "HAL/ThreadHeartBeat.h"
 
 DEFINE_LOG_CATEGORY_STATIC(LogReferenceChain, Log, All);
@@ -278,6 +279,29 @@ void FReferenceChainSearch::DumpChain(FReferenceChain* Chain)
 	}
 }
 
+void FReferenceChainSearch::WriteChain(FReferenceChain* Chain, FString& OutString)
+{
+	if (Chain->Num())
+	{
+		// Roots are at the end so iterate from the last to the first node
+		for (int32 NodeIndex = Chain->Num() - 1; NodeIndex >= 0; --NodeIndex)
+		{
+			UObject* Object = Chain->GetNode(NodeIndex)->Object;
+			const FNodeReferenceInfo& ReferenceInfo = Chain->GetReferenceInfo(NodeIndex);
+
+			if (NodeIndex != Chain->Num() - 1)
+			{
+				OutString.Append(LINE_TERMINATOR);
+				OutString.Append(FCString::Spc(Chain->Num() - NodeIndex - 1));
+			}
+			
+			OutString.Append(*GetObjectFlags(Object));
+			OutString.Append(*Object->GetFullName());
+			OutString.Append(*ReferenceInfo.ToString());
+		}
+	}
+}
+
 void FReferenceChainSearch::FReferenceChain::FillReferenceInfo()
 {
 	// The first entry is the object we were looking for references to so add an empty entry for it
@@ -349,7 +373,17 @@ public:
 			else
 			{
 				RefInfo.Type = FReferenceChainSearch::EReferenceType::AddReferencedObjects;
-				if (ReferencingObject)
+
+				if (FGCObject::GGCObjectReferencer && (!ReferencingObject || ReferencingObject == FGCObject::GGCObjectReferencer))
+				{
+					// If we have no referencing object, check the global referencer
+					FString ReferencerString;
+					if (FGCObject::GGCObjectReferencer->GetReferencerName(Object, ReferencerString, true))
+					{
+						RefInfo.ReferencerName = FName(*ReferencerString);
+					}
+				}
+				else if (ReferencingObject)
 				{
 					RefInfo.ReferencerName = ReferencingObject->GetFName();
 				}
@@ -433,7 +467,7 @@ void FReferenceChainSearch::FindDirectReferencesForObjects()
 	}
 }
 
-void FReferenceChainSearch::PrintResults(bool bDumpAllChains /*= false*/)
+void FReferenceChainSearch::PrintResults(bool bDumpAllChains /*= false*/) const
 {
 	if (ReferenceChains.Num())
 	{
@@ -466,6 +500,25 @@ void FReferenceChainSearch::PrintResults(bool bDumpAllChains /*= false*/)
 	}
 }
 
+FString FReferenceChainSearch::GetRootPath() const
+{
+	if (ReferenceChains.Num())
+	{
+		FReferenceChain* Chain = ReferenceChains[0];
+		FString OutString;
+
+		WriteChain(Chain, OutString);
+		return OutString;
+	}
+	else
+	{
+		return FString::Printf(TEXT("%s%s is not currently reachable."),
+			*GetObjectFlags(ObjectToFindReferencesTo),
+			*ObjectToFindReferencesTo->GetFullName()
+		);
+	}
+}
+
 void FReferenceChainSearch::Cleanup()
 {
 	for (int32 ChainIndex = 0; ChainIndex < ReferenceChains.Num(); ++ChainIndex)
diff --git a/Engine/Source/Runtime/CoreUObject/Public/CoreUObject.h b/Engine/Source/Runtime/CoreUObject/Public/CoreUObject.h
--- a/Engine/Source/Runtime/CoreUObject/Public/CoreUObject.h
+++ b/Engine/Source/Runtime/CoreUObject/Public/CoreUObject.h
@@ -62,9 +62,6 @@ MONOLITHIC_HEADER_BOILERPLATE()
 #include "Serialization/FindReferencersArchive.h"
 #include "Serialization/FindObjectReferencers.h"
 #include "Serialization/ArchiveFindCulprit.h"
-#include "Serialization/ArchiveObjectGraph.h"
-#include "Serialization/TraceReferences.h"
-#include "Serialization/ArchiveTraceRoute.h"
 #include "Serialization/DuplicatedObject.h"
 #include "Serialization/DuplicatedDataReader.h"
 #include "Serialization/DuplicatedDataWriter.h"
diff --git a/Engine/Source/Runtime/CoreUObject/Public/Serialization/ArchiveObjectGraph.h b/Engine/Source/Runtime/CoreUObject/Public/Serialization/ArchiveObjectGraph.h
--- a/Engine/Source/Runtime/CoreUObject/Public/Serialization/ArchiveObjectGraph.h
+++ b/Engine/Source/Runtime/CoreUObject/Public/Serialization/ArchiveObjectGraph.h
@@ -87,6 +87,7 @@ class FArchiveObjectGraph : public FArchiveUObject
 	EObjectFlags	RequiredFlags;
 
 public:
+	UE_DEPRECATED(4.23, "This class is out of date and misses many references, replace with FReferenceChainSearch or FFindReferencersArchive")
 	FArchiveObjectGraph(bool IncludeTransients, EObjectFlags KeepFlags);
 	~FArchiveObjectGraph();
 
diff --git a/Engine/Source/Runtime/CoreUObject/Public/Serialization/ArchiveTraceRoute.h b/Engine/Source/Runtime/CoreUObject/Public/Serialization/ArchiveTraceRoute.h
--- a/Engine/Source/Runtime/CoreUObject/Public/Serialization/ArchiveTraceRoute.h
+++ b/Engine/Source/Runtime/CoreUObject/Public/Serialization/ArchiveTraceRoute.h
@@ -6,10 +6,7 @@
 #include "Serialization/ArchiveUObject.h"
 #include "Serialization/ArchiveObjectGraph.h"
 
-/**
- * Archive for finding shortest path from root to a particular object.
- * Depth-first search.
- */
+// DEPRECATED: This class is out of date and misses many GC references, replace with FReferenceChainSearch
 class FArchiveTraceRoute : public FArchiveUObject
 {
 	/**
@@ -36,6 +33,8 @@ class FArchiveTraceRoute : public FArchiveUObject
 	};
 
 public:
+
+	UE_DEPRECATED(4.23, "This function is out of date and misses many GC references, replace with FReferenceChainSearch")
 	static COREUOBJECT_API TMap<UObject*,UProperty*> FindShortestRootPath( UObject* Object, bool bIncludeTransients, EObjectFlags KeepFlags );
 
 	/**
diff --git a/Engine/Source/Runtime/CoreUObject/Public/Serialization/TraceReferences.h b/Engine/Source/Runtime/CoreUObject/Public/Serialization/TraceReferences.h
--- a/Engine/Source/Runtime/CoreUObject/Public/Serialization/TraceReferences.h
+++ b/Engine/Source/Runtime/CoreUObject/Public/Serialization/TraceReferences.h
@@ -16,6 +16,7 @@ class FTraceReferences
 	void GetReferencedInternal( UObject* CurrentObject, TArray<FObjectGraphNode*> &OutReferenced, int32 CurrentDepth, int32 TargetDepth );
 
 public:
+	UE_DEPRECATED(4.23, "This class is out of date and misses many references, replace with FReferenceChainSearch or FFindReferencersArchive")
 	FTraceReferences( bool bIncludeTransients = false, EObjectFlags KeepFlags = RF_AllFlags );
 
 	// returns referencer string of an object
diff --git a/Engine/Source/Runtime/CoreUObject/Public/UObject/GCObject.h b/Engine/Source/Runtime/CoreUObject/Public/UObject/GCObject.h
--- a/Engine/Source/Runtime/CoreUObject/Public/UObject/GCObject.h
+++ b/Engine/Source/Runtime/CoreUObject/Public/UObject/GCObject.h
@@ -29,6 +29,8 @@ class COREUOBJECT_API UGCObjectReferencer : public UObject
 	TArray<FGCObject*> ReferencedObjects;
 	/** Critical section used when adding and removing objects */
 	FCriticalSection ReferencedObjectsCritical;
+	/** True if we are currently inside AddReferencedObjects */
+	bool bIsAddingReferencedObjects = false;
 
 public:
 	DECLARE_CASTED_CLASS_INTRINSIC_WITH_API(UGCObjectReferencer, UObject, CLASS_Transient, TEXT("/Script/CoreUObject"), CASTCLASS_None, NO_API);
@@ -52,9 +54,10 @@ public:
 	 *
 	 * @param Object The object that we're looking for.
 	 * @param OutName the name of the FGCObject that reports this object.
+	 * @param bOnlyIfAddingReferenced Only try to find the name if we are currently inside AddReferencedObjects
 	 * @return true if the object was found.
 	 */
-	bool GetReferencerName(UObject* Object, FString& OutName) const;
+	bool GetReferencerName(UObject* Object, FString& OutName, bool bOnlyIfAddingReferenced = false) const;
 
 	/**
 	 * Forwards this call to all registered objects so they can reference
diff --git a/Engine/Source/Runtime/CoreUObject/Public/UObject/ReferenceChainSearch.h b/Engine/Source/Runtime/CoreUObject/Public/UObject/ReferenceChainSearch.h
--- a/Engine/Source/Runtime/CoreUObject/Public/UObject/ReferenceChainSearch.h
+++ b/Engine/Source/Runtime/CoreUObject/Public/UObject/ReferenceChainSearch.h
@@ -91,7 +91,14 @@ public:
 			}
 			else if (Type == EReferenceType::AddReferencedObjects)
 			{
-				return TEXT("::AddReferencedObjects()");
+				if (!ReferencerName.IsNone())
+				{
+					return FString::Printf(TEXT("::AddReferencedObjects(): %s"), *ReferencerName.ToString());
+				}
+				else
+				{
+					return TEXT("::AddReferencedObjects()");
+				}
 			}
 			return FString();
 		}
@@ -183,19 +190,22 @@ public:
 	};
 
 	/** Constructs a new search engine and finds references to the specified object */
-	FReferenceChainSearch(UObject* InObjectToFindReferencesTo, EReferenceChainSearchMode Mode = EReferenceChainSearchMode::PrintResults);
+	COREUOBJECT_API FReferenceChainSearch(UObject* InObjectToFindReferencesTo, EReferenceChainSearchMode Mode = EReferenceChainSearchMode::PrintResults);
 
 	/** Destructor */
-	~FReferenceChainSearch();
+	COREUOBJECT_API ~FReferenceChainSearch();
 
 	/** 
 	 * Dumps results to log 
 	 * @param bDumpAllChains - if set to false, the output will be trimmed to the first 100 reference chains
 	 */
-	void PrintResults(bool bDumpAllChains = false);
+	COREUOBJECT_API void PrintResults(bool bDumpAllChains = false) const;
+
+	/** Returns a string with a short report explaining the root path, will contain newlines */
+	COREUOBJECT_API FString GetRootPath() const;
 
 	/** Returns all reference chains */
-	const TArray<FReferenceChain*>& GetReferenceChains() const
+	COREUOBJECT_API const TArray<FReferenceChain*>& GetReferenceChains() const
 	{
 		return ReferenceChains;
 	}
@@ -235,5 +245,7 @@ private:
 	static FString GetObjectFlags(UObject* InObject);
 	/** Dumps a reference chain to log */
 	static void DumpChain(FReferenceChain* Chain);
+	/** Writes a reference chain to a string */
+	static void WriteChain(FReferenceChain* Chain, FString& OutString);
 };
 
diff --git a/Engine/Source/Runtime/Engine/Classes/Engine/StreamableManager.h b/Engine/Source/Runtime/Engine/Classes/Engine/StreamableManager.h
--- a/Engine/Source/Runtime/Engine/Classes/Engine/StreamableManager.h
+++ b/Engine/Source/Runtime/Engine/Classes/Engine/StreamableManager.h
@@ -313,27 +313,18 @@ struct ENGINE_API FStreamableManager : public FGCObject
 	/** Checks for any redirectors that were previously loaded, and returns the redirected target if found. This will not handle redirects that it doesn't yet know about */
 	FSoftObjectPath ResolveRedirects(const FSoftObjectPath& Target) const;
 
-	UE_DEPRECATED(4.16, "Call LoadSynchronous with bManageActiveHandle=true instead if you want the manager to keep the handle alive")
-	UObject* SynchronousLoad(FSoftObjectPath const& Target);
+	/** Returns the debug name for this manager */
+	const FString& GetManagerName() const;
 
-	template< typename T >
-	UE_DEPRECATED(4.16, "Call LoadSynchronous with bManageActiveHandle=true instead if you want the manager to keep the handle alive")
-	T* SynchronousLoadType(FSoftObjectPath const& Target)
-	{
-		PRAGMA_DISABLE_DEPRECATION_WARNINGS
-		return Cast< T >(SynchronousLoad(Target));
-		PRAGMA_ENABLE_DEPRECATION_WARNINGS
-	}
-
-	UE_DEPRECATED(4.16, "Call RequestAsyncLoad with bManageActiveHandle=true instead if you want the manager to keep the handle alive")
-	void SimpleAsyncLoad(const FSoftObjectPath& Target, TAsyncLoadPriority Priority = DefaultAsyncLoadPriority);
-
-	UE_DEPRECATED(4.16, "AddStructReferencedObjects is no longer necessary, as it is a GCObject now")
-	void AddStructReferencedObjects(class FReferenceCollector& Collector) const {}
+	/** Modifies the debug name of this manager, used for debugging GC references */
+	void SetManagerName(const FString& InName);
 
 	/** Add referenced objects to stop them from GCing */
 	virtual void AddReferencedObjects(FReferenceCollector& Collector) override;
 
+	/** Returns debug name */
+	virtual FString GetReferencerName() const override;
+
 	FStreamableManager();
 	~FStreamableManager();
 
@@ -382,6 +373,9 @@ private:
 
 	/** If True, temporarily force synchronous loading */
 	bool bForceSynchronousLoads;
+
+	/** Debug name of this manager */
+	FString ManagerName;
 };
 
 
diff --git a/Engine/Source/Runtime/Engine/Private/AssetManager.cpp b/Engine/Source/Runtime/Engine/Private/AssetManager.cpp
--- a/Engine/Source/Runtime/Engine/Private/AssetManager.cpp
+++ b/Engine/Source/Runtime/Engine/Private/AssetManager.cpp
@@ -165,6 +165,8 @@ void UAssetManager::PostInitProperties()
 		}
 
 		LoadRedirectorMaps();
+
+		StreamableManager.SetManagerName(FString::Printf(TEXT("%s.StreamableManager"), *GetPathName()));
 	}
 }
 
diff --git a/Engine/Source/Runtime/Engine/Private/LevelStreaming.cpp b/Engine/Source/Runtime/Engine/Private/LevelStreaming.cpp
--- a/Engine/Source/Runtime/Engine/Private/LevelStreaming.cpp
+++ b/Engine/Source/Runtime/Engine/Private/LevelStreaming.cpp
@@ -4,7 +4,7 @@
 #include "ContentStreaming.h"
 #include "Misc/App.h"
 #include "UObject/Package.h"
-#include "Serialization/ArchiveTraceRoute.h"
+#include "UObject/ReferenceChainSearch.h"
 #include "Misc/PackageName.h"
 #include "UObject/LinkerLoad.h"
 #include "EngineGlobals.h"
@@ -957,16 +957,14 @@ bool ULevelStreaming::RequestLevel(UWorld* PersistentWorld, bool bAllowLevelLoad
 			}
 
 #if !(UE_BUILD_SHIPPING || UE_BUILD_TEST)
-			if (World->PersistentLevel == NULL)
+			if (World->PersistentLevel == nullptr)
 			{
-				UE_LOG(LogLevelStreaming, Log, TEXT("World exists but PersistentLevel doesn't for %s, most likely caused by reference to world of unloaded level and GC setting reference to NULL while keeping world object"), *World->GetOutermost()->GetName());
-				// print out some debug information...
-				StaticExec(World, *FString::Printf(TEXT("OBJ REFS CLASS=WORLD NAME=%s shortest"), *World->GetPathName()));
-				TMap<UObject*,UProperty*> Route = FArchiveTraceRoute::FindShortestRootPath( World, true, GARBAGE_COLLECTION_KEEPFLAGS );
-				FString ErrorString = FArchiveTraceRoute::PrintRootPath( Route, World );
-				UE_LOG(LogLevelStreaming, Log, TEXT("%s"), *ErrorString);
-				// before asserting
-				checkf(World->PersistentLevel,TEXT("Most likely caused by reference to world of unloaded level and GC setting reference to NULL while keeping world object"));
+				UE_LOG(LogLevelStreaming, Error, TEXT("World exists but PersistentLevel doesn't for %s, most likely caused by reference to world of unloaded level and GC setting reference to null while keeping world object"), *World->GetOutermost()->GetName());
+				UE_LOG(LogLevelStreaming, Error, TEXT("Most likely caused by reference to world of unloaded level and GC setting reference to null while keeping world object. Referenced by:"));
+
+				FReferenceChainSearch RefChainSearch(World, EReferenceChainSearchMode::Shortest | EReferenceChainSearchMode::PrintResults);
+				UE_LOG(LogLoad, Fatal, TEXT("World exists but PersistentLevel doesn't for %s! Referenced by:") LINE_TERMINATOR TEXT("%s"), *World->GetPathName(), *RefChainSearch.GetRootPath());
+
 				return false;
 			}
 #endif
diff --git a/Engine/Source/Runtime/Engine/Private/StreamableManager.cpp b/Engine/Source/Runtime/Engine/Private/StreamableManager.cpp
--- a/Engine/Source/Runtime/Engine/Private/StreamableManager.cpp
+++ b/Engine/Source/Runtime/Engine/Private/StreamableManager.cpp
@@ -721,6 +721,7 @@ FStreamableManager::FStreamableManager()
 {
 	FCoreUObjectDelegates::GetPreGarbageCollectDelegate().AddRaw(this, &FStreamableManager::OnPreGarbageCollect);
 	bForceSynchronousLoads = false;
+	ManagerName = TEXT("StreamableManager");
 }
 
 FStreamableManager::~FStreamableManager()
@@ -796,6 +797,21 @@ void FStreamableManager::AddReferencedObjects(FReferenceCollector& Collector)
 	}
 }
 
+FString FStreamableManager::GetReferencerName() const
+{
+	return ManagerName;
+}
+
+const FString& FStreamableManager::GetManagerName() const
+{
+	return ManagerName;
+}
+
+void FStreamableManager::SetManagerName(const FString& InName)
+{
+	ManagerName = InName;
+}
+
 FStreamable* FStreamableManager::FindStreamable(const FSoftObjectPath& Target) const
 {
 	FStreamable* Existing = StreamableItems.FindRef(Target);
@@ -891,16 +907,6 @@ FStreamable* FStreamableManager::StreamInternal(const FSoftObjectPath& InTargetN
 	return Existing;
 }
 
-void FStreamableManager::SimpleAsyncLoad(const FSoftObjectPath& Target, TAsyncLoadPriority Priority)
-{
-	RequestAsyncLoad(Target, FStreamableDelegate(), Priority, true);
-}
-
-UObject* FStreamableManager::SynchronousLoad(FSoftObjectPath const& Target)
-{
-	return LoadSynchronous(Target, true);
-}
-
 TSharedPtr<FStreamableHandle> FStreamableManager::RequestAsyncLoad(const TArray<FSoftObjectPath>& TargetsToStream, FStreamableDelegate DelegateToCall, TAsyncLoadPriority Priority, bool bManageActiveHandle, bool bStartStalled, const FString& DebugName)
 {
 	// Schedule a new callback, this will get called when all related async loads are completed
diff --git a/Engine/Source/Runtime/Engine/Private/UnrealEngine.cpp b/Engine/Source/Runtime/Engine/Private/UnrealEngine.cpp
--- a/Engine/Source/Runtime/Engine/Private/UnrealEngine.cpp
+++ b/Engine/Source/Runtime/Engine/Private/UnrealEngine.cpp
@@ -44,7 +44,6 @@ UnrealEngine.cpp: Implements the UEngine class and helpers.
 #include "Serialization/ArchiveCountMem.h"
 #include "Serialization/ObjectWriter.h"
 #include "Serialization/ObjectReader.h"
-#include "Serialization/ArchiveTraceRoute.h"
 #include "Misc/PackageName.h"
 #include "Misc/EngineVersion.h"
 #include "UObject/LinkerLoad.h"
@@ -223,6 +222,7 @@ UnrealEngine.cpp: Implements the UEngine class and helpers.
 
 #include "HAL/FileManagerGeneric.h"
 #include "UObject/UObjectThreadContext.h"
+#include "UObject/ReferenceChainSearch.h"
 
 #include "Particles/ParticleSystemManager.h"
 #include "Components/SkinnedMeshComponent.h"
@@ -13135,15 +13135,11 @@ void UEngine::VerifyLoadMapWorldCleanup()
 		{
 			if ((World->PersistentLevel == nullptr || !WorldHasValidContext(World->PersistentLevel->OwningWorld)) && !IsWorldDuplicate(World))
 			{
-				// Print some debug information...
-				UE_LOG(LogLoad, Log, TEXT("%s not cleaned up by garbage collection! "), *World->GetFullName());
-				StaticExec(World, *FString::Printf(TEXT("OBJ REFS CLASS=WORLD NAME=%s"), *World->GetPathName()));
-				TMap<UObject*,UProperty*>	Route		= FArchiveTraceRoute::FindShortestRootPath( World, true, GARBAGE_COLLECTION_KEEPFLAGS );
-				FString						ErrorString	= FArchiveTraceRoute::PrintRootPath( Route, World );
-				UE_LOG(LogLoad, Log, TEXT("%s"),*ErrorString);
-				// before asserting.
+				UE_LOG(LogLoad, Error, TEXT("Previously active world %s not cleaned up by garbage collection!"), *World->GetPathName());
+				UE_LOG(LogLoad, Error, TEXT("Once a world has become active, it cannot be reused and must be destroyed and reloaded. World referenced by:"));
 
-				UE_LOG(LogLoad, Fatal, TEXT("%s not cleaned up by garbage collection!") LINE_TERMINATOR TEXT("%s") , *World->GetFullName(), *ErrorString );
+				FReferenceChainSearch RefChainSearch(World, EReferenceChainSearchMode::Shortest | EReferenceChainSearchMode::PrintResults);
+				UE_LOG(LogLoad, Fatal, TEXT("Previously active world %s not cleaned up by garbage collection! Referenced by:") LINE_TERMINATOR TEXT("%s"), *World->GetPathName(), *RefChainSearch.GetRootPath());
 			}
 		}
 	}
@@ -13318,6 +13314,11 @@ public:
 	{
 		Collector.AddReferencedObjects( Levels ); 
 	}
+
+	virtual FString GetReferencerName() const override
+	{
+		return TEXT("FPendingStreamingLevelHolder");
+	}
 };
 
 
diff --git a/Engine/Source/Runtime/Engine/Private/World.cpp b/Engine/Source/Runtime/Engine/Private/World.cpp
--- a/Engine/Source/Runtime/Engine/Private/World.cpp
+++ b/Engine/Source/Runtime/Engine/Private/World.cpp
@@ -17,7 +17,7 @@
 #include "UObject/Package.h"
 #include "UObject/ObjectRedirector.h"
 #include "UObject/UObjectAnnotation.h"
-#include "Serialization/ArchiveTraceRoute.h"
+#include "UObject/ReferenceChainSearch.h"
 #include "Misc/PackageName.h"
 #include "Serialization/AsyncLoading.h"
 #include "GameMapsSettings.h"
@@ -2632,12 +2632,8 @@ void FLevelStreamingGCHelper::VerifyLevelsGotRemovedByGC()
 			// But disregard package object itself.
 			&&	!Object->IsA(UPackage::StaticClass()) )
 			{
-				UE_LOG(LogWorld, Log, TEXT("%s didn't get garbage collected! Trying to find culprit, though this might crash. Try increasing stack size if it does."), *Object->GetFullName());
-				StaticExec(NULL, *FString::Printf(TEXT("OBJ REFS CLASS=%s NAME=%s shortest"),*Object->GetClass()->GetName(), *Object->GetPathName()));
-				TMap<UObject*,UProperty*>	Route		= FArchiveTraceRoute::FindShortestRootPath( Object, true, GARBAGE_COLLECTION_KEEPFLAGS );
-				FString						ErrorString	= FArchiveTraceRoute::PrintRootPath( Route, Object );
-				// Print out error message. We don't assert here as there might be multiple culprits.
-				UE_LOG(LogWorld, Warning, TEXT("%s didn't get garbage collected!") LINE_TERMINATOR TEXT("%s"), *Object->GetFullName(), *ErrorString );
+				UE_LOG(LogWorld, Warning, TEXT("Level object %s didn't get garbage collected! Trying to find culprit, though this might crash. Try increasing stack size if it does. Referenced by:"), *Object->GetFullName());
+				FReferenceChainSearch RefChainSearch(Object, EReferenceChainSearchMode::Shortest | EReferenceChainSearchMode::PrintResults);
 				FailCount++;
 			}
 		}
