[NO NON-PII COMMIT MESSAGE]

diff --git a/Engine/Source/Editor/UnrealEd/Private/DataTableEditorUtils.cpp b/Engine/Source/Editor/UnrealEd/Private/DataTableEditorUtils.cpp
--- a/Engine/Source/Editor/UnrealEd/Private/DataTableEditorUtils.cpp
+++ b/Engine/Source/Editor/UnrealEd/Private/DataTableEditorUtils.cpp
@@ -10,11 +10,436 @@
 #include "Engine/UserDefinedStruct.h"
 #include "ScopedTransaction.h"
 #include "K2Node_GetDataTableRow.h"
+#include "Input/Reply.h"
+#include "Widgets/Input/SComboButton.h"
+#include "Widgets/Input/SEditableTextBox.h"
+#include "Widgets/Views/SListView.h"
+#include "Widgets/Input/SComboBox.h"
+#include "AssetRegistryModule.h"
 
 #define LOCTEXT_NAMESPACE "DataTableEditorUtils"
 
+/** Combobox that allows selecting a struct row for a data table. Based off of SSearchableComboBox */
+class SDataTableStructComboBox : public SComboButton
+{
+public:
+	/** Type of list used for showing menu options. */
+	typedef SListView< TSharedPtr<FString> > SComboListType;
+	/** Delegate type used to generate widgets that represent Options */
+	typedef typename TSlateDelegates< TSharedPtr<FString> >::FOnGenerateWidget FOnGenerateWidget;
+	typedef typename TSlateDelegates< TSharedPtr<FString> >::FOnSelectionChanged FOnSelectionChanged;
+	DECLARE_DELEGATE_OneParam(FOnFillComboBoxStrings, TArray<TSharedPtr<FString>>&);
+
+	SLATE_BEGIN_ARGS(SDataTableStructComboBox)
+		: _Content()
+		, _ComboBoxStyle(&FCoreStyle::Get().GetWidgetStyle< FComboBoxStyle >("ComboBox"))
+		, _ButtonStyle(nullptr)
+		, _ItemStyle(&FCoreStyle::Get().GetWidgetStyle< FTableRowStyle >("TableView.Row"))
+		, _ContentPadding(FMargin(4.0, 2.0))
+		, _ForegroundColor(FCoreStyle::Get().GetSlateColor("InvertedForeground"))
+		, _OnStructSelected()
+		, _InitiallySelectedItem(nullptr)
+		, _Method()
+		, _MaxListHeight(450.0f)
+		, _HasDownArrow(true)
+	{}
+
+	/** Slot for this button's content (optional) */
+	SLATE_DEFAULT_SLOT(FArguments, Content)
+
+	SLATE_STYLE_ARGUMENT(FComboBoxStyle, ComboBoxStyle)
+
+	/** The visual style of the button part of the combo box (overrides ComboBoxStyle) */
+	SLATE_STYLE_ARGUMENT(FButtonStyle, ButtonStyle)
+
+	SLATE_STYLE_ARGUMENT(FTableRowStyle, ItemStyle)
+
+	SLATE_ATTRIBUTE(FMargin, ContentPadding)
+	SLATE_ATTRIBUTE(FSlateColor, ForegroundColor)
+
+	SLATE_EVENT(FDataTableEditorUtils::FOnDataTableStructSelected, OnStructSelected)
+	
+	/** The custom scrollbar to use in the ListView */
+	SLATE_ARGUMENT(TSharedPtr<SScrollBar>, CustomScrollbar)
+
+	/** The option that should be selected when the combo box is first created */
+	SLATE_ARGUMENT(TSharedPtr<FString>, InitiallySelectedItem)
+
+	SLATE_ARGUMENT(TOptional<EPopupMethod>, Method)
+
+	/** The max height of the combo box menu */
+	SLATE_ARGUMENT(float, MaxListHeight)
+
+	/**
+	 * When false, the down arrow is not generated and it is up to the API consumer
+	 * to make their own visual hint that this is a drop down.
+	 */
+	SLATE_ARGUMENT(bool, HasDownArrow)
+
+	SLATE_END_ARGS()
+
+	/**
+	 * Construct the widget from a declaration
+	 *
+	 * @param InArgs   Declaration from which to construct the combo box
+	 */
+	void Construct(const FArguments& InArgs);
+
+	void ClearSelection();
+
+	void SetSelectedItem(TSharedPtr<FString> InSelectedItem);
+
+	/** @return the item currently selected by the combo box. */
+	TSharedPtr<FString> GetSelectedItem();
+
+	/**
+	 * Requests a list refresh after updating options
+	 * Call SetSelectedItem to update the selected item if required
+	 * @see SetSelectedItem
+	 */
+	void RefreshOptions();
+
+	/** Returns the asset data for a specific string, or null if not found */
+	const FAssetData* FindAssetDataForString(TSharedPtr<FString> StringOption) const;
+
+	/** Returns struct from AssetData, possibly loading it */
+	UScriptStruct* GetOrLoadStruct(const FAssetData* AssetData);
+
+private:
+
+	/** Generate a row for the InItem in the combo box's list (passed in as OwnerTable). Do this by calling the user-specified OnGenerateWidget */
+	TSharedRef<ITableRow> GenerateMenuItemRow(TSharedPtr<FString> InItem, const TSharedRef<STableViewBase>& OwnerTable);
+
+	/** Called if the menu is closed */
+	void OnMenuOpenChanged(bool bOpen);
+
+	/** Invoked when the selection in the list changes */
+	void OnSelectionChanged_Internal(TSharedPtr<FString> ProposedSelection, ESelectInfo::Type SelectInfo);
+
+	/** Invoked when the search text changes */
+	void OnSearchTextChanged(const FText& ChangedText);
+
+	/** Text to display inside box */
+	FText GetSelectedText() const;
+
+	/** Show tooltip text for a specific option */
+	FText GetTooltipText(TSharedPtr<FString> StringOption);
+
+	/** Handle clicking on the content menu */
+	virtual FReply OnButtonClicked() override;
+
+	/** The item style to use. */
+	const FTableRowStyle* ItemStyle;
+
+private:
+	/** Delegate that is invoked when the selected item in the combo box changes */
+	FDataTableEditorUtils::FOnDataTableStructSelected OnStructSelected;
+
+	/** The item currently selected in the combo box */
+	TSharedPtr<FString> SelectedItem;
+	/** The search field used for the combox box's contents */
+	TSharedPtr< SEditableTextBox > SearchField;
+	/** The ListView that we pop up; visualized the available options. */
+	TSharedPtr< SComboListType > ComboListView;
+	/** The Scrollbar used in the ListView. */
+	TSharedPtr< SScrollBar > CustomScrollbar;
+
+	/** List of names to show in combo box, there is a 1:1 mapping to PossibleStructs */
+	TArray< TSharedPtr<FString> > CurrentOptions;
+	/** List of AssetData representing rows */
+	TArray<FAssetData> PossibleStructs;
+};
+
+void SDataTableStructComboBox::Construct(const FArguments& InArgs)
+{
+	check(InArgs._ComboBoxStyle);
+
+	ItemStyle = InArgs._ItemStyle;
+
+	// Work out which values we should use based on whether we were given an override, or should use the style's version
+	const FComboButtonStyle& OurComboButtonStyle = InArgs._ComboBoxStyle->ComboButtonStyle;
+	const FButtonStyle* const OurButtonStyle = InArgs._ButtonStyle ? InArgs._ButtonStyle : &OurComboButtonStyle.ButtonStyle;
+
+	this->OnStructSelected = InArgs._OnStructSelected;
+
+	CustomScrollbar = InArgs._CustomScrollbar;
+
+	TSharedRef<SWidget> ComboBoxMenuContent =
+		SNew(SBox)
+		.MaxDesiredHeight(InArgs._MaxListHeight)
+		[
+			SNew(SVerticalBox)
+			+ SVerticalBox::Slot()
+			.AutoHeight()
+			[
+				SAssignNew(this->SearchField, SEditableTextBox)
+				.HintText(LOCTEXT("Search", "Search"))
+				.OnTextChanged(this, &SDataTableStructComboBox::OnSearchTextChanged)
+			]
+
+			+ SVerticalBox::Slot()
+			[
+				SAssignNew(this->ComboListView, SComboListType)
+				.ListItemsSource(&CurrentOptions)
+				.OnGenerateRow(this, &SDataTableStructComboBox::GenerateMenuItemRow)
+				.OnSelectionChanged(this, &SDataTableStructComboBox::OnSelectionChanged_Internal)
+				.SelectionMode(ESelectionMode::Single)
+				.ExternalScrollbar(InArgs._CustomScrollbar)
+			]
+		];
+
+	// Set up content
+	TSharedPtr<SWidget> ButtonContent = InArgs._Content.Widget;
+	if (InArgs._Content.Widget == SNullWidget::NullWidget)
+	{
+		SAssignNew(ButtonContent, STextBlock)
+			.Text(this, &SDataTableStructComboBox::GetSelectedText);
+	}
+
+	SComboButton::Construct(SComboButton::FArguments()
+		.ComboButtonStyle(&OurComboButtonStyle)
+		.ButtonStyle(OurButtonStyle)
+		.Method(InArgs._Method)
+		.ButtonContent()
+		[
+			ButtonContent.ToSharedRef()
+		]
+		.MenuContent()
+		[
+			ComboBoxMenuContent
+		]
+		.HasDownArrow(InArgs._HasDownArrow)
+		.ContentPadding(InArgs._ContentPadding)
+		.ForegroundColor(InArgs._ForegroundColor)
+		.OnMenuOpenChanged(this, &SDataTableStructComboBox::OnMenuOpenChanged)
+		.IsFocusable(true)
+		);
+
+	// Better to select search field so you can type right away
+	SetMenuContentWidgetToFocus(SearchField);
+
+	// Refresh options now
+	RefreshOptions();
+
+	// Need to establish the selected item at point of construction so its available for querying
+	// NB: If you need a selection to fire use SetItemSelection rather than setting an IntiallySelectedItem
+	SelectedItem = InArgs._InitiallySelectedItem;
+	if (TListTypeTraits<TSharedPtr<FString>>::IsPtrValid(SelectedItem))
+	{
+		ComboListView->Private_SetItemSelection(SelectedItem, true);
+	}
+}
+
+void SDataTableStructComboBox::ClearSelection()
+{
+	ComboListView->ClearSelection();
+}
+
+void SDataTableStructComboBox::SetSelectedItem(TSharedPtr<FString> InSelectedItem)
+{
+	if (TListTypeTraits<TSharedPtr<FString>>::IsPtrValid(InSelectedItem))
+	{
+		ComboListView->SetSelection(InSelectedItem);
+	}
+	else
+	{
+		ComboListView->ClearSelection();
+	}
+}
+
+TSharedPtr<FString> SDataTableStructComboBox::GetSelectedItem()
+{
+	return SelectedItem;
+}
+
+FText SDataTableStructComboBox::GetSelectedText() const
+{
+	if (SelectedItem.IsValid())
+	{
+		return FText::FromString(*SelectedItem);
+	}
+
+	return FText::GetEmpty();
+}
+
+FText SDataTableStructComboBox::GetTooltipText(TSharedPtr<FString> StringOption)
+{
+	const FAssetData* FoundAsset = FindAssetDataForString(StringOption);
+
+	if (FoundAsset)
+	{
+		return FText::FromString(FoundAsset->PackageName.ToString());
+	}
+	return FText::GetEmpty();
+}
+
+void SDataTableStructComboBox::RefreshOptions()
+{
+	if (PossibleStructs.Num() == 0)
+	{
+		FDataTableEditorUtils::GetPossibleStructAssetData(PossibleStructs);
+
+		CurrentOptions.Reset();
+		for (const FAssetData& FoundStruct : PossibleStructs)
+		{
+			CurrentOptions.Add(MakeShareable(new FString(FoundStruct.AssetName.ToString())));
+		}
+	}
+
+	if (!ComboListView->IsPendingRefresh())
+	{
+		ComboListView->RequestListRefresh();
+	}
+}
+
+const FAssetData* SDataTableStructComboBox::FindAssetDataForString(TSharedPtr<FString> StringOption) const
+{
+	check(CurrentOptions.Num() == PossibleStructs.Num());
+	for (int32 i = 0; i < CurrentOptions.Num(); i++)
+	{
+		if (StringOption == CurrentOptions[i])
+		{
+			return &PossibleStructs[i];
+		}
+	}
+	return nullptr;
+}
+
+UScriptStruct* SDataTableStructComboBox::GetOrLoadStruct(const FAssetData* AssetData)
+{
+	if (!AssetData)
+	{
+		return nullptr;
+	}
+
+	return Cast<UScriptStruct>(AssetData->GetAsset());
+}
+
+TSharedRef<ITableRow> SDataTableStructComboBox::GenerateMenuItemRow(TSharedPtr<FString> InItem, const TSharedRef<STableViewBase>& OwnerTable)
+{
+	FString SearchToken = SearchField->GetText().ToString().ToLower();
+	EVisibility WidgetVisibility = EVisibility::Visible;
+	if (!SearchToken.IsEmpty())
+	{
+		if (InItem->ToLower().Find(SearchToken) < 0)
+		{
+			WidgetVisibility = EVisibility::Collapsed;
+		}
+	}
+	
+	TAttribute<FText> OnGetToolTip = TAttribute<FText>::Create(TAttribute<FText>::FGetter::CreateSP(this, &SDataTableStructComboBox::GetTooltipText, InItem));
+
+	return SNew(SComboRow<TSharedPtr<FString>>, OwnerTable)
+		.Style(ItemStyle)
+		.Visibility(WidgetVisibility)
+		[
+			SNew(STextBlock)
+			.Text(FText::FromString(*InItem))
+			.ToolTipText(OnGetToolTip)
+		];
+}
+
+void SDataTableStructComboBox::OnMenuOpenChanged(bool bOpen)
+{
+	if (bOpen == false)
+	{
+		if (TListTypeTraits<TSharedPtr<FString>>::IsPtrValid(SelectedItem))
+		{
+			// Ensure the ListView selection is set back to the last committed selection
+			ComboListView->SetSelection(SelectedItem, ESelectInfo::OnNavigation);
+			ComboListView->RequestScrollIntoView(SelectedItem, 0);
+		}
+
+		// Set focus back to ComboBox for users focusing the ListView that just closed
+		FSlateApplication::Get().ForEachUser([&](FSlateUser* User) {
+			if (FSlateApplication::Get().HasUserFocusedDescendants(AsShared(), User->GetUserIndex()))
+			{
+				FSlateApplication::Get().SetUserFocus(User->GetUserIndex(), AsShared(), EFocusCause::SetDirectly);
+			}
+		});
+	}
+}
+
+void SDataTableStructComboBox::OnSelectionChanged_Internal(TSharedPtr<FString> ProposedSelection, ESelectInfo::Type SelectInfo)
+{
+	// Ensure that the proposed selection is different
+	if (SelectInfo != ESelectInfo::OnNavigation)
+	{
+		// Ensure that the proposed selection is different from selected
+		if (ProposedSelection != SelectedItem)
+		{
+			SelectedItem = ProposedSelection;
+			
+			UScriptStruct* SelectedStruct = GetOrLoadStruct(FindAssetDataForString(SelectedItem));
+
+			OnStructSelected.ExecuteIfBound(SelectedStruct);
+
+		}
+		// close combo even if user reselected item
+		this->SetIsOpen(false);
+	}
+}
+
+void SDataTableStructComboBox::OnSearchTextChanged(const FText& ChangedText)
+{
+	FString SearchToken = ChangedText.ToString().ToLower();
+	for (int32 i = 0; i < CurrentOptions.Num(); i++)
+	{
+		TSharedPtr<ITableRow> Row = ComboListView->WidgetFromItem(CurrentOptions[i]);
+		if (Row)
+		{
+			if (SearchToken.IsEmpty())
+			{
+				Row->AsWidget()->SetVisibility(EVisibility::Visible);
+			}
+			else if (CurrentOptions[i]->ToLower().Find(SearchToken) >= 0)
+			{
+				Row->AsWidget()->SetVisibility(EVisibility::Visible);
+			}
+			else
+			{
+				Row->AsWidget()->SetVisibility(EVisibility::Collapsed);
+			}
+		}
+	}
+
+	ComboListView->RequestListRefresh();
+
+	SelectedItem = TSharedPtr< FString >();
+}
+
+FReply SDataTableStructComboBox::OnButtonClicked()
+{
+	// if user clicked to close the combo menu
+	if (this->IsOpen())
+	{
+		// Re-select first selected item, just in case it was selected by navigation previously
+		TArray<TSharedPtr<FString>> SelectedItems = ComboListView->GetSelectedItems();
+		if (SelectedItems.Num() > 0)
+		{
+			OnSelectionChanged_Internal(SelectedItems[0], ESelectInfo::Direct);
+		}
+	}
+	else
+	{
+		SearchField->SetText(FText::GetEmpty());
+		RefreshOptions();
+	}
+
+	return SComboButton::OnButtonClicked();
+}
+
 const FString FDataTableEditorUtils::VariableTypesTooltipDocLink = TEXT("Shared/Editor/Blueprint/VariableTypes");
 
+TSharedRef<SWidget> FDataTableEditorUtils::MakeRowStructureComboBox(FOnDataTableStructSelected OnSelected)
+{
+	TSharedRef<SDataTableStructComboBox> ComboBox = SNew(SDataTableStructComboBox)
+		.OnStructSelected(OnSelected)
+		.ContentPadding(3);
+
+	return ComboBox;
+}
+
 FDataTableEditorUtils::FDataTableEditorManager& FDataTableEditorUtils::FDataTableEditorManager::Get()
 {
 	static TSharedRef< FDataTableEditorManager > EditorManager(new FDataTableEditorManager());
@@ -360,6 +785,37 @@ TArray<UScriptStruct*> FDataTableEditorUtils::GetPossibleStructs()
 	return RowStructs;
 }
 
+void FDataTableEditorUtils::GetPossibleStructAssetData(TArray<FAssetData>& StructAssets)
+{
+	StructAssets.Reset();
+
+	// Make combo of table rowstruct options
+	for (TObjectIterator<UScriptStruct> It; It; ++It)
+	{
+		UScriptStruct* Struct = *It;
+		if (IsValidTableStruct(Struct))
+		{
+			StructAssets.Add(FAssetData(Struct));
+		}
+	}
+
+	// Now get unloaded ones
+	const FAssetRegistryModule& AssetRegistryModule = FModuleManager::LoadModuleChecked<FAssetRegistryModule>(TEXT("AssetRegistry"));
+	TArray<FAssetData> AssetData;
+	AssetRegistryModule.Get().GetAssetsByClass(UUserDefinedStruct::StaticClass()->GetFName(), AssetData);
+
+	for (int32 AssetIndex = 0; AssetIndex < AssetData.Num(); ++AssetIndex)
+	{
+		const FAssetData& Asset = AssetData[AssetIndex];
+		if (Asset.IsValid() && !Asset.IsAssetLoaded())
+		{
+			StructAssets.Add(Asset);
+		}
+	}
+
+	StructAssets.Sort([](const FAssetData& A, const FAssetData& B) { return A.AssetName < B.AssetName; });
+}
+
 bool FDataTableEditorUtils::IsValidTableStruct(UScriptStruct* Struct)
 {
 	UScriptStruct* TableRowStruct = FindObjectChecked<UScriptStruct>(ANY_PACKAGE, TEXT("TableRowBase"));
diff --git a/Engine/Source/Editor/UnrealEd/Private/Factories/DataTableFactory.cpp b/Engine/Source/Editor/UnrealEd/Private/Factories/DataTableFactory.cpp
--- a/Engine/Source/Editor/UnrealEd/Private/Factories/DataTableFactory.cpp
+++ b/Engine/Source/Editor/UnrealEd/Private/Factories/DataTableFactory.cpp
@@ -28,26 +28,19 @@ bool UDataTableFactory::ConfigureProperties()
 	class FDataTableFactoryUI : public TSharedFromThis < FDataTableFactoryUI >
 	{
 		TSharedPtr<SWindow> PickerWindow;
-		TSharedPtr<SComboBox<UScriptStruct*>> RowStructCombo;
+		TSharedPtr<SWidget> RowStructCombo;
 		TSharedPtr<SButton> OkButton;
 		UScriptStruct* ResultStruct;
 	public:
 		FDataTableFactoryUI() : ResultStruct(nullptr) {}
 
-		TSharedRef<SWidget> MakeRowStructItemWidget(class UScriptStruct* InStruct) const
+		void OnStructSelected(UScriptStruct* NewStruct)
 		{
-			return SNew(STextBlock).Text(InStruct ? InStruct->GetDisplayNameText() : FText::GetEmpty());
-		}
-
-		FText GetSelectedRowOptionText() const
-		{
-			UScriptStruct* RowStruct = RowStructCombo.IsValid() ? RowStructCombo->GetSelectedItem() : nullptr;
-			return RowStruct ? RowStruct->GetDisplayNameText() : FText::GetEmpty();
+			ResultStruct = NewStruct;
 		}
 
 		FReply OnCreate()
 		{
-			ResultStruct = RowStructCombo.IsValid() ? RowStructCombo->GetSelectedItem() : nullptr;
 			if (PickerWindow.IsValid())
 			{
 				PickerWindow->RequestDestroyWindow();
@@ -67,21 +60,14 @@ bool UDataTableFactory::ConfigureProperties()
 
 		bool IsAnyRowSelected() const
 		{
-			return  RowStructCombo.IsValid() && RowStructCombo->GetSelectedItem();
+			return ResultStruct != nullptr;
 		}
 
 		UScriptStruct* OpenStructSelector()
 		{
 			ResultStruct = nullptr;
-			auto RowStructs = FDataTableEditorUtils::GetPossibleStructs();
 
-			RowStructCombo = SNew(SComboBox<UScriptStruct*>)
-				.OptionsSource(&RowStructs)
-				.OnGenerateWidget(this, &FDataTableFactoryUI::MakeRowStructItemWidget)
-				[
-					SNew(STextBlock)
-					.Text(this, &FDataTableFactoryUI::GetSelectedRowOptionText)
-				];
+			RowStructCombo = FDataTableEditorUtils::MakeRowStructureComboBox(FDataTableEditorUtils::FOnDataTableStructSelected::CreateSP(this, &FDataTableFactoryUI::OnStructSelected));
 
 			PickerWindow = SNew(SWindow)
 				.Title(LOCTEXT("DataTableFactoryOptions", "Pick Row Structure"))
diff --git a/Engine/Source/Editor/UnrealEd/Private/SCSVImportOptions.cpp b/Engine/Source/Editor/UnrealEd/Private/SCSVImportOptions.cpp
--- a/Engine/Source/Editor/UnrealEd/Private/SCSVImportOptions.cpp
+++ b/Engine/Source/Editor/UnrealEd/Private/SCSVImportOptions.cpp
@@ -12,6 +12,7 @@
 #include "Modules/ModuleManager.h"
 #include "PropertyEditorModule.h"
 #include "ObjectEditorUtils.h"
+#include "DataTableEditorUtils.h"
 
 #define LOCTEXT_NAMESPACE "CSVImportFactory"
 
@@ -27,28 +28,6 @@ void SCSVImportOptions::Construct(const FArguments& InArgs)
 	ImportTypes.Add(MakeShareable(new ECSVImportType(ECSVImportType::ECSV_CurveFloat)));
 	ImportTypes.Add(MakeShareable(new ECSVImportType(ECSVImportType::ECSV_CurveVector)));
 
-	// Find table row struct info
-	UScriptStruct* TableRowStruct = FindObjectChecked<UScriptStruct>(ANY_PACKAGE, TEXT("TableRowBase"));
-	if (TableRowStruct != nullptr)
-	{
-		// Make combo of table rowstruct options
-		for (TObjectIterator<UScriptStruct> It; It; ++It)
-		{
-			UScriptStruct* Struct = *It;
-			// If a child of the table row struct base, but not itself
-			const bool bBasedOnTableRowBase = Struct->IsChildOf(TableRowStruct) && (Struct != TableRowStruct);
-			const bool bUDStruct = Struct->IsA<UUserDefinedStruct>();
-			const bool bValidStruct = (Struct->GetOutermost() != GetTransientPackage());
-			if ((bBasedOnTableRowBase || bUDStruct) && bValidStruct)
-			{
-				RowStructs.Add(Struct);
-			}
-		}
-
-		// Alphabetically sort the row structs by name
-		RowStructs.Sort([](const UScriptStruct& ElementA, const UScriptStruct& ElementB) { return (ElementA.GetName() < ElementB.GetName()); } );
-	}
-
 	// Create properties view
 	FPropertyEditorModule & EditModule = FModuleManager::Get().GetModuleChecked<FPropertyEditorModule>("PropertyEditor");
 	FDetailsViewArgs DetailsViewArgs(/*bUpdateFromSelection=*/ false, /*bLockable=*/ false, /*bAllowSearch=*/ false, /*InNameAreaSettings=*/ FDetailsViewArgs::HideNameArea, /*bHideSelectionTip=*/ true);
@@ -69,6 +48,9 @@ void SCSVImportOptions::Construct(const FArguments& InArgs)
 		return false;
 	}));
 
+	RowStructCombo = FDataTableEditorUtils::MakeRowStructureComboBox(FDataTableEditorUtils::FOnDataTableStructSelected::CreateSP(this, &SCSVImportOptions::OnStructSelected));
+	RowStructCombo->SetVisibility(TAttribute<EVisibility>::Create(TAttribute<EVisibility>::FGetter::CreateSP(this, &SCSVImportOptions::GetTableRowOptionVis)));
+
 	// Create widget
 	this->ChildSlot
 	[
@@ -137,14 +119,7 @@ void SCSVImportOptions::Construct(const FArguments& InArgs)
 			+SVerticalBox::Slot()
 			.AutoHeight()
 			[
-				SAssignNew(RowStructCombo, SComboBox<UScriptStruct*>)
-				.OptionsSource( &RowStructs )
-				.OnGenerateWidget( this, &SCSVImportOptions::MakeRowStructItemWidget )
-				.Visibility( this, &SCSVImportOptions::GetTableRowOptionVis )
-				[
-					SNew(STextBlock)
-					.Text(this, &SCSVImportOptions::GetSelectedRowOptionText)
-				]
+				RowStructCombo.ToSharedRef()
 			]
 			// Curve interpolation
 			+SVerticalBox::Slot()
@@ -301,12 +276,9 @@ void SCSVImportOptions::OnImportTypeSelected(TSharedPtr<ECSVImportType> Selectio
 	}
 }
 
-/** Called to create a widget for each struct */
-TSharedRef<SWidget> SCSVImportOptions::MakeRowStructItemWidget(UScriptStruct* Struct)
+void SCSVImportOptions::OnStructSelected(UScriptStruct* NewStruct)
 {
-	check(Struct != nullptr);
-	return	SNew(STextBlock)
-			.Text(FText::FromString(Struct->GetName()));
+	SelectedStruct = NewStruct;
 }
 
 FString SCSVImportOptions::GetCurveTypeText(CurveInterpModePtr InterpMode) const
@@ -340,7 +312,6 @@ TSharedRef<SWidget> SCSVImportOptions::MakeCurveTypeWidget(CurveInterpModePtr In
 /** Called when 'OK' button is pressed */
 FReply SCSVImportOptions::OnImport()
 {
-	SelectedStruct = RowStructCombo->GetSelectedItem();
 	SelectedImportType = *ImportTypeCombo->GetSelectedItem();
 	if (CurveInterpCombo->GetSelectedItem().IsValid())
 	{
@@ -361,7 +332,7 @@ bool SCSVImportOptions::CanImport() const
 	switch (ImportType)
 	{
 	case ECSVImportType::ECSV_DataTable:
-		return RowStructCombo->GetSelectedItem() != nullptr;
+		return SelectedStruct != nullptr;
 		break;
 	case ECSVImportType::ECSV_CurveTable:
 		return CurveInterpCombo->GetSelectedItem().IsValid();
@@ -395,14 +366,6 @@ FText SCSVImportOptions::GetSelectedItemText() const
 		: FText::GetEmpty();
 }
 
-FText SCSVImportOptions::GetSelectedRowOptionText() const
-{
-	UScriptStruct* SelectedScript = RowStructCombo->GetSelectedItem();
-	return (SelectedScript)
-		? FText::FromString(SelectedScript->GetName())
-		: FText::GetEmpty();
-}
-
 FText SCSVImportOptions::GetSelectedCurveTypeText() const
 {
 	CurveInterpModePtr CurveModePtr = CurveInterpCombo->GetSelectedItem();
diff --git a/Engine/Source/Editor/UnrealEd/Public/DataTableEditorUtils.h b/Engine/Source/Editor/UnrealEd/Public/DataTableEditorUtils.h
--- a/Engine/Source/Editor/UnrealEd/Public/DataTableEditorUtils.h
+++ b/Engine/Source/Editor/UnrealEd/Public/DataTableEditorUtils.h
@@ -4,6 +4,8 @@
 #include "CoreMinimal.h"
 #include "Engine/DataTable.h"
 #include "Kismet2/ListenerManager.h"
+#include "Widgets/SWidget.h"
+#include "AssetData.h"
 
 struct FDataTableEditorColumnHeaderData
 {
@@ -82,7 +84,12 @@ struct UNREALED_API FDataTableEditorUtils
 
 	static void CacheDataTableForEditing(const UDataTable* DataTable, TArray<FDataTableEditorColumnHeaderDataPtr>& OutAvailableColumns, TArray<FDataTableEditorRowListViewDataPtr>& OutAvailableRows);
 
+	/** Returns all script structs that can be used as a data table row. This only includes loaded ones */
 	static TArray<UScriptStruct*> GetPossibleStructs();
+
+	/** Fills in an array with all possible DataTable structs, unloaded and loaded */
+	static void GetPossibleStructAssetData(TArray<FAssetData>& StructAssets);
+	
 	/** Utility function which verifies that the specified struct type is viable for data tables */
 	static bool IsValidTableStruct(UScriptStruct* Struct);
 
@@ -94,4 +101,10 @@ struct UNREALED_API FDataTableEditorUtils
 
 	/** Link to variable type doc  */
 	static const FString VariableTypesTooltipDocLink;
+
+	/** Delegate called when a data table struct is selected */
+	DECLARE_DELEGATE_OneParam(FOnDataTableStructSelected, UScriptStruct*);
+
+	/** Creates a combo box that allows selecting from the list of possible row structures */
+	static TSharedRef<SWidget> MakeRowStructureComboBox(FOnDataTableStructSelected OnSelected);
 };
diff --git a/Engine/Source/Editor/UnrealEd/Public/SCSVImportOptions.h b/Engine/Source/Editor/UnrealEd/Public/SCSVImportOptions.h
--- a/Engine/Source/Editor/UnrealEd/Public/SCSVImportOptions.h
+++ b/Engine/Source/Editor/UnrealEd/Public/SCSVImportOptions.h
@@ -71,8 +71,8 @@ public:
 	/** Called when import type changes */
 	void OnImportTypeSelected(TSharedPtr<ECSVImportType> Selection, ESelectInfo::Type SelectionType);
 
-	/** Called to create a widget for each struct */
-	TSharedRef<SWidget> MakeRowStructItemWidget(UScriptStruct* Struct);
+	/** Called when datatable row is selected */
+	void OnStructSelected(UScriptStruct* NewStruct);
 
 	FString GetCurveTypeText(CurveInterpModePtr InterpMode) const;
 
@@ -90,8 +90,6 @@ public:
 
 	FText GetSelectedItemText() const;
 
-	FText GetSelectedRowOptionText() const;
-
 	FText GetSelectedCurveTypeText() const;
 
 private:
@@ -115,11 +113,8 @@ private:
 
 	// Row type
 
-	/** Array of row struct options */
-	TArray< UScriptStruct* >						RowStructs;
-
 	/** The row struct combo box */
-	TSharedPtr< SComboBox<UScriptStruct*> >			RowStructCombo;
+	TSharedPtr< SWidget >							RowStructCombo;
 
 	/** The selected row struct */
 	UScriptStruct*									SelectedStruct;
