Don't copy instance only properties to a blueprint when applying property changes from an instance.

diff --git a/Engine/Source/Editor/Kismet/Private/SSCSEditor.cpp b/Engine/Source/Editor/Kismet/Private/SSCSEditor.cpp
--- a/Engine/Source/Editor/Kismet/Private/SSCSEditor.cpp
+++ b/Engine/Source/Editor/Kismet/Private/SSCSEditor.cpp
@@ -6268,7 +6268,7 @@ FText SSCSEditor::OnGetApplyChangesToBlueprintTooltip() const
 		AActor* BlueprintCDO = Actor->GetClass()->GetDefaultObject<AActor>();
 		if(BlueprintCDO != NULL)
 		{
-			const EditorUtilities::ECopyOptions::Type CopyOptions = (EditorUtilities::ECopyOptions::Type)(EditorUtilities::ECopyOptions::PreviewOnly|EditorUtilities::ECopyOptions::OnlyCopyEditOrInterpProperties);
+			const EditorUtilities::ECopyOptions::Type CopyOptions = (EditorUtilities::ECopyOptions::Type)(EditorUtilities::ECopyOptions::PreviewOnly|EditorUtilities::ECopyOptions::OnlyCopyEditOrInterpProperties|EditorUtilities::ECopyOptions::SkipInstanceOnlyProperties);
 			NumChangedProperties += EditorUtilities::CopyActorProperties(Actor, BlueprintCDO, CopyOptions);
 		}
 		NumChangedProperties += Actor->GetInstanceComponents().Num();
@@ -6447,7 +6447,7 @@ void SSCSEditor::OnApplyChangesToBlueprint() const
 				AActor* BlueprintCDO = Actor->GetClass()->GetDefaultObject<AActor>();
 				if (BlueprintCDO != NULL)
 				{
-					const EditorUtilities::ECopyOptions::Type CopyOptions = (EditorUtilities::ECopyOptions::Type)(EditorUtilities::ECopyOptions::OnlyCopyEditOrInterpProperties | EditorUtilities::ECopyOptions::PropagateChangesToArchetypeInstances);
+					const EditorUtilities::ECopyOptions::Type CopyOptions = (EditorUtilities::ECopyOptions::Type)(EditorUtilities::ECopyOptions::OnlyCopyEditOrInterpProperties | EditorUtilities::ECopyOptions::PropagateChangesToArchetypeInstances | EditorUtilities::ECopyOptions::SkipInstanceOnlyProperties);
 					NumChangedProperties = EditorUtilities::CopyActorProperties(Actor, BlueprintCDO, CopyOptions);
 					if (Actor->GetInstanceComponents().Num() > 0)
 					{
diff --git a/Engine/Source/Editor/UnrealEd/Private/Editor.cpp b/Engine/Source/Editor/UnrealEd/Private/Editor.cpp
--- a/Engine/Source/Editor/UnrealEd/Private/Editor.cpp
+++ b/Engine/Source/Editor/UnrealEd/Private/Editor.cpp
@@ -1217,7 +1217,8 @@ namespace EditorUtilities
 
 			if( !bIsTransient && !bIsIdentical && !bIsComponentContainer && !bIsComponentProp && !bIsBlueprintReadonly)
 			{
-				const bool bIsSafeToCopy = !( Options.Flags & ECopyOptions::OnlyCopyEditOrInterpProperties ) || ( Property->HasAnyPropertyFlags( CPF_Edit | CPF_Interp ) );
+				const bool bIsSafeToCopy = (!( Options.Flags & ECopyOptions::OnlyCopyEditOrInterpProperties ) || ( Property->HasAnyPropertyFlags( CPF_Edit | CPF_Interp ) ))
+				                        && (!( Options.Flags & ECopyOptions::SkipInstanceOnlyProperties) || ( !Property->HasAllPropertyFlags(CPF_DisableEditOnTemplate) ) );
 				if( bIsSafeToCopy )
 				{
 					if (!Options.CanCopyProperty(*Property, *SourceActor))
@@ -1342,7 +1343,8 @@ namespace EditorUtilities
 					if( !bIsTransient && !bIsIdentical && !bIsComponent && !SourceUCSModifiedProperties.Contains(Property)
 						&& ( !bIsTransform || SourceComponent != SourceActor->GetRootComponent() || ( !SourceActor->HasAnyFlags( RF_ClassDefaultObject | RF_ArchetypeObject ) && !TargetActor->HasAnyFlags( RF_ClassDefaultObject | RF_ArchetypeObject ) ) ) )
 					{
-						const bool bIsSafeToCopy = !( Options.Flags & ECopyOptions::OnlyCopyEditOrInterpProperties ) || ( Property->HasAnyPropertyFlags( CPF_Edit | CPF_Interp ) );
+						const bool bIsSafeToCopy = (!(Options.Flags & ECopyOptions::OnlyCopyEditOrInterpProperties) || (Property->HasAnyPropertyFlags(CPF_Edit | CPF_Interp)))
+						                        && (!(Options.Flags & ECopyOptions::SkipInstanceOnlyProperties) || (!Property->HasAllPropertyFlags(CPF_DisableEditOnTemplate)));
 						if( bIsSafeToCopy )
 						{
 							if (!Options.CanCopyProperty(*Property, *SourceActor))
diff --git a/Engine/Source/Editor/UnrealEd/Private/Kismet2/Kismet2.cpp b/Engine/Source/Editor/UnrealEd/Private/Kismet2/Kismet2.cpp
--- a/Engine/Source/Editor/UnrealEd/Private/Kismet2/Kismet2.cpp
+++ b/Engine/Source/Editor/UnrealEd/Private/Kismet2/Kismet2.cpp
@@ -1632,7 +1632,7 @@ UBlueprint* FKismetEditorUtilities::CreateBlueprintFromActor(const FName Bluepri
 			if (NewBlueprint->GeneratedClass != nullptr)
 			{
 				AActor* CDO = CastChecked<AActor>(NewBlueprint->GeneratedClass->GetDefaultObject());
-				const auto CopyOptions = (EditorUtilities::ECopyOptions::Type)(EditorUtilities::ECopyOptions::OnlyCopyEditOrInterpProperties | EditorUtilities::ECopyOptions::PropagateChangesToArchetypeInstances);
+				const EditorUtilities::ECopyOptions::Type CopyOptions = (EditorUtilities::ECopyOptions::Type)(EditorUtilities::ECopyOptions::OnlyCopyEditOrInterpProperties | EditorUtilities::ECopyOptions::PropagateChangesToArchetypeInstances | EditorUtilities::ECopyOptions::SkipInstanceOnlyProperties);
 				EditorUtilities::CopyActorProperties(Actor, CDO, CopyOptions);
 
 				if (USceneComponent* DstSceneRoot = CDO->GetRootComponent())
diff --git a/Engine/Source/Editor/UnrealEd/Public/Editor.h b/Engine/Source/Editor/UnrealEd/Public/Editor.h
--- a/Engine/Source/Editor/UnrealEd/Public/Editor.h
+++ b/Engine/Source/Editor/UnrealEd/Public/Editor.h
@@ -653,6 +653,9 @@ namespace EditorUtilities
 
 			/** Filters out Blueprint Read-only properties */
 			FilterBlueprintReadOnly = 1 << 5,
+
+			/** Filters out properties that are marked instance only. */
+			SkipInstanceOnlyProperties = 1 << 6,
 		};
 	}
 
