Fix ensure after undoing following an ed graph paste operation

diff --git a/Engine/Source/Editor/UnrealEd/Private/EdGraphUtilities.cpp b/Engine/Source/Editor/UnrealEd/Private/EdGraphUtilities.cpp
--- a/Engine/Source/Editor/UnrealEd/Private/EdGraphUtilities.cpp
+++ b/Engine/Source/Editor/UnrealEd/Private/EdGraphUtilities.cpp
@@ -412,6 +412,9 @@ void FEdGraphUtilities::ImportNodesFromText(UEdGraph* DestinationGraph, const FS
 	// Fix up pin cross-links, etc...
 	FEdGraphUtilities::PostProcessPastedNodes(Factory.SpawnedNodes);
 
+	// If a pin link wasn't resolved by now it was connected to something outside the clipboard and should be cleared:
+	UEdGraphPin::ResolveAllPinReferences();
+
 	ImportedNodeSet.Append(Factory.SpawnedNodes);
 }
 
