When getting debug info for a blueprint function's output parameter, display the value of the connected pin if the output parameter is connected to another node.

diff --git a/Engine/Source/Editor/UnrealEd/Private/Kismet2/KismetDebugUtilities.cpp b/Engine/Source/Editor/UnrealEd/Private/Kismet2/KismetDebugUtilities.cpp
--- a/Engine/Source/Editor/UnrealEd/Private/Kismet2/KismetDebugUtilities.cpp
+++ b/Engine/Source/Editor/UnrealEd/Private/Kismet2/KismetDebugUtilities.cpp
@@ -913,16 +913,14 @@ void FKismetDebugUtilities::GetValidBreakpointLocations(const UK2Node_MacroInsta
 	if (!bIsMacroPure && MacroEntryNode)
 	{
 		// Get the execute pin outputs on the entry node
-		for (auto PinIt = MacroEntryNode->Pins.CreateConstIterator(); PinIt; ++PinIt)
+		for (const UEdGraphPin* ExecPin : MacroEntryNode->Pins)
 		{
-			const UEdGraphPin* ExecPin = *PinIt;
 			if (ExecPin && ExecPin->Direction == EGPD_Output
 				&& ExecPin->PinType.PinCategory == UEdGraphSchema_K2::PC_Exec)
 			{
 				// For each pin linked to each execute pin, collect the node that owns it
-				for (auto LinkedToPinIt = ExecPin->LinkedTo.CreateConstIterator(); LinkedToPinIt; ++LinkedToPinIt)
+				for (const UEdGraphPin* LinkedToPin : ExecPin->LinkedTo)
 				{
-					const UEdGraphPin* LinkedToPin = *LinkedToPinIt;
 					check(LinkedToPin);
 
 					const UEdGraphNode* LinkedToNode = LinkedToPin->GetOwningNode();
@@ -1185,6 +1183,13 @@ FKismetDebugUtilities::EWatchTextResult FKismetDebugUtilities::FindDebuggingData
 					{
 						if (OutParmRec->Property == Property)
 						{
+							// try to use the output pin we're linked to
+							// otherwise the output param won't show any data since the return node hasn't executed when we stop here
+							if (WatchPin->Direction == EEdGraphPinDirection::EGPD_Input && WatchPin->LinkedTo.Num() == 1)
+							{
+								return FindDebuggingData(Blueprint, ActiveObject, WatchPin->LinkedTo[0], OutProperty, OutData, OutDelta, OutParent, SeenObjects);
+							}
+
 							PropertyBase = OutParmRec->PropAddr;
 							break;
 						}
@@ -1208,7 +1213,7 @@ FKismetDebugUtilities::EWatchTextResult FKismetDebugUtilities::FindDebuggingData
 				else if (AActor* Actor = Cast<AActor>(ActiveObject))
 				{
 					// Try and locate the propertybase in the actor components
-					for (auto ComponentIter : Actor->GetComponents())
+					for (UActorComponent* ComponentIter : Actor->GetComponents())
 					{
 						if (ComponentIter->GetClass()->IsChildOf(PropertyClass))
 						{
