Don't allow names of functions to be used as properties/functions/etc.

diff --git a/Engine/Source/Editor/UnrealEd/Private/Kismet2/Kismet2NameValidators.cpp b/Engine/Source/Editor/UnrealEd/Private/Kismet2/Kismet2NameValidators.cpp
--- a/Engine/Source/Editor/UnrealEd/Private/Kismet2/Kismet2NameValidators.cpp
+++ b/Engine/Source/Editor/UnrealEd/Private/Kismet2/Kismet2NameValidators.cpp
@@ -103,6 +103,7 @@ FKismetNameValidator::FKismetNameValidator(const class UBlueprint* Blueprint, FN
 	ExistingName = InExistingName;
 	BlueprintObject = Blueprint;
 	FBlueprintEditorUtils::GetClassVariableList(BlueprintObject, Names, true);
+	FBlueprintEditorUtils::GetFunctionNameList(BlueprintObject, Names);
 	FBlueprintEditorUtils::GetAllGraphNames(BlueprintObject, Names);
 	FBlueprintEditorUtils::GetSCSVariableNameList(Blueprint, Names);
 	FBlueprintEditorUtils::GetImplementingBlueprintsFunctionNameList(Blueprint, Names);
