PR #5821: Make BPs with compiling errors locatable (Contributed by [REDACTED_PERSON])

diff --git a/Engine/Source/Editor/UnrealEd/Private/PlayLevel.cpp b/Engine/Source/Editor/UnrealEd/Private/PlayLevel.cpp
--- a/Engine/Source/Editor/UnrealEd/Private/PlayLevel.cpp
+++ b/Engine/Source/Editor/UnrealEd/Private/PlayLevel.cpp
@@ -1743,12 +1743,43 @@ struct FInternalPlayLevelUtils
 {
 	static int32 ResolveDirtyBlueprints(const bool bPromptForCompile, TArray<UBlueprint*>& ErroredBlueprints, const bool bForceLevelScriptRecompile = true)
 	{
+		struct FLocal
+		{
+			static void OnMessageLogLinkActivated(const class TSharedRef<IMessageToken>& Token)
+			{
+				if (Token->GetType() == EMessageToken::Object)
+				{
+					const TSharedRef<FUObjectToken> UObjectToken = StaticCastSharedRef<FUObjectToken>(Token);
+					if (UObjectToken->GetObject().IsValid())
+					{
+						FKismetEditorUtilities::BringKismetToFocusAttentionOnObject(UObjectToken->GetObject().Get());
+					}
+				}
+			}
+
+			static void AddCompileErrorToLog(UBlueprint* ErroredBlueprint, FMessageLog& BlueprintLog)
+			{
+				FFormatNamedArguments Arguments;
+				Arguments.Add(TEXT("Name"), FText::FromString(ErroredBlueprint->GetName()));
+
+				TSharedRef<FTokenizedMessage> Message = FTokenizedMessage::Create(EMessageSeverity::Warning);
+				Message->AddToken(FTextToken::Create(LOCTEXT("BlueprintCompileFailed", "Blueprint failed to compile: ")));
+				Message->AddToken(FUObjectToken::Create(ErroredBlueprint, FText::FromString(ErroredBlueprint->GetName()))
+					->OnMessageTokenActivated(FOnMessageTokenActivated::CreateStatic(&FLocal::OnMessageLogLinkActivated))
+				);
+
+				BlueprintLog.AddMessage(Message);
+			}
+		};
+
 		const bool bAutoCompile = !bPromptForCompile;
 		FString PromptDirtyList;
 
 		TArray<UBlueprint*> InNeedOfRecompile;
 		ErroredBlueprints.Empty();
 
+		FMessageLog BlueprintLog("BlueprintLog");
+
 		double BPRegenStartTime = FPlatformTime::Seconds();
 		for (TObjectIterator<UBlueprint> BlueprintIt; BlueprintIt; ++BlueprintIt)
 		{
@@ -1777,6 +1808,7 @@ struct FInternalPlayLevelUtils
 			else if (BS_Error == Blueprint->Status && Blueprint->bDisplayCompilePIEWarning)
 			{
 				ErroredBlueprints.Add(Blueprint);
+				FLocal::AddCompileErrorToLog(Blueprint, BlueprintLog);
 			}
 		}
 
@@ -1792,7 +1824,6 @@ struct FInternalPlayLevelUtils
 		}
 		int32 RecompiledCount = 0;
 
-		FMessageLog BlueprintLog("BlueprintLog");
 		if (bRunCompilation && (InNeedOfRecompile.Num() > 0))
 		{
 			const FText LogPageLabel = (bAutoCompile) ? LOCTEXT("BlueprintAutoCompilationPageLabel", "Pre-Play auto-recompile") :
@@ -1849,11 +1880,7 @@ struct FInternalPlayLevelUtils
 					if (bHadError && ErroredBlueprints.Find(CompiledBlueprint) == INDEX_NONE)
 					{
 						ErroredBlueprints.Add(CompiledBlueprint);
-
-						FFormatNamedArguments Arguments;
-						Arguments.Add(TEXT("Name"), FText::FromString(CompiledBlueprint->GetName()));
-
-						BlueprintLog.Info(FText::Format(LOCTEXT("BlueprintCompileFailed", "Blueprint {Name} failed to compile"), Arguments));
+						FLocal::AddCompileErrorToLog(CompiledBlueprint, BlueprintLog);
 					}
 
 					++RecompiledCount;
@@ -2343,6 +2370,8 @@ void UEditorEngine::PlayInEditor( UWorld* InWorld, bool bInSimulateInEditor, FPl
 		const bool bContinuePIE = EAppReturnType::Yes == FMessageDialog::Open( EAppMsgType::YesNo, FText::Format( NSLOCTEXT("PlayInEditor", "PrePIE_BlueprintErrors", "One or more blueprints has an unresolved compiler error, are you sure you want to Play in Editor?{ErrorBlueprints}"), Args ) );
 		if ( !bContinuePIE )
 		{
+			FMessageLog("BlueprintLog").Open(EMessageSeverity::Warning);
+
 			FEditorDelegates::EndPIE.Broadcast(bInSimulateInEditor);
 			FNavigationSystem::OnPIEEnd(*InWorld);
 			return;
