PR #5785: Snap default blueprint event nodes to the grid (Contributed by [REDACTED_PERSON])

diff --git a/Engine/Source/Editor/UnrealEd/Private/Kismet2/Kismet2.cpp b/Engine/Source/Editor/UnrealEd/Private/Kismet2/Kismet2.cpp
--- a/Engine/Source/Editor/UnrealEd/Private/Kismet2/Kismet2.cpp
+++ b/Engine/Source/Editor/UnrealEd/Private/Kismet2/Kismet2.cpp
@@ -80,6 +80,7 @@
 #include "Framework/Notifications/NotificationManager.h"
 #include "Widgets/Notifications/SNotificationList.h"
 #include "Engine/InheritableComponentHandler.h"
+#include "Classes/EditorStyleSettings.h"
 
 DECLARE_CYCLE_STAT(TEXT("Compile Blueprint"), EKismetCompilerStats_CompileBlueprint, STATGROUP_KismetCompiler);
 DECLARE_CYCLE_STAT(TEXT("Broadcast Precompile"), EKismetCompilerStats_BroadcastPrecompile, STATGROUP_KismetCompiler);
@@ -654,6 +655,14 @@ UK2Node_Event* FKismetEditorUtilities::AddDefaultEventNode(UBlueprint* InBluepri
 		NewEventNode = NewObject<UK2Node_Event>(InGraph);
 		NewEventNode->EventReference = EventReference;
 
+		// Snap the new position to the grid
+		const UEditorStyleSettings* StyleSettings = GetDefault<UEditorStyleSettings>();
+		if (StyleSettings)
+		{
+			const uint32 GridSnapSize = StyleSettings->GridSnapSize;
+			InOutNodePosY = GridSnapSize * FMath::RoundFromZero(InOutNodePosY / (float)GridSnapSize);
+		}
+		
 		// add update event graph
 		NewEventNode->bOverrideFunction=true;
 		NewEventNode->CreateNewGuid();
