Add limited support for reordering and recategorizing Blueprint functions
Adapted from PR #4955: Fix reordering blueprint functions (Contributed by [REDACTED_PERSON])

diff --git a/Engine/Source/Editor/BlueprintGraph/Private/EdGraphSchema_K2_Actions.cpp b/Engine/Source/Editor/BlueprintGraph/Private/EdGraphSchema_K2_Actions.cpp
--- a/Engine/Source/Editor/BlueprintGraph/Private/EdGraphSchema_K2_Actions.cpp
+++ b/Engine/Source/Editor/BlueprintGraph/Private/EdGraphSchema_K2_Actions.cpp
@@ -157,8 +157,14 @@ FEdGraphSchemaActionDefiningObject FEdGraphSchemaAction_K2Graph::GetPersistentIt
 	UObject* DefiningObject = GetSourceBlueprint();
 	if (UFunction* Func = GetFunction())
 	{
-		DefiningObject = Func->GetOwnerStruct();
+		// Use the class where the function was initially introduced as the defining object
+		while (Func->GetSuperFunction())
+		{
+			Func = Func->GetSuperFunction();
+		}
+		DefiningObject = Func->GetOuterUClassUnchecked();
 	}
+
 	return FEdGraphSchemaActionDefiningObject(DefiningObject, (void*)GraphType);
 }
 
diff --git a/Engine/Source/Editor/Kismet/Private/SMyBlueprint.cpp b/Engine/Source/Editor/Kismet/Private/SMyBlueprint.cpp
--- a/Engine/Source/Editor/Kismet/Private/SMyBlueprint.cpp
+++ b/Engine/Source/Editor/Kismet/Private/SMyBlueprint.cpp
@@ -1160,6 +1160,12 @@ void SMyBlueprint::CollectAllActions(FGraphActionListBuilderBase& OutAllActions)
 			}
 		}
 
+		// Default, so place in 'non' category
+		if (FunctionCategory.EqualTo(FText::FromString(BlueprintObj->GetName())) || FunctionCategory.EqualTo(UEdGraphSchema_K2::VR_DefaultCategory))
+		{
+			FunctionCategory = FText::GetEmpty();
+		}
+
 		//@TODO: Should be a bit more generic (or the AnimGraph shouldn't be stored as a FunctionGraph...)
 		const bool bIsConstructionScript = Graph->GetFName() == UEdGraphSchema_K2::FN_UserConstructionScript;
 		
diff --git a/Engine/Source/Editor/UnrealEd/Private/Kismet2/BlueprintEditorUtils.cpp b/Engine/Source/Editor/UnrealEd/Private/Kismet2/BlueprintEditorUtils.cpp
--- a/Engine/Source/Editor/UnrealEd/Private/Kismet2/BlueprintEditorUtils.cpp
+++ b/Engine/Source/Editor/UnrealEd/Private/Kismet2/BlueprintEditorUtils.cpp
@@ -3955,7 +3955,8 @@ void FBlueprintEditorUtils::SetBlueprintFunctionOrMacroCategory(UEdGraph* Graph,
 	UBlueprint* Blueprint = FBlueprintEditorUtils::FindBlueprintForGraphChecked(Graph);
 	if (FKismetUserDeclaredFunctionMetadata* MetaData = FBlueprintEditorUtils::GetGraphFunctionMetaData(Graph))
 	{
-		if(!MetaData->Category.EqualTo(InCategoryName))
+		const FText& NewCategory = InCategoryName.IsEmpty() ? UEdGraphSchema_K2::VR_DefaultCategory : InCategoryName;
+		if (!MetaData->Category.EqualTo(NewCategory))
 		{
 			FScopedTransaction Transaction(LOCTEXT("SetBlueprintFunctionOrMacroCategory", "Set Category"));
 
@@ -3971,7 +3972,6 @@ void FBlueprintEditorUtils::SetBlueprintFunctionOrMacroCategory(UEdGraph* Graph,
 				}
 			}
 
-			const FText& NewCategory = InCategoryName.IsEmpty() ? UEdGraphSchema_K2::VR_DefaultCategory : InCategoryName;
 			MetaData->Category = NewCategory;
 
 			if (Function)
@@ -4086,9 +4086,16 @@ bool FBlueprintEditorUtils::MoveGraphBeforeOtherGraph(UEdGraph* Graph, int32 New
 			}
 		}
 
-		if (bModified && !bDontRecompile)
+		if (bModified)
 		{
-			FBlueprintEditorUtils::MarkBlueprintAsStructurallyModified(Blueprint);
+			if (!bDontRecompile)
+			{
+				FBlueprintEditorUtils::MarkBlueprintAsStructurallyModified(Blueprint);
+			}
+			else
+			{
+				FBlueprintEditorUtils::MarkBlueprintAsModified(Blueprint);
+			}
 		}
 
 		return bModified;
