Fix formatting issues in CookOnTheFLyServer.cpp

#rnx

diff --git a/Engine/Source/Editor/UnrealEd/Private/CookOnTheFlyServer.cpp b/Engine/Source/Editor/UnrealEd/Private/CookOnTheFlyServer.cpp
--- a/Engine/Source/Editor/UnrealEd/Private/CookOnTheFlyServer.cpp
+++ b/Engine/Source/Editor/UnrealEd/Private/CookOnTheFlyServer.cpp
@@ -237,7 +237,9 @@ public:
 	void Start()
 	{
 		if (StartTime)
+		{
 			return;
+		}
 
 		StartTime = FPlatformTime::Cycles64();
 	}
@@ -245,7 +247,9 @@ public:
 	void Stop()
 	{
 		if (!StartTime)
+		{
 			return;
+		}
 
 		HierarchyTimerInfo->Length += FPlatformTime::ToSeconds64(FPlatformTime::Cycles64() - StartTime);
 		++HierarchyTimerInfo->HitCount;
@@ -969,12 +973,15 @@ struct CookRequestQueue
 		FScopeLock ScopeLock(&SynchronizationObject);
 		const TArray<FName>* Platforms = PlatformList.Find(Filename);
 		if (Platforms == NULL)
+		{
 			return false;
-
+		}
 		for (const FName& PlatformName : PlatformNames)
 		{
 			if (!Platforms->Contains(PlatformName))
+			{
 				return false;
+			}
 		}
 		return true;
 	}
@@ -984,7 +991,9 @@ struct CookRequestQueue
 		FScopeLock ScopeLock(&SynchronizationObject);
 		const TArray<FName>* Platforms = PlatformList.Find(Filename);
 		if (Platforms == NULL)
+		{
 			return false;
+		}
 		return true;
 	}
 
@@ -1227,7 +1236,7 @@ struct FPackageTracker : public FUObjectArray::FUObjectCreateListener, public FU
 			auto Package = const_cast<UPackage*>(static_cast<const UPackage*>(Object));
 
 			if (Package->GetOuter() == nullptr)
-	{
+			{
 				LoadedPackages.Add(Package);
 				NewPackages.Add(Package);
 
@@ -1239,7 +1248,7 @@ struct FPackageTracker : public FUObjectArray::FUObjectCreateListener, public FU
 	virtual void NotifyUObjectDeleted(const class UObjectBase *Object, int32 Index) override
 	{
 		if (Object->GetClass() == UPackage::StaticClass())
-	{
+		{
 			auto Package = const_cast<UPackage*>(static_cast<const UPackage*>(Object));
 
 			LoadedPackages.Remove(Package);
@@ -1343,10 +1352,14 @@ void FPackageTracker::FilterLoadedPackage(UPackage* Package)
 	const FName StandardPackageFName = PackageNameCache->GetCachedStandardPackageFileFName(Package);
 
 	if (StandardPackageFName == NAME_None)
+	{
 		return;	// if we have name none that means we are in core packages or something...
+	}
 
 	if (CookedPackages.Exists(StandardPackageFName, AllTargetPlatformNames))
+	{
 		return;
+	}
 
 	PackagesPendingSave.Add(Package);
 }
@@ -1354,7 +1367,9 @@ void FPackageTracker::FilterLoadedPackage(UPackage* Package)
 void FPackageTracker::UpdatePackagesPendingSave()
 {
 	if (bPackagesPendingSaveDirty == false)
+	{
 		return;
+	}
 
 	PackagesPendingSave.Empty(PackagesPendingSave.Num());
 
@@ -2054,17 +2069,17 @@ UCookOnTheFlyServer::FReentryData& UCookOnTheFlyServer::GetReentryData(const UPa
 {
 	FReentryData& CurrentReentryData = PackageReentryData.FindOrAdd(Package->GetFName());
 
-	if ( (CurrentReentryData.bIsValid == false) && (Package->IsFullyLoaded() == true))
+	if ((CurrentReentryData.bIsValid == false) && (Package->IsFullyLoaded() == true))
 	{
 		CurrentReentryData.bIsValid = true;
 		CurrentReentryData.FileName = Package->GetFName();
-		GetObjectsWithOuter(Package, CurrentReentryData.CachedObjectsInOuter);	
-}
+		GetObjectsWithOuter(Package, CurrentReentryData.CachedObjectsInOuter);
+	}
 	return CurrentReentryData;
 }
 
 
-uint32 UCookOnTheFlyServer::TickCookOnTheSide( const float TimeSlice, uint32 &CookedPackageCount, ECookTickFlags TickFlags )
+uint32 UCookOnTheFlyServer::TickCookOnTheSide(const float TimeSlice, uint32 &CookedPackageCount, ECookTickFlags TickFlags)
 {
 	if (IsCookByTheBookMode() && CookByTheBookOptions->bFullLoadAndSave)
 	{
@@ -2091,41 +2106,41 @@ uint32 UCookOnTheFlyServer::TickCookOnTheSide( const float TimeSlice, uint32 &Co
 	// This is all the target platforms which we needed to process requests for this iteration
 	// we use this in the unsolicited packages processing below
 	TArray<FName> AllTargetPlatformNames;
-	
+
 	while (!GIsRequestingExit || CurrentCookMode == ECookMode::CookByTheBook)
 	{
 		if (HasCookRequests())
 		{
 			const float CurrentProgressDisplayTime = FPlatformTime::Seconds();
 
-			if (	LastCookedPackagesCount != PackageTracker->CookedPackages.Num() 
-				||	LastCookRequestsCount != PackageTracker->CookRequests.Num() 
-				||	(CurrentProgressDisplayTime - LastProgressDisplayTime) > GCookProgressRepeatTime)
+			if (LastCookedPackagesCount != PackageTracker->CookedPackages.Num()
+				|| LastCookRequestsCount != PackageTracker->CookRequests.Num()
+				|| (CurrentProgressDisplayTime - LastProgressDisplayTime) > GCookProgressRepeatTime)
 			{
-				UE_CLOG(!(TickFlags & ECookTickFlags::HideProgressDisplay) && (GCookProgressDisplay & (int32)ECookProgressDisplayMode::RemainingPackages), 
-					LogCook, 
-					Display, 
-					TEXT("Cooked packages %d Packages Remain %d Total %d"), 
-					PackageTracker->CookedPackages.Num(), 
-					PackageTracker->CookRequests.Num(), 
+				UE_CLOG(!(TickFlags & ECookTickFlags::HideProgressDisplay) && (GCookProgressDisplay & (int32)ECookProgressDisplayMode::RemainingPackages),
+					LogCook,
+					Display,
+					TEXT("Cooked packages %d Packages Remain %d Total %d"),
+					PackageTracker->CookedPackages.Num(),
+					PackageTracker->CookRequests.Num(),
 					PackageTracker->CookedPackages.Num() + PackageTracker->CookRequests.Num());
 
 				LastCookedPackagesCount = PackageTracker->CookedPackages.Num();
-				LastCookRequestsCount	= PackageTracker->CookRequests.Num();
+				LastCookRequestsCount = PackageTracker->CookRequests.Num();
 				LastProgressDisplayTime = CurrentProgressDisplayTime;
 			}
 		}
 
 		// if we just cooked a map then don't process anything the rest of this tick
-		if ( Result & COSR_RequiresGC )
+		if (Result & COSR_RequiresGC)
 		{
 			break;
 		}
 
-		if ( IsCookByTheBookMode() )
+		if (IsCookByTheBookMode())
 		{
-			check( CookByTheBookOptions );
-			if ( CookByTheBookOptions->bCancel )
+			check(CookByTheBookOptions);
+			if (CookByTheBookOptions->bCancel)
 			{
 				CancelCookByTheBook();
 			}
@@ -2163,7 +2178,7 @@ uint32 UCookOnTheFlyServer::TickCookOnTheSide( const float TimeSlice, uint32 &Co
 
 		// prevent autosave from happening until we are finished cooking
 		// causes really bad hitches
-		if ( GUnrealEd )
+		if (GUnrealEd)
 		{
 			const static float SecondsWarningTillAutosave = 10.0f;
 			GUnrealEd->GetPackageAutoSaver().ForceMinimumTimeTillAutoSave(SecondsWarningTillAutosave);
@@ -2182,21 +2197,21 @@ uint32 UCookOnTheFlyServer::TickCookOnTheSide( const float TimeSlice, uint32 &Co
 #endif
 		SCOPE_TIMER(TickCookOnTheSide);
 
-		check( ToBuild.IsValid() );
+		check(ToBuild.IsValid());
 		const TArray<FName>& TargetPlatformNames = ToBuild.GetPlatformNames();
 
 #if OUTPUT_TIMING
 		//FScopeTimer PackageManualTimer( ToBuild.GetFilename().ToString(), false );
 #endif
 
-		for ( const FName& PlatformName : TargetPlatformNames )
+		for (const FName& PlatformName : TargetPlatformNames)
 		{
 			AllTargetPlatformNames.AddUnique(PlatformName);
 		}
 
-		for ( const FName& PlatformName : AllTargetPlatformNames )
+		for (const FName& PlatformName : AllTargetPlatformNames)
 		{
-			if ( ToBuild.HasPlatform(PlatformName) == false )
+			if (ToBuild.HasPlatform(PlatformName) == false)
 			{
 				ToBuild.AddPlatform(PlatformName);
 			}
@@ -2205,14 +2220,14 @@ uint32 UCookOnTheFlyServer::TickCookOnTheSide( const float TimeSlice, uint32 &Co
 		const FString BuildFilename = ToBuild.GetFilename().ToString();
 
 		bool bShouldCook = true;
-		
-		if( CookByTheBookOptions && CookByTheBookOptions->bErrorOnEngineContentUse )
+
+		if (CookByTheBookOptions && CookByTheBookOptions->bErrorOnEngineContentUse)
 		{
 			check(IsCookingDLC());
 			FString DLCPath = FPaths::Combine(*GetBaseDirectoryForDLC(), TEXT("Content"));
-			if ( ToBuild.GetFilename().ToString().StartsWith(DLCPath) == false ) // if we don't start with the dlc path then we shouldn't be cooking this data 
+			if (ToBuild.GetFilename().ToString().StartsWith(DLCPath) == false) // if we don't start with the dlc path then we shouldn't be cooking this data 
 			{
-				UE_LOG(LogCook, Error, TEXT("Engine or Game content %s is being referenced by DLC!"), *ToBuild.GetFilename().ToString() );
+				UE_LOG(LogCook, Error, TEXT("Engine or Game content %s is being referenced by DLC!"), *ToBuild.GetFilename().ToString());
 				bShouldCook = false;
 			}
 		}
@@ -2228,11 +2243,11 @@ uint32 UCookOnTheFlyServer::TickCookOnTheSide( const float TimeSlice, uint32 &Co
 
 		UPackage* PackageForCooking = nullptr;
 
-		if ( bShouldCook ) // if we should cook the package then cook it otherwise add it to the list of already cooked packages below
+		if (bShouldCook) // if we should cook the package then cook it otherwise add it to the list of already cooked packages below
 		{
 			UPackage* Package = LoadPackageForCooking(BuildFilename);
 
-			if ( Package )
+			if (Package)
 			{
 				FString Name = Package->GetPathName();
 				FString PackageFilename(PackageNameCache->GetCachedStandardPackageFilename(Package));
@@ -2257,7 +2272,7 @@ uint32 UCookOnTheFlyServer::TickCookOnTheSide( const float TimeSlice, uint32 &Co
 			}
 		}
 
-		
+
 		if (PackageForCooking == nullptr)
 		{
 			// if we are iterative cooking the package might already be cooked
@@ -2278,7 +2293,7 @@ uint32 UCookOnTheFlyServer::TickCookOnTheSide( const float TimeSlice, uint32 &Co
 					// if we find the file this means it was cooked on a previous cook, however source package can't be found now. 
 					// this could be because the source package was deleted or renamed, and we are using iterative cooking
 					// perhaps in this case we should delete it?
-					UE_LOG(LogCook, Warning, TEXT("Found cooked file which shouldn't exist as it failed loading %s"), *SandboxFilename); 
+					UE_LOG(LogCook, Warning, TEXT("Found cooked file which shouldn't exist as it failed loading %s"), *SandboxFilename);
 					IFileManager::Get().Delete(*SandboxFilename);
 				}
 			}
@@ -2291,9 +2306,9 @@ uint32 UCookOnTheFlyServer::TickCookOnTheSide( const float TimeSlice, uint32 &Co
 
 		ITargetPlatformManagerModule& TPM = GetTargetPlatformManagerRef();
 		TArray<const ITargetPlatform*> TargetPlatforms;
-		for ( const FName& TargetPlatformName : AllTargetPlatformNames )
+		for (const FName& TargetPlatformName : AllTargetPlatformNames)
 		{
-			TargetPlatforms.Add( TPM.FindTargetPlatform( TargetPlatformName.ToString() ) );
+			TargetPlatforms.Add(TPM.FindTargetPlatform(TargetPlatformName.ToString()));
 		}
 
 
@@ -2304,10 +2319,10 @@ uint32 UCookOnTheFlyServer::TickCookOnTheSide( const float TimeSlice, uint32 &Co
 		{
 			SCOPE_TIMER(CallBeginCacheForCookedPlatformData);
 			// cache the resources for this package for each platform
-			
+
 			bIsAllDataCached &= BeginPackageCacheForCookedPlatformData(PackageForCooking, TargetPlatforms, Timer);
 
-			if( bIsAllDataCached )
+			if (bIsAllDataCached)
 			{
 				bIsAllDataCached &= FinishPackageCacheForCookedPlatformData(PackageForCooking, TargetPlatforms, Timer);
 			}
@@ -2317,7 +2332,7 @@ uint32 UCookOnTheFlyServer::TickCookOnTheSide( const float TimeSlice, uint32 &Co
 		bool ShouldTickPrecache = true;
 
 		// if we are ready to save then don't waste time precaching other stuff
-		if ( bIsAllDataCached == true )
+		if (bIsAllDataCached == true)
 		{
 			ShouldTickPrecache = false;
 		}
@@ -2329,14 +2344,14 @@ uint32 UCookOnTheFlyServer::TickCookOnTheSide( const float TimeSlice, uint32 &Co
 		else
 		{
 			// if we are doing no shader compilation right now try and precache something so that we load up the cpu
-			if ( GShaderCompilingManager->GetNumRemainingJobs() == 0 )
+			if (GShaderCompilingManager->GetNumRemainingJobs() == 0)
 			{
 				ShouldTickPrecache = true;
 			}
 		}
 
 		// cook on the fly mode we don't want to precache here because save package is going to stall on this package, we don't want to flood the system with precache requests before we stall
-		if (IsCookOnTheFlyMode()) 
+		if (IsCookOnTheFlyMode())
 		{
 			ShouldTickPrecache = false;
 		}
@@ -2353,7 +2368,7 @@ uint32 UCookOnTheFlyServer::TickCookOnTheSide( const float TimeSlice, uint32 &Co
 		}
 
 		ProcessUnsolicitedPackages();
-		
+
 		// in cook by the book bail out early because shaders and compiled for the primary package we are trying to save. 
 		// note in this case we also put the package at the end of the queue that queue might be reordered if we do partial gc
 		if ((bIsAllDataCached == false) && IsCookByTheBookMode() && !IsRealtimeMode())
@@ -2376,18 +2391,18 @@ uint32 UCookOnTheFlyServer::TickCookOnTheSide( const float TimeSlice, uint32 &Co
 
 		SaveCookedPackages(PackageForCooking, AllTargetPlatformNames, TargetPlatforms, Timer, /* out */ CookedPackageCount, /* out */ Result);
 
-		if ( Timer.IsTimeUp() )
+		if (Timer.IsTimeUp())
 		{
 			break;
 		}
 	}
-	
 
-	if ( IsCookOnTheFlyMode() && (IsCookingInEditor() == false) )
+
+	if (IsCookOnTheFlyMode() && (IsCookingInEditor() == false))
 	{
 		static int32 TickCounter = 0;
 		++TickCounter;
-		if ( TickCounter > 50 )
+		if (TickCounter > 50)
 		{
 			// dump stats every 50 ticks or so
 			DumpStats();
@@ -2426,10 +2441,14 @@ bool UCookOnTheFlyServer::BeginPackageCacheForCookedPlatformData(UPackage* Packa
 	FReentryData& CurrentReentryData = GetReentryData(Package);
 
 	if (CurrentReentryData.bIsValid == false)
+	{
 		return true;
+	}
 
 	if (CurrentReentryData.bBeginCacheFinished)
+	{
 		return true;
+	}
 
 	for (; CurrentReentryData.BeginCacheCount < CurrentReentryData.CachedObjectsInOuter.Num(); ++CurrentReentryData.BeginCacheCount)
 	{
@@ -2492,10 +2511,14 @@ bool UCookOnTheFlyServer::FinishPackageCacheForCookedPlatformData(UPackage* Pack
 	FReentryData& CurrentReentryData = GetReentryData(Package);
 
 	if (CurrentReentryData.bIsValid == false)
+	{
 		return true;
+	}
 
 	if (CurrentReentryData.bFinishedCacheFinished)
+	{
 		return true;
+	}
 
 	for (UObject* Obj : CurrentReentryData.CachedObjectsInOuter)
 	{
@@ -2619,8 +2642,9 @@ UPackage* UCookOnTheFlyServer::LoadPackageForCooking(const FString& BuildFilenam
 void UCookOnTheFlyServer::ProcessUnsolicitedPackages()
 {
 	if (IsCookByTheBookMode() && CookByTheBookOptions->bDisableUnsolicitedPackages)
+	{
 		return;
-
+	}
 	// Ensure sublevels are loaded by iterating all recently loaded packages and invoking
 	// PostLoadPackageFixup
 
@@ -2630,10 +2654,10 @@ void UCookOnTheFlyServer::ProcessUnsolicitedPackages()
 		TArray<UPackage*> NewPackages = PackageTracker->GetNewPackages();
 
 		for (UPackage* Package : NewPackages)
-			{
+		{
 			PostLoadPackageFixup(Package);
-			}
 		}
+	}
 }
 
 void UCookOnTheFlyServer::SaveCookedPackages(
@@ -2665,8 +2689,8 @@ void UCookOnTheFlyServer::SaveCookedPackages(
 
 	if (PackageToSave)
 	{
-	if (PackagesToSave.Num())
-	{
+		if (PackagesToSave.Num())
+		{
 			UPackage* First = PackagesToSave[0];
 			PackagesToSave.Add(First);
 			PackagesToSave[0] = PackageToSave;
@@ -3039,48 +3063,47 @@ void UCookOnTheFlyServer::PostLoadPackageFixup(UPackage* Package)
 
 	// Perform special processing for UWorld
 
-		UWorld* World = UWorld::FindWorldInPackage(Package);
-		check(World);
+	UWorld* World = UWorld::FindWorldInPackage(Package);
+	check(World);
 
-		World->PersistentLevel->HandleLegacyMapBuildData();
+	World->PersistentLevel->HandleLegacyMapBuildData();
 
 	if (IsCookByTheBookMode() == false)
-		{
+	{
 		return;
 	}
 
-			GIsCookerLoadingPackage = true;
-			if (World->GetStreamingLevels().Num())
-			{
-				TSet<FName> NeverCookPackageNames;
+	GIsCookerLoadingPackage = true;
+	if (World->GetStreamingLevels().Num())
+	{
+		TSet<FName> NeverCookPackageNames;
 		PackageTracker->NeverCookPackageList.GetValues(NeverCookPackageNames);
 
 		UE_LOG(LogCook, Display, TEXT("Loading secondary levels for package '%s'"), *World->GetName());
 
-				World->LoadSecondaryLevels(true, &NeverCookPackageNames);
-			}
-			GIsCookerLoadingPackage = false;
+		World->LoadSecondaryLevels(true, &NeverCookPackageNames);
+	}
+	GIsCookerLoadingPackage = false;
 
-			TArray<FString> NewPackagesToCook;
+	TArray<FString> NewPackagesToCook;
 
-			// Collect world composition tile packages to cook
-			if (World->WorldComposition)
-			{
-				World->WorldComposition->CollectTilesToCook(NewPackagesToCook);
-			}
+	// Collect world composition tile packages to cook
+	if (World->WorldComposition)
+	{
+		World->WorldComposition->CollectTilesToCook(NewPackagesToCook);
+	}
 
-			for (const FString& PackageName : NewPackagesToCook)
-			{
+	for (const FString& PackageName : NewPackagesToCook)
+	{
 		FName StandardPackageFName = PackageNameCache->GetCachedStandardPackageFileFName(FName(*PackageName));
 
-				if (StandardPackageFName != NAME_None)
-				{
-					RequestPackage(StandardPackageFName, false);
-				}
-			}
+		if (StandardPackageFName != NAME_None)
+		{
+			RequestPackage(StandardPackageFName, false);
+		}
+	}
 }
 
-
 void UCookOnTheFlyServer::TickPrecacheObjectsForPlatforms(const float TimeSlice, const TArray<const ITargetPlatform*>& TargetPlatforms) 
 {
 	
@@ -3115,7 +3138,9 @@ void UCookOnTheFlyServer::TickPrecacheObjectsForPlatforms(const float TimeSlice,
 	++LastUpdateTick;
 
 	if (Timer.IsTimeUp())
+	{
 		return;
+	}
 
 	bool AllMaterialsCompiled = true;
 	// queue up some shaders for compilation
@@ -3144,7 +3169,9 @@ void UCookOnTheFlyServer::TickPrecacheObjectsForPlatforms(const float TimeSlice,
 		}
 
 		if (Timer.IsTimeUp())
+		{
 			return;
+		}
 
 		if (GShaderCompilingManager->GetNumRemainingJobs() > MaxPrecacheShaderJobs)
 		{
@@ -3177,7 +3204,9 @@ void UCookOnTheFlyServer::TickPrecacheObjectsForPlatforms(const float TimeSlice,
 			Texture->BeginCacheForCookedPlatformData(TargetPlatform);
 		}
 		if (Timer.IsTimeUp())
+		{
 			return;
+		}
 	}
 
 }
@@ -3238,12 +3267,12 @@ TArray<UPackage*> UCookOnTheFlyServer::GetUnsolicitedPackages(const TArray<FName
 
 		const FName StandardPackageFName = PackageNameCache->GetCachedStandardPackageFileFName(Package);
 
-				if (StandardPackageFName == NAME_None)
+		if (StandardPackageFName == NAME_None)
 			continue;	// if we have name none that means we are in core packages or something...
 
 		if (PackageTracker->CookedPackages.Exists(StandardPackageFName, TargetPlatformNames))
-					continue;
-	
+			continue;
+
 		PackagesToSave.Add(Package);
 
 		UE_LOG(LogCook, Verbose, TEXT("Found unsolicited package to cook '%s'"), *Package->GetName());
@@ -3553,14 +3582,14 @@ void UCookOnTheFlyServer::MarkGCPackagesToKeepForCooker()
 	// then sort by the number of dependencies which are referenced by the package
 	// we want to process the packages with the highest dependencies so that they can
 	// be evicted from memory and are likely to be able to be released on next GC pass
-	PackageTracker->CookRequests.Sort([&PackageDependenciesCount,&LoadedPackages](const FName& A, const FName& B)
-		{
-			int32 ADependencies = PackageDependenciesCount.FindChecked(A);
-			int32 BDependencies = PackageDependenciesCount.FindChecked(B);
-			bool ALoaded = LoadedPackages.Contains(A);
-			bool BLoaded = LoadedPackages.Contains(B);
-			return (ALoaded == BLoaded) ? (ADependencies > BDependencies) : ALoaded > BLoaded;
-		}
+	PackageTracker->CookRequests.Sort([&PackageDependenciesCount, &LoadedPackages](const FName& A, const FName& B)
+	{
+		int32 ADependencies = PackageDependenciesCount.FindChecked(A);
+		int32 BDependencies = PackageDependenciesCount.FindChecked(B);
+		bool ALoaded = LoadedPackages.Contains(A);
+		bool BLoaded = LoadedPackages.Contains(B);
+		return (ALoaded == BLoaded) ? (ADependencies > BDependencies) : ALoaded > BLoaded;
+	}
 	);
 }
 
@@ -3569,7 +3598,6 @@ void UCookOnTheFlyServer::BeginDestroy()
 	EndNetworkFileServer();
 
 	Super::BeginDestroy();
-
 }
 
 void UCookOnTheFlyServer::TickRecompileShaderRequests()
@@ -3584,7 +3612,6 @@ void UCookOnTheFlyServer::TickRecompileShaderRequests()
 	{
 		HandleNetworkFileServerRecompileShaders(Request->RecompileData);
 
-
 		// all done! other thread can unblock now
 		Request->bComplete = true;
 	}
@@ -3610,8 +3637,10 @@ bool UCookOnTheFlyServer::ShouldConsiderCompressedPackageFileLengthRequirements(
 
 bool UCookOnTheFlyServer::MakePackageFullyLoaded(UPackage* Package) const
 {
-	if ( Package->IsFullyLoaded() )
+	if (Package->IsFullyLoaded())
+	{
 		return true;
+	}
 
 	bool bPackageFullyLoaded = false;
 	GIsCookerLoadingPackage = true;
@@ -4906,7 +4935,7 @@ bool UCookOnTheFlyServer::IniSettingsOutOfDate(const ITargetPlatform* TargetPlat
 
 bool UCookOnTheFlyServer::SaveCurrentIniSettings(const ITargetPlatform* TargetPlatform) const
 {
-	auto S = FScopeAssign<bool>(IniSettingRecurse, true);
+	FScopeAssign<bool> S = FScopeAssign<bool>(IniSettingRecurse, true);
 
 	TMap<FString, FString> AdditionalIniSettings;
 	GetAdditionalCurrentIniVersionStrings(TargetPlatform, AdditionalIniSettings);
@@ -6489,7 +6518,9 @@ void UCookOnTheFlyServer::BuildMapDependencyGraph(const FName& PlatformName)
 		FName Name = FName(*FPackageName::FilenameToLongPackageName(CookedPackage.ToString()));
 
 		if (!ContainsMap(Name))
+		{
 			continue;
+		}
 
 		TSet<FName> DependentPackages;
 		TSet<FName> Roots; 
@@ -6593,7 +6624,6 @@ void UCookOnTheFlyServer::ClearPlatformCookedData(const FName& PlatformName)
 
 void UCookOnTheFlyServer::ClearCachedCookedPlatformDataForPlatform( const FName& PlatformName )
 {
-
 	ITargetPlatformManagerModule& TPM = GetTargetPlatformManagerRef();
 	ITargetPlatform* TargetPlatform = TPM.FindTargetPlatform(PlatformName.ToString());
 	if ( TargetPlatform )
@@ -6616,7 +6646,6 @@ void UCookOnTheFlyServer::OnTargetPlatformChangedSupportedFormats(const ITargetP
 
 void UCookOnTheFlyServer::CreateSandboxFile()
 {
-
 	// initialize the sandbox file after determining if we are cooking dlc
 	// Local sandbox file wrapper. This will be used to handle path conversions,
 	// but will not be used to actually write/read files so we can safely
@@ -7003,8 +7032,10 @@ void UCookOnTheFlyServer::StartCookByTheBook( const FCookByTheBookStartupOptions
 	// add all the files for the requested platform to the cook list
 	for ( const FName& FileFName : FilesInPath )
 	{
-		if ( FileFName == NAME_None )
+		if (FileFName == NAME_None)
+		{
 			continue;
+		}
 
 		const FName PackageFileFName = PackageNameCache->GetCachedStandardPackageFileFName(FileFName);
 		
@@ -7045,9 +7076,7 @@ void UCookOnTheFlyServer::StartCookByTheBook( const FCookByTheBookStartupOptions
 				{
 					PackageTracker->CookRequests.EnqueueUnique( FFilePlatformRequest( PackageFilename, PlatformArray) );
 				}
-
 			}
-
 		}
 	}
 
@@ -7231,7 +7260,9 @@ void UCookOnTheFlyServer::HandleNetworkFileServerFileRequest(const FString& File
 	PackageTracker->CookRequests.EnqueueUnique(FileRequest, true);
 	
 	if (PackageTracker->CookRequestEvent)
+	{
 		PackageTracker->CookRequestEvent->Trigger();
+	}
 
 #if PROFILE_NETWORK
 	bool bFoundNetworkEventWait = true;
@@ -7306,7 +7337,6 @@ void UCookOnTheFlyServer::HandleNetworkGetPrecookedList(const FString& PlatformN
 	}
 }
 
-
 void UCookOnTheFlyServer::HandleNetworkFileServerRecompileShaders(const FShaderRecompileData& RecompileData)
 {
 	// shouldn't receive network requests unless we are in cook on the fly mode
