Fix deprecation warning
#rnx

diff --git a/Engine/Source/Editor/UnrealEd/Private/DataTableEditorUtils.cpp b/Engine/Source/Editor/UnrealEd/Private/DataTableEditorUtils.cpp
--- a/Engine/Source/Editor/UnrealEd/Private/DataTableEditorUtils.cpp
+++ b/Engine/Source/Editor/UnrealEd/Private/DataTableEditorUtils.cpp
@@ -813,7 +813,7 @@ void FDataTableEditorUtils::GetPossibleStructAssetData(TArray<FAssetData>& Struc
 		}
 	}
 
-	StructAssets.Sort([](const FAssetData& A, const FAssetData& B) { return A.AssetName < B.AssetName; });
+	StructAssets.Sort([](const FAssetData& A, const FAssetData& B) { return A.AssetName.LexicalLess(B.AssetName); });
 }
 
 bool FDataTableEditorUtils::IsValidTableStruct(UScriptStruct* Struct)
