PR #4282: Don't override window position when the play in editor window is centered

diff --git a/Engine/Source/Editor/UnrealEd/Private/PlayLevel.cpp b/Engine/Source/Editor/UnrealEd/Private/PlayLevel.cpp
--- a/Engine/Source/Editor/UnrealEd/Private/PlayLevel.cpp
+++ b/Engine/Source/Editor/UnrealEd/Private/PlayLevel.cpp
@@ -3242,8 +3242,12 @@ UGameInstance* UEditorEngine::CreatePIEGameInstance(int32 InPIEInstance, bool bI
 
 							if (index <= 0)
 							{
-								LevelEditorPlaySettings->NewWindowPosition.X = FPlatformMath::RoundToInt(PIEWindowPos.X);
-								LevelEditorPlaySettings->NewWindowPosition.Y = FPlatformMath::RoundToInt(PIEWindowPos.Y);
+								// only override the window position if the window isn't being centered
+								if (!LevelEditorPlaySettings->CenterNewWindow)
+								{
+									LevelEditorPlaySettings->NewWindowPosition.X = FPlatformMath::RoundToInt(PIEWindowPos.X);
+									LevelEditorPlaySettings->NewWindowPosition.Y = FPlatformMath::RoundToInt(PIEWindowPos.Y);
+								}
 							}
 							else
 							{
