Add option to control if the material can be reorder to the fbx file material order

diff --git a/Engine/Source/Editor/DetailCustomizations/Private/FbxImportUIDetails.cpp b/Engine/Source/Editor/DetailCustomizations/Private/FbxImportUIDetails.cpp
--- a/Engine/Source/Editor/DetailCustomizations/Private/FbxImportUIDetails.cpp
+++ b/Engine/Source/Editor/DetailCustomizations/Private/FbxImportUIDetails.cpp
@@ -593,6 +593,18 @@ void FFbxImportUIDetails::CustomizeDetails( IDetailLayoutBuilder& DetailBuilder
 				}
 			}
 		}
+
+		TSharedPtr<IPropertyHandle> ReorderMaterialToFbxOrderProp = nullptr;
+		if (ImportUI->MeshTypeToImport == FBXIT_StaticMesh)
+		{
+			ReorderMaterialToFbxOrderProp = StaticMeshDataProp->GetChildHandle(GET_MEMBER_NAME_CHECKED(UFbxStaticMeshImportData, bReorderMaterialToFbxOrder));
+		}
+		else if (ImportUI->MeshTypeToImport == FBXIT_SkeletalMesh)
+		{
+			ReorderMaterialToFbxOrderProp = SkeletalMeshDataProp->GetChildHandle(GET_MEMBER_NAME_CHECKED(UFbxSkeletalMeshImportData, bReorderMaterialToFbxOrder));
+		}
+		//Add the advance reorder option at the end
+		MaterialCategory.AddProperty(ReorderMaterialToFbxOrderProp);
 	}
 
 	//Information category
diff --git a/Engine/Source/Editor/UnrealEd/Classes/Factories/FbxMeshImportData.h b/Engine/Source/Editor/UnrealEd/Classes/Factories/FbxMeshImportData.h
--- a/Engine/Source/Editor/UnrealEd/Classes/Factories/FbxMeshImportData.h
+++ b/Engine/Source/Editor/UnrealEd/Classes/Factories/FbxMeshImportData.h
@@ -82,6 +82,10 @@ class UFbxMeshImportData : public UFbxAssetImportData
 	UPROPERTY(EditAnywhere, BlueprintReadWrite, config, AdvancedDisplay, Category = Mesh, meta=(ImportType="Mesh|GeoOnly"))
 	TEnumAsByte<enum EFBXNormalGenerationMethod::Type> NormalGenerationMethod;
 
+	/* If checked, The material list will be reorder to the same order has the FBX file. */
+	UPROPERTY(EditAnywhere, BlueprintReadWrite, config, AdvancedDisplay, Category = Material, meta = (OBJRestrict = "true"))
+	bool bReorderMaterialToFbxOrder;
+
 	bool CanEditChange( const UProperty* InProperty ) const override;
 
 	//////////////////////////////////////////////////////////////////////////
diff --git a/Engine/Source/Editor/UnrealEd/Private/Fbx/FbxMainImport.cpp b/Engine/Source/Editor/UnrealEd/Private/Fbx/FbxMainImport.cpp
--- a/Engine/Source/Editor/UnrealEd/Private/Fbx/FbxMainImport.cpp
+++ b/Engine/Source/Editor/UnrealEd/Private/Fbx/FbxMainImport.cpp
@@ -400,6 +400,7 @@ void ApplyImportUIToImportOptions(UFbxImportUI* ImportUI, FBXImportOptions& InOu
 		InOutImportOptions.bConvertSceneUnit		= StaticMeshData->bConvertSceneUnit;
 		InOutImportOptions.VertexColorImportOption	= StaticMeshData->VertexColorImportOption;
 		InOutImportOptions.VertexOverrideColor		= StaticMeshData->VertexOverrideColor;
+		InOutImportOptions.bReorderMaterialToFbxOrder = StaticMeshData->bReorderMaterialToFbxOrder;
 	}
 	else if ( ImportUI->MeshTypeToImport == FBXIT_SkeletalMesh )
 	{
@@ -419,6 +420,7 @@ void ApplyImportUIToImportOptions(UFbxImportUI* ImportUI, FBXImportOptions& InOu
 		InOutImportOptions.bConvertSceneUnit			= SkeletalMeshData->bConvertSceneUnit;
 		InOutImportOptions.VertexColorImportOption		= SkeletalMeshData->VertexColorImportOption;
 		InOutImportOptions.VertexOverrideColor			= SkeletalMeshData->VertexOverrideColor;
+		InOutImportOptions.bReorderMaterialToFbxOrder	= SkeletalMeshData->bReorderMaterialToFbxOrder;
 
 		if(ImportUI->bImportAnimations)
 		{
@@ -1423,6 +1425,7 @@ bool FFbxImporter::ImportFromFile(const FString& Filename, const FString& Type,
 				  * @EventParam ImportUniformScale float Returns the uniform scale apply on the import data
 				  * @EventParam MaterialBasePath string Returns the path pointing on the base material use to import material instance
 				  * @EventParam MaterialSearchLocation string Returns the scope of the search for existing materials (if material not found it can create one depending on bImportMaterials value)
+				  * @EventParam ReorderMaterialToFbxOrder boolean returns weather the importer should reorder the materials in the same order has the fbx file
 				  * @EventParam AutoGenerateCollision boolean Returns whether the importer should create collision primitive
 				  * @EventParam CombineToSingle boolean Returns whether the importer should combine all mesh part together or import many meshes
 				  * @EventParam BakePivotInVertex boolean Returns whether the importer should bake the fbx mesh pivot into the vertex position
@@ -1509,6 +1512,7 @@ bool FFbxImporter::ImportFromFile(const FString& Filename, const FString& Type,
 						Attribs.Add(FAnalyticsEventAttribute(TEXT("GenOpt ImportUniformScale"), ImportOptions->ImportUniformScale));
 						Attribs.Add(FAnalyticsEventAttribute(TEXT("GenOpt MaterialBasePath"), ImportOptions->MaterialBasePath));
 						Attribs.Add(FAnalyticsEventAttribute(TEXT("GenOpt MaterialSearchLocation"), MaterialSearchLocationEnum->GetNameStringByValue((uint64)(ImportOptions->MaterialSearchLocation))));
+						Attribs.Add(FAnalyticsEventAttribute(TEXT("GenOpt ReorderMaterialToFbxOrder"), ImportOptions->bReorderMaterialToFbxOrder));
 
 						//We cant capture a this member, so just assign the pointer here
 						FBXImportOptions* CaptureImportOptions = ImportOptions;
diff --git a/Engine/Source/Editor/UnrealEd/Private/Fbx/FbxMeshImportData.cpp b/Engine/Source/Editor/UnrealEd/Private/Fbx/FbxMeshImportData.cpp
--- a/Engine/Source/Editor/UnrealEd/Private/Fbx/FbxMeshImportData.cpp
+++ b/Engine/Source/Editor/UnrealEd/Private/Fbx/FbxMeshImportData.cpp
@@ -9,6 +9,7 @@ UFbxMeshImportData::UFbxMeshImportData(const FObjectInitializer& ObjectInitializ
 	NormalImportMethod = FBXNIM_ComputeNormals;
 	NormalGenerationMethod = EFBXNormalGenerationMethod::MikkTSpace;
 	bBakePivotInVertex = false;
+	bReorderMaterialToFbxOrder = true;
 }
 
 bool UFbxMeshImportData::CanEditChange(const UProperty* InProperty) const
diff --git a/Engine/Source/Editor/UnrealEd/Private/Fbx/FbxStaticMeshImport.cpp b/Engine/Source/Editor/UnrealEd/Private/Fbx/FbxStaticMeshImport.cpp
--- a/Engine/Source/Editor/UnrealEd/Private/Fbx/FbxStaticMeshImport.cpp
+++ b/Engine/Source/Editor/UnrealEd/Private/Fbx/FbxStaticMeshImport.cpp
@@ -34,6 +34,7 @@
 #include "PhysicsEngine/BodySetup.h"
 #include "MeshDescription.h"
 #include "MeshAttributes.h"
+#include "MeshDescriptionOperations.h"
 #include "IMeshBuilderModule.h"
 #include "Settings/EditorExperimentalSettings.h"
 
@@ -1754,7 +1755,7 @@ UStaticMesh* UnFbx::FFbxImporter::ImportStaticMeshAsSingle(UObject* InParent, TA
 	return StaticMesh;
 }
 
-void ReorderMaterialAfterImport(UStaticMesh* StaticMesh, TArray<FbxNode*>& MeshNodeArray)
+void ReorderMaterialAfterImport(UStaticMesh* StaticMesh, TArray<FbxNode*>& MeshNodeArray, bool bAllowFbxReorder)
 {
 	if (StaticMesh == nullptr)
 	{
@@ -1843,70 +1844,73 @@ void ReorderMaterialAfterImport(UStaticMesh* StaticMesh, TArray<FbxNode*>& MeshN
 	}
 
 	//Reorder the StaticMaterials array to reflect the order in the fbx file
-	//So we make sure the order reflect the material ID in the DCCs
-	FMeshSectionInfoMap OldSectionInfoMap = StaticMesh->SectionInfoMap;
-	TArray<int32> FbxRemapMaterials;
-	TArray<FStaticMaterial> NewStaticMaterials;
-	for (int32 FbxMaterialIndex = 0; FbxMaterialIndex < MeshMaterials.Num(); ++FbxMaterialIndex)
+	if (bAllowFbxReorder)
 	{
-		const FString &FbxMaterial = MeshMaterials[FbxMaterialIndex];
-		int32 FoundMaterialIndex = INDEX_NONE;
-		for (int32 BuildMaterialIndex = 0; BuildMaterialIndex < StaticMesh->StaticMaterials.Num(); ++BuildMaterialIndex)
+		//So we make sure the order reflect the material ID in the DCCs
+		FMeshSectionInfoMap OldSectionInfoMap = StaticMesh->SectionInfoMap;
+		TArray<int32> FbxRemapMaterials;
+		TArray<FStaticMaterial> NewStaticMaterials;
+		for (int32 FbxMaterialIndex = 0; FbxMaterialIndex < MeshMaterials.Num(); ++FbxMaterialIndex)
 		{
-			FStaticMaterial &BuildMaterial = StaticMesh->StaticMaterials[BuildMaterialIndex];
-			if (FbxMaterial.Compare(BuildMaterial.ImportedMaterialSlotName.ToString()) == 0)
+			const FString &FbxMaterial = MeshMaterials[FbxMaterialIndex];
+			int32 FoundMaterialIndex = INDEX_NONE;
+			for (int32 BuildMaterialIndex = 0; BuildMaterialIndex < StaticMesh->StaticMaterials.Num(); ++BuildMaterialIndex)
 			{
-				FoundMaterialIndex = BuildMaterialIndex;
-				break;
+				FStaticMaterial &BuildMaterial = StaticMesh->StaticMaterials[BuildMaterialIndex];
+				if (FbxMaterial.Compare(BuildMaterial.ImportedMaterialSlotName.ToString()) == 0)
+				{
+					FoundMaterialIndex = BuildMaterialIndex;
+					break;
+				}
 			}
-		}
 
-		if (FoundMaterialIndex != INDEX_NONE)
-		{
-			FbxRemapMaterials.Add(FoundMaterialIndex);
-			NewStaticMaterials.Add(StaticMesh->StaticMaterials[FoundMaterialIndex]);
+			if (FoundMaterialIndex != INDEX_NONE)
+			{
+				FbxRemapMaterials.Add(FoundMaterialIndex);
+				NewStaticMaterials.Add(StaticMesh->StaticMaterials[FoundMaterialIndex]);
+			}
 		}
-	}
-	//Add the materials not used by the LOD 0 at the end of the array. The order here is irrelevant since it can be used by many LOD other then LOD 0 and in different order
-	for (int32 BuildMaterialIndex = 0; BuildMaterialIndex < StaticMesh->StaticMaterials.Num(); ++BuildMaterialIndex)
-	{
-		const FStaticMaterial &StaticMaterial = StaticMesh->StaticMaterials[BuildMaterialIndex];
-		bool bFoundMaterial = false;
-		for (const FStaticMaterial &BuildMaterial : NewStaticMaterials)
+		//Add the materials not used by the LOD 0 at the end of the array. The order here is irrelevant since it can be used by many LOD other then LOD 0 and in different order
+		for (int32 BuildMaterialIndex = 0; BuildMaterialIndex < StaticMesh->StaticMaterials.Num(); ++BuildMaterialIndex)
 		{
-			if (StaticMaterial == BuildMaterial)
+			const FStaticMaterial &StaticMaterial = StaticMesh->StaticMaterials[BuildMaterialIndex];
+			bool bFoundMaterial = false;
+			for (const FStaticMaterial &BuildMaterial : NewStaticMaterials)
 			{
-				bFoundMaterial = true;
-				break;
+				if (StaticMaterial == BuildMaterial)
+				{
+					bFoundMaterial = true;
+					break;
+				}
+			}
+			if (!bFoundMaterial)
+			{
+				FbxRemapMaterials.Add(BuildMaterialIndex);
+				NewStaticMaterials.Add(StaticMaterial);
 			}
 		}
-		if (!bFoundMaterial)
+
+		StaticMesh->StaticMaterials.Empty();
+		for (const FStaticMaterial &BuildMaterial : NewStaticMaterials)
 		{
-			FbxRemapMaterials.Add(BuildMaterialIndex);
-			NewStaticMaterials.Add(StaticMaterial);
+			StaticMesh->StaticMaterials.Add(BuildMaterial);
 		}
-	}
 
-	StaticMesh->StaticMaterials.Empty();
-	for (const FStaticMaterial &BuildMaterial : NewStaticMaterials)
-	{
-		StaticMesh->StaticMaterials.Add(BuildMaterial);
-	}
-
-	//Remap the material instance of the staticmaterial array and remap the material index of all sections
-	for (int32 LODResoureceIndex = 0; LODResoureceIndex < StaticMesh->RenderData->LODResources.Num(); ++LODResoureceIndex)
-	{
-		FStaticMeshLODResources& LOD = StaticMesh->RenderData->LODResources[LODResoureceIndex];
-		int32 NumSections = LOD.Sections.Num();
-		for (int32 SectionIndex = 0; SectionIndex < NumSections; ++SectionIndex)
+		//Remap the material instance of the staticmaterial array and remap the material index of all sections
+		for (int32 LODResoureceIndex = 0; LODResoureceIndex < StaticMesh->RenderData->LODResources.Num(); ++LODResoureceIndex)
 		{
-			FMeshSectionInfo Info = OldSectionInfoMap.Get(LODResoureceIndex, SectionIndex);
-			int32 RemapIndex = FbxRemapMaterials.Find(Info.MaterialIndex);
-			if (StaticMesh->StaticMaterials.IsValidIndex(RemapIndex))
+			FStaticMeshLODResources& LOD = StaticMesh->RenderData->LODResources[LODResoureceIndex];
+			int32 NumSections = LOD.Sections.Num();
+			for (int32 SectionIndex = 0; SectionIndex < NumSections; ++SectionIndex)
 			{
-				Info.MaterialIndex = RemapIndex;
-				StaticMesh->SectionInfoMap.Set(LODResoureceIndex, SectionIndex, Info);
-				StaticMesh->OriginalSectionInfoMap.Set(LODResoureceIndex, SectionIndex, Info);
+				FMeshSectionInfo Info = OldSectionInfoMap.Get(LODResoureceIndex, SectionIndex);
+				int32 RemapIndex = FbxRemapMaterials.Find(Info.MaterialIndex);
+				if (StaticMesh->StaticMaterials.IsValidIndex(RemapIndex))
+				{
+					Info.MaterialIndex = RemapIndex;
+					StaticMesh->SectionInfoMap.Set(LODResoureceIndex, SectionIndex, Info);
+					StaticMesh->OriginalSectionInfoMap.Set(LODResoureceIndex, SectionIndex, Info);
+				}
 			}
 		}
 	}
@@ -2063,7 +2067,7 @@ void UnFbx::FFbxImporter::PostImportStaticMesh(UStaticMesh* StaticMesh, TArray<F
 	//If we have import a LOD other then the base, the material array cannot be sorted, because only the base LOD reorder the material array
 	if (LODIndex == 0 && StaticMesh->StaticMaterials.Num() > 1)
 	{
-		ReorderMaterialAfterImport(StaticMesh, MeshNodeArray);
+		ReorderMaterialAfterImport(StaticMesh, MeshNodeArray, ImportOptions->bReorderMaterialToFbxOrder);
 	}
 }
 
diff --git a/Engine/Source/Editor/UnrealEd/Public/FbxImporter.h b/Engine/Source/Editor/UnrealEd/Public/FbxImporter.h
--- a/Engine/Source/Editor/UnrealEd/Public/FbxImporter.h
+++ b/Engine/Source/Editor/UnrealEd/Public/FbxImporter.h
@@ -166,6 +166,8 @@ struct FBXImportOptions
 	FString BaseEmmisiveTextureName;
 	FString BaseSpecularTextureName;
 	EMaterialSearchLocation MaterialSearchLocation;
+	//If true the materials will be reorder to follow the fbx order
+	bool bReorderMaterialToFbxOrder;
 	// Skeletal Mesh options
 	bool bImportMorph;
 	bool bImportAnimations;
