Updated TraceResult when clicking in the level viewport to respect AllowSelectTranslucent

diff --git a/Engine/Source/Editor/UnrealEd/Private/Editor/ActorPositioning.cpp b/Engine/Source/Editor/UnrealEd/Private/Editor/ActorPositioning.cpp
--- a/Engine/Source/Editor/UnrealEd/Private/Editor/ActorPositioning.cpp
+++ b/Engine/Source/Editor/UnrealEd/Private/Editor/ActorPositioning.cpp
@@ -17,6 +17,7 @@
 #include "SnappingUtils.h"
 #include "LandscapeHeightfieldCollisionComponent.h"
 #include "LandscapeComponent.h"
+#include "Editor/EditorPerProjectUserSettings.h"
 
 FActorPositionTraceResult FActorPositioning::TraceWorldForPositionWithDefault(const FViewportCursorLocation& Cursor, const FSceneView& View, const TArray<AActor*>* IgnoreActors)
 {
@@ -86,16 +87,18 @@ bool IsHitIgnored(const FHitResult& InHit, const FSceneView& InSceneView)
 
 	// Only use this component if it is visible in the specified scene views
 	bool bIsRenderedOnScreen = false;
+	bool bIgnoreTranslucentPrimitive = false;
 	{				
 		if (PrimitiveComponent && PrimitiveComponent->SceneProxy)
 		{
 			const FPrimitiveViewRelevance ViewRelevance = PrimitiveComponent->SceneProxy->GetViewRelevance(&InSceneView);
 			// BSP is a bit special in that its bDrawRelevance is false even when drawn as wireframe because InSceneView.Family->EngineShowFlags.BSPTriangles is off
 			bIsRenderedOnScreen = ViewRelevance.bDrawRelevance || (PrimitiveComponent->IsA(UModelComponent::StaticClass()) && InSceneView.Family->EngineShowFlags.BSP);
+			bIgnoreTranslucentPrimitive = ViewRelevance.HasTranslucency() && !GetDefault<UEditorPerProjectUserSettings>()->bAllowSelectTranslucent;
 		}
 	}
 
-	return !bIsRenderedOnScreen;
+	return !bIsRenderedOnScreen || bIgnoreTranslucentPrimitive;
 }
 
 FActorPositionTraceResult FActorPositioning::TraceWorldForPosition(const UWorld& InWorld, const FSceneView& InSceneView, const FVector& RayStart, const FVector& RayEnd, const TArray<AActor*>* IgnoreActors)
