UDN crash with compiling optimized build
#fix converted float comparison to IsNearlyZero

diff --git a/Engine/Source/Editor/UnrealEd/Private/LODCluster.cpp b/Engine/Source/Editor/UnrealEd/Private/LODCluster.cpp
--- a/Engine/Source/Editor/UnrealEd/Private/LODCluster.cpp
+++ b/Engine/Source/Editor/UnrealEd/Private/LODCluster.cpp
@@ -47,7 +47,7 @@ const float CalculateOverlap(const FSphere& ASphere, const float AFillingFactor,
 	}
 
 	float Distance = (ASphere.Center-BSphere.Center).Size();
-	check (Distance > 0.f);
+	check(!FMath::IsNearlyZero(Distance));
 
 	float ARadius = ASphere.W;
 	float BRadius = BSphere.W;
@@ -67,13 +67,13 @@ const float CalculateOverlap(const FSphere& ASphere, const float AFillingFactor,
 	float OverlapRedius = FMath::Sqrt(OverlapRadius1) / OverlapRadius2;
 	float OverlapRediusSq = OverlapRedius*OverlapRedius;
 
-	check (OverlapRadius1 >= 0.f);
+	check(!FMath::IsNearlyZero(OverlapRadius1));
 
 	float ConstPI = PI/6.0f;
 	float AVolume = ConstPI*(3*OverlapRediusSq + ACapHeight*ACapHeight) * ACapHeight;
 	float BVolume = ConstPI*(3*OverlapRediusSq + BCapHeight*BCapHeight) * BCapHeight;
 
-	check (AVolume > 0.f &&  BVolume > 0.f);
+	check(!FMath::IsNearlyZero(AVolume) && !FMath::IsNearlyZero(BVolume));
 
 	float TotalVolume = AFillingFactor*AVolume + BFillingFactor*BVolume;
 	return TotalVolume;
