Follow up fix to 4973523

diff --git a/Engine/Source/Editor/UnrealEd/Classes/Animation/DebugSkelMeshComponent.h b/Engine/Source/Editor/UnrealEd/Classes/Animation/DebugSkelMeshComponent.h
--- a/Engine/Source/Editor/UnrealEd/Classes/Animation/DebugSkelMeshComponent.h
+++ b/Engine/Source/Editor/UnrealEd/Classes/Animation/DebugSkelMeshComponent.h
@@ -302,7 +302,7 @@ class UNREALED_API UDebugSkelMeshComponent : public USkeletalMeshComponent
 	//~ Begin SkeletalMeshComponent Interface
 	virtual void InitAnim(bool bForceReinit) override;
 	virtual bool IsWindEnabled() const override { return true; }
-	virtual void K2_SetAnimInstanceClass(class UClass* NewClass) override;
+	virtual void SetAnimClass(class UClass* NewClass) override;
 	//~ End SkeletalMeshComponent Interface
 	// Preview.
 	// @todo document
diff --git a/Engine/Source/Editor/UnrealEd/Private/Animation/DebugSkelMeshComponent.cpp b/Engine/Source/Editor/UnrealEd/Private/Animation/DebugSkelMeshComponent.cpp
--- a/Engine/Source/Editor/UnrealEd/Private/Animation/DebugSkelMeshComponent.cpp
+++ b/Engine/Source/Editor/UnrealEd/Private/Animation/DebugSkelMeshComponent.cpp
@@ -295,7 +295,7 @@ void UDebugSkelMeshComponent::InitAnim(bool bForceReinit)
 	}
 }
 
-void UDebugSkelMeshComponent::K2_SetAnimInstanceClass(class UClass* NewClass)
+void UDebugSkelMeshComponent::SetAnimClass(class UClass* NewClass)
 {
 	// Override this to do nothing and warn the user
 	UE_LOG(LogAnimation, Warning, TEXT("Attempting to destroy an animation preview actor, skipping."));
