Arcball Added/Screen to Rotation Gizmo.

MouseDeltaTracker: ConvertMovementDeltaToDragRot .. Added FSceneView parameter since we need to pass that down to the UnrealWidget when rotating arcball/screen which we do in world space for robustness.

EditorViewportClient:
    When rotating the arcball/screen we want to show the hand cursor  There is one function (IsUsingAbsoluteTranslation) that's used to show the cursor, pick the cursor to use, and do a seperate absolute translation path when dragging. Since we just want one we added a parameter to the function,  IsUsingAbsoluteTranslation(bool bAlsoCheckAbsoluteRotation).

FWidget:
To support the above added a function to say we are in absolute rotation mode (AllowsAbsoluteRotationMovement), which happens if we are in  Rotate and XYZ.    Also contains changes to Render_Rotate, and Drag functions to move the gizmos.  We need to the last drag position so we cache that now.
Finally in order to draw the transaprent sphere added a DrawSphereColored function since I couldn't get transparent to work with the default materials.

Added a toggle to Editor Preferences-> Level Editor-> Look and Feel for Enable Arcball Rotate and Enable Screen Rotate (Seemed best place since Enable Combined Translate/Rotate Widget) is also there.

diff --git a/Engine/Source/Editor/UnrealEd/Classes/Settings/LevelEditorViewportSettings.h b/Engine/Source/Editor/UnrealEd/Classes/Settings/LevelEditorViewportSettings.h
--- a/Engine/Source/Editor/UnrealEd/Classes/Settings/LevelEditorViewportSettings.h
+++ b/Engine/Source/Editor/UnrealEd/Classes/Settings/LevelEditorViewportSettings.h
@@ -238,6 +238,14 @@ class UNREALED_API ULevelEditorViewportSettings
 	UPROPERTY(EditAnywhere, config, Category=LookAndFeel, meta=( DisplayName = "Enable Combined Translate/Rotate Widget" ))
 	uint32 bAllowTranslateRotateZWidget:1;
 
+	/** Allow translate/rotate widget */
+	UPROPERTY(EditAnywhere, config, Category = LookAndFeel, meta = (DisplayName = "Enable Arcball Rotate"))
+	uint32 bAllowArcballRotate : 1;
+
+	/** Allow translate/rotate widget */
+	UPROPERTY(EditAnywhere, config, Category = LookAndFeel, meta = (DisplayName = "Enable Screen Rotate"))
+	uint32 bAllowScreenRotate : 1;
+
 	/** If true, Clicking a BSP selects the brush and ctrl+shift+click selects the surface. If false, vice versa */
 	UPROPERTY(EditAnywhere, config, Category=LookAndFeel, meta=( DisplayName = "Clicking BSP Enables Brush" ), AdvancedDisplay)
 	uint32 bClickBSPSelectsBrush:1;
diff --git a/Engine/Source/Editor/UnrealEd/Private/EditorViewportClient.cpp b/Engine/Source/Editor/UnrealEd/Private/EditorViewportClient.cpp
--- a/Engine/Source/Editor/UnrealEd/Private/EditorViewportClient.cpp
+++ b/Engine/Source/Editor/UnrealEd/Private/EditorViewportClient.cpp
@@ -1444,7 +1444,7 @@ EMouseCursor::Type FEditorViewportClient::GetCursor(FViewport* InViewport,int32
 		MouseCursor = EMouseCursor::CardinalCross;
 	}
 	//wyisyg mode
-	else if (IsUsingAbsoluteTranslation() && bHasMouseMovedSinceClick)
+	else if (IsUsingAbsoluteTranslation(true) && bHasMouseMovedSinceClick)
 	{
 		MouseCursor = EMouseCursor::CardinalCross;
 	}
@@ -2101,30 +2101,30 @@ void FEditorViewportClient::UpdateMouseDelta()
 			{
 				bWidgetAxisControlledByDrag = false;
 				Widget->SetCurrentAxis( EAxisList::None );
-				MouseDeltaTracker->ConvertMovementDeltaToDragRot(this, DragDelta, Drag, Rot, Scale);
+				MouseDeltaTracker->ConvertMovementDeltaToDragRot(DragStartView, this, DragDelta, Drag, Rot, Scale);
 				Widget->SetCurrentAxis( CurrentAxis );
 				CurrentAxis = EAxisList::None;
 			}
 			else
 			{
-				//if Absolute Translation, and not just moving the camera around
-				if (IsUsingAbsoluteTranslation())
+				if (DragStartView == nullptr)
 				{
-					if (DragStartView == nullptr)
-					{
-						// Compute a view.
-						DragStartViewFamily = new FSceneViewFamily(FSceneViewFamily::ConstructionValues(
-							Viewport,
-							GetScene(),
-							EngineShowFlags)
-							.SetRealtimeUpdate(IsRealtime()));
-						DragStartView = CalcSceneView(DragStartViewFamily);
-					}
+					// Compute a view.
+					DragStartViewFamily = new FSceneViewFamily(FSceneViewFamily::ConstructionValues(
+						Viewport,
+						GetScene(),
+						EngineShowFlags)
+						.SetRealtimeUpdate(IsRealtime()));
+					DragStartView = CalcSceneView(DragStartViewFamily);
+				}
+				//if Absolute Translation, and not just moving the camera around
+				if (IsUsingAbsoluteTranslation(false))
+				{ 
 					MouseDeltaTracker->AbsoluteTranslationConvertMouseToDragRot(DragStartView, this, Drag, Rot, Scale);
 				} 
 				else
 				{
-					MouseDeltaTracker->ConvertMovementDeltaToDragRot(this, DragDelta, Drag, Rot, Scale);
+					MouseDeltaTracker->ConvertMovementDeltaToDragRot(DragStartView, this, DragDelta, Drag, Rot, Scale);
 				}
 			}
 
@@ -2397,14 +2397,17 @@ void FEditorViewportClient::MarkMouseMovedSinceClick()
 }
 
 /** Determines whether this viewport is currently allowed to use Absolute Movement */
-bool FEditorViewportClient::IsUsingAbsoluteTranslation() const
+bool FEditorViewportClient::IsUsingAbsoluteTranslation(bool bAlsoCheckAbsoluteRotation) const
 {
 	bool bIsHotKeyAxisLocked = Viewport->KeyState(EKeys::LeftControl) || Viewport->KeyState(EKeys::RightControl);
 	bool bCameraLockedToWidget = !(Widget && Widget->GetCurrentAxis() & EAxisList::Screen) && (Viewport->KeyState(EKeys::LeftShift) || Viewport->KeyState(EKeys::RightShift));
 	// Screen-space movement must always use absolute translation
-	bool bScreenSpaceTransformation = Widget && (Widget->GetCurrentAxis() == EAxisList::Screen);
+	bool bScreenSpaceTransformation = Widget && (Widget->GetCurrentAxis() == EAxisList::Screen) && GetWidgetMode() != FWidget::WM_Rotate;
 	bool bAbsoluteMovementEnabled = GetDefault<ULevelEditorViewportSettings>()->bUseAbsoluteTranslation || bScreenSpaceTransformation;
-	bool bCurrentWidgetSupportsAbsoluteMovement = FWidget::AllowsAbsoluteTranslationMovement( GetWidgetMode() ) || bScreenSpaceTransformation;
+	bool bCurrentWidgetSupportsAbsoluteMovement = FWidget::AllowsAbsoluteTranslationMovement( GetWidgetMode()) || bScreenSpaceTransformation;
+	EAxisList::Type AxisType = Widget ? Widget->GetCurrentAxis() : EAxisList::None;
+
+	bool bCurrentWidgetSupportsAbsoluteRotation = bAlsoCheckAbsoluteRotation ? FWidget::AllowsAbsoluteRotationMovement(GetWidgetMode(), AxisType) : false;
 	bool bWidgetActivelyTrackingAbsoluteMovement = Widget && (Widget->GetCurrentAxis() != EAxisList::None);
 
 	const bool LeftMouseButtonDown = Viewport->KeyState(EKeys::LeftMouseButton);
@@ -2413,7 +2416,7 @@ bool FEditorViewportClient::IsUsingAbsoluteTranslation() const
 
 	const bool bAnyMouseButtonsDown = (LeftMouseButtonDown || MiddleMouseButtonDown || RightMouseButtonDown);
 
-	return (!bCameraLockedToWidget && !bIsHotKeyAxisLocked && bAbsoluteMovementEnabled && bCurrentWidgetSupportsAbsoluteMovement && bWidgetActivelyTrackingAbsoluteMovement && !IsOrtho() && bAnyMouseButtonsDown);
+	return (!bCameraLockedToWidget && !bIsHotKeyAxisLocked && bAbsoluteMovementEnabled && (bCurrentWidgetSupportsAbsoluteMovement || bCurrentWidgetSupportsAbsoluteRotation) && bWidgetActivelyTrackingAbsoluteMovement && !IsOrtho() && bAnyMouseButtonsDown);
 }
 
 void FEditorViewportClient::SetMatineeRecordingWindow (IMatineeBase* InInterpEd)
@@ -4627,8 +4630,8 @@ void FEditorViewportClient::UpdateRequiredCursorVisibility()
 		}
 	}
 
-	//if Absolute Translation and not just moving the camera around
-	if (IsUsingAbsoluteTranslation() && !MouseDeltaTracker->UsingDragTool() )
+	//if Absolute Translation or arc rotate and not just moving the camera around
+	if (IsUsingAbsoluteTranslation(true) && !MouseDeltaTracker->UsingDragTool())
 	{
 		//If we are dragging something we should hide the hardware cursor and show the s/w one
 		SetRequiredCursor(false, true);
diff --git a/Engine/Source/Editor/UnrealEd/Private/MouseDeltaTracker.cpp b/Engine/Source/Editor/UnrealEd/Private/MouseDeltaTracker.cpp
--- a/Engine/Source/Editor/UnrealEd/Private/MouseDeltaTracker.cpp
+++ b/Engine/Source/Editor/UnrealEd/Private/MouseDeltaTracker.cpp
@@ -475,7 +475,7 @@ void FMouseDeltaTracker::AddDelta(FEditorViewportClient* InViewportClient, FKey
 			}
 		}
 		//if Absolute Translation, and not just moving the camera around
-		else if (InViewportClient->IsUsingAbsoluteTranslation())
+		else if (InViewportClient->IsUsingAbsoluteTranslation(false))
 		{
 			// Compute a view.
 			FSceneViewFamilyContext ViewFamily(FSceneViewFamily::ConstructionValues( 
@@ -618,7 +618,7 @@ const FVector FMouseDeltaTracker::GetScreenDelta() const
 /**
  * Converts the delta movement to drag/rotation/scale based on the viewport type or widget axis
  */
-void FMouseDeltaTracker::ConvertMovementDeltaToDragRot(FEditorViewportClient* InViewportClient, FVector& InOutDragDelta, FVector& OutDrag, FRotator& OutRotation, FVector& OutScale) const
+void FMouseDeltaTracker::ConvertMovementDeltaToDragRot(FSceneView* InView, FEditorViewportClient* InViewportClient, FVector& InOutDragDelta, FVector& OutDrag, FRotator& OutRotation, FVector& OutScale) const
 {
 	OutDrag = FVector::ZeroVector;
 	OutRotation = FRotator::ZeroRotator;
@@ -626,7 +626,7 @@ void FMouseDeltaTracker::ConvertMovementDeltaToDragRot(FEditorViewportClient* In
 
 	if( InViewportClient->GetCurrentWidgetAxis() != EAxisList::None )
 	{
-		InViewportClient->Widget->ConvertMouseMovementToAxisMovement( InViewportClient, bUsedDragModifier, InOutDragDelta, OutDrag, OutRotation, OutScale );
+		InViewportClient->Widget->ConvertMouseMovementToAxisMovement( InView, InViewportClient, bUsedDragModifier, InOutDragDelta, OutDrag, OutRotation, OutScale );
 	}
 	else
 	{
diff --git a/Engine/Source/Editor/UnrealEd/Private/UnrealWidget.cpp b/Engine/Source/Editor/UnrealEd/Private/UnrealWidget.cpp
--- a/Engine/Source/Editor/UnrealEd/Private/UnrealWidget.cpp
+++ b/Engine/Source/Editor/UnrealEd/Private/UnrealWidget.cpp
@@ -154,7 +154,10 @@ FWidget::FWidget()
 	AxisColorX = FLinearColor(0.594f,0.0197f,0.0f);
 	AxisColorY = FLinearColor(0.1349f,0.3959f,0.0f);
 	AxisColorZ = FLinearColor(0.0251f,0.207f,0.85f);
+	ScreenAxisColor = FLinearColor(0.76, 0.72, 0.14f);
 	PlaneColorXY = FColor::Yellow;
+	ArcBallColor = FColor(128, 128, 128, 6);
+
 	ScreenSpaceColor  = FColor(196, 196, 196);
 	CurrentColor = FColor::Yellow;
 
@@ -202,6 +205,7 @@ FWidget::FWidget()
 	YAxisDir = FVector2D::ZeroVector;
 	ZAxisDir = FVector2D::ZeroVector;
 	DragStartPos = FVector2D::ZeroVector;
+	LastDragPos = FVector2D::ZeroVector;
 }
 
 extern ENGINE_API void StringSize(UFont* Font,int32& XL,int32& YL,const TCHAR* Text, FCanvas* Canvas);
@@ -637,6 +641,94 @@ void FWidget::Render_Translate( const FSceneView* View, FPrimitiveDrawInterface*
 	}
 }
 
+uint8 LargeInnerAlpha = 0x3f;
+uint8 SmallInnerAlpha = 0x0f;
+uint8 LargeOuterAlpha = 0x7f;
+uint8 SmallOuterAlpha = 0x0f;
+
+void FWidget::DrawColoredSphere(FPrimitiveDrawInterface* PDI, const FVector& Center, const FRotator& Orientation, FColor Color, const FVector& Radii, int32 NumSides, int32 NumRings, const FMaterialRenderProxy* MaterialRenderProxy, uint8 DepthPriority, bool bDisableBackfaceCulling)
+{
+	// Use a mesh builder to draw the sphere.
+	FDynamicMeshBuilder MeshBuilder(PDI->View->GetFeatureLevel());
+	{
+		// The first/last arc are on top of each other.
+		int32 NumVerts = (NumSides + 1) * (NumRings + 1);
+		FDynamicMeshVertex* Verts = (FDynamicMeshVertex*)FMemory::Malloc(NumVerts * sizeof(FDynamicMeshVertex));
+
+		// Calculate verts for one arc
+		FDynamicMeshVertex* ArcVerts = (FDynamicMeshVertex*)FMemory::Malloc((NumRings + 1) * sizeof(FDynamicMeshVertex));
+
+		for (int32 i = 0; i < NumRings + 1; i++)
+		{
+			FDynamicMeshVertex* ArcVert = &ArcVerts[i];
+
+			float angle = ((float)i / NumRings) * PI;
+
+			ArcVert->Color = Color;
+			// Note- unit sphere, so position always has mag of one. We can just use it for normal!			
+			ArcVert->Position.X = 0.0f;
+			ArcVert->Position.Y = FMath::Sin(angle);
+			ArcVert->Position.Z = FMath::Cos(angle);
+
+			ArcVert->SetTangents(
+				FVector(1, 0, 0),
+				FVector(0.0f, -ArcVert->Position.Z, ArcVert->Position.Y),
+				ArcVert->Position
+			);
+
+			ArcVert->TextureCoordinate[0].X = 0.0f;
+			ArcVert->TextureCoordinate[0].Y = ((float)i / NumRings);
+		}
+
+		// Then rotate this arc NumSides+1 times.
+		for (int32 s = 0; s < NumSides + 1; s++)
+		{
+			FRotator ArcRotator(0, 360.f * (float)s / NumSides, 0);
+			FRotationMatrix ArcRot(ArcRotator);
+			float XTexCoord = ((float)s / NumSides);
+
+			for (int32 v = 0; v < NumRings + 1; v++)
+			{
+				int32 VIx = (NumRings + 1)*s + v;
+				Verts[VIx].Color = Color;
+				Verts[VIx].Position = ArcRot.TransformPosition(ArcVerts[v].Position);
+
+				Verts[VIx].SetTangents(
+					ArcRot.TransformVector(ArcVerts[v].TangentX.ToFVector()),
+					ArcRot.TransformVector(ArcVerts[v].GetTangentY()),
+					ArcRot.TransformVector(ArcVerts[v].TangentZ.ToFVector())
+				);
+
+				Verts[VIx].TextureCoordinate[0].X = XTexCoord;
+				Verts[VIx].TextureCoordinate[0].Y = ArcVerts[v].TextureCoordinate[0].Y;
+			}
+		}
+
+		// Add all of the vertices we generated to the mesh builder.
+		for (int32 VertIdx = 0; VertIdx < NumVerts; VertIdx++)
+		{
+			MeshBuilder.AddVertex(Verts[VertIdx]);
+		}
+
+		// Add all of the triangles we generated to the mesh builder.
+		for (int32 s = 0; s < NumSides; s++)
+		{
+			int32 a0start = (s + 0) * (NumRings + 1);
+			int32 a1start = (s + 1) * (NumRings + 1);
+
+			for (int32 r = 0; r < NumRings; r++)
+			{
+				MeshBuilder.AddTriangle(a0start + r + 0, a1start + r + 0, a0start + r + 1);
+				MeshBuilder.AddTriangle(a1start + r + 0, a1start + r + 1, a0start + r + 1);
+			}
+		}
+
+		// Free our local copy of verts and arc verts
+		FMemory::Free(Verts);
+		FMemory::Free(ArcVerts);
+	}
+	MeshBuilder.Draw(PDI, FScaleMatrix(Radii) * FRotationMatrix(Orientation) * FTranslationMatrix(Center), MaterialRenderProxy, DepthPriority, bDisableBackfaceCulling);
+}
 /**
  * Draws the rotation widget.
  */
@@ -658,22 +750,62 @@ void FWidget::Render_Rotate( const FSceneView* View,FPrimitiveDrawInterface* PDI
 	if (bDrawWidget || bDragging)
 	{
 		bIsOrthoDrawingFullRing = false;
-
-		//no draw the arc segments
-		if( DrawAxis&EAxisList::X )
+		//now draw the arc segments
+		if (DrawAxis&EAxisList::X)
 		{
 			DrawRotationArc(View, PDI, EAxisList::X, InLocation, ZAxis, YAxis, DirectionToWidget, AxisColorX.ToFColor(true), Scale, XAxisDir);
 		}
 
-		if( DrawAxis&EAxisList::Y )
+		if (DrawAxis&EAxisList::Y)
 		{
 			DrawRotationArc(View, PDI, EAxisList::Y, InLocation, XAxis, ZAxis, DirectionToWidget, AxisColorY.ToFColor(true), Scale, YAxisDir);
 		}
 
-		if( DrawAxis&EAxisList::Z )
+		if (DrawAxis&EAxisList::Z)
 		{
 			DrawRotationArc(View, PDI, EAxisList::Z, InLocation, XAxis, YAxis, DirectionToWidget, AxisColorZ.ToFColor(true), Scale, ZAxisDir);
 		}
+
+		if (DrawAxis&EAxisList::XYZ  && (!bDragging || CurrentAxis == EAxisList::XYZ) && GetDefault<ULevelEditorViewportSettings>()->bAllowArcballRotate)
+		{
+			FVector Center = InLocation;
+			FRotator Orientation = FRotator::ZeroRotator;
+			const float InnerDistance = (INNER_AXIS_CIRCLE_RADIUS *0.925f * Scale) + GetDefault<ULevelEditorViewportSettings>()->TransformWidgetSizeAdjustment;
+			FVector Radii(InnerDistance, InnerDistance, InnerDistance);
+			const bool bDisabled = IsWidgetDisabled();
+			FColor ArcColor = ArcBallColor;
+			ArcColor.A = (CurrentAxis==EAxisList::XYZ) ? ArcBallColor.A + 5 : ArcBallColor.A; //less transparent if selected
+
+			PDI->SetHitProxy(new HWidgetAxis(EAxisList::XYZ, bDisabled));
+			{
+				DrawColoredSphere(PDI, InLocation, FRotator::ZeroRotator, ArcColor, Radii, 32, 24, TransparentPlaneMaterialXY->GetRenderProxy(), SDPG_Foreground, true);
+			}
+			PDI->SetHitProxy(NULL);
+		}
+		
+		if(DrawAxis&EAxisList::Screen && (!bDragging || CurrentAxis == EAxisList::Screen) && GetDefault<ULevelEditorViewportSettings>()->bAllowScreenRotate)
+		{
+			const bool bIsPerspective = (View->ViewMatrices.GetProjectionMatrix().M[3][3] < 1.0f);
+			const bool bIsOrtho = !bIsPerspective;
+
+			const FVector Axis0 = View->GetViewUp();
+			const FVector Axis1 = View->GetViewRight();
+			
+			const float OuterRadius = (OUTER_AXIS_CIRCLE_RADIUS * 1.25f * Scale) + GetDefault<ULevelEditorViewportSettings>()->TransformWidgetSizeAdjustment;
+			const float InnerRadius = ((OUTER_AXIS_CIRCLE_RADIUS -1) * 1.25f * Scale) + GetDefault<ULevelEditorViewportSettings>()->TransformWidgetSizeAdjustment;
+
+			FColor ArcColor = (CurrentAxis==EAxisList::Screen) ? CurrentColor : ScreenAxisColor.ToFColor(true);
+
+			const bool bDisabled = IsWidgetDisabled();
+
+			PDI->SetHitProxy(new HWidgetAxis(EAxisList::Screen, bDisabled));
+			{
+				FThickArcParams OuterArcParams(PDI, InLocation, TransparentPlaneMaterialXY, InnerRadius, OuterRadius);
+				//Pass through alpha
+				DrawThickArc(OuterArcParams, Axis0, Axis1, 0.0f, 2.0f * PI, ArcColor, DirectionToWidget, !bIsPerspective);
+			}
+			PDI->SetHitProxy(NULL);
+		}
 	}
 }
 
@@ -989,7 +1121,7 @@ void FWidget::Render_2D(const FSceneView* View, FPrimitiveDrawInterface* PDI, FE
 /**
  * Converts mouse movement on the screen to widget axis movement/rotation.
  */
-void FWidget::ConvertMouseMovementToAxisMovement( FEditorViewportClient* InViewportClient, bool bInUsedDragModifier, FVector& InOutDelta, FVector& OutDrag, FRotator& OutRotation, FVector& OutScale )
+void FWidget::ConvertMouseMovementToAxisMovement(FSceneView* InView, FEditorViewportClient* InViewportClient, bool bInUsedDragModifier, FVector& InOutDelta, FVector& OutDrag, FRotator& OutRotation, FVector& OutScale )
 {
 	OutDrag = FVector::ZeroVector;
 	OutRotation = FRotator::ZeroRotator;
@@ -1043,11 +1175,11 @@ void FWidget::ConvertMouseMovementToAxisMovement( FEditorViewportClient* InViewp
 
 		case WM_Rotate:
 			{
-				FRotator Rotation;
-				FVector2D EffectiveDelta;
 
 				if (CurrentAxis == EAxisList::X)
 				{
+					FRotator Rotation;
+					FVector2D EffectiveDelta;
 					// Get screen direction representing positive rotation
 					const FVector2D AxisDir = bIsOrthoDrawingFullRing ? TangentDir : XAxisDir;
 
@@ -1060,9 +1192,15 @@ void FWidget::ConvertMouseMovementToAxisMovement( FEditorViewportClient* InViewp
 
 					// Use to calculate the new input delta
 					EffectiveDelta = AxisDir * Rotation.Roll;
+					// Adjust the input delta according to how much rotation was actually applied
+					InOutDelta = FVector(EffectiveDelta.X, -EffectiveDelta.Y, 0.0f);
+					// Need to get the delta rotation in the current coordinate space of the widget
+					OutRotation = (CustomCoordSystem.Inverse() * FRotationMatrix(Rotation) * CustomCoordSystem).Rotator();
 				}
 				else if (CurrentAxis == EAxisList::Y)
 				{
+					FRotator Rotation;
+					FVector2D EffectiveDelta;
 					// TODO: Determine why -TangentDir is necessary here, and fix whatever is causing it
 					const FVector2D AxisDir = bIsOrthoDrawingFullRing ? -TangentDir : YAxisDir;
 
@@ -1071,9 +1209,15 @@ void FWidget::ConvertMouseMovementToAxisMovement( FEditorViewportClient* InViewp
 
 					CurrentDeltaRotation = Rotation.Pitch;
 					EffectiveDelta = AxisDir * Rotation.Pitch;
+					// Adjust the input delta according to how much rotation was actually applied
+					InOutDelta = FVector(EffectiveDelta.X, -EffectiveDelta.Y, 0.0f);
+					// Need to get the delta rotation in the current coordinate space of the widget
+					OutRotation = (CustomCoordSystem.Inverse() * FRotationMatrix(Rotation) * CustomCoordSystem).Rotator();
 				}
 				else if (CurrentAxis == EAxisList::Z)
 				{
+					FRotator Rotation;
+					FVector2D EffectiveDelta;
 					const FVector2D AxisDir = bIsOrthoDrawingFullRing ? TangentDir : ZAxisDir;
 
 					Rotation = FRotator(0, FVector2D::DotProduct(AxisDir, DragDir), 0);
@@ -1081,13 +1225,95 @@ void FWidget::ConvertMouseMovementToAxisMovement( FEditorViewportClient* InViewp
 
 					CurrentDeltaRotation = Rotation.Yaw;
 					EffectiveDelta = AxisDir * Rotation.Yaw;
+					// Adjust the input delta according to how much rotation was actually applied
+					InOutDelta = FVector(EffectiveDelta.X, -EffectiveDelta.Y, 0.0f);
+					// Need to get the delta rotation in the current coordinate space of the widget
+					OutRotation = (CustomCoordSystem.Inverse() * FRotationMatrix(Rotation) * CustomCoordSystem).Rotator();
 				}
+				else if (CurrentAxis == EAxisList::XYZ) //arcball rotate
+				{
+					//To do this we need to calculate the rotation axis and rotation angle.
+					//The Axis is the cross product of the current ray from eye to pixel in world space with the previous ray.
+					//The Angle is angle amount we rotate from the object's location to the imaginary sphere that matches up with the difference between the current and previous ray
+					//From the Camera. Those rays form a triangle, which we can bisect with a common side.
+					FVector2D MousePosition(InViewportClient->Viewport->GetMouseX(), InViewportClient->Viewport->GetMouseY());
+					FViewportCursorLocation OldMouseViewportRay(InView, InViewportClient, LastDragPos.X, LastDragPos.Y);
+					FViewportCursorLocation MouseViewportRay(InView, InViewportClient, MousePosition.X, MousePosition.Y);
+					
+					LastDragPos = MousePosition;
+					FVector DirectionToWidget = InViewportClient->GetWidgetLocation() - MouseViewportRay.GetOrigin();
+					float Length = DirectionToWidget.Size();
+					if (!FMath::IsNearlyZero(Length))
+					{
+						//Calc Axis
+						DirectionToWidget /= Length;
+						const FVector CameraToPixelDir = MouseViewportRay.GetDirection();						
+						const FVector OldCameraToPixelDir = OldMouseViewportRay.GetDirection();
+						FVector RotationAxis = FVector::CrossProduct(OldCameraToPixelDir, CameraToPixelDir);
+						RotationAxis.Normalize();
+						float RotationAngle = 0.0f;
+						FVector4 ScreenLocation = InView->WorldToScreen(InViewportClient->GetWidgetLocation());
+						FVector2D PixelLocation;
+						InView->ScreenToPixel(ScreenLocation, PixelLocation);
+						float Distance = FVector2D::Distance(PixelLocation, MousePosition);
+						const float MaxDiff = 2.0f * INNER_AXIS_CIRCLE_RADIUS;
+						//If outside radius, do screen rotate instead, like other DCC's
+						if (Distance > MaxDiff)
+						{
+							RotationAngle = FMath::Acos(FVector::DotProduct(OldCameraToPixelDir, CameraToPixelDir));
+							FVector Cross = FVector::CrossProduct(OldCameraToPixelDir, CameraToPixelDir);
+							if (FVector::DotProduct(DirectionToWidget, Cross) < 0.0f)
+							{
+								RotationAngle *= -1.0f;
+							}
+							RotationAxis = DirectionToWidget;
+						}
+						else //need to calc angle
+						{
+							const float Scale = ScreenLocation.W * (4.0f / InView->UnscaledViewRect.Width() / InView->ViewMatrices.GetProjectionMatrix().M[0][0]);
+							const float InnerRadius = (INNER_AXIS_CIRCLE_RADIUS * Scale) + GetDefault<ULevelEditorViewportSettings>()->TransformWidgetSizeAdjustment;
+							const float LengthOfAdjacent = Length - InnerRadius;
+							RotationAngle = FMath::Acos(FVector::DotProduct(DirectionToWidget, CameraToPixelDir));
+							RotationAngle = FMath::Acos(FVector::DotProduct(OldCameraToPixelDir, CameraToPixelDir));
+							const float OppositeSize = FMath::Tan(RotationAngle) * LengthOfAdjacent;
+							RotationAngle = FMath::Atan(OppositeSize / InnerRadius);
+							RotationAngle = RotationAngle < 0.0f ? RotationAngle : -RotationAngle;
+
+						}
+						const FQuat QuatRotation(RotationAxis, RotationAngle);
+						OutRotation = FRotator(QuatRotation);
+						FSnappingUtils::SnapRotatorToGrid(OutRotation);
+					}
+					break;
+				}
+				else if(CurrentAxis == EAxisList::Screen)
+				{
+					FVector2D MousePosition(InViewportClient->Viewport->GetMouseX(), InViewportClient->Viewport->GetMouseY());
+					FViewportCursorLocation OldMouseViewportRay(InView, InViewportClient, LastDragPos.X, LastDragPos.Y);
+					FViewportCursorLocation MouseViewportRay(InView, InViewportClient, MousePosition.X, MousePosition.Y);
 
-				// Adjust the input delta according to how much rotation was actually applied
-				InOutDelta = FVector(EffectiveDelta.X, -EffectiveDelta.Y, 0.0f);
+					LastDragPos = MousePosition;
+					FVector DirectionToWidget = InViewportClient->GetWidgetLocation() - MouseViewportRay.GetOrigin();
+					float Length = DirectionToWidget.Size();
 
-				// Need to get the delta rotation in the current coordinate space of the widget
-				OutRotation = (CustomCoordSystem.Inverse() * FRotationMatrix(Rotation) * CustomCoordSystem).Rotator();
+					if (!FMath::IsNearlyZero(Length))
+					{
+						DirectionToWidget /= Length;
+
+						const FVector CameraToPixelDir = MouseViewportRay.GetDirection();
+						const FVector OldCameraToPixelDir = OldMouseViewportRay.GetDirection();
+						float RotationAngle = FMath::Acos(FVector::DotProduct(OldCameraToPixelDir, CameraToPixelDir));
+						FVector Cross =  FVector::CrossProduct(OldCameraToPixelDir, CameraToPixelDir);
+						if (FVector::DotProduct(DirectionToWidget, Cross) < 0.0f)
+						{
+							RotationAngle *= -1.0f;
+						}	
+						const FQuat QuatRotation(DirectionToWidget, RotationAngle);
+						OutRotation = FRotator(QuatRotation);
+						FSnappingUtils::SnapRotatorToGrid(OutRotation);
+					}
+					break;
+				}
 			}
 			break;
 
@@ -1463,6 +1689,16 @@ bool FWidget::AllowsAbsoluteTranslationMovement(EWidgetMode WidgetMode)
 	return false;
 }
 
+/** Only some modes support Absolute Rotation Movement/arcball*/
+bool FWidget::AllowsAbsoluteRotationMovement(EWidgetMode WidgetMode, EAxisList::Type InAxisType)
+{
+	if (WidgetMode == WM_Rotate && (InAxisType == EAxisList::XYZ || InAxisType == EAxisList::Screen))
+	{
+		return true;
+	}
+	return false;
+}
+
 /** 
  * Serializes the widget references so they don't get garbage collected.
  *
@@ -1590,11 +1826,6 @@ float FWidget::GetDeltaRotation() const
 }
 
 
-uint8 LargeInnerAlpha = 0x3f;
-uint8 SmallInnerAlpha = 0x0f;
-uint8 LargeOuterAlpha = 0x7f;
-uint8 SmallOuterAlpha = 0x0f;
-
 /**
  * If actively dragging, draws a ring representing the potential rotation of the selected objects, snap ticks, and "delta" markers
  * If not actively dragging, draws a quarter ring representing the closest quadrant to the camera
@@ -1623,6 +1854,8 @@ void FWidget::DrawRotationArc(const FSceneView* View, FPrimitiveDrawInterface* P
 	{
 		if ((CurrentAxis&InAxis) || (bIsOrthoDrawingFullRing))
 		{
+			bool bDrawingArcBallAlso = CurrentAxis == EAxisList::XYZ;
+
 			float DeltaRotation = GetDeltaRotation();
 			float AdjustedDeltaRotation = IsRotationLocalSpace() ? -DeltaRotation : DeltaRotation;
 			float AbsRotation = FRotator::ClampAxis(FMath::Abs(DeltaRotation));
@@ -1640,57 +1873,60 @@ void FWidget::DrawRotationArc(const FSceneView* View, FPrimitiveDrawInterface* P
 			ArcColor.A = SmallOuterAlpha;
 			DrawPartialRotationArc(View, PDI, InAxis, InLocation,  Axis0, Axis1, StartAngle + FilledAngle, StartAngle + 2*PI, ArcColor, InScale, InDirectionToWidget);
 
-			ArcColor = (CurrentAxis&InAxis) ? CurrentColor : ArcColor;
-			//Hallow Arrow
-			ArcColor.A = 0;
-			DrawStartStopMarker(PDI, InLocation, Axis0, Axis1, 0, ArcColor, InScale);
-			//Filled Arrow
-			ArcColor.A = LargeOuterAlpha;
-			DrawStartStopMarker(PDI, InLocation, Axis0, Axis1, AdjustedDeltaRotation, ArcColor, InScale);
+			ArcColor = (CurrentAxis&InAxis && !bDrawingArcBallAlso) ? CurrentColor : ArcColor;
+			if (!bDrawingArcBallAlso)
+			{
+				//Hallow Arrow
+				ArcColor.A = 0;
+				DrawStartStopMarker(PDI, InLocation, Axis0, Axis1, 0, ArcColor, InScale);
+				//Filled Arrow
+				ArcColor.A = LargeOuterAlpha;
+				DrawStartStopMarker(PDI, InLocation, Axis0, Axis1, AdjustedDeltaRotation, ArcColor, InScale);
+				ArcColor.A = 255;
 
-			ArcColor.A = 255;
 
-			FVector SnapLocation = InLocation;
+				FVector SnapLocation = InLocation;
 
-			if (GetDefault<ULevelEditorViewportSettings>()->RotGridEnabled)
-			{
-				float DeltaAngle = GEditor->GetRotGridSize().Yaw;
-				//every 22.5 degrees
-				float TickMarker = 22.5f;
-				for (float Angle = 0; Angle < 360.f; Angle+=DeltaAngle)
-				{ 
-					FVector GridAxis = Axis0.RotateAngleAxis(Angle, ZAxis);
-					float PercentSize = (FMath::Fmod(Angle, TickMarker)==0) ? .75f : .25f;
-					if (FMath::Fmod(Angle, 90.f) != 0)
+				if (GetDefault<ULevelEditorViewportSettings>()->RotGridEnabled)
+				{
+					float DeltaAngle = GEditor->GetRotGridSize().Yaw;
+					//every 22.5 degrees
+					float TickMarker = 22.5f;
+					for (float Angle = 0; Angle < 360.f; Angle += DeltaAngle)
 					{
-						DrawSnapMarker(PDI, SnapLocation,  GridAxis,  FVector::ZeroVector, ArcColor, InScale, 0.0f, PercentSize);
+						FVector GridAxis = Axis0.RotateAngleAxis(Angle, ZAxis);
+						float PercentSize = (FMath::Fmod(Angle, TickMarker) == 0) ? .75f : .25f;
+						if (FMath::Fmod(Angle, 90.f) != 0)
+						{
+							DrawSnapMarker(PDI, SnapLocation, GridAxis, FVector::ZeroVector, ArcColor, InScale, 0.0f, PercentSize);
+						}
 					}
 				}
-			}
 
-			//draw axis tick marks
-			FColor AxisColor = InColor;
-			//Rotate Colors to match Axis 0
-			Swap(AxisColor.R, AxisColor.G);
-			Swap(AxisColor.B, AxisColor.R);
-			AxisColor.A = (AdjustedDeltaRotation == 0) ? MAX_uint8 : LargeOuterAlpha;
-			DrawSnapMarker(PDI, SnapLocation,  Axis0,  Axis1, AxisColor, InScale, .25f);
-			AxisColor.A = (AdjustedDeltaRotation == 180.f) ? MAX_uint8 : LargeOuterAlpha;
-			DrawSnapMarker(PDI, SnapLocation, -Axis0, -Axis1, AxisColor, InScale, .25f);
-
-			//Rotate Colors to match Axis 1
-			Swap(AxisColor.R, AxisColor.G);
-			Swap(AxisColor.B, AxisColor.R);
-			AxisColor.A = (AdjustedDeltaRotation == 90.f) ? MAX_uint8 : LargeOuterAlpha;
-			DrawSnapMarker(PDI, SnapLocation,  Axis1, -Axis0, AxisColor, InScale, .25f);
-			AxisColor.A = (AdjustedDeltaRotation == 270.f) ? MAX_uint8 : LargeOuterAlpha;
-			DrawSnapMarker(PDI, SnapLocation, -Axis1,  Axis0, AxisColor, InScale, .25f);
-
-			if (bDragging)
-			{
-				float OffsetAngle = IsRotationLocalSpace() ? 0 : AdjustedDeltaRotation;
+				//draw axis tick marks
+				FColor AxisColor = InColor;
+				//Rotate Colors to match Axis 0
+				Swap(AxisColor.R, AxisColor.G);
+				Swap(AxisColor.B, AxisColor.R);
+				AxisColor.A = (AdjustedDeltaRotation == 0) ? MAX_uint8 : LargeOuterAlpha;
+				DrawSnapMarker(PDI, SnapLocation, Axis0, Axis1, AxisColor, InScale, .25f);
+				AxisColor.A = (AdjustedDeltaRotation == 180.f) ? MAX_uint8 : LargeOuterAlpha;
+				DrawSnapMarker(PDI, SnapLocation, -Axis0, -Axis1, AxisColor, InScale, .25f);
+
+				//Rotate Colors to match Axis 1
+				Swap(AxisColor.R, AxisColor.G);
+				Swap(AxisColor.B, AxisColor.R);
+				AxisColor.A = (AdjustedDeltaRotation == 90.f) ? MAX_uint8 : LargeOuterAlpha;
+				DrawSnapMarker(PDI, SnapLocation, Axis1, -Axis0, AxisColor, InScale, .25f);
+				AxisColor.A = (AdjustedDeltaRotation == 270.f) ? MAX_uint8 : LargeOuterAlpha;
+				DrawSnapMarker(PDI, SnapLocation, -Axis1, Axis0, AxisColor, InScale, .25f);
+
+				if (bDragging)
+				{
+					float OffsetAngle = IsRotationLocalSpace() ? 0 : AdjustedDeltaRotation;
 
-				CacheRotationHUDText(View, PDI, InLocation, Axis0.RotateAngleAxis(OffsetAngle, ZAxis), Axis1.RotateAngleAxis(OffsetAngle, ZAxis), DeltaRotation, InScale);
+					CacheRotationHUDText(View, PDI, InLocation, Axis0.RotateAngleAxis(OffsetAngle, ZAxis), Axis1.RotateAngleAxis(OffsetAngle, ZAxis), DeltaRotation, InScale);
+				}
 			}
 		}
 	}
@@ -1745,7 +1981,7 @@ void FWidget::DrawPartialRotationArc(const FSceneView* View, FPrimitiveDrawInter
 	PDI->SetHitProxy( new HWidgetAxis( InAxis ) );
 	{
 		FThickArcParams OuterArcParams(PDI, InLocation, TransparentPlaneMaterialXY, InnerRadius, OuterRadius);
-		FColor OuterColor = ( CurrentAxis&InAxis ? CurrentColor : InColor );
+		FColor OuterColor = ( (CurrentAxis&InAxis  && ( CurrentAxis != EAxisList::XYZ ) )? CurrentColor : InColor );
 		//Pass through alpha
 		OuterColor.A = InColor.A;
 		DrawThickArc(OuterArcParams, Axis0, Axis1, InStartAngle, InEndAngle, OuterColor, InDirectionToWidget, !bIsPerspective );
diff --git a/Engine/Source/Editor/UnrealEd/Public/EditorViewportClient.h b/Engine/Source/Editor/UnrealEd/Public/EditorViewportClient.h
--- a/Engine/Source/Editor/UnrealEd/Public/EditorViewportClient.h
+++ b/Engine/Source/Editor/UnrealEd/Public/EditorViewportClient.h
@@ -949,7 +949,7 @@ public:
 	void MarkMouseMovedSinceClick();
 
 	/** Determines whether this viewport is currently allowed to use Absolute Movement */
-	bool IsUsingAbsoluteTranslation() const;
+	bool IsUsingAbsoluteTranslation(bool bAlsoCheckAbsoluteRotation = false) const;
 
 	bool IsForcedRealtimeAudio() const { return bForceAudioRealtime; }
 
diff --git a/Engine/Source/Editor/UnrealEd/Public/MouseDeltaTracker.h b/Engine/Source/Editor/UnrealEd/Public/MouseDeltaTracker.h
--- a/Engine/Source/Editor/UnrealEd/Public/MouseDeltaTracker.h
+++ b/Engine/Source/Editor/UnrealEd/Public/MouseDeltaTracker.h
@@ -87,7 +87,7 @@ public:
 	/**
 	 * Converts the delta movement to drag/rotation/scale based on the viewport type or widget axis.
 	 */
-	void UNREALED_API ConvertMovementDeltaToDragRot(FEditorViewportClient* InViewportClient, FVector& InDragDelta, FVector& OutDrag, FRotator& OutRotation, FVector& OutScale) const;
+	void UNREALED_API ConvertMovementDeltaToDragRot(FSceneView* InView, FEditorViewportClient* InViewportClient, FVector& InDragDelta, FVector& OutDrag, FRotator& OutRotation, FVector& OutScale) const;
 	/**
 	 * Absolute Translation conversion from mouse position on the screen to widget axis movement/rotation.
 	 */
diff --git a/Engine/Source/Editor/UnrealEd/Public/UnrealWidget.h b/Engine/Source/Editor/UnrealEd/Public/UnrealWidget.h
--- a/Engine/Source/Editor/UnrealEd/Public/UnrealWidget.h
+++ b/Engine/Source/Editor/UnrealEd/Public/UnrealWidget.h
@@ -96,7 +96,7 @@ public:
 	/**
 	 * Converts mouse movement on the screen to widget axis movement/rotation.
 	 */
-	void ConvertMouseMovementToAxisMovement( FEditorViewportClient* InViewportClient, bool bInUsedDragModifier, FVector& InDiff, FVector& OutDrag, FRotator& OutRotation, FVector& OutScale );
+	void ConvertMouseMovementToAxisMovement(FSceneView* InView,  FEditorViewportClient* InViewportClient, bool bInUsedDragModifier, FVector& InDiff, FVector& OutDrag, FRotator& OutRotation, FVector& OutScale );
 
 	/**
 	 * Absolute Translation conversion from mouse movement on the screen to widget axis movement/rotation.
@@ -114,6 +114,8 @@ public:
 	/** Only some modes support Absolute Translation Movement.  Check current mode */
 	static bool AllowsAbsoluteTranslationMovement(FWidget::EWidgetMode WidgetMode);
 
+	/** Only some modes support Absolute Rotation Movement.  Check current mode */
+	static bool AllowsAbsoluteRotationMovement(EWidgetMode WidgetMode, EAxisList::Type InAxisType);
 	/**
 	 * Sets the default visibility of the widget, if it is not overridden by an active editor mode tool.
 	 *
@@ -156,6 +158,7 @@ public:
 	void SetDragStartPosition(const FVector2D& Position)
 	{
 		DragStartPos = Position;
+		LastDragPos = DragStartPos;
 	}
 
 	/**
@@ -371,6 +374,9 @@ private:
 	 */
 	uint32 GetDominantAxisIndex( const FVector& InDiff, FEditorViewportClient* ViewportClient ) const;
 
+
+	void DrawColoredSphere(FPrimitiveDrawInterface* PDI, const FVector& Center, const FRotator& Orientation, FColor Color, const FVector& Radii, int32 NumSides, int32 NumRings, const FMaterialRenderProxy* MaterialRenderProxy, uint8 DepthPriority, bool bDisableBackfaceCulling);
+
 	/** The axis currently being moused over */
 	EAxisList::Type CurrentAxis;
 
@@ -380,7 +386,8 @@ private:
 	FVector2D XAxisDir, YAxisDir, ZAxisDir;
 	/** Drag start position in viewport space */
 	FVector2D DragStartPos;
-
+	/** Last mouse position in viewport space */
+	FVector2D LastDragPos;
 	enum
 	{
 		AXIS_ARROW_SEGMENTS = 16
@@ -397,7 +404,9 @@ private:
 	UMaterialInstanceDynamic* OpaquePlaneMaterialXY;
 
 	FLinearColor AxisColorX, AxisColorY, AxisColorZ;
+	FLinearColor ScreenAxisColor;
 	FColor PlaneColorXY, ScreenSpaceColor, CurrentColor;
+	FColor ArcBallColor;
 
 	/** Any mode tools being used */
 	FEditorModeTools* EditorModeTools;
