Correct tooltips for new rotation widgets.-

diff --git a/Engine/Source/Editor/UnrealEd/Classes/Settings/LevelEditorViewportSettings.h b/Engine/Source/Editor/UnrealEd/Classes/Settings/LevelEditorViewportSettings.h
--- a/Engine/Source/Editor/UnrealEd/Classes/Settings/LevelEditorViewportSettings.h
+++ b/Engine/Source/Editor/UnrealEd/Classes/Settings/LevelEditorViewportSettings.h
@@ -238,11 +238,11 @@ class UNREALED_API ULevelEditorViewportSettings
 	UPROPERTY(EditAnywhere, config, Category=LookAndFeel, meta=( DisplayName = "Enable Combined Translate/Rotate Widget" ))
 	uint32 bAllowTranslateRotateZWidget:1;
 
-	/** Allow translate/rotate widget */
+	/** Allow arcball rotation with rotate widget */
 	UPROPERTY(EditAnywhere, config, Category = LookAndFeel, meta = (DisplayName = "Enable Arcball Rotate"))
 	uint32 bAllowArcballRotate : 1;
 
-	/** Allow translate/rotate widget */
+	/** Allow screen rotation with rotate widget */
 	UPROPERTY(EditAnywhere, config, Category = LookAndFeel, meta = (DisplayName = "Enable Screen Rotate"))
 	uint32 bAllowScreenRotate : 1;
 
