Calc screen angle from correct widget location and disable snapping

diff --git a/Engine/Source/Editor/UnrealEd/Private/UnrealWidget.cpp b/Engine/Source/Editor/UnrealEd/Private/UnrealWidget.cpp
--- a/Engine/Source/Editor/UnrealEd/Private/UnrealWidget.cpp
+++ b/Engine/Source/Editor/UnrealEd/Private/UnrealWidget.cpp
@@ -1258,9 +1258,18 @@ void FWidget::ConvertMouseMovementToAxisMovement(FSceneView* InView, FEditorView
 						float Distance = FVector2D::Distance(PixelLocation, MousePosition);
 						const float MaxDiff = 2.0f * INNER_AXIS_CIRCLE_RADIUS;
 						//If outside radius, do screen rotate instead, like other DCC's
+
 						if (Distance > MaxDiff)
 						{
-							RotationAngle = FMath::Acos(FVector::DotProduct(OldCameraToPixelDir, CameraToPixelDir));
+							FPlane Plane(InViewportClient->GetWidgetLocation(), DirectionToWidget);
+							FVector StartOnPlane = FMath::RayPlaneIntersection(MouseViewportRay.GetOrigin(), CameraToPixelDir, Plane);
+							FVector OldOnPlane = FMath::RayPlaneIntersection(MouseViewportRay.GetOrigin(), OldCameraToPixelDir, Plane);
+							StartOnPlane -= InViewportClient->GetWidgetLocation();
+							OldOnPlane -= InViewportClient->GetWidgetLocation();
+							StartOnPlane.Normalize();
+							OldOnPlane.Normalize();
+							RotationAngle = FMath::Acos(FVector::DotProduct(StartOnPlane, OldOnPlane));
+							
 							FVector Cross = FVector::CrossProduct(OldCameraToPixelDir, CameraToPixelDir);
 							if (FVector::DotProduct(DirectionToWidget, Cross) < 0.0f)
 							{
@@ -1268,21 +1277,19 @@ void FWidget::ConvertMouseMovementToAxisMovement(FSceneView* InView, FEditorView
 							}
 							RotationAxis = DirectionToWidget;
 						}
-						else //need to calc angle
+						else
 						{
 							const float Scale = ScreenLocation.W * (4.0f / InView->UnscaledViewRect.Width() / InView->ViewMatrices.GetProjectionMatrix().M[0][0]);
 							const float InnerRadius = (INNER_AXIS_CIRCLE_RADIUS * Scale) + GetDefault<ULevelEditorViewportSettings>()->TransformWidgetSizeAdjustment;
 							const float LengthOfAdjacent = Length - InnerRadius;
-							RotationAngle = FMath::Acos(FVector::DotProduct(DirectionToWidget, CameraToPixelDir));
 							RotationAngle = FMath::Acos(FVector::DotProduct(OldCameraToPixelDir, CameraToPixelDir));
 							const float OppositeSize = FMath::Tan(RotationAngle) * LengthOfAdjacent;
 							RotationAngle = FMath::Atan(OppositeSize / InnerRadius);
 							RotationAngle = RotationAngle < 0.0f ? RotationAngle : -RotationAngle;
-
 						}
+				
 						const FQuat QuatRotation(RotationAxis, RotationAngle);
 						OutRotation = FRotator(QuatRotation);
-						FSnappingUtils::SnapRotatorToGrid(OutRotation);
 					}
 					break;
 				}
@@ -1302,7 +1309,14 @@ void FWidget::ConvertMouseMovementToAxisMovement(FSceneView* InView, FEditorView
 
 						const FVector CameraToPixelDir = MouseViewportRay.GetDirection();
 						const FVector OldCameraToPixelDir = OldMouseViewportRay.GetDirection();
-						float RotationAngle = FMath::Acos(FVector::DotProduct(OldCameraToPixelDir, CameraToPixelDir));
+						FPlane Plane(InViewportClient->GetWidgetLocation(), DirectionToWidget);
+						FVector StartOnPlane = FMath::RayPlaneIntersection(MouseViewportRay.GetOrigin(), CameraToPixelDir, Plane);
+						FVector OldOnPlane = FMath::RayPlaneIntersection(MouseViewportRay.GetOrigin(), OldCameraToPixelDir, Plane);
+						StartOnPlane -= InViewportClient->GetWidgetLocation();
+						OldOnPlane -= InViewportClient->GetWidgetLocation();
+						StartOnPlane.Normalize();
+						OldOnPlane.Normalize();
+						float RotationAngle = FMath::Acos(FVector::DotProduct(StartOnPlane, OldOnPlane));
 						FVector Cross =  FVector::CrossProduct(OldCameraToPixelDir, CameraToPixelDir);
 						if (FVector::DotProduct(DirectionToWidget, Cross) < 0.0f)
 						{
@@ -1310,7 +1324,6 @@ void FWidget::ConvertMouseMovementToAxisMovement(FSceneView* InView, FEditorView
 						}	
 						const FQuat QuatRotation(DirectionToWidget, RotationAngle);
 						OutRotation = FRotator(QuatRotation);
-						FSnappingUtils::SnapRotatorToGrid(OutRotation);
 					}
 					break;
 				}
