Added forward declaration

diff --git a/Engine/Source/Editor/UnrealEd/Public/UnrealWidget.h b/Engine/Source/Editor/UnrealEd/Public/UnrealWidget.h
--- a/Engine/Source/Editor/UnrealEd/Public/UnrealWidget.h
+++ b/Engine/Source/Editor/UnrealEd/Public/UnrealWidget.h
@@ -13,6 +13,7 @@ class FPrimitiveDrawInterface;
 class FSceneView;
 class UMaterialInstanceDynamic;
 class UMaterialInterface;
+class FMaterialRenderProxy;
 
 /** Coordinate system identifiers. */
 enum ECoordSystem
