Based on feedback, always show local transform when rotating arcball.  Also made ball little smaller to help with hittesting with the existing rotation axis. Just expose the local widget transform, had to pass it up through EditTools into the ViewportClient.
Also default the arcball to on.

diff --git a/Engine/Source/Editor/UnrealEd/Private/EditorModeManager.cpp b/Engine/Source/Editor/UnrealEd/Private/EditorModeManager.cpp
--- a/Engine/Source/Editor/UnrealEd/Private/EditorModeManager.cpp
+++ b/Engine/Source/Editor/UnrealEd/Private/EditorModeManager.cpp
@@ -470,34 +470,7 @@ FMatrix FEditorModeTools::GetCustomDrawingCoordinateSystem()
 	{
 		case COORD_Local:
 		{
-			// Let the current mode have a shot at setting the local coordinate system.
-			// If it doesn't want to, create it by looking at the currently selected actors list.
-
-			bool CustomCoordinateSystemProvided = false;
-			for (const auto& Mode : Modes)
-			{
-				if (Mode->GetCustomDrawingCoordinateSystem(Matrix, nullptr))
-				{
-					CustomCoordinateSystemProvided = true;
-					break;
-				}
-			}
-
-			if (!CustomCoordinateSystemProvided)
-			{
-				const int32 Num = GetSelectedActors()->CountSelections<AActor>();
-
-				// Coordinate system needs to come from the last actor selected
-				if (Num > 0)
-				{
-					Matrix = FQuatRotationMatrix(GetSelectedActors()->GetBottom<AActor>()->GetActorQuat());
-				}
-			}
-
-			if (!Matrix.Equals(FMatrix::Identity))
-			{
-				Matrix.RemoveScaling();
-			}
+			Matrix = GetLocalCoordinateSystem();
 		}
 		break;
 
@@ -516,6 +489,42 @@ FMatrix FEditorModeTools::GetCustomInputCoordinateSystem()
 	return GetCustomDrawingCoordinateSystem();
 }
 
+FMatrix FEditorModeTools::GetLocalCoordinateSystem()
+{
+	FMatrix Matrix = FMatrix::Identity;
+	// Let the current mode have a shot at setting the local coordinate system.
+	// If it doesn't want to, create it by looking at the currently selected actors list.
+
+	bool CustomCoordinateSystemProvided = false;
+	for (const auto& Mode : Modes)
+	{
+		if (Mode->GetCustomDrawingCoordinateSystem(Matrix, nullptr))
+		{
+			CustomCoordinateSystemProvided = true;
+			break;
+		}
+	}
+
+	if (!CustomCoordinateSystemProvided)
+	{
+		const int32 Num = GetSelectedActors()->CountSelections<AActor>();
+
+		// Coordinate system needs to come from the last actor selected
+		if (Num > 0)
+		{
+			Matrix = FQuatRotationMatrix(GetSelectedActors()->GetBottom<AActor>()->GetActorQuat());
+		}
+	}
+
+	if (!Matrix.Equals(FMatrix::Identity))
+	{
+		Matrix.RemoveScaling();
+	}
+
+	return Matrix;
+}
+
+
 /** Gets the widget axis to be drawn */
 EAxisList::Type FEditorModeTools::GetWidgetAxisToDraw( FWidget::EWidgetMode InWidgetMode ) const
 {
diff --git a/Engine/Source/Editor/UnrealEd/Private/EditorViewportClient.cpp b/Engine/Source/Editor/UnrealEd/Private/EditorViewportClient.cpp
--- a/Engine/Source/Editor/UnrealEd/Private/EditorViewportClient.cpp
+++ b/Engine/Source/Editor/UnrealEd/Private/EditorViewportClient.cpp
@@ -3439,6 +3439,11 @@ FMatrix FEditorViewportClient::GetWidgetCoordSystem() const
 	return ModeTools->GetCustomInputCoordinateSystem();
 }
 
+FMatrix FEditorViewportClient::GetLocalCoordinateSystem() const
+{
+	return ModeTools->GetLocalCoordinateSystem();
+}
+
 void FEditorViewportClient::SetWidgetCoordSystemSpace(ECoordSystem NewCoordSystem)
 {
 	ModeTools->SetCoordSystem(NewCoordSystem);
diff --git a/Engine/Source/Editor/UnrealEd/Private/Settings/SettingsClasses.cpp b/Engine/Source/Editor/UnrealEd/Private/Settings/SettingsClasses.cpp
--- a/Engine/Source/Editor/UnrealEd/Private/Settings/SettingsClasses.cpp
+++ b/Engine/Source/Editor/UnrealEd/Private/Settings/SettingsClasses.cpp
@@ -647,7 +647,8 @@ ULevelEditorViewportSettings::ULevelEditorViewportSettings( const FObjectInitial
 	BillboardScale = 1.0f;
 	TransformWidgetSizeAdjustment = 0.0f;
 	MeasuringToolUnits = MeasureUnits_Centimeters;
-
+	bAllowArcballRotate = true;
+	bAllowScreenRotate = true;
 	// Set a default preview mesh
 	PreviewMeshes.Add(FSoftObjectPath("/Engine/EditorMeshes/ColorCalibrator/SM_ColorCalibrator.SM_ColorCalibrator"));
 }
diff --git a/Engine/Source/Editor/UnrealEd/Private/UnrealWidget.cpp b/Engine/Source/Editor/UnrealEd/Private/UnrealWidget.cpp
--- a/Engine/Source/Editor/UnrealEd/Private/UnrealWidget.cpp
+++ b/Engine/Source/Editor/UnrealEd/Private/UnrealWidget.cpp
@@ -21,7 +21,7 @@ static const float INNER_AXIS_CIRCLE_RADIUS = 48.0f;
 static const float OUTER_AXIS_CIRCLE_RADIUS = 56.0f;
 static const float ROTATION_TEXT_RADIUS = 75.0f;
 static const int32 AXIS_CIRCLE_SIDES = 24;
-
+static const float ARCALL_RELATIVE_INNER_SIZE = 0.9f;
 /*
  *  Simple struct used to create and group data related to the current window's / viewport's space,
  *  orientation, and scale.
@@ -741,7 +741,16 @@ void FWidget::Render_Rotate( const FSceneView* View,FPrimitiveDrawInterface* PDI
 	FVector YAxis = CustomCoordSystem.TransformVector(FVector(0, 1, 0));
 	FVector ZAxis = CustomCoordSystem.TransformVector(FVector(0, 0, 1));
 
-	EAxisList::Type DrawAxis = GetAxisToDraw( ViewportClient->GetWidgetMode() );
+	EAxisList::Type DrawAxis = GetAxisToDraw(ViewportClient->GetWidgetMode());
+
+	//if drawing arcball, and dragging make sure axis is showing local
+	if (DrawAxis&EAxisList::XYZ  && CurrentAxis == EAxisList::XYZ && bDragging)
+	{
+		FMatrix LocalCoordinate = ViewportClient->GetLocalCoordinateSystem();
+		XAxis = LocalCoordinate.TransformVector(FVector(1, 0, 0));
+		YAxis = LocalCoordinate.TransformVector(FVector(0, 1, 0));
+		ZAxis = LocalCoordinate.TransformVector(FVector(0, 0, 1));
+	}
 
 	FVector DirectionToWidget = View->IsPerspectiveProjection() ? (InLocation - View->ViewMatrices.GetViewOrigin()) : -View->GetViewDirection();
 	DirectionToWidget.Normalize();
@@ -770,7 +779,7 @@ void FWidget::Render_Rotate( const FSceneView* View,FPrimitiveDrawInterface* PDI
 		{
 			FVector Center = InLocation;
 			FRotator Orientation = FRotator::ZeroRotator;
-			const float InnerDistance = (INNER_AXIS_CIRCLE_RADIUS *0.925f * Scale) + GetDefault<ULevelEditorViewportSettings>()->TransformWidgetSizeAdjustment;
+			const float InnerDistance = (INNER_AXIS_CIRCLE_RADIUS * ARCALL_RELATIVE_INNER_SIZE * Scale) + GetDefault<ULevelEditorViewportSettings>()->TransformWidgetSizeAdjustment;
 			FVector Radii(InnerDistance, InnerDistance, InnerDistance);
 			const bool bDisabled = IsWidgetDisabled();
 			FColor ArcColor = ArcBallColor;
@@ -1256,7 +1265,7 @@ void FWidget::ConvertMouseMovementToAxisMovement(FSceneView* InView, FEditorView
 						FVector2D PixelLocation;
 						InView->ScreenToPixel(ScreenLocation, PixelLocation);
 						float Distance = FVector2D::Distance(PixelLocation, MousePosition);
-						const float MaxDiff = 2.0f * INNER_AXIS_CIRCLE_RADIUS;
+						const float MaxDiff = 2.0f * (INNER_AXIS_CIRCLE_RADIUS);
 						//If outside radius, do screen rotate instead, like other DCC's
 
 						if (Distance > MaxDiff)
diff --git a/Engine/Source/Editor/UnrealEd/Public/EditorModeManager.h b/Engine/Source/Editor/UnrealEd/Public/EditorModeManager.h
--- a/Engine/Source/Editor/UnrealEd/Public/EditorModeManager.h
+++ b/Engine/Source/Editor/UnrealEd/Public/EditorModeManager.h
@@ -110,6 +110,7 @@ public:
 
 	FMatrix GetCustomDrawingCoordinateSystem();
 	FMatrix GetCustomInputCoordinateSystem();
+	FMatrix GetLocalCoordinateSystem();
 	
 	/** 
 	 * Returns true if the passed in editor mode is active 
diff --git a/Engine/Source/Editor/UnrealEd/Public/EditorViewportClient.h b/Engine/Source/Editor/UnrealEd/Public/EditorViewportClient.h
--- a/Engine/Source/Editor/UnrealEd/Public/EditorViewportClient.h
+++ b/Engine/Source/Editor/UnrealEd/Public/EditorViewportClient.h
@@ -564,6 +564,11 @@ public:
 	 */
 	virtual FMatrix GetWidgetCoordSystem() const;
 
+	/**
+	* @return The local coordinate system for the transform widget.
+	* For world coordiante system return the identity matrix
+	*/
+	virtual FMatrix GetLocalCoordinateSystem() const;
 	/**
 	 * Sets the coordinate system space to use
 	 */
