Make hit sphere smaller for arcball
#trivial

diff --git a/Engine/Source/Editor/UnrealEd/Private/UnrealWidget.cpp b/Engine/Source/Editor/UnrealEd/Private/UnrealWidget.cpp
--- a/Engine/Source/Editor/UnrealEd/Private/UnrealWidget.cpp
+++ b/Engine/Source/Editor/UnrealEd/Private/UnrealWidget.cpp
@@ -21,7 +21,7 @@ static const float INNER_AXIS_CIRCLE_RADIUS = 48.0f;
 static const float OUTER_AXIS_CIRCLE_RADIUS = 56.0f;
 static const float ROTATION_TEXT_RADIUS = 75.0f;
 static const int32 AXIS_CIRCLE_SIDES = 24;
-static const float ARCALL_RELATIVE_INNER_SIZE = 0.9f;
+static const float ARCALL_RELATIVE_INNER_SIZE = 0.75f;
 /*
  *  Simple struct used to create and group data related to the current window's / viewport's space,
  *  orientation, and scale.
