Suppress Autosave during Sequencer Playback/Recording

diff --git a/Engine/Source/Editor/Sequencer/Private/Sequencer.cpp b/Engine/Source/Editor/Sequencer/Private/Sequencer.cpp
--- a/Engine/Source/Editor/Sequencer/Private/Sequencer.cpp
+++ b/Engine/Source/Editor/Sequencer/Private/Sequencer.cpp
@@ -3199,13 +3199,21 @@ void FSequencer::SetPlaybackStatus(EMovieScenePlayerStatus::Type InPlaybackStatu
 	PlaybackState = InPlaybackStatus;
 
 	// Inform the renderer when Sequencer is in a 'paused' state for the sake of inter-frame effects
-	const bool bIsPaused = (InPlaybackStatus == EMovieScenePlayerStatus::Stopped || InPlaybackStatus == EMovieScenePlayerStatus::Scrubbing || InPlaybackStatus == EMovieScenePlayerStatus::Stepping);
-
+	ESequencerState SequencerState = ESS_None;
+	if (InPlaybackStatus == EMovieScenePlayerStatus::Playing || InPlaybackStatus == EMovieScenePlayerStatus::Recording)
+	{
+		SequencerState = ESS_Playing;
+	}
+	else if (InPlaybackStatus == EMovieScenePlayerStatus::Stopped || InPlaybackStatus == EMovieScenePlayerStatus::Scrubbing || InPlaybackStatus == EMovieScenePlayerStatus::Stepping)
+	{
+		SequencerState = ESS_Paused;
+	}
+	
 	for (FLevelEditorViewportClient* LevelVC : GEditor->GetLevelViewportClients())
 	{
 		if (LevelVC && LevelVC->AllowsCinematicControl())
 		{
-			LevelVC->ViewState.GetReference()->SetSequencerState(bIsPaused);
+			LevelVC->ViewState.GetReference()->SetSequencerState(SequencerState);
 		}
 	}
 
diff --git a/Engine/Source/Editor/UnrealEd/Private/PackageAutoSaver.cpp b/Engine/Source/Editor/UnrealEd/Private/PackageAutoSaver.cpp
--- a/Engine/Source/Editor/UnrealEd/Private/PackageAutoSaver.cpp
+++ b/Engine/Source/Editor/UnrealEd/Private/PackageAutoSaver.cpp
@@ -27,6 +27,7 @@
 #include "ShaderCompiler.h"
 #include "EditorLevelUtils.h"
 #include "IVREditorModule.h"
+#include "LevelEditorViewport.h"
 
 namespace PackageAutoSaverJson
 {
@@ -438,7 +439,17 @@ bool FPackageAutoSaver::CanAutoSave() const
 	const bool bAreShadersCompiling = GShaderCompilingManager->IsCompiling();
 	const bool bIsVREditorActive = IVREditorModule::Get().IsVREditorEnabled();	// @todo vreditor: Eventually we should support this while in VR (modal VR progress, with sufficient early warning)
 
-	return (bAutosaveEnabled && !bSlowTask && !bInterpEditMode && !bPlayWorldValid && !bAnyMenusVisible && !bAutomationTesting && !bIsInteracting && !GIsDemoMode && bHasGameOrProjectLoaded && !bAreShadersCompiling && !bIsVREditorActive);
+	bool bIsSequencerPlaying = false;
+	for (FLevelEditorViewportClient* LevelVC : GEditor->GetLevelViewportClients())
+	{
+		if (LevelVC && LevelVC->AllowsCinematicControl() && LevelVC->ViewState.GetReference()->GetSequencerState() == ESS_Playing)
+		{
+			bIsSequencerPlaying = true;
+			break;
+		}
+	}
+
+	return (bAutosaveEnabled && !bSlowTask && !bInterpEditMode && !bPlayWorldValid && !bAnyMenusVisible && !bAutomationTesting && !bIsInteracting && !GIsDemoMode && bHasGameOrProjectLoaded && !bAreShadersCompiling && !bIsVREditorActive && !bIsSequencerPlaying);
 }
 
 bool FPackageAutoSaver::DoPackagesNeedAutoSave() const
diff --git a/Engine/Source/Runtime/Engine/Public/SceneManagement.h b/Engine/Source/Runtime/Engine/Public/SceneManagement.h
--- a/Engine/Source/Runtime/Engine/Public/SceneManagement.h
+++ b/Engine/Source/Runtime/Engine/Public/SceneManagement.h
@@ -91,6 +91,13 @@ struct ENGINE_API FTemporalLODState
 	void UpdateTemporalLODTransition(const class FViewInfo& View, float LastRenderTime);
 };
 
+enum ESequencerState
+{
+	ESS_None,
+	ESS_Paused,
+	ESS_Playing,
+};
+
 /**
  * The scene manager's persistent view state.
  */
@@ -195,9 +202,9 @@ public:
 	//
 	virtual uint32 GetCurrentTemporalAASampleIndex() const = 0;
 
-	virtual void SetSequencerState(const bool bIsPaused) = 0;
+	virtual void SetSequencerState(ESequencerState InSequencerState) = 0;
 
-	virtual bool GetSequencerState() = 0;
+	virtual ESequencerState GetSequencerState() = 0;
 
 	/** Returns the current PreExposure value. PreExposure is a custom scale applied to the scene color to prevent buffer overflow. */
 	virtual float GetPreExposure() const = 0;
diff --git a/Engine/Source/Runtime/Renderer/Private/PostProcess/PostProcessTemporalAA.cpp b/Engine/Source/Runtime/Renderer/Private/PostProcess/PostProcessTemporalAA.cpp
--- a/Engine/Source/Runtime/Renderer/Private/PostProcess/PostProcessTemporalAA.cpp
+++ b/Engine/Source/Runtime/Renderer/Private/PostProcess/PostProcessTemporalAA.cpp
@@ -193,7 +193,7 @@ public:
 
 		SetShaderValue(RHICmdList, ShaderRHI, DitherScale, bUseDither ? 1.0f : 0.0f);
 
-		const bool bIgnoreVelocity = (Context.View.ViewState && Context.View.ViewState->bSequencerIsPaused);
+		const bool bIgnoreVelocity = (Context.View.ViewState && Context.View.ViewState->GetSequencerState() == ESS_Paused);
 		SetShaderValue(RHICmdList, ShaderRHI, VelocityScaling, bIgnoreVelocity ? 0.0f : 1.0f);
 
 		SetShaderValue(RHICmdList, ShaderRHI, CurrentFrameWeight, CVarTemporalAACurrentFrameWeight.GetValueOnRenderThread());
diff --git a/Engine/Source/Runtime/Renderer/Private/PostProcess/PostProcessing.cpp b/Engine/Source/Runtime/Renderer/Private/PostProcess/PostProcessing.cpp
--- a/Engine/Source/Runtime/Renderer/Private/PostProcess/PostProcessing.cpp
+++ b/Engine/Source/Runtime/Renderer/Private/PostProcess/PostProcessing.cpp
@@ -1721,7 +1721,7 @@ void FPostProcessing::Process(FRHICommandListImmediate& RHICmdList, const FViewI
 				float MaxVelocity = View.FinalPostProcessSettings.MotionBlurMax / 100.0f;
 				float MaxVelocityTiles = MaxVelocity * SizeX * (0.5f / 16.0f);
 				float MaxTileDistGathered = 3.0f;
-				if( MaxVelocityTiles > MaxTileDistGathered || CVarMotionBlurScatter.GetValueOnRenderThread() || (ViewState && ViewState->bSequencerIsPaused) )
+				if( MaxVelocityTiles > MaxTileDistGathered || CVarMotionBlurScatter.GetValueOnRenderThread() || (ViewState && ViewState->GetSequencerState() == ESS_Paused) )
 				{
 					FRenderingCompositePass* VelocityScatterPass = Context.Graph.RegisterPass( new(FMemStack::Get()) FRCPassPostProcessVelocityScatter() );
 					VelocityScatterPass->SetInput( ePId_Input0, MaxTileVelocity );
diff --git a/Engine/Source/Runtime/Renderer/Private/RendererScene.cpp b/Engine/Source/Runtime/Renderer/Private/RendererScene.cpp
--- a/Engine/Source/Runtime/Renderer/Private/RendererScene.cpp
+++ b/Engine/Source/Runtime/Renderer/Private/RendererScene.cpp
@@ -168,7 +168,7 @@ FSceneViewState::FSceneViewState()
 	// Sets the mipbias to invalid large number.
 	MaterialTextureCachedMipBias = BIG_NUMBER;
 
-	bSequencerIsPaused = false;
+	SequencerState = ESS_None;
 
 	LightPropagationVolume = NULL; 
 
diff --git a/Engine/Source/Runtime/Renderer/Private/ScenePrivate.h b/Engine/Source/Runtime/Renderer/Private/ScenePrivate.h
--- a/Engine/Source/Runtime/Renderer/Private/ScenePrivate.h
+++ b/Engine/Source/Runtime/Renderer/Private/ScenePrivate.h
@@ -973,8 +973,8 @@ public:
 	// Is DOFHistoryRT2 set from DepthOfField?
 	bool bDOFHistory2;
 
-	// True when Sequencer has paused
-	bool bSequencerIsPaused;
+	// Sequencer state for view management
+	ESequencerState SequencerState;
 
 	FTemporalLODState TemporalLODState;
 
@@ -1470,14 +1470,14 @@ public:
 
 	virtual SIZE_T GetSizeBytes() const override;
 
-	virtual void SetSequencerState(const bool bIsPaused) override
+	virtual void SetSequencerState(ESequencerState InSequencerState) override
 	{
-		bSequencerIsPaused = bIsPaused;
+		SequencerState = InSequencerState;
 	}
 
-	virtual bool GetSequencerState() override
+	virtual ESequencerState GetSequencerState() override
 	{
-		return bSequencerIsPaused;
+		return SequencerState;
 	}
 
 	/** Information about visibility/occlusion states in past frames for individual primitives. */
diff --git a/Engine/Source/Runtime/Renderer/Private/SceneVisibility.cpp b/Engine/Source/Runtime/Renderer/Private/SceneVisibility.cpp
--- a/Engine/Source/Runtime/Renderer/Private/SceneVisibility.cpp
+++ b/Engine/Source/Runtime/Renderer/Private/SceneVisibility.cpp
@@ -3318,7 +3318,7 @@ void FSceneViewState::UpdateMotionBlurTimeScale(const FViewInfo& View)
 	if (MotionBlurTargetFPS <= 0)
 	{
 		// Keep motion vector lengths stable for paused sequencer frames.
-		if (bSequencerIsPaused)
+		if (GetSequencerState() == ESS_Paused)
 		{
 			// Reset the moving average to the current delta time.
 			MotionBlurTargetDeltaTime = DeltaWorldTime;
@@ -3333,7 +3333,7 @@ void FSceneViewState::UpdateMotionBlurTimeScale(const FViewInfo& View)
 	{
 		// Keep motion vector lengths stable for paused sequencer frames. Assumes a 60 FPS tick.
 		// Tuned for content compatibility with existing content when target is the default 30 FPS.
-		if (bSequencerIsPaused)
+		if (GetSequencerState() == ESS_Paused)
 		{
 			DeltaWorldTime = 1.0f / 60.0f;
 		}
