Sequencer: Fixed a crash when reloading a level sequence asset that had the sequence director open
  - Asset editor manager was not closing toolkits for non-asset objects that had been repointed by a package reload

diff --git a/Engine/Source/Editor/UnrealEd/Private/Toolkits/AssetEditorManager.cpp b/Engine/Source/Editor/UnrealEd/Private/Toolkits/AssetEditorManager.cpp
--- a/Engine/Source/Editor/UnrealEd/Private/Toolkits/AssetEditorManager.cpp
+++ b/Engine/Source/Editor/UnrealEd/Private/Toolkits/AssetEditorManager.cpp
@@ -772,14 +772,14 @@ void FAssetEditorManager::HandlePackageReloaded(const EPackageReloadPhase InPack
 	if (InPackageReloadPhase == EPackageReloadPhase::PrePackageFixup)
 	{
 		/** Call close for all old assets even if not open, so global callback will go off */
-		TArray<UObject*> OldAssets;
+		TArray<UObject*> ObjectsToClose;
 		const TMap<UObject*, UObject*>& RepointedMap = InPackageReloadedEvent->GetRepointedObjects();
 
 		for (const TPair<UObject*, UObject*> RepointPair : RepointedMap)
 		{
 			if (RepointPair.Key->IsAsset())
 			{
-				OldAssets.Add(RepointPair.Key);
+				ObjectsToClose.Add(RepointPair.Key);
 			}
 		}
 
@@ -793,11 +793,13 @@ void FAssetEditorManager::HandlePackageReloaded(const EPackageReloadPhase InPack
 				{
 					PendingAssetsToOpen.AddUnique(NewAsset);
 				}
+
+				ObjectsToClose.AddUnique(AssetEditorPair.Key);
 			}
 		}
 
 		int32 NumAssetEditorsClosed = 0;
-		for (UObject* OldAsset : OldAssets)
+		for (UObject* OldAsset : ObjectsToClose)
 		{
 			NumAssetEditorsClosed += CloseAllEditorsForAsset(OldAsset);
 		}
