Back out changelist 6850583 for now to unblock testing.

#rnx

diff --git a/Engine/Source/Editor/Kismet/Private/BlueprintCompilationManager.cpp b/Engine/Source/Editor/Kismet/Private/BlueprintCompilationManager.cpp
--- a/Engine/Source/Editor/Kismet/Private/BlueprintCompilationManager.cpp
+++ b/Engine/Source/Editor/Kismet/Private/BlueprintCompilationManager.cpp
@@ -1959,7 +1959,8 @@ void FBlueprintCompilationManagerImpl::ReinstanceBatch(TArray<FReinstancingJob>&
 
 	for(UObject* ArchetypeReferencer : ArchetypeReferencers)
 	{
-		FArchiveReplaceObjectRef<UObject> ReplaceInCDOAr(ArchetypeReferencer, OldArchetypeToNewArchetype, false, false, false, false, true);
+		UPackage* NewPackage = ArchetypeReferencer->GetOutermost();
+		FArchiveReplaceOrClearExternalReferences<UObject> ReplaceInCDOAr(ArchetypeReferencer, OldArchetypeToNewArchetype, NewPackage);
 	}
 }
 
diff --git a/Engine/Source/Editor/UnrealEd/Private/Kismet2/KismetReinstanceUtilities.cpp b/Engine/Source/Editor/UnrealEd/Private/Kismet2/KismetReinstanceUtilities.cpp
--- a/Engine/Source/Editor/UnrealEd/Private/Kismet2/KismetReinstanceUtilities.cpp
+++ b/Engine/Source/Editor/UnrealEd/Private/Kismet2/KismetReinstanceUtilities.cpp
@@ -31,7 +31,6 @@
 #include "BlueprintEditor.h"
 #include "Engine/Selection.h"
 #include "BlueprintEditorSettings.h"
-#include "Algo/RemoveIf.h"
 
 extern COREUOBJECT_API bool GBlueprintUseCompilationManager;
 
@@ -89,24 +88,6 @@ struct FReplaceReferenceHelper
 			BP_SCOPED_COMPILER_EVENT_STAT(EKismetReinstancerStats_FindReferencers);
 
 			Targets = FArchiveHasReferences::GetAllReferencers(SourceObjects, ObjectsThatShouldUseOldStuff);
-
-			// Make sure we don't update properties in old objects, as they
-			// may take ownership of objects referenced in new objects (e.g.
-			// delete components owned by new actors)
-			int32 RangeEnd = Algo::RemoveIf(Targets, [&ObjectsToReplace](UObject* Obj)
-			{
-				for (UObject* ObjectToReplace : ObjectsToReplace)
-				{
-					if (ObjectToReplace == Obj || Obj->IsIn(ObjectToReplace))
-					{
-						return true;
-					}
-				}
-
-				return false;
-			});
-
-			Targets.RemoveAtSwap(RangeEnd, Targets.Num() - RangeEnd, false);
 		}
 
 		{
@@ -114,45 +95,51 @@ struct FReplaceReferenceHelper
 
 			for (UObject* Obj : Targets)
 			{
-				// The class for finding and replacing weak references.
-				// We can't relay on "standard" weak references replacement as
-				// it depends on FSoftObjectPath::ResolveObject, which
-				// tries to find the object with the stored path. It is
-				// impossible, cause above we deleted old actors (after
-				// spawning new ones), so during objects traverse we have to
-				// find FSoftObjectPath with the raw given path taken
-				// before deletion of old actors and fix them.
-				class ReferenceReplace : public FArchiveReplaceObjectRef < UObject >
+				// Make sure we don't update properties in old objects, as they
+				// may take ownership of objects referenced in new objects (e.g.
+				// delete components owned by new actors)
+				if (!ObjectsToReplace.Contains(Obj))
 				{
-				public:
-					ReferenceReplace(UObject* InSearchObject, const TMap<UObject*, UObject*>& InReplacementMap, const TMap<FSoftObjectPath, UObject*>& InWeakReferencesMap)
-						: FArchiveReplaceObjectRef<UObject>(InSearchObject, InReplacementMap, false, false, false, true), WeakReferencesMap(InWeakReferencesMap)
+					// The class for finding and replacing weak references.
+					// We can't relay on "standard" weak references replacement as
+					// it depends on FSoftObjectPath::ResolveObject, which
+					// tries to find the object with the stored path. It is
+					// impossible, cause above we deleted old actors (after
+					// spawning new ones), so during objects traverse we have to
+					// find FSoftObjectPath with the raw given path taken
+					// before deletion of old actors and fix them.
+					class ReferenceReplace : public FArchiveReplaceObjectRef < UObject >
 					{
-						SerializeSearchObject();
-					}
-
-					FArchive& operator<<(FSoftObjectPath& Ref) override
-					{
-						const UObject*const* PtrToObjPtr = WeakReferencesMap.Find(Ref);
-
-						if (PtrToObjPtr != nullptr)
+					public:
+						ReferenceReplace(UObject* InSearchObject, const TMap<UObject*, UObject*>& InReplacementMap, const TMap<FSoftObjectPath, UObject*>& InWeakReferencesMap)
+							: FArchiveReplaceObjectRef<UObject>(InSearchObject, InReplacementMap, false, false, false, true), WeakReferencesMap(InWeakReferencesMap)
 						{
-							Ref = *PtrToObjPtr;
+							SerializeSearchObject();
 						}
 
-						return *this;
-					}
+						FArchive& operator<<(FSoftObjectPath& Ref) override
+						{
+							const UObject*const* PtrToObjPtr = WeakReferencesMap.Find(Ref);
 
-					FArchive& operator<<(FSoftObjectPtr& Ref) override
-					{
-						return operator<<(Ref.GetUniqueID());
-					}
+							if (PtrToObjPtr != nullptr)
+							{
+								Ref = *PtrToObjPtr;
+							}
+
+							return *this;
+						}
 
-				private:
-					const TMap<FSoftObjectPath, UObject*>& WeakReferencesMap;
-				};
+						FArchive& operator<<(FSoftObjectPtr& Ref) override
+						{
+							return operator<<(Ref.GetUniqueID());
+						}
 
-				ReferenceReplace ReplaceAr(Obj, OldToNewInstanceMap, ReinstancedObjectsWeakReferenceMap);
+					private:
+						const TMap<FSoftObjectPath, UObject*>& WeakReferencesMap;
+					};
+
+					ReferenceReplace ReplaceAr(Obj, OldToNewInstanceMap, ReinstancedObjectsWeakReferenceMap);
+				}
 			}
 		}
 	}
@@ -1704,7 +1691,9 @@ static void ReplaceObjectHelper(UObject*& OldObject, UClass* OldClass, UObject*&
 		}
 	}
 
-	// NOTE: Don't mark the old object as pending kill yet! Otherwise, any weak references to this object will be invalidated too soon and won't get replaced.
+	OldObject->RemoveFromRoot();
+	OldObject->MarkPendingKill();
+
 	OldToNewInstanceMap.Add(OldObject, NewUObject);
 
 	if (bIsComponent)
@@ -1861,7 +1850,7 @@ static void ReplaceActorHelper(UObject* OldObject, UClass* OldClass, UObject*& N
 		GEditor->Layers->DisassociateActorFromLayers(OldActor);
 	}
 
-	// NOTE: Don't destroy the Actor yet! Otherwise, any weak references to this Actor will be invalidated too soon and won't get replaced.
+	World->EditorDestroyActor(OldActor, /*bShouldModifyLevel =*/true);
 	OldToNewInstanceMap.Add(OldActor, NewActor);
 }
 
@@ -1916,6 +1905,18 @@ void FBlueprintCompileReinstancer::ReplaceInstancesOfClass_Inner(TMap<UClass*, U
 		OnObjectsReplacedHandle = GEditor->OnObjectsReplaced().AddRaw(&ObjectRemappingHelper, &FObjectRemappingHelper::OnObjectsReplaced);
 	}
 
+	auto UpdateObjectBeingDebugged = [](UObject* InOldObject, UObject* InNewObject)
+	{
+		if (UBlueprint* OldObjBlueprint = Cast<UBlueprint>(InOldObject->GetClass()->ClassGeneratedBy))
+		{
+			const UObject* DebugObj = OldObjBlueprint->GetObjectBeingDebugged(EGetObjectOrWorldBeingDebuggedFlags::IgnorePendingKill);
+			if (DebugObj == InOldObject)
+			{
+				OldObjBlueprint->SetObjectBeingDebugged(InNewObject);
+			}
+		}
+	};
+
 	{
 		TArray<UObject*> ObjectsToReplace;
 
@@ -1966,7 +1967,7 @@ void FBlueprintCompileReinstancer::ReplaceInstancesOfClass_Inner(TMap<UClass*, U
 					{
 						UObject* NewUObject = nullptr;
 						ReplaceObjectHelper(OldObject, OldClass, NewUObject, NewClass, OldToNewInstanceMap, OldToNewNameMap, OldObjIndex, ObjectsToReplace, PotentialEditorsForRefreshing, OwnersToRerunConstructionScript, &FDirectAttachChildrenAccessor::Get, bIsComponent, bArchetypesAreUpToDate);
-						
+						UpdateObjectBeingDebugged(OldObject, NewUObject);
 						ObjectsReplaced.Add(OldObject);
 
 						if (bLogConversions)
@@ -2038,7 +2039,7 @@ void FBlueprintCompileReinstancer::ReplaceInstancesOfClass_Inner(TMap<UClass*, U
 							// Actors that are not in a level cannot be reconstructed, sequencer team decided to reinstance these as normal objects:
 							ReplaceObjectHelper(OldObject, OldClass, NewUObject, NewClass, OldToNewInstanceMap, OldToNewNameMap, OldObjIndex, ObjectsToReplace, PotentialEditorsForRefreshing, OwnersToRerunConstructionScript, &FDirectAttachChildrenAccessor::Get, false, bArchetypesAreUpToDate);
 						}
-
+						UpdateObjectBeingDebugged(OldObject, NewUObject);
 						ObjectsReplaced.Add(OldObject);
 
 						if (bLogConversions)
diff --git a/Engine/Source/Runtime/CoreUObject/Private/Serialization/ArchiveHasReferences.cpp b/Engine/Source/Runtime/CoreUObject/Private/Serialization/ArchiveHasReferences.cpp
--- a/Engine/Source/Runtime/CoreUObject/Private/Serialization/ArchiveHasReferences.cpp
+++ b/Engine/Source/Runtime/CoreUObject/Private/Serialization/ArchiveHasReferences.cpp
@@ -12,7 +12,6 @@ FArchiveHasReferences::FArchiveHasReferences(UObject* InTarget, const TSet<UObje
 	check(InTarget);
 
 	ArIsObjectReferenceCollector = true;
-	ArIsModifyingWeakAndStrongReferences = true;
 	InTarget->Serialize(*this);
 
 	class FArchiveProxyCollector : public FReferenceCollector
diff --git a/Engine/Source/Runtime/Engine/Classes/Engine/Blueprint.h b/Engine/Source/Runtime/Engine/Classes/Engine/Blueprint.h
--- a/Engine/Source/Runtime/Engine/Classes/Engine/Blueprint.h
+++ b/Engine/Source/Runtime/Engine/Classes/Engine/Blueprint.h
@@ -658,11 +658,9 @@ public:
 
 protected:
 	/** Current object being debugged for this blueprint */
-	UPROPERTY(transient, duplicatetransient)
 	TWeakObjectPtr< class UObject > CurrentObjectBeingDebugged;
 
 	/** Current world being debugged for this blueprint */
-	UPROPERTY(transient, duplicatetransient)
 	TWeakObjectPtr< class UWorld > CurrentWorldBeingDebugged;
 
 	/** Delegate called when the debug object is set */
