Adding some extra analytics to Engine Heartbeat

diff --git a/Engine/Source/Editor/UnrealEd/Private/UnrealEdMisc.cpp b/Engine/Source/Editor/UnrealEd/Private/UnrealEdMisc.cpp
--- a/Engine/Source/Editor/UnrealEd/Private/UnrealEdMisc.cpp
+++ b/Engine/Source/Editor/UnrealEd/Private/UnrealEdMisc.cpp
@@ -614,6 +614,11 @@ void FUnrealEdMisc::InitEngineAnalytics()
 * @EventParam IsDebugger (bool) Whether the debugger is currently present
 * @EventParam WasDebuggerPresent (bool) Whether the debugger was present previously
 * @EventParam IsInVRMode (bool) If the current heartbeat occurred while VR mode was active
+* @EventParam bIsEnterprise (bool) If the editor has an enterprise project loaded
+* @EventParam Is5MinIdle (bool) Whether the user is idle, using the 5 minute methodology
+* @EventParam Is30MinIdle (bool) Whether the user is idle, using the 30 minute methodology
+* @EventParam IsInPIE (bool) Whether the user is/was in PIE during this heartbeat
+* @EventParam RealIntervalSec (bool) The real time since the last heartbeat
 *
 * @Source Editor
 *
@@ -634,6 +639,13 @@ void FUnrealEdMisc::EditorAnalyticsHeartbeat()
 	//////////////////////////////////////////////////////////////////////////////////////////////////////////////////////////////////////////////
 	const double HeartbeatInvervalSec = (double)(UnrealEdMiscDefs::HeartbeatIntervalSeconds);
 	const double Now = FPlatformTime::Seconds();
+
+	static double LastTimeInPIE = 0;
+	if (FPlayWorldCommandCallbacks::IsInPIE())
+	{
+		LastTimeInPIE = Now;
+	}
+
 	// Initialize to ensure the first time this is called we will DEF execute the heartbeat.
 	static double LastHeartbeatTime = Now - HeartbeatInvervalSec;
 	// allow a little bit of slop in the timing in case the event is only firing slightly too soon, even though that is technically impossible.
@@ -654,10 +666,10 @@ void FUnrealEdMisc::EditorAnalyticsHeartbeat()
 	const bool bInVRMode = IVREditorModule::Get().IsVREditorModeActive();
 	const bool bIsEnterprise = IProjectManager::Get().IsEnterpriseProject();
 
-	double LastInteractionTime = FSlateApplication::Get().GetLastUserInteractionTime();
+	const double LastInteractionTime = FSlateApplication::Get().GetLastUserInteractionTime();
 	
 	// Did the user interact since the last heartbeat
-	bool bIdle = LastInteractionTime < LastHeartbeatTime;
+	const bool bIdle = LastInteractionTime < LastHeartbeatTime;
 	
 	extern ENGINE_API float GAverageFPS;
 
@@ -678,9 +690,22 @@ void FUnrealEdMisc::EditorAnalyticsHeartbeat()
 	Attributes.Add(FAnalyticsEventAttribute(TEXT("IsInVRMode"), bInVRMode));
 	Attributes.Add(FAnalyticsEventAttribute(TEXT("Enterprise"), bIsEnterprise));
 
+	// In the 5 Min and 30 Min idle case we treat longer periods of time as none idle post user interaction
+	const bool b5MinIdle = LastInteractionTime + (5.0f * 60.0f) < LastHeartbeatTime;
+	Attributes.Add(FAnalyticsEventAttribute(TEXT("Is5MinIdle"), b5MinIdle));
+	const bool b30MinIdle = LastInteractionTime + (30.0f * 60.0f) < LastHeartbeatTime;
+	Attributes.Add(FAnalyticsEventAttribute(TEXT("Is30MinIdle"), b30MinIdle));
+
+	// InPIE bool, added a 5 second grace period to include PIE shutdown issues in the "InPIE" window
+	const bool bInPIE = LastTimeInPIE + 5.0f > LastHeartbeatTime;
+	Attributes.Add(FAnalyticsEventAttribute(TEXT("IsInPIE"), bInPIE));
+
+	// Interval seconds actually track the entirety of the time since last heartbeat.
+	Attributes.Add(FAnalyticsEventAttribute(TEXT("RealIntervalSec"), Now - LastHeartbeatTime));
+
 	FEngineAnalytics::GetProvider().RecordEvent(TEXT("Editor.Usage.Heartbeat"), Attributes);
 	
-	LastHeartbeatTime = FPlatformTime::Seconds();
+	LastHeartbeatTime = Now;
 }
 
 void FUnrealEdMisc::TickAssetAnalytics()
