Back out straggler from changelist 6850583 (missed from 6896870).

#rnx

diff --git a/Engine/Source/Editor/UnrealEd/Private/Kismet2/KismetReinstanceUtilities.cpp b/Engine/Source/Editor/UnrealEd/Private/Kismet2/KismetReinstanceUtilities.cpp
--- a/Engine/Source/Editor/UnrealEd/Private/Kismet2/KismetReinstanceUtilities.cpp
+++ b/Engine/Source/Editor/UnrealEd/Private/Kismet2/KismetReinstanceUtilities.cpp
@@ -2131,22 +2131,6 @@ void FBlueprintCompileReinstancer::ReplaceInstancesOfClass_Inner(TMap<UClass*, U
 				}
 			}
 		}
-
-		Obj->RemoveFromRoot();
-		checkf(!Obj->IsPendingKill(), TEXT("'%s' should not have been marked as PendingKill prior to this point."), *Obj->GetPathName());
-
-		AActor* ObjAsActor = Cast<AActor>(Obj);
-		if (ObjAsActor && ObjAsActor->GetLevel())
-		{
-			// This will take care of marking the object as PendingKill.
-			UWorld* World = ObjAsActor->GetWorld();
-			check(World != nullptr);
-			World->EditorDestroyActor(ObjAsActor, /*bShouldModifyLevel =*/true);
-		}
-		else
-		{
-			Obj->MarkPendingKill();
-		}
 	}
 
 
