- Modify commandlet to gather mismatched bone names for UParticleModuleLocationBoneSocket between LOD levels
- We don't need to check for mismatched SourceLocation array size since that should be covered on load

#rnx

diff --git a/Engine/Source/Editor/UnrealEd/Private/Commandlets/ParticleSystemAuditCommandlet.cpp b/Engine/Source/Editor/UnrealEd/Private/Commandlets/ParticleSystemAuditCommandlet.cpp
--- a/Engine/Source/Editor/UnrealEd/Private/Commandlets/ParticleSystemAuditCommandlet.cpp
+++ b/Engine/Source/Editor/UnrealEd/Private/Commandlets/ParticleSystemAuditCommandlet.cpp
@@ -192,18 +192,75 @@ bool UParticleSystemAuditCommandlet::ProcessParticleSystems()
 								{
 									bHasSpawnPerUnit = true;
 								}
-								else if (UParticleModuleLocationBoneSocket* BoneModule = Cast<UParticleModuleLocationBoneSocket>(Module))
+							}
+						}
+					}
+
+					// Cross LOD validation for things where indices are per particle and used to index into module specific data for example
+					if ( !bSingleLOD )
+					{
+						UParticleLODLevel* HighestLODLevel = Emitter->LODLevels[0];
+						check(HighestLODLevel != nullptr);
+
+						for ( int32 ModuleIdx=0; ModuleIdx < HighestLODLevel->Modules.Num(); ++ModuleIdx )
+						{
+							UParticleModule* HighestLODModule = HighestLODLevel->Modules[ModuleIdx];
+							if ( UParticleModuleLocationBoneSocket* HighestBoneModule = Cast<UParticleModuleLocationBoneSocket>(HighestLODModule) )
+							{
+								bool bRequiresSourceLocationValidation = HighestBoneModule->bInheritBoneVelocity || HighestBoneModule->bUpdatePositionEachFrame;
+								if (!bRequiresSourceLocationValidation )
 								{
-									int32 SourceArraySize = BoneModule->SourceLocations.Num();
-									if (BoneLocationArraySize >= 0 && BoneLocationArraySize != SourceArraySize)
+									for (int32 LODIdx = 1; LODIdx < Emitter->LODLevels.Num(); ++LODIdx)
 									{
-										// we assume that there is not more than one bone location module per emitter,
-										// so the mismatch has to come from the LOD levels
-										bMismatchedLODBoneModules = true;
+										UParticleLODLevel* LODLevel = Emitter->LODLevels[LODIdx];
+										if (LODLevel == nullptr)
+											continue;
+
+										UParticleModule* LODModule = LODLevel->Modules[ModuleIdx];
+										UParticleModuleLocationBoneSocket* LODBoneModule = Cast<UParticleModuleLocationBoneSocket>(LODModule);
+										if (LODBoneModule == nullptr)
+										{
+											UE_LOG(LogParticleSystemAuditCommandlet, Warning, TEXT("ParticleSystem '%s' Emitter '%s' LOD index '%d' has incorrect module type '%s' when should be 'UParticleModuleLocationBoneSocket'"), *PSys->GetPathName(), *Emitter->GetName(), LODIdx, LODModule == nullptr ? TEXT("nullptr") : *LODModule->GetClass()->GetName());
+											continue;
+										}
+
+										if (LODBoneModule->bInheritBoneVelocity || LODBoneModule->bUpdatePositionEachFrame)
+										{
+											bRequiresSourceLocationValidation = true;
+											break;
+										}
 									}
-									else
+								}
+
+								if ( bRequiresSourceLocationValidation )
+								{
+									for (int32 LODIdx = 1; LODIdx < Emitter->LODLevels.Num(); ++LODIdx)
 									{
-										BoneLocationArraySize = SourceArraySize;
+										UParticleLODLevel* LODLevel = Emitter->LODLevels[LODIdx];
+										if (LODLevel == nullptr)
+											continue;
+
+										UParticleModule* LODModule = LODLevel->Modules[ModuleIdx];
+										UParticleModuleLocationBoneSocket* LODBoneModule = Cast<UParticleModuleLocationBoneSocket>(LODModule);
+										if (LODBoneModule == nullptr)
+											continue;
+
+										if ( HighestBoneModule->SourceLocations.Num() != LODBoneModule->SourceLocations.Num() )
+										{
+											bMismatchedLODBoneModules = true;
+											UE_LOG(LogParticleSystemAuditCommandlet, Log, TEXT("ParticleSystem '%s' Emitter '%s' LOD index '%d' SourceLocations does not match HighestLOD, this should have been fixed on load"), *PSys->GetPathName(), *Emitter->GetName(), LODIdx);
+											continue;
+										}
+
+										for ( int32 SourceLocationIdx=0; SourceLocationIdx < HighestBoneModule->SourceLocations.Num(); ++SourceLocationIdx )
+										{
+											if ( HighestBoneModule->SourceLocations[SourceLocationIdx].BoneSocketName != LODBoneModule->SourceLocations[SourceLocationIdx].BoneSocketName)
+											{
+												bMismatchedLODBoneModules = true;
+												UE_LOG(LogParticleSystemAuditCommandlet, Log, TEXT("ParticleSystem '%s' Emitter '%s' LOD index '%d' BoneSocketNames do not match between LOD levels, this is not fatal but may be an error"), *PSys->GetPathName(), *Emitter->GetName(), LODIdx);
+												break;
+											}
+										}
 									}
 								}
 							}
