Fixed the blueprint editor stealing focus and thus making the PIE blueprint compilation warning dialog uninteractible.

Source: 7307677 (Dev-Editor)

diff --git a/Engine/Source/Editor/UnrealEd/Private/PlayLevel.cpp b/Engine/Source/Editor/UnrealEd/Private/PlayLevel.cpp
--- a/Engine/Source/Editor/UnrealEd/Private/PlayLevel.cpp
+++ b/Engine/Source/Editor/UnrealEd/Private/PlayLevel.cpp
@@ -2279,13 +2279,20 @@ static bool ShowBlueprintErrorDialog( TArray<UBlueprint*> ErroredBlueprints )
 {
 	struct Local
 	{
-		static void OnHyperlinkClicked( TWeakObjectPtr<UBlueprint> InBlueprint )
+		static void OnHyperlinkClicked( TWeakObjectPtr<UBlueprint> InBlueprint, TSharedPtr<SCustomDialog> InDialog )
 		{
 			if (UBlueprint* BlueprintToEdit = InBlueprint.Get())
 			{
 				// Open the blueprint
 				GEditor->EditObject( BlueprintToEdit );
 			}
+
+			if (InDialog.IsValid())
+			{
+				// Opening the blueprint editor above may end up creating an invisible new window on top of the dialog, 
+				// thus making it not interactable, so we have to force the dialog back to the front
+				InDialog->BringToFront(true);
+			}
 		}
 	};
 
@@ -2297,6 +2304,8 @@ static bool ShowBlueprintErrorDialog( TArray<UBlueprint*> ErroredBlueprints )
 			.Text(NSLOCTEXT("PlayInEditor", "PrePIE_BlueprintErrors", "One or more blueprints has an unresolved compiler error, are you sure you want to Play in Editor?"))
 		];
 
+	TSharedPtr<SCustomDialog> CustomDialog;
+
 	for (UBlueprint* Blueprint : ErroredBlueprints)
 	{
 		TWeakObjectPtr<UBlueprint> BlueprintPtr = Blueprint;
@@ -2307,7 +2316,7 @@ static bool ShowBlueprintErrorDialog( TArray<UBlueprint*> ErroredBlueprints )
 			[
 				SNew(SHyperlink)
 				.Style(FEditorStyle::Get(), "Common.GotoBlueprintHyperlink")
-				.OnNavigate_Static(&Local::OnHyperlinkClicked, BlueprintPtr)
+				.OnNavigate(FSimpleDelegate::CreateLambda([BlueprintPtr, &CustomDialog]() { Local::OnHyperlinkClicked(BlueprintPtr, CustomDialog); }))
 				.Text(FText::FromString(Blueprint->GetName()))
 				.ToolTipText(NSLOCTEXT("SourceHyperlink", "EditBlueprint_ToolTip", "Click to edit the blueprint"))
 			];
@@ -2325,7 +2334,7 @@ static bool ShowBlueprintErrorDialog( TArray<UBlueprint*> ErroredBlueprints )
 	FText OKText = NSLOCTEXT("PlayInEditor", "PrePIE_OkText", "Play in Editor");
 	FText CancelText = NSLOCTEXT("Dialogs", "EAppReturnTypeCancel", "Cancel");
 
-	TSharedRef<SCustomDialog> CustomDialog = SNew(SCustomDialog)
+	CustomDialog = SNew(SCustomDialog)
 		.Title(DialogTitle)
 		.IconBrush("NotificationList.DefaultMessage")
 		.DialogContent(DialogContents)
