Copied cl 7277154 from Dev-Enterprise

Fix actors labels changing when reimporting scenes

diff --git a/Engine/Source/Editor/UnrealEd/Classes/Editor/EditorEngine.h b/Engine/Source/Editor/UnrealEd/Classes/Editor/EditorEngine.h
--- a/Engine/Source/Editor/UnrealEd/Classes/Editor/EditorEngine.h
+++ b/Engine/Source/Editor/UnrealEd/Classes/Editor/EditorEngine.h
@@ -237,6 +237,12 @@ struct FCachedActorLabels
 		ActorLabels.Add(InLabel);
 	}
 
+	/** Remove a label from this set */
+	FORCEINLINE void Remove(const FString& InLabel)
+	{
+		ActorLabels.Remove(InLabel);
+	}
+
 	/** Check if the specified label exists */
 	FORCEINLINE bool Contains(const FString& InLabel) const
 	{
