Fixed null pointer crash when loading a level from a commandlet if GEditor is a UUnrealEdEngine

diff --git a/Engine/Source/Editor/UnrealEd/Private/EditorSelectUtils.cpp b/Engine/Source/Editor/UnrealEd/Private/EditorSelectUtils.cpp
--- a/Engine/Source/Editor/UnrealEd/Private/EditorSelectUtils.cpp
+++ b/Engine/Source/Editor/UnrealEd/Private/EditorSelectUtils.cpp
@@ -389,7 +389,7 @@ void UUnrealEdEngine::NoteSelectionChange(bool bNotify)
 	UpdatePivotLocationForSelection( true );
 
 	// Clear active editing visualizer on selection change
-	GUnrealEd->ComponentVisManager.ClearActiveComponentVis();
+	ComponentVisManager.ClearActiveComponentVis();
 
 	TArray<FEdMode*> ActiveModes;
 	GLevelEditorModeTools().GetActiveModes( ActiveModes );
