Fix virtual bones being missing from imported meshes

diff --git a/Engine/Source/Editor/UnrealEd/Private/Fbx/FbxSkeletalMeshImport.cpp b/Engine/Source/Editor/UnrealEd/Private/Fbx/FbxSkeletalMeshImport.cpp
--- a/Engine/Source/Editor/UnrealEd/Private/Fbx/FbxSkeletalMeshImport.cpp
+++ b/Engine/Source/Editor/UnrealEd/Private/Fbx/FbxSkeletalMeshImport.cpp
@@ -1844,7 +1844,7 @@ USkeletalMesh* UnFbx::FFbxImporter::ImportSkeletalMesh(FImportSkeletalMeshArgs &
 
 	// process reference skeleton from import data
 	int32 SkeletalDepth = 0;
-	USkeleton* ExistingSkeleton = ExistSkelMeshDataPtr != nullptr ? ExistSkelMeshDataPtr->ExistingSkeleton : SkeletalMesh->Skeleton;
+	USkeleton* ExistingSkeleton = ExistSkelMeshDataPtr != nullptr ? ExistSkelMeshDataPtr->ExistingSkeleton : ImportOptions->SkeletonForAnimation;
 	if (!ProcessImportMeshSkeleton(ExistingSkeleton, SkeletalMesh->RefSkeleton, SkeletalDepth, *SkelMeshImportDataPtr))
 	{
 		SkeletalMesh->ClearFlags(RF_Standalone);
