Brush: Check for valid DefaultBrush()

diff --git a/Engine/Source/Editor/UnrealEd/Public/DragAndDrop/BrushBuilderDragDropOp.h b/Engine/Source/Editor/UnrealEd/Public/DragAndDrop/BrushBuilderDragDropOp.h
--- a/Engine/Source/Editor/UnrealEd/Public/DragAndDrop/BrushBuilderDragDropOp.h
+++ b/Engine/Source/Editor/UnrealEd/Public/DragAndDrop/BrushBuilderDragDropOp.h
@@ -32,7 +32,7 @@ public:
 	~FBrushBuilderDragDropOp()
 	{
 		UWorld* World = GEditor->GetEditorWorldContext().World();
-		if(World != nullptr)
+		if ((World != nullptr) && (World->GetDefaultBrush() != nullptr))
 		{
 			// Deselect & hide the builder brush
 			World->GetDefaultBrush()->SetIsTemporarilyHiddenInEditor(true);
@@ -111,8 +111,11 @@ private:
 	{
 		// show & select the builder brush
 		UWorld* World = GEditor->GetEditorWorldContext().World();
-		World->GetDefaultBrush()->SetIsTemporarilyHiddenInEditor(false);
-		GEditor->SelectActor(World->GetDefaultBrush(), true, false);
+		if ((World != nullptr) && (World->GetDefaultBrush() != nullptr))
+		{
+			World->GetDefaultBrush()->SetIsTemporarilyHiddenInEditor(false);
+			GEditor->SelectActor(World->GetDefaultBrush(), true, false);
+		}
 	}
 
 private:
