Changing the setting 'Allow Cooked Data in The Editor' will now prompt for the editor restart

diff --git a/Engine/Source/Editor/UnrealEd/Public/CookerSettings.h b/Engine/Source/Editor/UnrealEd/Public/CookerSettings.h
--- a/Engine/Source/Editor/UnrealEd/Public/CookerSettings.h
+++ b/Engine/Source/Editor/UnrealEd/Public/CookerSettings.h
@@ -132,7 +132,8 @@ public:
 	/** Allows opening cooked assets in the editor */
 	UPROPERTY(EditAnywhere, config, Category = Editor, meta = (
 		ConsoleVariable = "cook.AllowCookedDataInEditorBuilds", DisplayName = "Allow Cooked Content In The Editor",
-		ToolTip = "If true, the editor will be able to open cooked assets (limited to a subset of supported asset types)."))
+		ToolTip = "If true, the editor will be able to open cooked assets (limited to a subset of supported asset types).",
+		ConfigRestartRequired = true))
 	uint32 bAllowCookedDataInEditorBuilds : 1;
 
 private:
