Brush: Check for valid DefaultBrush()

diff --git a/Engine/Source/Editor/UnrealEd/Public/DragAndDrop/BrushBuilderDragDropOp.h b/Engine/Source/Editor/UnrealEd/Public/DragAndDrop/BrushBuilderDragDropOp.h
--- a/Engine/Source/Editor/UnrealEd/Public/DragAndDrop/BrushBuilderDragDropOp.h
+++ b/Engine/Source/Editor/UnrealEd/Public/DragAndDrop/BrushBuilderDragDropOp.h
@@ -68,7 +68,7 @@ public:
 	virtual void OnDragged( const FDragDropEvent& DragDropEvent ) override
 	{
 		UWorld* World = GEditor->GetEditorWorldContext().World();
-		if(World != nullptr)
+		if ((World != nullptr) && (World->GetDefaultBrush() != nullptr))
 		{
 			if(CursorDecoratorWindow->IsVisible() && !World->GetDefaultBrush()->IsTemporarilyHiddenInEditor())
 			{
