SkelMesh preflight test.

diff --git a/Engine/Source/Editor/UnrealEd/Private/Commandlets/DiffAssetRegistriesCommandlet.cpp b/Engine/Source/Editor/UnrealEd/Private/Commandlets/DiffAssetRegistriesCommandlet.cpp
--- a/Engine/Source/Editor/UnrealEd/Private/Commandlets/DiffAssetRegistriesCommandlet.cpp
+++ b/Engine/Source/Editor/UnrealEd/Private/Commandlets/DiffAssetRegistriesCommandlet.cpp
@@ -871,10 +871,10 @@ void UDiffAssetRegistriesCommandlet::LogChangedFiles(FArchive *CSVFile, FString
 				UE_LOG(LogDiffAssets, Display, TEXT("%c %s : (Class=%s,NewSize=%d bytes,OldSize=%d bytes)"), classification, *AssetPath.ToString(), *ClassName.ToString(), ChangeInfo.ChangedBytes, PrevData->DiskSize);
 			}
 			//UE_LOG(LogDiffAssets, Display, TEXT("Source file changed: %s"), (flags & EAssetFlags::GuidChange) ? TEXT("true") : TEXT("false"));
-			if ((flags & EAssetFlags::GuidChange) && bMatchChangelists)
+			/*if ((flags & EAssetFlags::GuidChange) && bMatchChangelists)
 			{
 				UE_LOG(LogDiffAssets, Display, TEXT("Last change: %d"), Changelist);
-			}
+			}*/
 #if 0			
 			TArray<FAssetIdentifier> Dependencies;
 			NewState.GetDependencies(AssetPath, Dependencies, EAssetRegistryDependencyType::Hard);
