Fix CIS.

diff --git a/Engine/Source/Editor/UnrealEd/Private/DragTool_BoxSelect.cpp b/Engine/Source/Editor/UnrealEd/Private/DragTool_BoxSelect.cpp
--- a/Engine/Source/Editor/UnrealEd/Private/DragTool_BoxSelect.cpp
+++ b/Engine/Source/Editor/UnrealEd/Private/DragTool_BoxSelect.cpp
@@ -6,6 +6,7 @@
 #include "CanvasItem.h"
 #include "Settings/LevelEditorViewportSettings.h"
 #include "GameFramework/Volume.h"
+#include "UnrealEdGlobals.h"
 #include "EngineUtils.h"
 #include "EditorModeManager.h"
 #include "EditorModes.h"
diff --git a/Engine/Source/Editor/UnrealEd/Private/DragTool_FrustumSelect.cpp b/Engine/Source/Editor/UnrealEd/Private/DragTool_FrustumSelect.cpp
--- a/Engine/Source/Editor/UnrealEd/Private/DragTool_FrustumSelect.cpp
+++ b/Engine/Source/Editor/UnrealEd/Private/DragTool_FrustumSelect.cpp
@@ -6,6 +6,7 @@
 #include "CanvasItem.h"
 #include "Settings/LevelEditorViewportSettings.h"
 #include "GameFramework/Volume.h"
+#include "UnrealEdGlobals.h"
 #include "EngineUtils.h"
 #include "Engine/Selection.h"
 #include "EditorModeManager.h"
diff --git a/Engine/Source/Editor/UnrealEd/Public/ComponentVisualizer.h b/Engine/Source/Editor/UnrealEd/Public/ComponentVisualizer.h
--- a/Engine/Source/Editor/UnrealEd/Public/ComponentVisualizer.h
+++ b/Engine/Source/Editor/UnrealEd/Public/ComponentVisualizer.h
@@ -6,6 +6,7 @@
 #include "InputCoreTypes.h"
 #include "Components/ActorComponent.h"
 #include "HitProxies.h"
+#include "ConvexVolume.h"
 
 class AActor;
 class FCanvas;
diff --git a/Engine/Source/Runtime/Engine/Private/Components/SplineComponent.cpp b/Engine/Source/Runtime/Engine/Private/Components/SplineComponent.cpp
--- a/Engine/Source/Runtime/Engine/Private/Components/SplineComponent.cpp
+++ b/Engine/Source/Runtime/Engine/Private/Components/SplineComponent.cpp
@@ -9,6 +9,7 @@
 #include "PrimitiveViewRelevance.h"
 #include "PrimitiveSceneProxy.h"
 #include "SceneManagement.h"
+#include "UnrealEngine.h"
 
 #define SPLINE_FAST_BOUNDS_CALCULATION 0
 
