Made level cleanup output more verbose.

diff --git a/Engine/Source/Editor/UnrealEd/Private/EditorServer.cpp b/Engine/Source/Editor/UnrealEd/Private/EditorServer.cpp
--- a/Engine/Source/Editor/UnrealEd/Private/EditorServer.cpp
+++ b/Engine/Source/Editor/UnrealEd/Private/EditorServer.cpp
@@ -1955,6 +1955,7 @@ void UEditorEngine::BSPIntersectionHelper(UWorld* InWorld, ECsgOper Operation)
 
 void UEditorEngine::CheckForWorldGCLeaks( UWorld* NewWorld, UPackage* WorldPackage )
 {
+	int32 NumFailedToCleanup = 0;
 	// Make sure the old world is completely gone, except if the new world was one of it's sublevels
 	for(TObjectIterator<UWorld> It; It; ++It)
 	{
@@ -1964,7 +1965,8 @@ void UEditorEngine::CheckForWorldGCLeaks( UWorld* NewWorld, UPackage* WorldPacka
 		if(!bIsNewWorld && !bIsPersistantWorldType && !WorldHasValidContext(RemainingWorld))
 		{
 			FReferenceChainSearch RefChainSearch(RemainingWorld, EReferenceChainSearchMode::Shortest | EReferenceChainSearchMode::PrintResults);
-			UE_LOG(LogEditorServer, Fatal, TEXT("Old world %s not cleaned up by garbage collection while loading new map! Referenced by:") LINE_TERMINATOR TEXT("%s"), *RemainingWorld->GetPathName(), *RefChainSearch.GetRootPath());
+			UE_LOG(LogEditorServer, Error, TEXT("Old world %s not cleaned up by garbage collection while loading new map! Referenced by:") LINE_TERMINATOR TEXT("%s"), *RemainingWorld->GetPathName(), *RefChainSearch.GetRootPath());
+			NumFailedToCleanup++;
 		}
 	}
 
@@ -1978,10 +1980,16 @@ void UEditorEngine::CheckForWorldGCLeaks( UWorld* NewWorld, UPackage* WorldPacka
 			if(!bIsNewWorldPackage && RemainingPackage == WorldPackage)
 			{
 				FReferenceChainSearch RefChainSearch(RemainingPackage, EReferenceChainSearchMode::Shortest | EReferenceChainSearchMode::PrintResults);
-				UE_LOG(LogEditorServer, Fatal, TEXT("Old level package %s not cleaned up by garbage collection while loading new map! Referenced by:") LINE_TERMINATOR TEXT("%s"), *RemainingPackage->GetPathName(), *RefChainSearch.GetRootPath());
+				UE_LOG(LogEditorServer, Error, TEXT("Old level package %s not cleaned up by garbage collection while loading new map! Referenced by:") LINE_TERMINATOR TEXT("%s"), *RemainingPackage->GetPathName(), *RefChainSearch.GetRootPath());
+				NumFailedToCleanup++;
 			}
 		}
 	}
+	
+	if (NumFailedToCleanup > 0)
+	{
+		UE_LOG(LogEditorServer, Fatal, TEXT("World Memory Leaks: %d leaks objects and packages. See The output above."), NumFailedToCleanup);
+	}
 }
 
 void UEditorEngine::EditorDestroyWorld( FWorldContext & Context, const FText& CleanseText, UWorld* NewWorld )
