Fix CIS.

diff --git a/Engine/Source/Editor/UnrealEd/Private/DragTool_BoxSelect.cpp b/Engine/Source/Editor/UnrealEd/Private/DragTool_BoxSelect.cpp
--- a/Engine/Source/Editor/UnrealEd/Private/DragTool_BoxSelect.cpp
+++ b/Engine/Source/Editor/UnrealEd/Private/DragTool_BoxSelect.cpp
@@ -6,6 +6,7 @@
 #include "CanvasItem.h"
 #include "Settings/LevelEditorViewportSettings.h"
 #include "GameFramework/Volume.h"
+#include "Editor/UnrealEdEngine.h"
 #include "UnrealEdGlobals.h"
 #include "EngineUtils.h"
 #include "EditorModeManager.h"
diff --git a/Engine/Source/Editor/UnrealEd/Private/DragTool_FrustumSelect.cpp b/Engine/Source/Editor/UnrealEd/Private/DragTool_FrustumSelect.cpp
--- a/Engine/Source/Editor/UnrealEd/Private/DragTool_FrustumSelect.cpp
+++ b/Engine/Source/Editor/UnrealEd/Private/DragTool_FrustumSelect.cpp
@@ -6,6 +6,7 @@
 #include "CanvasItem.h"
 #include "Settings/LevelEditorViewportSettings.h"
 #include "GameFramework/Volume.h"
+#include "Editor/UnrealEdEngine.h"
 #include "UnrealEdGlobals.h"
 #include "EngineUtils.h"
 #include "Engine/Selection.h"
