Fixed transaction property diff detecting when the last write for a property was a write back effictively reducing the check range

diff --git a/Engine/Source/Editor/UnrealEd/Classes/Editor/Transactor.h b/Engine/Source/Editor/UnrealEd/Classes/Editor/Transactor.h
--- a/Engine/Source/Editor/UnrealEd/Classes/Editor/Transactor.h
+++ b/Engine/Source/Editor/UnrealEd/Classes/Editor/Transactor.h
@@ -116,7 +116,7 @@ protected:
 				else
 				{
 					DataOffset = FMath::Min(DataOffset, InOffset);
-					DataSize = InOffset + InSize - DataOffset;
+					DataSize = FMath::Max(InOffset + InSize - DataOffset, DataSize);
 				}
 			}
 
