Fixed wrong mip level being sampled when visualizing volume texture.

diff --git a/Engine/Shaders/Private/SimpleElementVolumeTexturePreviewPixelShader.usf b/Engine/Shaders/Private/SimpleElementVolumeTexturePreviewPixelShader.usf
--- a/Engine/Shaders/Private/SimpleElementVolumeTexturePreviewPixelShader.usf
+++ b/Engine/Shaders/Private/SimpleElementVolumeTexturePreviewPixelShader.usf
@@ -63,7 +63,7 @@ void TileMain(
 	}
 	else
 	{
-		Sample = Texture3DSample(InTexture, InTextureSampler, TexCoord3D);
+		Sample = Texture3DSampleGrad(InTexture, InTextureSampler, TexCoord3D, ddx(float3(TextureCoordinate, 0)), ddy(float3(TextureCoordinate, 0)));
 	}
 
 
diff --git a/Engine/Source/Editor/UnrealEd/Private/VolumeTexturePreview.cpp b/Engine/Source/Editor/UnrealEd/Private/VolumeTexturePreview.cpp
--- a/Engine/Source/Editor/UnrealEd/Private/VolumeTexturePreview.cpp
+++ b/Engine/Source/Editor/UnrealEd/Private/VolumeTexturePreview.cpp
@@ -84,7 +84,7 @@ public:
 		}
 		SetShaderValue(RHICmdList, GetPixelShader(),ColorWeights,ColorWeightsValue);
 
-		const int32 MipSizeZ = FMath::Max<int32>(SizeZ >> FMath::FloorToInt(MipLevel), 1);
+		const int32 MipSizeZ = MipLevel >= 0 ? FMath::Max<int32>(SizeZ >> FMath::FloorToInt(MipLevel), 1) : SizeZ;
 		FVector4 PackedParametersValue(GammaValue, MipLevel, (float)MipSizeZ, Opacity);
 		SetShaderValue(RHICmdList, GetPixelShader(), PackedParameters, PackedParametersValue);
 
