Fix for material instance with no parent being opened for edit
#rnx

diff --git a/Engine/Source/Editor/UnrealEd/Private/ThumbnailHelpers.cpp b/Engine/Source/Editor/UnrealEd/Private/ThumbnailHelpers.cpp
--- a/Engine/Source/Editor/UnrealEd/Private/ThumbnailHelpers.cpp
+++ b/Engine/Source/Editor/UnrealEd/Private/ThumbnailHelpers.cpp
@@ -361,7 +361,7 @@ void FMaterialThumbnailScene::SetMaterialInterface(UMaterialInterface* InMateria
 
 bool FMaterialThumbnailScene::ShouldSetSeparateTranslucency(class UMaterialInterface* InMaterial) const
 {
-	return InMaterial->GetMaterialResource(GMaxRHIFeatureLevel)->IsTranslucencyAfterDOFEnabled();
+	return InMaterial->GetMaterialResource(GMaxRHIFeatureLevel) != nullptr ? InMaterial->GetMaterialResource(GMaxRHIFeatureLevel)->IsTranslucencyAfterDOFEnabled() : false;
 }
 
 void FMaterialThumbnailScene::GetViewMatrixParameters(const float InFOVDegrees, FVector& OutOrigin, float& OutOrbitPitch, float& OutOrbitYaw, float& OutOrbitZoom) const
