Allow cooked content to work when dumping properties as well (see CL# 5708311)
#rnx

diff --git a/Engine/Source/Editor/UnrealEd/Private/Commandlets/PackageUtilities.cpp b/Engine/Source/Editor/UnrealEd/Private/Commandlets/PackageUtilities.cpp
--- a/Engine/Source/Editor/UnrealEd/Private/Commandlets/PackageUtilities.cpp
+++ b/Engine/Source/Editor/UnrealEd/Private/Commandlets/PackageUtilities.cpp
@@ -1585,6 +1585,7 @@ int32 UPkgInfoCommandlet::Main( const FString& Params )
 			if (Reader)
 			{
 				TGuardValue<bool> GuardAllowUnversionedContentInEditor(GAllowUnversionedContentInEditor, true);
+				TGuardValue<int32> GuardAllowCookedContentInEditor(GAllowCookedDataInEditorBuilds, 1);
 				UPackage* LoadedPackage = LoadPackage(Package, *Filename, LOAD_NoVerify, Reader);
 				if (LoadedPackage)
 				{
