#UE4 Add extra logging to fullloadandsave cooks, when it is requested

diff --git a/Engine/Source/Editor/UnrealEd/Private/CookOnTheFlyServer.cpp b/Engine/Source/Editor/UnrealEd/Private/CookOnTheFlyServer.cpp
--- a/Engine/Source/Editor/UnrealEd/Private/CookOnTheFlyServer.cpp
+++ b/Engine/Source/Editor/UnrealEd/Private/CookOnTheFlyServer.cpp
@@ -7895,6 +7895,8 @@ uint32 UCookOnTheFlyServer::FullLoadAndSave(uint32& CookedPackageCount)
 					{
 						FString PlatFilename = Filename.Replace(TEXT("[Platform]"), *Target->PlatformName());
 
+						UE_CLOG(GCookProgressDisplay & (int32)ECookProgressDisplayMode::PackageNames, LogCook, Display, TEXT("Cooking %s -> %s"), *Package->GetName(), *PlatFilename);
+
 						bool bSwap = (!Target->IsLittleEndian()) ^ (!PLATFORM_LITTLE_ENDIAN);
 						if (!Target->HasEditorOnlyData())
 						{
