Support dual-layer BD storage type and chunk layer assignment.
- Fix a bug in code generating pakchunklayers.txt

diff --git a/Engine/Source/Editor/UnrealEd/Private/Commandlets/AssetRegistryGenerator.cpp b/Engine/Source/Editor/UnrealEd/Private/Commandlets/AssetRegistryGenerator.cpp
--- a/Engine/Source/Editor/UnrealEd/Private/Commandlets/AssetRegistryGenerator.cpp
+++ b/Engine/Source/Editor/UnrealEd/Private/Commandlets/AssetRegistryGenerator.cpp
@@ -222,6 +222,50 @@ public:
 	}
 };
 
+static void ParseChunkLayerAssignment(TArray<FString> ChunkLayerAssignmentArray, TMap<int32, int32>& OutChunkLayerAssignment)
+{
+	OutChunkLayerAssignment.Empty();
+
+	const TCHAR* PropertyChunkId = TEXT("ChunkId=");
+	const TCHAR* PropertyLayerId = TEXT("Layer=");
+	for (FString& Entry : ChunkLayerAssignmentArray)
+	{
+		// Remove parentheses
+		Entry.TrimStartAndEndInline();
+		Entry.ReplaceInline(TEXT("("), TEXT(""));
+		Entry.ReplaceInline(TEXT(")"), TEXT(""));
+
+		int32 ChunkId = -1;
+		int32 LayerId = -1;
+		FParse::Value(*Entry, PropertyChunkId, ChunkId);
+		FParse::Value(*Entry, PropertyLayerId, LayerId);
+
+		if (ChunkId >= 0 && LayerId >= 0 && !OutChunkLayerAssignment.Contains(ChunkId))
+		{
+			OutChunkLayerAssignment.Add(ChunkId, LayerId);
+		}
+	}
+}
+
+static void AssignLayerChunkDelegate(const FAssignLayerChunkMap* ChunkManifest, const FString& Platform, const int32 ChunkIndex, int32& OutChunkLayer)
+{
+	OutChunkLayer = 0;
+
+	FConfigFile PlatformIniFile;
+	FConfigCacheIni::LoadLocalIniFile(PlatformIniFile, TEXT("Game"), true, *Platform);
+	TArray<FString> ChunkLayerAssignmentArray;
+	PlatformIniFile.GetArray(TEXT("/Script/UnrealEd.ProjectPackagingSettings"), TEXT("ChunkLayerAssignment"), ChunkLayerAssignmentArray);
+
+	TMap<int32, int32> ChunkLayerAssignment;
+	ParseChunkLayerAssignment(ChunkLayerAssignmentArray, ChunkLayerAssignment);
+
+	int32* LayerId = ChunkLayerAssignment.Find(ChunkIndex);
+	if (LayerId)
+	{
+		OutChunkLayer = *LayerId;
+	}
+}
+
 bool FAssetRegistryGenerator::GenerateStreamingInstallManifest(int64 InExtraFlavorChunkSize, FSandboxPlatformFile* InSandboxFile)
 {
 	const FString Platform = TargetPlatform->PlatformName();
@@ -316,6 +360,13 @@ bool FAssetRegistryGenerator::GenerateStreamingInstallManifest(int64 InExtraFlav
 	// generate per-chunk pak list files
 	for (int32 PakchunkIndex = 0; PakchunkIndex < FinalChunkManifests.Num(); ++PakchunkIndex)
 	{
+		// Serialize chunk layers whether chunk is empty or not
+		int32 TargetLayer = 0;
+		FGameDelegates::Get().GetAssignLayerChunkDelegate().ExecuteIfBound(FinalChunkManifests[PakchunkIndex], Platform, PakchunkIndex, TargetLayer);
+
+		FString LayerString = FString::Printf(TEXT("%d\r\n"), TargetLayer);
+		ChunkLayerFile->Serialize(TCHAR_TO_ANSI(*LayerString), LayerString.Len());
+
 		// Is this chunk empty?
 		if (!FinalChunkManifests[PakchunkIndex] || FinalChunkManifests[PakchunkIndex]->Num() == 0)
 		{
@@ -407,15 +458,7 @@ bool FAssetRegistryGenerator::GenerateStreamingInstallManifest(int64 InExtraFlav
 			// add this pakfilelist to our master list of pakfilelists
 			FString PakChunkListLine = FString::Printf(TEXT("%s%s\r\n"), *PakChunkFilename, *PakChunkOptions);
 			PakChunkListFile->Serialize(TCHAR_TO_ANSI(*PakChunkListLine), PakChunkListLine.Len());
-
-			int32 TargetLayer = 0;
-			FGameDelegates::Get().GetAssignLayerChunkDelegate().ExecuteIfBound(FinalChunkManifests[PakchunkIndex], Platform, PakchunkIndex, TargetLayer);
-
-			FString LayerString = FString::Printf(TEXT("%d\r\n"), TargetLayer);
-
-			ChunkLayerFile->Serialize(TCHAR_TO_ANSI(*LayerString), LayerString.Len());
 		}
-
 	}
 
 	ChunkLayerFile->Close();
@@ -682,6 +725,8 @@ void FAssetRegistryGenerator::Initialize(const TArray<FName> &InStartupPackages)
 	AssetRegistry.InitializeSerializationOptions(SaveOptions, TargetPlatform->IniPlatformName());
 
 	AssetRegistry.InitializeTemporaryAssetRegistryState(State, SaveOptions);
+
+	FGameDelegates::Get().GetAssignLayerChunkDelegate() = FAssignLayerChunkDelegate::CreateStatic(AssignLayerChunkDelegate);
 }
 
 void FAssetRegistryGenerator::ComputePackageDifferences(TSet<FName>& ModifiedPackages, TSet<FName>& NewPackages, TSet<FName>& RemovedPackages, TSet<FName>& IdenticalCookedPackages, TSet<FName>& IdenticalUncookedPackages, bool bRecurseModifications, bool bRecurseScriptModifications)
