Avoid creating what a child blueprint will consider a duplicate name when pasting.

diff --git a/Engine/Source/Editor/UnrealEd/Private/Kismet2/ComponentEditorUtils.cpp b/Engine/Source/Editor/UnrealEd/Private/Kismet2/ComponentEditorUtils.cpp
--- a/Engine/Source/Editor/UnrealEd/Private/Kismet2/ComponentEditorUtils.cpp
+++ b/Engine/Source/Editor/UnrealEd/Private/Kismet2/ComponentEditorUtils.cpp
@@ -262,8 +262,7 @@ FString FComponentEditorUtils::GenerateValidVariableNameFromAsset(UObject* Asset
 	int32 Counter = 1;
 	FString AssetName = Asset->GetName();
 
-	UClass* Class = Cast<UClass>(Asset);
-	if (Class)
+	if (UClass* Class = Cast<UClass>(Asset))
 	{
 		if (!Class->HasAnyClassFlags(CLASS_CompiledFromBlueprint))
 		{
@@ -274,6 +273,10 @@ FString FComponentEditorUtils::GenerateValidVariableNameFromAsset(UObject* Asset
 			AssetName.RemoveFromEnd("_C");
 		}
 	}
+	else if (UActorComponent* Comp = Cast <UActorComponent>(Asset))
+	{
+		AssetName.RemoveFromEnd(UActorComponent::ComponentTemplateNameSuffix);
+	}
 
 	// Try to create a name without any numerical suffix first
 	FString NewName = AssetName;
