AssetRegistryGenerator: Correct error message when a given encrypted root asset exists in two different chunks

diff --git a/Engine/Source/Editor/UnrealEd/Private/Commandlets/AssetRegistryGenerator.cpp b/Engine/Source/Editor/UnrealEd/Private/Commandlets/AssetRegistryGenerator.cpp
--- a/Engine/Source/Editor/UnrealEd/Private/Commandlets/AssetRegistryGenerator.cpp
+++ b/Engine/Source/Editor/UnrealEd/Private/Commandlets/AssetRegistryGenerator.cpp
@@ -575,7 +575,7 @@ void FAssetRegistryGenerator::InjectEncryptionData(FAssetRegistryState& TargetSt
 
 						if (AssetData->ChunkIDs.Num() > 1)
 						{
-							UE_LOG(LogAssetRegistryGenerator, Error, TEXT("Encrypted primary asset '%s' exists in two chunks. Only secondary assets should be shared between chunks."));
+							UE_LOG(LogAssetRegistryGenerator, Error, TEXT("Encrypted root asset '%s' exists in two chunks. Only secondary assets should be shared between chunks."), *AssetData->ObjectPath.ToString());
 						}
 						else if (AssetData->ChunkIDs.Num() == 1)
 						{
