Fix crash at export

diff --git a/Engine/Source/Editor/UnrealEd/Private/Fbx/FbxSkeletalMeshExport.cpp b/Engine/Source/Editor/UnrealEd/Private/Fbx/FbxSkeletalMeshExport.cpp
--- a/Engine/Source/Editor/UnrealEd/Private/Fbx/FbxSkeletalMeshExport.cpp
+++ b/Engine/Source/Editor/UnrealEd/Private/Fbx/FbxSkeletalMeshExport.cpp
@@ -251,7 +251,7 @@ FbxNode* FFbxExporter::CreateMesh(const USkeletalMesh* SkelMesh, const TCHAR* Me
 	}
 
 
-	if (GetExportOptions()->bExportMorphTargets)
+	if (GetExportOptions()->bExportMorphTargets && SkelMesh->Skeleton) //The skeleton can be null if this is a destructible mesh.
 	{
 		const FSmartNameMapping* SmartNameMapping = SkelMesh->Skeleton->GetSmartNameContainer(USkeleton::AnimCurveMappingName);
 		TMap<FName, FbxAnimCurve*> BlendShapeCurvesMap;
