Fixing leaked SWindow preventing UnrealEditor to close.

Alexis.Matte

diff --git a/Engine/Source/Editor/UnrealEd/Private/Fbx/SFbxSceneOptionWindow.cpp b/Engine/Source/Editor/UnrealEd/Private/Fbx/SFbxSceneOptionWindow.cpp
--- a/Engine/Source/Editor/UnrealEd/Private/Fbx/SFbxSceneOptionWindow.cpp
+++ b/Engine/Source/Editor/UnrealEd/Private/Fbx/SFbxSceneOptionWindow.cpp
@@ -1134,7 +1134,7 @@ TSharedRef<SDockTab> SFbxSceneOptionWindow::SpawnStaticMeshReimportTab(const FSp
 
 TSharedPtr<SWidget> SFbxSceneOptionWindow::SpawnDockTab()
 {
-	return FbxSceneImportTabManager->RestoreFrom(Layout.ToSharedRef(), OwnerWindow).ToSharedRef();
+	return FbxSceneImportTabManager->RestoreFrom(Layout.ToSharedRef(), OwnerWindow.Pin()).ToSharedRef();
 }
 
 void SFbxSceneOptionWindow::InitAllTabs()
@@ -1380,7 +1380,7 @@ void SFbxSceneOptionWindow::CloseFbxSceneOption()
 	if (OwnerWindow.IsValid())
 	{
 		//Close the window
-		OwnerWindow->RequestDestroyWindow();
+		OwnerWindow.Pin()->RequestDestroyWindow();
 	}
 	OwnerWindow = nullptr;
 }
diff --git a/Engine/Source/Editor/UnrealEd/Private/SFbxSceneOptionWindow.h b/Engine/Source/Editor/UnrealEd/Private/SFbxSceneOptionWindow.h
--- a/Engine/Source/Editor/UnrealEd/Private/SFbxSceneOptionWindow.h
+++ b/Engine/Source/Editor/UnrealEd/Private/SFbxSceneOptionWindow.h
@@ -49,7 +49,7 @@ public:
 		SLATE_ARGUMENT( class UFbxSceneImportOptionsStaticMesh*, SceneImportOptionsStaticMeshDisplay)
 		SLATE_ARGUMENT( ImportOptionsNameMapPtr, OverrideNameOptionsMap)
 		SLATE_ARGUMENT( class UFbxSceneImportOptionsSkeletalMesh*, SceneImportOptionsSkeletalMeshDisplay)
-		SLATE_ARGUMENT( TSharedPtr<SWindow>, OwnerWindow)
+		SLATE_ARGUMENT( TWeakPtr<SWindow>, OwnerWindow)
 		SLATE_ARGUMENT( FString, FullPath )
 	SLATE_END_ARGS()
 
@@ -142,7 +142,7 @@ private:
 	class UFbxSceneImportOptionsStaticMesh* SceneImportOptionsStaticMeshDisplay;
 	ImportOptionsNameMapPtr OverrideNameOptionsMap;
 	class UFbxSceneImportOptionsSkeletalMesh* SceneImportOptionsSkeletalMeshDisplay;
-	TSharedPtr< SWindow > OwnerWindow;
+	TWeakPtr< SWindow > OwnerWindow;
 	FString FullPath;
 
 	bool bCanReimportHierarchy;
