Fix "Updating Grass Maps" notification appearing while dragging.  This causes a loss of control over the drag

diff --git a/Engine/Source/Editor/UnrealEd/Private/GrassRenderingNotification.cpp b/Engine/Source/Editor/UnrealEd/Private/GrassRenderingNotification.cpp
--- a/Engine/Source/Editor/UnrealEd/Private/GrassRenderingNotification.cpp
+++ b/Engine/Source/Editor/UnrealEd/Private/GrassRenderingNotification.cpp
@@ -19,7 +19,8 @@ FGrassRenderingNotificationImpl GGrassRenderingNotification;
 
 bool FGrassRenderingNotificationImpl::ShouldShowNotification(const bool bIsNotificationAlreadyActive) const
 {
-	return ALandscapeProxy::TotalComponentsNeedingGrassMapRender > 0;
+	// Avoid showing the notification when a user is interacting as it causes a focus lost event which prevents them from dragging
+	return ALandscapeProxy::TotalComponentsNeedingGrassMapRender > 0 && !(FSlateApplication::Get().HasAnyMouseCaptor() || GUnrealEd->IsUserInteracting());
 }
 
 void FGrassRenderingNotificationImpl::SetNotificationText(const TSharedPtr<SNotificationItem>& InNotificationItem) const
