non-unity

diff --git a/Engine/Source/Editor/UnrealEd/Private/GrassRenderingNotification.cpp b/Engine/Source/Editor/UnrealEd/Private/GrassRenderingNotification.cpp
--- a/Engine/Source/Editor/UnrealEd/Private/GrassRenderingNotification.cpp
+++ b/Engine/Source/Editor/UnrealEd/Private/GrassRenderingNotification.cpp
@@ -4,6 +4,7 @@
 #include "GlobalEditorNotification.h"
 #include "LandscapeProxy.h"
 #include "Widgets/Notifications/SNotificationList.h"
+#include "Framework/Application/SlateApplication.h"
 
 /** Notification class for grassmmap rendering. */
 class FGrassRenderingNotificationImpl : public FGlobalEditorNotification
