#UE4 Non-unity fix

diff --git a/Engine/Source/Editor/UnrealEd/Private/GrassRenderingNotification.cpp b/Engine/Source/Editor/UnrealEd/Private/GrassRenderingNotification.cpp
--- a/Engine/Source/Editor/UnrealEd/Private/GrassRenderingNotification.cpp
+++ b/Engine/Source/Editor/UnrealEd/Private/GrassRenderingNotification.cpp
@@ -5,6 +5,7 @@
 #include "LandscapeProxy.h"
 #include "Widgets/Notifications/SNotificationList.h"
 #include "Framework/Application/SlateApplication.h"
+#include "UnrealEdGlobals.h"
 
 /** Notification class for grassmmap rendering. */
 class FGrassRenderingNotificationImpl : public FGlobalEditorNotification
