#UE4 Nonunity fix

diff --git a/Engine/Source/Editor/UnrealEd/Private/GrassRenderingNotification.cpp b/Engine/Source/Editor/UnrealEd/Private/GrassRenderingNotification.cpp
--- a/Engine/Source/Editor/UnrealEd/Private/GrassRenderingNotification.cpp
+++ b/Engine/Source/Editor/UnrealEd/Private/GrassRenderingNotification.cpp
@@ -6,6 +6,7 @@
 #include "Widgets/Notifications/SNotificationList.h"
 #include "Framework/Application/SlateApplication.h"
 #include "UnrealEdGlobals.h"
+#include "Editor/UnrealEdEngine.h"
 
 /** Notification class for grassmmap rendering. */
 class FGrassRenderingNotificationImpl : public FGlobalEditorNotification
