#UE4 Fix threading issue in concurrentsave

diff --git a/Engine/Source/Editor/UnrealEd/Classes/CookOnTheSide/CookOnTheFlyServer.h b/Engine/Source/Editor/UnrealEd/Classes/CookOnTheSide/CookOnTheFlyServer.h
--- a/Engine/Source/Editor/UnrealEd/Classes/CookOnTheSide/CookOnTheFlyServer.h
+++ b/Engine/Source/Editor/UnrealEd/Classes/CookOnTheSide/CookOnTheFlyServer.h
@@ -225,6 +225,7 @@ private:
 	// tmap of the Config name, Section name, Key name, to the value
 	typedef TMap<FName, TMap<FName, TMap<FName, TArray<FString>>>> FIniSettingContainer;
 
+	mutable FCriticalSection ConfigFileCS;
 	mutable bool IniSettingRecurse = false;
 	mutable FIniSettingContainer AccessedIniStrings;
 	TArray<const FConfigFile*> OpenConfigFiles;
diff --git a/Engine/Source/Editor/UnrealEd/Private/CookOnTheFlyServer.cpp b/Engine/Source/Editor/UnrealEd/Private/CookOnTheFlyServer.cpp
--- a/Engine/Source/Editor/UnrealEd/Private/CookOnTheFlyServer.cpp
+++ b/Engine/Source/Editor/UnrealEd/Private/CookOnTheFlyServer.cpp
@@ -4693,6 +4693,7 @@ bool UCookOnTheFlyServer::GetCookedIniVersionStrings(const ITargetPlatform* Targ
 
 void UCookOnTheFlyServer::OnFConfigCreated(const FConfigFile* Config)
 {
+	FScopeLock Lock(&ConfigFileCS);
 	if (IniSettingRecurse)
 	{
 		return;
@@ -4703,6 +4704,7 @@ void UCookOnTheFlyServer::OnFConfigCreated(const FConfigFile* Config)
 
 void UCookOnTheFlyServer::OnFConfigDeleted(const FConfigFile* Config)
 {
+	FScopeLock Lock(&ConfigFileCS);
 	if (IniSettingRecurse)
 	{
 		return;
