Unshelved from pending changelist '7678487':

Network Emulation in PIE
- Added ability to set network emulation in the editor Play launch options
- When enabled you can set minimum, maximum latency and packet loss to incoming and outgoing traffic
- Has a dropdown box populated with preconfigured network profiles (list configured in UNetworkSettings)
- Works on clients and servers (defaults to Server only)
- Works with "Run Under One Process" and when launching an external process

NetDriver
- Latency and packet loss can now be applied to incoming packets
- Network emulation console commands now prefixed by "NetEmulation". Ex: "NetEmulation PktEmulationProfile=Bad"
-Open URL's can now add netemulation commands. Ex: "Open 127.0.0.1?PktEmulationProfile=Bad"

[RN] major engine/network

diff --git a/Engine/Config/BaseEngine.ini b/Engine/Config/BaseEngine.ini
--- a/Engine/Config/BaseEngine.ini
+++ b/Engine/Config/BaseEngine.ini
@@ -2174,6 +2174,28 @@ LLVMWasmBackend=False
 [Browser]
 bForceMessageLoop=true
 
+[PacketSimulationProfile.Average]
+; Simulates network conditions with a ping of 60 to 120
+PktLoss=1
+PktIncomingLoss=1
+PktLagMin=30
+PktLagMax=60
+PktIncomingLagMin=30
+PktIncomingLagMax=60
+
+[PacketSimulationProfile.Bad]
+; Simulates network conditions with high packet loss and a ping of 200 to 400
+PktLoss=5
+PktIncomingLoss=5
+PktLagMin=100
+PktLagMax=200
+PktIncomingLagMin=100
+PktIncomingLagMax=200
+
+[/Script/Engine.NetworkSettings]
++NetworkEmulationProfiles=(ProfileName="Average",ToolTip="Simulates average internet conditions")
++NetworkEmulationProfiles=(ProfileName="Bad",ToolTip="Simulates laggy internet conditions")
+
 [/Script/HTML5Networking.WebSocketNetDriver]
 AllowPeerConnections=False
 AllowPeerVoice=False
diff --git a/Engine/Source/Editor/UnrealEd/Classes/Settings/LevelEditorPlayNetworkEmulationSettings.h b/Engine/Source/Editor/UnrealEd/Classes/Settings/LevelEditorPlayNetworkEmulationSettings.h
new file mode 100644
--- /dev/null
+++ b/Engine/Source/Editor/UnrealEd/Classes/Settings/LevelEditorPlayNetworkEmulationSettings.h
@@ -0,0 +1,110 @@
+// Copyright 1998-2019 Epic Games, Inc. All Rights Reserved.
+
+#pragma once
+
+#include "CoreMinimal.h"
+#include "IPropertyTypeCustomization.h"
+#include "PropertyHandle.h"
+#include "Widgets/Input/SComboBox.h"
+#include "LevelEditorPlayNetworkEmulationSettings.generated.h"
+
+class FDetailWidgetRow;
+class IDetailChildrenBuilder;
+class SWidget;
+struct FPacketSimulationSettings;
+
+USTRUCT()
+struct FNetworkEmulationPacketSettings
+{
+	GENERATED_USTRUCT_BODY()
+
+	// Minimum latency to add to packets
+	UPROPERTY(EditAnywhere, Category = "Network Settings", meta = (DisplayName = "Minimum Latency", ClampMin = "0", ClampMax = "5000"))
+	int32 MinLatency = 0;
+
+	// Maximum latency to add to packets. We use a random value between the minimum and maximum (when 0 = always the minimum value)
+	UPROPERTY(EditAnywhere, Category = "Network Settings", meta = (DisplayName = "Maximum Latency", ClampMin = "0", ClampMax = "5000"))
+	int32 MaxLatency = 0;
+
+	// Ratio of packets to randomly drop (0 = none, 100 = all)
+	UPROPERTY(EditAnywhere, Category = "Network Settings", meta = (ClampMin = "0", ClampMax = "100"))
+	int32 PacketLossPercentage = 0;
+};
+
+UENUM()
+enum class NetworkEmulationTarget
+{
+	Server UMETA(DisplayName = "Server Only"),
+	Client UMETA(DisplayName = "Clients Only"),
+	Any UMETA(DisplayName = "Everyone"),
+};
+
+USTRUCT()
+struct FLevelEditorPlayNetworkEmulationSettings
+{
+	GENERATED_USTRUCT_BODY()
+
+	// When true will apply the emulation settings when launching the game
+	UPROPERTY(EditAnywhere, Category = "Network Settings", meta = (DisplayName = "Enable Network Emulation"))
+	bool bIsNetworkEmulationEnabled;
+
+	UPROPERTY(EditAnywhere, Category = "Network Settings", meta = (DisplayName = "Emulation Target"))
+	NetworkEmulationTarget EmulationTarget;
+
+	// The profile name of the settings currently applied
+	UPROPERTY(EditAnywhere, Category = "Network Settings")
+	FString CurrentProfile;
+
+	// Settings that add latency and packet loss to all outgoing packets
+	UPROPERTY(EditAnywhere, Category = "Network Settings", meta = (DisplayName = "Outgoing traffic"))
+	FNetworkEmulationPacketSettings OutPackets;
+
+	// Settings that add latency and packet loss to all incoming packets
+	UPROPERTY(EditAnywhere, Category = "Network Settings", meta = (DisplayName = "Incoming traffic"))
+	FNetworkEmulationPacketSettings InPackets;
+
+	bool IsCustomProfile() const;
+
+	UNREALED_API bool IsEmulationEnabledForTarget(NetworkEmulationTarget CurrentTarget) const;
+
+	/** Convert the PIE settings into NetDriver settings */
+	UNREALED_API void ConvertToNetDriverSettings(FPacketSimulationSettings& OutNetDriverSettings) const;
+
+	/** Convert the PIE settings into CmdLine settings*/
+	UNREALED_API FString BuildPacketSettingsForCmdLine() const;
+
+	/** Convert the PIE settings into URL settings*/
+	UNREALED_API FString BuildPacketSettingsForURL() const;
+
+	void OnPostInitProperties();
+	void OnPostEditChange(struct FPropertyChangedEvent& PropertyChangedEvent);
+};
+
+
+/** Details customization for FEditorPlayNetworkEmulationSettings */
+class FLevelEditorPlayNetworkEmulationSettingsDetail : public IPropertyTypeCustomization
+{
+public:
+	static TSharedRef<IPropertyTypeCustomization> MakeInstance();
+
+	// IPropertyTypeCustomization interface
+	virtual void CustomizeHeader(TSharedRef<IPropertyHandle> StructPropertyHandle, FDetailWidgetRow& HeaderRow, IPropertyTypeCustomizationUtils& StructCustomizationUtils) override;
+	virtual void CustomizeChildren(TSharedRef<IPropertyHandle> StructPropertyHandle, IDetailChildrenBuilder& StructBuilder, IPropertyTypeCustomizationUtils& StructCustomizationUtils) override;
+
+	FText GetSelectedNetworkEmulationProfile() const;
+
+	// Combo button handlers
+	TSharedRef<SWidget> OnGetProfilesMenuContent() const;
+	void OnEmulationProfileChanged(int32 Index) const;
+
+	bool HandleAreSettingsEnabled() const;
+
+private:
+
+	// Property Handles
+	TSharedPtr<IPropertyHandle>	NetworkEmulationSettingsHandle;
+	TSharedPtr<IPropertyHandle>	IsNetworkEmulationEnabledHandle;
+	TSharedPtr<IPropertyHandle> CurrentProfileHandle;
+	TSharedPtr<IPropertyHandle> OutPacketsHandle;
+	TSharedPtr<IPropertyHandle> InPacketsHandle;
+};
\ No newline at end of file
diff --git a/Engine/Source/Editor/UnrealEd/Classes/Settings/LevelEditorPlaySettings.h b/Engine/Source/Editor/UnrealEd/Classes/Settings/LevelEditorPlaySettings.h
--- a/Engine/Source/Editor/UnrealEd/Classes/Settings/LevelEditorPlaySettings.h
+++ b/Engine/Source/Editor/UnrealEd/Classes/Settings/LevelEditorPlaySettings.h
@@ -7,6 +7,7 @@
 #include "UObject/Object.h"
 #include "Layout/Margin.h"
 #include "Layout/Visibility.h"
+#include "Settings/LevelEditorPlayNetworkEmulationSettings.h"
 #include "LevelEditorPlaySettings.generated.h"
 
 class SWindow;
@@ -348,6 +349,20 @@ private:
 	UPROPERTY(config, EditAnywhere, Category=MultiplayerOptions)
 	FString AdditionalLaunchOptions;
 
+public:
+
+	/** Should network emulation settings be applied or not */
+	bool IsNetworkEmulationEnabled() const
+	{
+		return NetworkEmulationSettings.bIsNetworkEmulationEnabled;
+	}
+
+	/**
+	 * Customizable settings allowing to emulate latency and packetloss for game network transmissions
+	 */
+	UPROPERTY(config, EditAnywhere, Category = MultiplayerOptions )
+	FLevelEditorPlayNetworkEmulationSettings NetworkEmulationSettings;
+
 public:
 
 	// Accessors for fetching the values of multiplayer options, and returning whether the option is valid at this time
@@ -380,6 +395,8 @@ public:
 	bool GetRouteGamepadToSecondWindow( bool &OutRouteGamepadToSecondWindow ) const { OutRouteGamepadToSecondWindow = RouteGamepadToSecondWindow; return IsRouteGamepadToSecondWindowActive(); }
 	EVisibility GetRouteGamepadToSecondWindowVisibility() const { return (RunUnderOneProcess ? EVisibility::Visible : EVisibility::Hidden); }
 
+	EVisibility GetNetworkEmulationVisibility() const { return (PlayNumberOfClients > 1 || PlayNetDedicated) ? EVisibility::Visible : EVisibility::Hidden; }
+
 	bool IsServerMapNameOverrideActive() const { return (PlayNetMode == PIE_StandaloneWithServer); }
 	bool GetServerMapNameOverride( FString& OutStandaloneServerMapName ) const { OutStandaloneServerMapName = ServerMapNameOverride; return IsServerMapNameOverrideActive(); }
 	EVisibility GetServerMapNameOverrideVisibility() const { return (PlayNetMode == PIE_StandaloneWithServer ? EVisibility::Visible : EVisibility::Hidden); }
diff --git a/Engine/Source/Editor/UnrealEd/Private/PlayLevel.cpp b/Engine/Source/Editor/UnrealEd/Private/PlayLevel.cpp
--- a/Engine/Source/Editor/UnrealEd/Private/PlayLevel.cpp
+++ b/Engine/Source/Editor/UnrealEd/Private/PlayLevel.cpp
@@ -1428,6 +1428,15 @@ void UEditorEngine::PlayStandaloneLocalPc(FString MapNameOverride, FIntPoint* Wi
 		AdditionalParameters += FString::Printf(TEXT(" -port=%hu"), ServerPort);
 	}
 
+	if (PlayInSettings->IsNetworkEmulationEnabled())
+	{
+		NetworkEmulationTarget CurrentTarget = bIsServer ? NetworkEmulationTarget::Server : NetworkEmulationTarget::Client;
+		if (PlayInSettings->NetworkEmulationSettings.IsEmulationEnabledForTarget(CurrentTarget))
+		{
+			AdditionalParameters += PlayInSettings->NetworkEmulationSettings.BuildPacketSettingsForCmdLine();
+		}
+	}
+
 	// Decide if fullscreen or windowed based on what is specified in the params
 	if (!AdditionalParameters.Contains(TEXT("-fullscreen")) && !AdditionalParameters.Contains(TEXT("-windowed")))
 	{
diff --git a/Engine/Source/Editor/UnrealEd/Private/Settings/LevelEditorPlayNetworkEmulationSettings.cpp b/Engine/Source/Editor/UnrealEd/Private/Settings/LevelEditorPlayNetworkEmulationSettings.cpp
new file mode 100644
--- /dev/null
+++ b/Engine/Source/Editor/UnrealEd/Private/Settings/LevelEditorPlayNetworkEmulationSettings.cpp
@@ -0,0 +1,349 @@
+// Copyright 1998-2019 Epic Games, Inc. All Rights Reserved.
+
+#include "Settings/LevelEditorPlayNetworkEmulationSettings.h"
+
+#include "IDetailChildrenBuilder.h"
+#include "IDetailPropertyRow.h"
+#include "DetailWidgetRow.h"
+#include "DetailLayoutBuilder.h"
+#include "Engine/NetDriver.h"
+#include "Engine/NetworkSettings.h"
+#include "Framework/MultiBox/MultiBoxBuilder.h"
+
+namespace NetworkEmulationSettingsHelper
+{
+	const FString& GetCustomProfileName()
+	{
+		static const FString CustomProfileName = TEXT("Custom");
+		return CustomProfileName;
+	}
+
+	void ConvertNetDriverSettingsToLevelEditorSettings(const FPacketSimulationSettings& NetDriverSettings, FNetworkEmulationPacketSettings& OutgoingTrafficPieSettings, FNetworkEmulationPacketSettings& IncomingTrafficPieSettings)
+	{
+		if (NetDriverSettings.PktLag > 0)
+		{
+			OutgoingTrafficPieSettings.MinLatency = FMath::Max(NetDriverSettings.PktLag - NetDriverSettings.PktLagVariance, 0);
+			OutgoingTrafficPieSettings.MaxLatency = FMath::Max(OutgoingTrafficPieSettings.MinLatency, NetDriverSettings.PktLag + NetDriverSettings.PktLagVariance); ;
+		}
+		else if (NetDriverSettings.PktLagMin > 0 || NetDriverSettings.PktLagMax > 0)
+		{
+			OutgoingTrafficPieSettings.MinLatency = NetDriverSettings.PktLagMin;
+			OutgoingTrafficPieSettings.MaxLatency = NetDriverSettings.PktLagMax;
+		}
+
+		OutgoingTrafficPieSettings.PacketLossPercentage = NetDriverSettings.PktLoss;
+
+		IncomingTrafficPieSettings.MinLatency = NetDriverSettings.PktIncomingLagMin;
+		IncomingTrafficPieSettings.MaxLatency = NetDriverSettings.PktIncomingLagMax;
+		IncomingTrafficPieSettings.PacketLossPercentage = NetDriverSettings.PktIncomingLoss;
+	}
+
+	FNetworkEmulationPacketSettings* GetPacketSettingsFromHandle(const TSharedPtr<IPropertyHandle>& PropertyHandle)
+	{
+		void* ValueData(nullptr);
+		FPropertyAccess::Result Result = PropertyHandle->GetValueData(ValueData);
+		if (Result != FPropertyAccess::Success)
+		{
+			return nullptr;
+		}
+
+		return (FNetworkEmulationPacketSettings*)ValueData;
+	}
+}
+
+//-------------------------------------------------------------------------------------------------
+// FLevelEditorPlayNetworkEmulationSettingsDetail
+
+TSharedRef<IPropertyTypeCustomization> FLevelEditorPlayNetworkEmulationSettingsDetail::MakeInstance()
+{
+	return MakeShareable(new FLevelEditorPlayNetworkEmulationSettingsDetail());
+}
+
+void FLevelEditorPlayNetworkEmulationSettingsDetail::CustomizeHeader(TSharedRef<IPropertyHandle> StructPropertyHandle, FDetailWidgetRow& HeaderRow, IPropertyTypeCustomizationUtils& StructCustomizationUtils)
+{
+	IsNetworkEmulationEnabledHandle = StructPropertyHandle->GetChildHandle(GET_MEMBER_NAME_CHECKED(FLevelEditorPlayNetworkEmulationSettings, bIsNetworkEmulationEnabled));
+
+	HeaderRow.NameContent()
+	[
+		IsNetworkEmulationEnabledHandle->CreatePropertyNameWidget()
+		
+	]
+	.ValueContent()
+	[
+		IsNetworkEmulationEnabledHandle->CreatePropertyValueWidget()
+	];
+}
+
+
+void FLevelEditorPlayNetworkEmulationSettingsDetail::CustomizeChildren(TSharedRef<IPropertyHandle> StructPropertyHandle, IDetailChildrenBuilder& StructBuilder, IPropertyTypeCustomizationUtils& StructCustomizationUtils)
+{
+	// Grab the widget handles
+	NetworkEmulationSettingsHandle = StructPropertyHandle;
+	CurrentProfileHandle = NetworkEmulationSettingsHandle->GetChildHandle(GET_MEMBER_NAME_CHECKED(FLevelEditorPlayNetworkEmulationSettings, CurrentProfile));
+	OutPacketsHandle = NetworkEmulationSettingsHandle->GetChildHandle(GET_MEMBER_NAME_CHECKED(FLevelEditorPlayNetworkEmulationSettings, OutPackets));
+	InPacketsHandle = NetworkEmulationSettingsHandle->GetChildHandle(GET_MEMBER_NAME_CHECKED(FLevelEditorPlayNetworkEmulationSettings, InPackets));
+
+	StructBuilder.AddProperty(NetworkEmulationSettingsHandle->GetChildHandle(GET_MEMBER_NAME_CHECKED(FLevelEditorPlayNetworkEmulationSettings, EmulationTarget)).ToSharedRef())
+		.IsEnabled(TAttribute<bool>(this, &FLevelEditorPlayNetworkEmulationSettingsDetail::HandleAreSettingsEnabled));
+
+	StructBuilder.AddCustomRow(NSLOCTEXT("FLevelEditorPlayNetworkEmulationSettings", "NetworkEmulationProfileName", "Network Emulation Profile"))
+	.NameContent()
+	[
+		SNew(STextBlock)
+		.Text(NSLOCTEXT("FLevelEditorPlayNetworkEmulationSettings", "NetworkEmulationProfileName", "Network Emulation Profile"))
+		.Font(IDetailLayoutBuilder::GetDetailFont())
+
+	]
+	.ValueContent()
+	[
+		SNew(SComboButton)
+		.OnGetMenuContent(this, &FLevelEditorPlayNetworkEmulationSettingsDetail::OnGetProfilesMenuContent)
+		.ButtonContent()
+		[
+			SNew(STextBlock)
+			.Text(this, &FLevelEditorPlayNetworkEmulationSettingsDetail::GetSelectedNetworkEmulationProfile)
+			.Font(IDetailLayoutBuilder::GetDetailFont())
+		]
+	]
+	.IsEnabled(TAttribute<bool>(this, &FLevelEditorPlayNetworkEmulationSettingsDetail::HandleAreSettingsEnabled));
+
+	
+
+	StructBuilder.AddProperty(InPacketsHandle.ToSharedRef())
+		.IsEnabled(TAttribute<bool>(this, &FLevelEditorPlayNetworkEmulationSettingsDetail::HandleAreSettingsEnabled));
+	
+	StructBuilder.AddProperty(OutPacketsHandle.ToSharedRef())
+		.IsEnabled(TAttribute<bool>(this, &FLevelEditorPlayNetworkEmulationSettingsDetail::HandleAreSettingsEnabled));
+}
+
+bool FLevelEditorPlayNetworkEmulationSettingsDetail::HandleAreSettingsEnabled() const
+{
+	if (IsNetworkEmulationEnabledHandle.IsValid())
+	{
+		bool bIsEnabled(false);
+		IsNetworkEmulationEnabledHandle->GetValue(bIsEnabled);
+		return bIsEnabled;
+	}
+
+	return false;
+}
+
+TSharedRef<SWidget> FLevelEditorPlayNetworkEmulationSettingsDetail::OnGetProfilesMenuContent() const
+{
+	FMenuBuilder MenuBuilder(true, nullptr);
+
+	if (const UNetworkSettings* NetworkSettings = GetDefault<UNetworkSettings>())
+	{
+		// Add preconfigured profiles
+		for (int32 Index=0; Index < NetworkSettings->NetworkEmulationProfiles.Num(); Index++ )
+		{
+			const FNetworkEmulationProfileDescription& ProfileDescription = NetworkSettings->NetworkEmulationProfiles[Index];
+			
+			MenuBuilder.AddMenuEntry(
+				FText::FromString(ProfileDescription.ProfileName),
+				FText::FromString(ProfileDescription.ToolTip),
+				FSlateIcon(),
+				FUIAction(
+					FExecuteAction::CreateSP(this, &FLevelEditorPlayNetworkEmulationSettingsDetail::OnEmulationProfileChanged, Index),
+					FCanExecuteAction()
+				)
+			);
+		}
+
+		// Add the custom entry
+		MenuBuilder.AddMenuEntry(
+			NSLOCTEXT("FLevelEditorPlayNetworkEmulationSettings", "NetworkEmulationCustomProfile", "Custom"),
+			NSLOCTEXT("FLevelEditorPlayNetworkEmulationSettings", "NetworkEmulationCustomProfileToolTip", "Customizable profile"),
+			FSlateIcon(),
+			FUIAction(
+				FExecuteAction::CreateSP(this, &FLevelEditorPlayNetworkEmulationSettingsDetail::OnEmulationProfileChanged, -1),
+				FCanExecuteAction()
+			)
+		);
+	}
+
+	return MenuBuilder.MakeWidget();
+}
+
+void FLevelEditorPlayNetworkEmulationSettingsDetail::OnEmulationProfileChanged(int32 Index) const
+{
+	FString Selection;
+	if (const UNetworkSettings* NetworkSettings = GetDefault<UNetworkSettings>())
+	{
+		if(NetworkSettings->NetworkEmulationProfiles.IsValidIndex(Index))
+		{
+			Selection = NetworkSettings->NetworkEmulationProfiles[Index].ProfileName;
+		}
+		else
+		{
+			Selection = NetworkEmulationSettingsHelper::GetCustomProfileName();
+		}
+	}
+
+	// Get the current settings
+	FNetworkEmulationPacketSettings* OutSettings = NetworkEmulationSettingsHelper::GetPacketSettingsFromHandle(OutPacketsHandle);
+	FNetworkEmulationPacketSettings* InSettings = NetworkEmulationSettingsHelper::GetPacketSettingsFromHandle(InPacketsHandle);
+
+	// Reload the settings when switching to an official profile
+	if (Selection != NetworkEmulationSettingsHelper::GetCustomProfileName())
+	{
+		FPacketSimulationSettings NetDriverSettings;
+		NetDriverSettings.LoadEmulationProfile(*Selection);
+
+		FNetworkEmulationPacketSettings OutgoingTrafficPIESettings;
+		FNetworkEmulationPacketSettings IncomingTrafficPIESettings;
+		NetworkEmulationSettingsHelper::ConvertNetDriverSettingsToLevelEditorSettings(NetDriverSettings, OutgoingTrafficPIESettings, IncomingTrafficPIESettings);
+
+		*OutSettings = OutgoingTrafficPIESettings;
+		*InSettings = IncomingTrafficPIESettings;
+	}
+
+	CurrentProfileHandle->SetValue(Selection);
+}
+
+FText FLevelEditorPlayNetworkEmulationSettingsDetail::GetSelectedNetworkEmulationProfile() const
+{
+	FString SelectedProfile;
+	CurrentProfileHandle->GetValue(SelectedProfile);
+
+	return FText::FromString(SelectedProfile);
+}
+
+//-------------------------------------------------------------------------------------------------
+// FLevelEditorPlayNetworkEmulationSettings
+
+bool FLevelEditorPlayNetworkEmulationSettings::IsCustomProfile() const
+{
+	return CurrentProfile == NetworkEmulationSettingsHelper::GetCustomProfileName();
+}
+
+bool FLevelEditorPlayNetworkEmulationSettings::IsEmulationEnabledForTarget(NetworkEmulationTarget CurrentTarget) const
+{
+	// Settings are applied to everyone
+	if (EmulationTarget == NetworkEmulationTarget::Any)
+	{
+		return true;
+	}
+
+	// Otherwise check if target matches the wanted setting
+	return EmulationTarget == CurrentTarget;
+}
+
+void FLevelEditorPlayNetworkEmulationSettings::ConvertToNetDriverSettings(FPacketSimulationSettings& OutNetDriverSettings) const
+{
+	OutNetDriverSettings.ResetSettings();
+
+	bool bProfileFound(false);
+
+	if ( !IsCustomProfile() )
+	{
+		// Load hardcoded profiles from the config since they might have advanced settings not found in PIE settings
+		bProfileFound = OutNetDriverSettings.LoadEmulationProfile(*CurrentProfile);
+	}
+
+	if (IsCustomProfile() || !bProfileFound)
+	{
+		// For custom set the settings manually from the PIE variables
+		OutNetDriverSettings.PktLagMin = OutPackets.MinLatency;
+		OutNetDriverSettings.PktLagMax = OutPackets.MaxLatency;
+		OutNetDriverSettings.PktLoss = OutPackets.PacketLossPercentage;
+		OutNetDriverSettings.PktIncomingLagMin = InPackets.MinLatency;
+		OutNetDriverSettings.PktIncomingLagMax = InPackets.MaxLatency;
+		OutNetDriverSettings.PktIncomingLoss = InPackets.PacketLossPercentage;
+	}
+}
+
+FString FLevelEditorPlayNetworkEmulationSettings::BuildPacketSettingsForCmdLine() const
+{
+	// Empty string when disabled
+	if (!bIsNetworkEmulationEnabled)
+	{
+		return FString();
+	}
+
+	FString CmdLine;
+
+	if (IsCustomProfile())
+	{
+		// Set each setting manually for user-edited profiles
+		CmdLine += FString::Printf(TEXT(" -PktLagMin=%d"), OutPackets.MinLatency);
+		CmdLine += FString::Printf(TEXT(" -PktLagMax=%d"), OutPackets.MaxLatency);
+		CmdLine += FString::Printf(TEXT(" -PktLoss=%d"), OutPackets.PacketLossPercentage);
+
+		CmdLine += FString::Printf(TEXT(" -PktIncomingLagMin=%d"), InPackets.MinLatency);
+		CmdLine += FString::Printf(TEXT(" -PktIncomingLagMax=%d"), InPackets.MaxLatency);
+		CmdLine += FString::Printf(TEXT(" -PktIncomingLoss=%d"), InPackets.PacketLossPercentage);
+	}
+	else
+	{
+		CmdLine += FString::Printf(TEXT(" -PktEmulationProfile=%s"), *CurrentProfile);
+	}
+	
+	return CmdLine;
+}
+
+FString FLevelEditorPlayNetworkEmulationSettings::BuildPacketSettingsForURL() const
+{
+	// Empty string when disabled
+	if (!bIsNetworkEmulationEnabled)
+	{
+		return FString();
+	}
+
+	FString CmdLine;
+
+	if (IsCustomProfile())
+	{
+		// Set each setting manually for user-edited profiles
+		CmdLine += FString::Printf(TEXT("?PktLagMin=%d"), OutPackets.MinLatency);
+		CmdLine += FString::Printf(TEXT("?PktLagMax=%d"), OutPackets.MaxLatency);
+		CmdLine += FString::Printf(TEXT("?PktLoss=%d"), OutPackets.PacketLossPercentage);
+
+		CmdLine += FString::Printf(TEXT("?PktIncomingLagMin=%d"), InPackets.MinLatency);
+		CmdLine += FString::Printf(TEXT("?PktIncomingLagMax=%d"), InPackets.MaxLatency);
+		CmdLine += FString::Printf(TEXT("?PktIncomingLoss=%d"), InPackets.PacketLossPercentage);
+	}
+	else
+	{
+		CmdLine += FString::Printf(TEXT("?PktEmulationProfile=%s"), *CurrentProfile);
+	}
+
+	return CmdLine;
+}
+
+void FLevelEditorPlayNetworkEmulationSettings::OnPostInitProperties()
+{
+	if (!IsCustomProfile())
+	{
+		// For official profiles reload the settings from the configs in case some values got changed
+		FPacketSimulationSettings NetDriverSettings;
+		bool bProfileFound = NetDriverSettings.LoadEmulationProfile(*CurrentProfile);
+
+		if (bProfileFound)
+		{
+			NetworkEmulationSettingsHelper::ConvertNetDriverSettingsToLevelEditorSettings(NetDriverSettings, OutPackets, InPackets);
+		}
+		else
+		{
+			CurrentProfile = NetworkEmulationSettingsHelper::GetCustomProfileName();
+		}
+	}
+}
+
+void FLevelEditorPlayNetworkEmulationSettings::OnPostEditChange(struct FPropertyChangedEvent& PropertyChangedEvent)
+{
+	if (PropertyChangedEvent.Property)
+	{
+		if (GET_MEMBER_NAME_CHECKED(FLevelEditorPlayNetworkEmulationSettings, CurrentProfile) == PropertyChangedEvent.Property->GetFName() ||
+			GET_MEMBER_NAME_CHECKED(FLevelEditorPlayNetworkEmulationSettings, bIsNetworkEmulationEnabled) == PropertyChangedEvent.Property->GetFName() ||
+			GET_MEMBER_NAME_CHECKED(FLevelEditorPlayNetworkEmulationSettings, EmulationTarget) == PropertyChangedEvent.Property->GetFName())
+		{
+			// Don't consider the settings dirty when these properties change
+			return;
+		}
+
+		// Set profile to Custom when the user dirties any other property
+		CurrentProfile = NetworkEmulationSettingsHelper::GetCustomProfileName();
+	}
+}
+
diff --git a/Engine/Source/Editor/UnrealEd/Private/Settings/LevelEditorPlaySettingsCustomization.h b/Engine/Source/Editor/UnrealEd/Private/Settings/LevelEditorPlaySettingsCustomization.h
--- a/Engine/Source/Editor/UnrealEd/Private/Settings/LevelEditorPlaySettingsCustomization.h
+++ b/Engine/Source/Editor/UnrealEd/Private/Settings/LevelEditorPlaySettingsCustomization.h
@@ -700,6 +700,9 @@ public:
 				.DisplayName(LOCTEXT("AdditionalLaunchOptionsLabel", "Command Line Arguments"))
 				.Visibility(TAttribute<EVisibility>::Create(TAttribute<EVisibility>::FGetter::CreateSP(this, &FLevelEditorPlaySettingsCustomization::HandleCmdLineVisibility)));
 
+			NetworkCategory.AddProperty("NetworkEmulationSettings")
+				.Visibility(TAttribute<EVisibility>::Create(TAttribute<EVisibility>::FGetter::CreateSP(this, &FLevelEditorPlaySettingsCustomization::HandleNetworkEmulationVisibility)));
+
 			NetworkCategory.AddCustomRow(LOCTEXT("PlayInNetworkViewportSize", "Multiplayer Viewport Size"), false)
 				.NameContent()
 				[
diff --git a/Engine/Source/Editor/UnrealEd/Private/Settings/SettingsClasses.cpp b/Engine/Source/Editor/UnrealEd/Private/Settings/SettingsClasses.cpp
--- a/Engine/Source/Editor/UnrealEd/Private/Settings/SettingsClasses.cpp
+++ b/Engine/Source/Editor/UnrealEd/Private/Settings/SettingsClasses.cpp
@@ -429,6 +429,11 @@ void ULevelEditorPlaySettings::PostEditChangeProperty(struct FPropertyChangedEve
 		}
 	}
 
+	if (PropertyChangedEvent.MemberProperty && PropertyChangedEvent.MemberProperty->GetFName() == GET_MEMBER_NAME_CHECKED(ULevelEditorPlaySettings, NetworkEmulationSettings))
+	{
+		NetworkEmulationSettings.OnPostEditChange(PropertyChangedEvent);
+	}
+
 	Super::PostEditChangeProperty(PropertyChangedEvent);
 }
 
@@ -439,6 +444,8 @@ void ULevelEditorPlaySettings::PostInitProperties()
 	NewWindowWidth = FMath::Max(0, NewWindowWidth);
 	NewWindowHeight = FMath::Max(0, NewWindowHeight);
 
+	NetworkEmulationSettings.OnPostInitProperties();
+
 #if WITH_EDITOR
 	FCoreDelegates::OnSafeFrameChangedEvent.AddUObject(this, &ULevelEditorPlaySettings::SwapSafeZoneTypes);
 #endif
diff --git a/Engine/Source/Editor/UnrealEd/Private/UnrealEdEngine.cpp b/Engine/Source/Editor/UnrealEd/Private/UnrealEdEngine.cpp
--- a/Engine/Source/Editor/UnrealEd/Private/UnrealEdEngine.cpp
+++ b/Engine/Source/Editor/UnrealEd/Private/UnrealEdEngine.cpp
@@ -49,6 +49,7 @@
 #include "Settings/GameMapsSettingsCustomization.h"
 #include "Settings/LevelEditorPlaySettingsCustomization.h"
 #include "Settings/ProjectPackagingSettingsCustomization.h"
+#include "Settings/LevelEditorPlayNetworkEmulationSettings.h"
 #include "StatsViewerModule.h"
 #include "SnappingUtils.h"
 #include "PackageAutoSaver.h"
@@ -152,6 +153,9 @@ void UUnrealEdEngine::Init(IEngineLoop* InEngineLoop)
 		PropertyModule.RegisterCustomClassLayout("GameMapsSettings", FOnGetDetailCustomizationInstance::CreateStatic(&FGameMapsSettingsCustomization::MakeInstance));
 		PropertyModule.RegisterCustomClassLayout("LevelEditorPlaySettings", FOnGetDetailCustomizationInstance::CreateStatic(&FLevelEditorPlaySettingsCustomization::MakeInstance));
 		PropertyModule.RegisterCustomClassLayout("ProjectPackagingSettings", FOnGetDetailCustomizationInstance::CreateStatic(&FProjectPackagingSettingsCustomization::MakeInstance));
+
+		PropertyModule.RegisterCustomPropertyTypeLayout("LevelEditorPlayNetworkEmulationSettings", FOnGetPropertyTypeCustomizationInstance::CreateStatic(&FLevelEditorPlayNetworkEmulationSettingsDetail::MakeInstance));
+		
 	}
 
 	if (!IsRunningCommandlet())
diff --git a/Engine/Source/Runtime/Engine/Classes/Engine/NetConnection.h b/Engine/Source/Runtime/Engine/Classes/Engine/NetConnection.h
--- a/Engine/Source/Runtime/Engine/Classes/Engine/NetConnection.h
+++ b/Engine/Source/Runtime/Engine/Classes/Engine/NetConnection.h
@@ -170,7 +170,7 @@ DECLARE_DELEGATE_ThreeParams(FOnLowLevelSend, void* /*Data*/, int32 /*Count*/, b
 /**
  * An artificially lagged packet
  */
-struct DelayedPacket
+struct FDelayedPacket
 {
 	/** The packet data to send */
 	TArray<uint8> Data;
@@ -186,7 +186,7 @@ struct DelayedPacket
 
 public:
 	UE_DEPRECATED(4.21, "Use the constructor that takes PacketTraits for allowing for analytics and flags")
-	FORCEINLINE DelayedPacket(uint8* InData, int32 InSizeBytes, int32 InSizeBits)
+	FORCEINLINE FDelayedPacket(uint8* InData, int32 InSizeBytes, int32 InSizeBits)
 		: Data()
 		, SizeBits(InSizeBits)
 		, Traits()
@@ -196,7 +196,7 @@ public:
 		FMemory::Memcpy(Data.GetData(), InData, InSizeBytes);
 	}
 
-	FORCEINLINE DelayedPacket(uint8* InData, int32 InSizeBits, FOutPacketTraits& InTraits)
+	FORCEINLINE FDelayedPacket(uint8* InData, int32 InSizeBits, FOutPacketTraits& InTraits)
 		: Data()
 		, SizeBits(InSizeBits)
 		, Traits(InTraits)
@@ -213,7 +213,16 @@ public:
 		Data.CountBytes(Ar);
 	}
 };
-#endif
+
+struct FDelayedIncomingPacket
+{
+	TUniquePtr<FBitReader> PacketData;
+
+	/** Time at which the packet should be reinjected into the connection */
+	double ReinjectionTime = 0.0;
+};
+
+#endif //#if DO_ENABLE_NET_TEST
 
 /** Record of channels with data written into each outgoing packet. */
 struct FWrittenChannelsRecord
@@ -610,16 +619,31 @@ public:
 	ENGINE_API void UpdateLevelVisibility(const FName& PackageName, bool bIsVisible);
 
 #if DO_ENABLE_NET_TEST
+
 	// For development.
 	/** Packet settings for testing lag, net errors, etc */
 	FPacketSimulationSettings PacketSimulationSettings;
 
-	/** delayed packet array */
-	TArray<DelayedPacket> Delayed;
-
 	/** Copies the settings from the net driver to our local copy */
-	void UpdatePacketSimulationSettings(void);
-#endif
+	void UpdatePacketSimulationSettings();
+
+private:
+
+	/** delayed outgoing packet array */
+	TArray<FDelayedPacket> Delayed;
+
+	/** delayed incoming packet array */
+	TArray<FDelayedIncomingPacket> DelayedIncomingPackets;
+
+	/** set to true when already delayed packets are reinjected */
+	bool bIsReinjectingDelayedPackets;
+
+	/** Process incoming packets that have been delayed for long enough */
+	void ReinjectDelayedPackets();
+
+#endif //#if DO_ENABLE_NET_TEST
+
+public:
 
 	/** 
 	 * If true, will resend everything this connection has ever sent, since the connection has been open.
@@ -1200,6 +1224,14 @@ private:
 	/** Returns true if an outgoing packet should be dropped due to packet simulation settings, including loss burst simulation. */
 	bool ShouldDropOutgoingPacketForLossSimulation(int64 NumBits) const;
 
+	/**
+	*	Test the current emulation settings to delay, drop, etc the current packet
+	*	Returns true if the packet was emulated and false when the packet must be sent via the normal path
+	*/
+#if DO_ENABLE_NET_TEST
+	bool CheckOutgoingPacketEmulation(FOutPacketTraits& Traits);
+#endif
+
 	/** Write packetHeader */
 	void WritePacketHeader(FBitWriter& Writer);
 
diff --git a/Engine/Source/Runtime/Engine/Classes/Engine/NetDriver.h b/Engine/Source/Runtime/Engine/Classes/Engine/NetDriver.h
--- a/Engine/Source/Runtime/Engine/Classes/Engine/NetDriver.h
+++ b/Engine/Source/Runtime/Engine/Classes/Engine/NetDriver.h
@@ -328,6 +328,7 @@ struct FNetworkObjectInfo;
 class UChannel;
 class IAnalyticsProvider;
 class FNetAnalyticsAggregator;
+class UNetDriver;
 
 using FConnectionMap = TMap<TSharedRef<const FInternetAddr>, UNetConnection*, FDefaultSetAllocator, FInternetAddrConstKeyMapFuncs<UNetConnection*>>;
 
@@ -373,7 +374,7 @@ DECLARE_DELEGATE_SevenParams(FOnSendRPC, AActor* /*Actor*/, UFunction* /*Functio
 USTRUCT()
 struct ENGINE_API FPacketSimulationSettings
 {
-	GENERATED_USTRUCT_BODY()
+	GENERATED_BODY()
 
 	/**
 	 * When set, will cause calls to FlushNet to drop packets.
@@ -449,10 +450,30 @@ struct ENGINE_API FPacketSimulationSettings
 	 */
 	UPROPERTY(EditAnywhere, Category = "Simulation Settings")
 	int32	PktLagMin = 0;
-	
 	UPROPERTY(EditAnywhere, Category = "Simulation Settings")
 	int32	PktLagMax = 0;
 
+	/**
+	 * Set a value to add a minimum delay in milliseconds to incoming
+	 * packets before they are processed.
+	 */
+	UPROPERTY(EditAnywhere, Category = "Simulation Settings")
+	int32	PktIncomingLagMin = 0;
+	
+	/**
+	 * The maximum delay in milliseconds to add to incoming
+	 * packets before they are processed.
+	 */
+	UPROPERTY(EditAnywhere, Category = "Simulation Settings")
+	int32	PktIncomingLagMax = 0;
+
+	/**
+	 * The ratio of incoming packets that will be dropped
+	 * to simulate packet loss
+	 */
+	UPROPERTY(EditAnywhere, Category = "Simulation Settings")
+	int32	PktIncomingLoss = 0;
+
 	/** reads in settings from the .ini file 
 	 * @note: overwrites all previous settings
 	 */
@@ -1619,4 +1640,10 @@ private:
 
 	/** Count the number of notified packets, i.e. packets that we know if they are delivered or not. Used to reliably measure outgoing packet loss */
 	uint32 OutTotalNotifiedPackets;
-};
+
+#if DO_ENABLE_NET_TEST
+	/** Dont load packet settings from config or cmdline when true*/
+	bool bForcedPacketSettings;
+#endif 
+
+};
\ No newline at end of file
diff --git a/Engine/Source/Runtime/Engine/Classes/Engine/NetworkSettings.h b/Engine/Source/Runtime/Engine/Classes/Engine/NetworkSettings.h
--- a/Engine/Source/Runtime/Engine/Classes/Engine/NetworkSettings.h
+++ b/Engine/Source/Runtime/Engine/Classes/Engine/NetworkSettings.h
@@ -9,6 +9,18 @@
 
 struct FPropertyChangedEvent;
 
+USTRUCT()
+struct FNetworkEmulationProfileDescription
+{
+	GENERATED_BODY()
+
+	UPROPERTY()
+	FString ProfileName;
+
+	UPROPERTY()
+	FString ToolTip;
+};
+
 /**
  * Network settings.
  */
@@ -47,6 +59,10 @@ class ENGINE_API UNetworkSettings : public UDeveloperSettings
 		ClampMin = "1", ClampMax = "65535", UIMin = "1", UIMax = "65535"))
 	int32 MaxRepArrayMemory = DefaultMaxRepArrayMemory;
 
+	/** This lists the common network emulation profiles that will be selectable in PIE settings */
+	UPROPERTY(config)
+	TArray<FNetworkEmulationProfileDescription> NetworkEmulationProfiles;
+
 public:
 
 	//~ Begin UObject Interface
diff --git a/Engine/Source/Runtime/Engine/Private/GameInstance.cpp b/Engine/Source/Runtime/Engine/Private/GameInstance.cpp
--- a/Engine/Source/Runtime/Engine/Private/GameInstance.cpp
+++ b/Engine/Source/Runtime/Engine/Private/GameInstance.cpp
@@ -318,6 +318,14 @@ FGameInstancePIEResult UGameInstance::StartPlayInEditorGameInstance(ULocalPlayer
 			URLString += FString::Printf(TEXT(":%hu"), ServerPort);
 		}
 
+		if (PlayInSettings->IsNetworkEmulationEnabled())
+		{
+			if (PlayInSettings->NetworkEmulationSettings.IsEmulationEnabledForTarget(NetworkEmulationTarget::Client))
+			{
+				URLString += PlayInSettings->NetworkEmulationSettings.BuildPacketSettingsForURL();
+			}
+		}
+
 		if (EditorEngine->Browse(*WorldContext, FURL(&BaseURL, *URLString, (ETravelType)TRAVEL_Absolute), Error) == EBrowseReturnVal::Pending)
 		{
 			EditorEngine->TransitionType = ETransitionType::WaitingToConnect;
@@ -332,13 +340,26 @@ FGameInstancePIEResult UGameInstance::StartPlayInEditorGameInstance(ULocalPlayer
 		// we're going to be playing in the current world, get it ready for play
 		UWorld* const PlayWorld = GetWorld();
 
+		FString ExtraURLOptions;
+		if (PlayInSettings->IsNetworkEmulationEnabled())
+		{
+			NetworkEmulationTarget CurrentTarget = PlayNetMode == PIE_ListenServer ? NetworkEmulationTarget::Server : NetworkEmulationTarget::Client;
+			if (PlayInSettings->NetworkEmulationSettings.IsEmulationEnabledForTarget(CurrentTarget))
+			{
+				ExtraURLOptions += PlayInSettings->NetworkEmulationSettings.BuildPacketSettingsForURL();
+			}
+		}
+
 		// make a URL
 		FURL URL;
 		// If the user wants to start in spectator mode, do not use the custom play world for now
 		if (EditorEngine->UserEditedPlayWorldURL.Len() > 0 && !Params.bStartInSpectatorMode)
 		{
+			FString UserURL = EditorEngine->UserEditedPlayWorldURL;
+			UserURL += ExtraURLOptions;
+
 			// If the user edited the play world url. Verify that the map name is the same as the currently loaded map.
-			URL = FURL(NULL, *EditorEngine->UserEditedPlayWorldURL, TRAVEL_Absolute);
+			URL = FURL(NULL, *UserURL, TRAVEL_Absolute);
 			if (URL.Map != PIEMapName)
 			{
 				// Ensure the URL map name is the same as the generated play world map name.
@@ -348,7 +369,7 @@ FGameInstancePIEResult UGameInstance::StartPlayInEditorGameInstance(ULocalPlayer
 		else
 		{
 			// The user did not edit the url, just build one from scratch.
-			URL = FURL(NULL, *EditorEngine->BuildPlayWorldURL(*PIEMapName, Params.bStartInSpectatorMode), TRAVEL_Absolute);
+			URL = FURL(NULL, *EditorEngine->BuildPlayWorldURL(*PIEMapName, Params.bStartInSpectatorMode, ExtraURLOptions), TRAVEL_Absolute);
 		}
 
 		// If a start location is specified, spawn a temporary PlayerStartPIE actor at the start location and use it as the portal.
diff --git a/Engine/Source/Runtime/Engine/Private/NetConnection.cpp b/Engine/Source/Runtime/Engine/Private/NetConnection.cpp
--- a/Engine/Source/Runtime/Engine/Private/NetConnection.cpp
+++ b/Engine/Source/Runtime/Engine/Private/NetConnection.cpp
@@ -605,7 +605,7 @@ PRAGMA_ENABLE_DEPRECATION_WARNINGS
 #if DO_ENABLE_NET_TEST
 		GRANULAR_NETWORK_MEMORY_TRACKING_TRACK("Delayed",
 			Delayed.CountBytes(Ar);
-			for (const DelayedPacket& Packet : Delayed)
+			for (const FDelayedPacket& Packet : Delayed)
 			{
 				Packet.CountBytes(Ar);
 			}
@@ -1173,6 +1173,10 @@ void UNetConnection::PostTickDispatch()
 {
 	if (!InternalAck)
 	{
+#if DO_ENABLE_NET_TEST
+		ReinjectDelayedPackets();
+#endif
+
 		FlushPacketOrderCache(/*bFlushWholeCache=*/true);
 		PacketAnalytics.Tick();
 	}
@@ -1220,6 +1224,30 @@ void UNetConnection::FlushPacketOrderCache(bool bFlushWholeCache/*=false*/)
 	}
 }
 
+#if DO_ENABLE_NET_TEST
+void UNetConnection::ReinjectDelayedPackets()
+{
+	const double CurrentTime = FPlatformTime::Seconds();
+
+	TGuardValue<bool> ReinjectingGuard(bIsReinjectingDelayedPackets, true);
+
+	uint32 NbReinjected(0);
+	for (const FDelayedIncomingPacket& DelayedPacket : DelayedIncomingPackets)
+	{
+		if (DelayedPacket.ReinjectionTime > CurrentTime)
+		{
+			break;
+		}
+		
+		++NbReinjected;
+		ReceivedPacket(*DelayedPacket.PacketData.Get());
+	}
+
+	// Delete processed packets
+	DelayedIncomingPackets.RemoveAt(0, NbReinjected, false);
+}
+#endif //#if DO_ENABLE_NET_TEST
+
 uint32 GNetOutBytes = 0;
 
 void UNetConnection::FlushNet(bool bIgnoreSimulation)
@@ -1290,61 +1318,17 @@ void UNetConnection::FlushNet(bool bIgnoreSimulation)
 
 		// Send now.
 #if DO_ENABLE_NET_TEST
+
+		bool bWasPacketEmulated(false);
+
 		// if the connection is closing/being destroyed/etc we need to send immediately regardless of settings
 		// because we won't be around to send it delayed
-		if (State == USOCK_Closed || IsGarbageCollecting() || bIgnoreSimulation || InternalAck)
+		if (State != USOCK_Closed && !IsGarbageCollecting() && !bIgnoreSimulation && !InternalAck)
 		{
-			// Checked in FlushNet() so each child class doesn't have to implement this
-			if (Driver->IsNetResourceValid())
-			{
-				LowLevelSend(SendBuffer.GetData(), SendBuffer.GetNumBits(), Traits);
-			}
+			bWasPacketEmulated = CheckOutgoingPacketEmulation(Traits);
 		}
-		else if (PacketSimulationSettings.PktOrder)
-		{
-			DelayedPacket& B = *(new(Delayed)DelayedPacket(SendBuffer.GetData(), SendBuffer.GetNumBits(), Traits));
 
-			for (int32 i=Delayed.Num() - 1; i >= 0; i--)
-			{
-				if (FMath::FRand() > 0.50)
-				{
-					if (!ShouldDropOutgoingPacketForLossSimulation(SendBuffer.GetNumBits()))
-					{
-						// Checked in FlushNet() so each child class doesn't have to implement this
-						if (Driver->IsNetResourceValid())
-						{
-							LowLevelSend((char*) &Delayed[i].Data[0], Delayed[i].SizeBits, Delayed[i].Traits);
-						}
-					}
-					Delayed.RemoveAt(i);
-				}
-			}
-		}
-		else if (PacketSimulationSettings.PktLag)
-		{
-			if (!ShouldDropOutgoingPacketForLossSimulation(SendBuffer.GetNumBits()))
-			{
-				DelayedPacket& B = *(new(Delayed)DelayedPacket(SendBuffer.GetData(), SendBuffer.GetNumBits(), Traits));
-
-				// ExtraLag goes from PktLag + [-PktLagVariance, PktLagVariance]
-				const double LagVariance = 2.0f * (FMath::FRand() - 0.5f) * double(PacketSimulationSettings.PktLagVariance);
-				const double ExtraLag = (double(PacketSimulationSettings.PktLag) + LagVariance) / 1000.f;
-				B.SendTime = FPlatformTime::Seconds() + ExtraLag;
-			}
-		}
-		else if (PacketSimulationSettings.PktLagMin > 0 && PacketSimulationSettings.PktLagMax > 0)
-		{
-			if (!ShouldDropOutgoingPacketForLossSimulation(SendBuffer.GetNumBits()))
-			{
-				DelayedPacket& B = *(new(Delayed)DelayedPacket(SendBuffer.GetData(), SendBuffer.GetNumBits(), Traits));
-
-				// ExtraLag goes from [PktLagMin, PktLagMax]
-				const double LagVariance = FMath::FRand() * double(PacketSimulationSettings.PktLagMax - PacketSimulationSettings.PktLagMin);
-				const double ExtraLag = (double(PacketSimulationSettings.PktLagMin) + LagVariance) / 1000.f;
-				B.SendTime = FPlatformTime::Seconds() + ExtraLag;
-			}
-		}
-		else if (!ShouldDropOutgoingPacketForLossSimulation(SendBuffer.GetNumBits()))
+		if (!bWasPacketEmulated)
 		{
 #endif
 			// Checked in FlushNet() so each child class doesn't have to implement this
@@ -1410,6 +1394,62 @@ void UNetConnection::FlushNet(bool bIgnoreSimulation)
 	}
 }
 
+#if DO_ENABLE_NET_TEST
+bool UNetConnection::CheckOutgoingPacketEmulation(FOutPacketTraits& Traits)
+{
+	if (ShouldDropOutgoingPacketForLossSimulation(SendBuffer.GetNumBits()))
+	{
+		UE_LOG(LogNet, VeryVerbose, TEXT("Dropping outgoing packet at %f"), FPlatformTime::Seconds());
+		return true;
+	}
+	else if (PacketSimulationSettings.PktOrder)
+	{
+		FDelayedPacket& B = *(new(Delayed)FDelayedPacket(SendBuffer.GetData(), SendBuffer.GetNumBits(), Traits));
+
+		for (int32 i = Delayed.Num() - 1; i >= 0; i--)
+		{
+			if (FMath::FRand() > 0.50)
+			{
+				// Checked in FlushNet() so each child class doesn't have to implement this
+				if (Driver->IsNetResourceValid())
+				{
+					LowLevelSend((char*)&Delayed[i].Data[0], Delayed[i].SizeBits, Delayed[i].Traits);
+				}
+
+				Delayed.RemoveAt(i);
+			}
+		}
+
+		return true;
+	}
+	else if (PacketSimulationSettings.PktLag)
+	{
+		FDelayedPacket& B = *(new(Delayed)FDelayedPacket(SendBuffer.GetData(), SendBuffer.GetNumBits(), Traits));
+
+		// ExtraLag goes from PktLag + [-PktLagVariance, PktLagVariance]
+		const double LagVariance = 2.0f * (FMath::FRand() - 0.5f) * double(PacketSimulationSettings.PktLagVariance);
+		const double ExtraLag = (double(PacketSimulationSettings.PktLag) + LagVariance) / 1000.f;
+		B.SendTime = FPlatformTime::Seconds() + ExtraLag;
+
+		return true;
+
+	}
+	else if (PacketSimulationSettings.PktLagMin > 0 || PacketSimulationSettings.PktLagMax > 0)
+	{
+		FDelayedPacket& B = *(new(Delayed)FDelayedPacket(SendBuffer.GetData(), SendBuffer.GetNumBits(), Traits));
+
+		// ExtraLag goes from [PktLagMin, PktLagMax]
+		const double LagVariance = FMath::FRand() * double(PacketSimulationSettings.PktLagMax - PacketSimulationSettings.PktLagMin);
+		const double ExtraLag = (double(PacketSimulationSettings.PktLagMin) + LagVariance) / 1000.f;
+		B.SendTime = FPlatformTime::Seconds() + ExtraLag;
+		
+		return true;
+	}
+
+	return false;
+}
+#endif
+
 bool UNetConnection::ShouldDropOutgoingPacketForLossSimulation(int64 NumBits) const
 {
 #if DO_ENABLE_NET_TEST
@@ -1680,6 +1720,35 @@ void UNetConnection::ReceivedPacket( FBitReader& Reader )
 		return;
 	}
 
+#if DO_ENABLE_NET_TEST
+	if (!InternalAck && !bIsReinjectingDelayedPackets)
+	{
+		if (PacketSimulationSettings.PktIncomingLoss)
+		{
+			if (FMath::FRand() * 100.f < PacketSimulationSettings.PktIncomingLoss)
+			{
+				UE_LOG(LogNet, VeryVerbose, TEXT("Dropped incoming packet at %f"), FPlatformTime::Seconds());
+				return;
+			}
+
+		}
+		if (PacketSimulationSettings.PktIncomingLagMin > 0 || PacketSimulationSettings.PktIncomingLagMax > 0)
+		{
+			const double LagVariance = FMath::FRand() * double(PacketSimulationSettings.PktIncomingLagMax - PacketSimulationSettings.PktIncomingLagMin);
+			const double ExtraLag = (double(PacketSimulationSettings.PktIncomingLagMin) + LagVariance) / 1000.f;
+
+			FDelayedIncomingPacket DelayedPacket;
+			DelayedPacket.PacketData = MakeUnique<FBitReader>(Reader);
+			DelayedPacket.ReinjectionTime = FPlatformTime::Seconds() + ExtraLag;
+
+			DelayedIncomingPackets.Emplace(MoveTemp(DelayedPacket));
+
+			UE_LOG(LogNet, VeryVerbose, TEXT("Delaying incoming packet for %f seconds"), ExtraLag);
+			return;
+		}
+	}
+#endif //#if DO_ENABLE_NET_TEST
+
 
 	FBitReaderMark ResetReaderMark(Reader);
 
diff --git a/Engine/Source/Runtime/Engine/Private/NetworkDriver.cpp b/Engine/Source/Runtime/Engine/Private/NetworkDriver.cpp
--- a/Engine/Source/Runtime/Engine/Private/NetworkDriver.cpp
+++ b/Engine/Source/Runtime/Engine/Private/NetworkDriver.cpp
@@ -294,11 +294,16 @@ void UNetDriver::InitPacketSimulationSettings()
 		return;
 	}
 
-	PacketSimulationSettings.ResetSettings();
+	PacketSimulationSettings.RegisterCommands();
+
+	if (bForcedPacketSettings)
+	{
+		return;
+	}
 
-	// read the settings from .ini and command line, with the command line taking precedence	
+	// Read the settings from .ini and command line, with the command line taking precedence	
+	PacketSimulationSettings.ResetSettings();
 	PacketSimulationSettings.LoadConfig(*NetDriverName.ToString());
-	PacketSimulationSettings.RegisterCommands();
 	PacketSimulationSettings.ParseSettings(FCommandLine::Get(), *NetDriverName.ToString());
 #endif
 }
@@ -1419,6 +1424,22 @@ bool UNetDriver::InitBase(bool bInitAsClient, FNetworkNotify* InNotify, const FU
 		});
 	}
 
+#if DO_ENABLE_NET_TEST
+	bool bSettingFound(false);
+	FPacketSimulationSettings PacketSettings;
+
+	for (const FString& URLOption : URL.Op)
+	{
+		bSettingFound |= PacketSettings.ParseSettings(*URLOption);
+	}
+
+	if( bSettingFound )
+	{
+		SetPacketSimulationSettings(PacketSettings);
+		bForcedPacketSettings = true;
+	}
+#endif //#if DO_ENABLE_NET_TEST
+
 	Notify = InNotify;
 
 	return bSuccess;
@@ -3283,6 +3304,10 @@ void FPacketSimulationSettings::LoadConfig(const TCHAR* OptionalQualifier)
 	
 	ConfigHelperInt(TEXT("PktDup"), PktDup, OptionalQualifier);
 
+	ConfigHelperInt(TEXT("PktIncomingLagMin"), PktIncomingLagMin, OptionalQualifier);
+	ConfigHelperInt(TEXT("PktIncomingLagMax"), PktIncomingLagMax, OptionalQualifier);
+	ConfigHelperInt(TEXT("PktIncomingLoss"), PktIncomingLoss, OptionalQualifier);
+
 	ValidateSettings();
 }
 
@@ -3344,6 +3369,10 @@ void FPacketSimulationSettings::ValidateSettings()
 	PktLagMax = FMath::Max(PktLagMin, PktLagMax);
 
 	PktDup = FMath::Clamp<int32>(PktDup, 0, 100);
+
+	PktIncomingLagMin = FMath::Max(PktIncomingLagMin, 0);
+	PktIncomingLagMax = FMath::Max(PktIncomingLagMin, PktIncomingLagMax);
+	PktIncomingLoss = FMath::Clamp<int32>(PktIncomingLoss, 0, 100);
 }
 
 bool FPacketSimulationSettings::ConfigHelperInt(const TCHAR* Name, int32& Value, const TCHAR* OptionalQualifier)
@@ -3387,16 +3416,22 @@ void FPacketSimulationSettings::RegisterCommands()
 	IConsoleManager& ConsoleManager = IConsoleManager::Get();
 	
 	// Register exec commands with the console manager for auto-completion if they havent been registered already by another net driver
-	if (!ConsoleManager.IsNameRegistered(TEXT("Net PktLoss=")))
+	if (!ConsoleManager.IsNameRegistered(TEXT("NetEmulation PktEmulationProfile=")))
 	{
-		ConsoleManager.RegisterConsoleCommand(TEXT("Net PktLoss="), TEXT("PktLoss=<n> (simulates network packet loss)"));
-		ConsoleManager.RegisterConsoleCommand(TEXT("Net PktOrder="), TEXT("PktOrder=<n> (simulates network packet received out of order)"));
-		ConsoleManager.RegisterConsoleCommand(TEXT("Net PktDup="), TEXT("PktDup=<n> (simulates sending/receiving duplicate network packets)"));
-		ConsoleManager.RegisterConsoleCommand(TEXT("Net PktLag="), TEXT("PktLag=<n> (simulates network packet lag)"));
-		ConsoleManager.RegisterConsoleCommand(TEXT("Net PktLagVariance="), TEXT("PktLagVariance=<n> (simulates variable network packet lag)"));
-		ConsoleManager.RegisterConsoleCommand(TEXT("Net PktLagMin="), TEXT("PktLagMin=<n> (minimum packet latency)"));
-		ConsoleManager.RegisterConsoleCommand(TEXT("Net PktLagMax="), TEXT("PktLagMax=<n> (maximum packet latency)"));
-		ConsoleManager.RegisterConsoleCommand(TEXT("Net EmulationProfile="), TEXT("EmulationProfile=<NAME> (apply an emulation profile)"));
+		ConsoleManager.RegisterConsoleCommand(TEXT("NetEmulation PktEmulationProfile="), TEXT("PktEmulationProfile=<NAME> (apply an emulation profile)"));
+
+		ConsoleManager.RegisterConsoleCommand(TEXT("NetEmulation PktLoss="), TEXT("PktLoss=<n> (simulates network packet loss)"));
+		ConsoleManager.RegisterConsoleCommand(TEXT("NetEmulation PktOrder="), TEXT("PktOrder=<n> (simulates network packet received out of order)"));
+		ConsoleManager.RegisterConsoleCommand(TEXT("NetEmulation PktDup="), TEXT("PktDup=<n> (simulates sending/receiving duplicate network packets)"));
+		ConsoleManager.RegisterConsoleCommand(TEXT("NetEmulation PktLag="), TEXT("PktLag=<n> (simulates network packet lag)"));
+		ConsoleManager.RegisterConsoleCommand(TEXT("NetEmulation PktLagVariance="), TEXT("PktLagVariance=<n> (simulates variable network packet lag)"));
+
+		ConsoleManager.RegisterConsoleCommand(TEXT("NetEmulation PktLagMin="), TEXT("PktLagMin=<n> (minimum outgoing packet latency)"));
+		ConsoleManager.RegisterConsoleCommand(TEXT("NetEmulation PktLagMax="), TEXT("PktLagMax=<n> (maximum outgoing packet latency)"));
+
+		ConsoleManager.RegisterConsoleCommand(TEXT("NetEmulation PktIncomingLagMin="), TEXT("PktIncomingLagMin=<n> (minimum incoming packet latency)"));
+		ConsoleManager.RegisterConsoleCommand(TEXT("NetEmulation PktIncomingLagMax="), TEXT("PktIncomingLagMax=<n> (maximum incoming packet latency)"));
+		ConsoleManager.RegisterConsoleCommand(TEXT("NetEmulation PktIncomingLoss="), TEXT("PktIncomingLoss=<n> (simulates incoming packet loss)"));
 	}
 
 	// Note we never unregister the console commands since net drivers come and go, and we can sometimes have more than 1, etc. 
@@ -3410,17 +3445,15 @@ void FPacketSimulationSettings::RegisterCommands()
  */
 bool FPacketSimulationSettings::ParseSettings(const TCHAR* Cmd, const TCHAR* OptionalQualifier)
 {
-	UE_LOG(LogTemp, Display, TEXT("ParseSettings for %s"), OptionalQualifier);
 	// note that each setting is tested.
 	// this is because the same function will be used to parse the command line as well
 	bool bParsed = false;
 
 	FString EmulationProfileName;
-	if (FParse::Value(Cmd, TEXT("EmulationProfile="), EmulationProfileName))
+	if (FParse::Value(Cmd, TEXT("PktEmulationProfile="), EmulationProfileName))
 	{
-		bParsed = true;
 		UE_LOG(LogNet, Log, TEXT("Applying EmulationProfile %s"), *EmulationProfileName);
-		LoadEmulationProfile(*EmulationProfileName);
+		bParsed = LoadEmulationProfile(*EmulationProfileName);
 	}
 	if( ParseHelper(Cmd, TEXT("PktLoss="), PktLoss, OptionalQualifier) )
 	{
@@ -3467,6 +3500,21 @@ bool FPacketSimulationSettings::ParseSettings(const TCHAR* Cmd, const TCHAR* Opt
 		bParsed = true;
 		UE_LOG(LogNet, Log, TEXT("PktLagMax set to %d"), PktLagMax);
 	}
+	if (ParseHelper(Cmd, TEXT("PktIncomingLagMin="), PktIncomingLagMin, OptionalQualifier))
+	{
+		bParsed = true;
+		UE_LOG(LogNet, Log, TEXT("PktIncomingLagMin set to %d"), PktIncomingLagMin);
+	}
+	if (ParseHelper(Cmd, TEXT("PktIncomingLagMax="), PktIncomingLagMax, OptionalQualifier))
+	{
+		bParsed = true;
+		UE_LOG(LogNet, Log, TEXT("PktIncomingLagMax set to %d"), PktIncomingLagMax);
+	}
+	if (ParseHelper(Cmd, TEXT("PktIncomingLoss="), PktIncomingLoss, OptionalQualifier))
+	{
+		bParsed = true;
+		UE_LOG(LogNet, Log, TEXT("PktIncomingLoss set to %d"), PktIncomingLoss);
+	}
 
 	ValidateSettings();
 	return bParsed;
@@ -5625,7 +5673,7 @@ static bool NetDriverExec(UWorld* InWorld, const TCHAR* Cmd, FOutputDevice& Ar)
 	bool bHandled = false;
 
 	// Ignore any execs that don't start with NET
-	if (FParse::Command(&Cmd, TEXT("NET")))
+	if (FParse::Command(&Cmd, TEXT("NET")) || FParse::Command(&Cmd, TEXT("NETEMULATION")))
 	{
 		UNetDriver* NamedDriver = NULL;
 		TCHAR TokenStr[128];
diff --git a/Engine/Source/Runtime/Engine/Private/PendingNetGame.cpp b/Engine/Source/Runtime/Engine/Private/PendingNetGame.cpp
--- a/Engine/Source/Runtime/Engine/Private/PendingNetGame.cpp
+++ b/Engine/Source/Runtime/Engine/Private/PendingNetGame.cpp
@@ -470,13 +470,20 @@ void UPendingNetGame::Tick( float DeltaTime )
 	 *   ****may NULL itself via CancelPending if a disconnect/error occurs****
 	 */
 	NetDriver->TickDispatch(DeltaTime);
+
+	if (NetDriver)
+	{
+		NetDriver->PostTickDispatch();
+	}
+
 	if (NetDriver)
 	{
 		NetDriver->TickFlush(DeltaTime);
-		if (NetDriver)
-		{
-			NetDriver->PostTickFlush();
-		}
+	}
+
+	if (NetDriver)
+	{
+		NetDriver->PostTickFlush();
 	}
 }
 
