Network emulation PIE settings
* Set custom profile when no current profile have been assigned yet.
* Downgrade warning to log when NetEmulation profile is not found in configs

diff --git a/Engine/Source/Editor/UnrealEd/Private/Settings/LevelEditorPlayNetworkEmulationSettings.cpp b/Engine/Source/Editor/UnrealEd/Private/Settings/LevelEditorPlayNetworkEmulationSettings.cpp
--- a/Engine/Source/Editor/UnrealEd/Private/Settings/LevelEditorPlayNetworkEmulationSettings.cpp
+++ b/Engine/Source/Editor/UnrealEd/Private/Settings/LevelEditorPlayNetworkEmulationSettings.cpp
@@ -313,10 +313,15 @@ FString FLevelEditorPlayNetworkEmulationSettings::BuildPacketSettingsForURL() co
 
 void FLevelEditorPlayNetworkEmulationSettings::OnPostInitProperties()
 {
-	if (!IsCustomProfile())
+	if (CurrentProfile.IsEmpty())
 	{
-		// For official profiles reload the settings from the configs in case some values got changed
+		CurrentProfile = NetworkEmulationSettingsHelper::GetCustomProfileName();
+	}
+	else if (!IsCustomProfile())
+	{
+		// For official profiles reload the settings from the config in case some values got changed
 		FPacketSimulationSettings NetDriverSettings;
+
 		bool bProfileFound = NetDriverSettings.LoadEmulationProfile(*CurrentProfile);
 
 		if (bProfileFound)
diff --git a/Engine/Source/Runtime/Engine/Private/NetworkDriver.cpp b/Engine/Source/Runtime/Engine/Private/NetworkDriver.cpp
--- a/Engine/Source/Runtime/Engine/Private/NetworkDriver.cpp
+++ b/Engine/Source/Runtime/Engine/Private/NetworkDriver.cpp
@@ -3319,7 +3319,7 @@ bool FPacketSimulationSettings::LoadEmulationProfile(const TCHAR* ProfileName)
 	bool bSectionExists = GConfig->GetSection(*SectionName, SectionConfigs, GEngineIni);
 	if (!bSectionExists)
 	{
-		UE_LOG(LogNet, Warning, TEXT("EmulationProfile [%s] was not found in %s. Packet settings were not changed"), *SectionName, *GEngineIni);
+		UE_LOG(LogNet, Log, TEXT("EmulationProfile [%s] was not found in %s. Packet settings were not changed"), *SectionName, *GEngineIni);
 		return false;
 	}
 
