Fix missing function due to bad integration resolve

diff --git a/Engine/Source/Editor/UnrealEd/Private/Settings/LevelEditorPlaySettingsCustomization.h b/Engine/Source/Editor/UnrealEd/Private/Settings/LevelEditorPlaySettingsCustomization.h
--- a/Engine/Source/Editor/UnrealEd/Private/Settings/LevelEditorPlaySettingsCustomization.h
+++ b/Engine/Source/Editor/UnrealEd/Private/Settings/LevelEditorPlaySettingsCustomization.h
@@ -901,6 +901,12 @@ private:
 		return GetDefault<ULevelEditorPlaySettings>()->GetPlayNetModeVisibility();
 	}
 
+	// Callback for checking if the Network Emulation can be used
+	EVisibility HandleNetworkEmulationVisibility() const
+	{
+		return GetDefault<ULevelEditorPlaySettings>()->GetNetworkEmulationVisibility();
+	}
+
 	// Callback for getting the visibility of the StandaloneServerMapName property.
 	EVisibility HandleServerMapNameOverrideVisibility() const
 	{
