Moving the test for low free UObjectArray slot count when cooking before testing available memory and exiting early

diff --git a/Engine/Source/Editor/UnrealEd/Private/CookOnTheFlyServer.cpp b/Engine/Source/Editor/UnrealEd/Private/CookOnTheFlyServer.cpp
--- a/Engine/Source/Editor/UnrealEd/Private/CookOnTheFlyServer.cpp
+++ b/Engine/Source/Editor/UnrealEd/Private/CookOnTheFlyServer.cpp
@@ -3223,6 +3223,14 @@ bool UCookOnTheFlyServer::HasExceededMaxMemory() const
 		return true;
 	}
 
+#if UE_GC_TRACK_OBJ_AVAILABLE
+	if (GUObjectArray.GetObjectArrayEstimatedAvailable() < MinFreeUObjectIndicesBeforeGC)
+	{
+		UE_LOG(LogCook, Display, TEXT("Running out of available UObject indices (%d remaining)"), GUObjectArray.GetObjectArrayEstimatedAvailable());
+		return true;
+	}
+#endif // UE_GC_TRACK_OBJ_AVAILABLE
+
 	// don't gc if we haven't reached our min gc level yet
 	if (MemStats.UsedVirtual < MinMemoryBeforeGC)
 	{
@@ -3238,14 +3246,6 @@ bool UCookOnTheFlyServer::HasExceededMaxMemory() const
 		return true;
 	}
 
-#if UE_GC_TRACK_OBJ_AVAILABLE
-	if (GUObjectArray.GetObjectArrayEstimatedAvailable() < MinFreeUObjectIndicesBeforeGC)
-	{
-		UE_LOG(LogCook, Display, TEXT("Running out of available UObject indices (%d remaining)"), GUObjectArray.GetObjectArrayEstimatedAvailable());
-		return true;
-	}
-#endif // UE_GC_TRACK_OBJ_AVAILABLE
-
 	return false;
 }
 
@@ -4139,7 +4139,7 @@ void UCookOnTheFlyServer::Initialize( ECookMode::Type DesiredCookMode, ECookInit
 	MinMemoryBeforeGC = MinMemoryBeforeGCInMB * 1024LL * 1024LL;
 	MinMemoryBeforeGC = FMath::Min(MaxMemoryAllowance, MinMemoryBeforeGC);
 
-	MinFreeUObjectIndicesBeforeGC = 5000;
+	MinFreeUObjectIndicesBeforeGC = 100000;
 	GConfig->GetInt(TEXT("CookSettings"), TEXT("MinFreeUObjectIndicesBeforeGC"), MinFreeUObjectIndicesBeforeGC, GEditorIni);
 	MinFreeUObjectIndicesBeforeGC = FMath::Max(MinFreeUObjectIndicesBeforeGC, 0);
 
