Cook commandlet will now properly report the number of available UObject indices before and after collecting garbage

diff --git a/Engine/Source/Editor/UnrealEd/Private/Commandlets/CookCommandlet.cpp b/Engine/Source/Editor/UnrealEd/Private/Commandlets/CookCommandlet.cpp
--- a/Engine/Source/Editor/UnrealEd/Private/Commandlets/CookCommandlet.cpp
+++ b/Engine/Source/Editor/UnrealEd/Private/Commandlets/CookCommandlet.cpp
@@ -991,7 +991,7 @@ bool UCookCommandlet::CookByTheBook( const TArray<ITargetPlatform*>& Platforms,
 						bShouldGC = false;
 
 						int32 NumObjectsBeforeGC = GUObjectArray.GetObjectArrayNumMinusAvailable();
-						int32 NumObjectsAvailableBeforeGC = GUObjectArray.GetObjectArrayNum();
+						int32 NumObjectsAvailableBeforeGC = GUObjectArray.GetObjectArrayEstimatedAvailable();
 
 						UE_LOG(LogCookCommandlet, Display, TEXT("GarbageCollection... (%s)"), *GCReason);
 						GCReason = FString();
@@ -1003,7 +1003,7 @@ bool UCookCommandlet::CookByTheBook( const TArray<ITargetPlatform*>& Platforms,
 						CollectGarbage(RF_NoFlags);
 
 						int32 NumObjectsAfterGC = GUObjectArray.GetObjectArrayNumMinusAvailable();
-						int32 NumObjectsAvailableAfterGC = GUObjectArray.GetObjectArrayNum();
+						int32 NumObjectsAvailableAfterGC = GUObjectArray.GetObjectArrayEstimatedAvailable();
 						UE_LOG(LogCookCommandlet, Display, TEXT("Full GC before %d available %d after %d available %d"), NumObjectsBeforeGC, NumObjectsAvailableBeforeGC, NumObjectsAfterGC, NumObjectsAvailableAfterGC);
 
 						DumpMemStats();
