Change arcball hit priority to foreground to make sure axis get hit first always.

diff --git a/Engine/Source/Editor/UnrealEd/Private/UnrealWidget.cpp b/Engine/Source/Editor/UnrealEd/Private/UnrealWidget.cpp
--- a/Engine/Source/Editor/UnrealEd/Private/UnrealWidget.cpp
+++ b/Engine/Source/Editor/UnrealEd/Private/UnrealWidget.cpp
@@ -787,7 +787,8 @@ void FWidget::Render_Rotate( const FSceneView* View,FPrimitiveDrawInterface* PDI
 			FColor ArcColor = ArcBallColor;
 			ArcColor.A = (CurrentAxis==EAxisList::XYZ) ? ArcBallColor.A + 5 : ArcBallColor.A; //less transparent if selected
 
-			PDI->SetHitProxy(new HWidgetAxis(EAxisList::XYZ, bDisabled));
+			//set priority to foreground so axis get hit first
+			PDI->SetHitProxy(new HWidgetAxis(EAxisList::XYZ, bDisabled, HPP_Foreground));
 			{
 				DrawColoredSphere(PDI, InLocation, FRotator::ZeroRotator, ArcColor, Radii, 32, 24, TransparentPlaneMaterialXY->GetRenderProxy(), SDPG_Foreground, true);
 			}
diff --git a/Engine/Source/Editor/UnrealEd/Public/UnrealWidget.h b/Engine/Source/Editor/UnrealEd/Public/UnrealWidget.h
--- a/Engine/Source/Editor/UnrealEd/Public/UnrealWidget.h
+++ b/Engine/Source/Editor/UnrealEd/Public/UnrealWidget.h
@@ -457,8 +457,8 @@ struct HWidgetAxis : public HHitProxy
 	EAxisList::Type Axis;
 	uint32 bDisabled:1;
 
-	HWidgetAxis(EAxisList::Type InAxis, bool InbDisabled = false):
-		HHitProxy(HPP_UI),
+	HWidgetAxis(EAxisList::Type InAxis, bool InbDisabled = false, EHitProxyPriority InHitProxy = HPP_UI):
+		HHitProxy(InHitProxy),
 		Axis(InAxis),
 		bDisabled(InbDisabled) {}
 
