Fixed a crash when running the -AutomatedMapBuild command line with the editor due to the static light building not being correctly waited on.

The problem was that if a static lighting build was kicked off by the building of the map, the code would not wait for the build at any point. This meant it would be active during the shutdown of the editor and would attempt to access other systems (such as the SwarmInterface) after they were no longer valid.

I updated the build step to wait on the lighting build, similar to the code in ContentCommandlets that does the same thing. This will allow the system to shut itself down while things are still valid.

diff --git a/Engine/Source/Editor/UnrealEd/Private/EditorBuildUtils.cpp b/Engine/Source/Editor/UnrealEd/Private/EditorBuildUtils.cpp
--- a/Engine/Source/Editor/UnrealEd/Private/EditorBuildUtils.cpp
+++ b/Engine/Source/Editor/UnrealEd/Private/EditorBuildUtils.cpp
@@ -1017,7 +1017,11 @@ void FBuildAllHandler::ProcessBuild(const TWeakPtr<SBuildProgressWidget>& BuildP
 			LightingOptions.QualityLevel = (ELightingBuildQuality)QualityLevel;
 
 			GUnrealEd->BuildLighting(LightingOptions);
-
+			while (GUnrealEd->IsLightingBuildCurrentlyRunning())
+			{
+				GUnrealEd->UpdateBuildLighting();
+			}
+			
 			// TODO!
 			//bShouldMapCheck = false;
 
