Cinema4D - Fix reimport of Static Mesh random failure

diff --git a/Engine/Source/Editor/UnrealEd/Private/Factories/EditorFactories.cpp b/Engine/Source/Editor/UnrealEd/Private/Factories/EditorFactories.cpp
--- a/Engine/Source/Editor/UnrealEd/Private/Factories/EditorFactories.cpp
+++ b/Engine/Source/Editor/UnrealEd/Private/Factories/EditorFactories.cpp
@@ -5556,7 +5556,9 @@ bool UReimportFbxStaticMeshFactory::CanReimport( UObject* Obj, TArray<FString>&
 				return false;
 			}
 
-			if (FPaths::GetExtension(Mesh->AssetImportData->GetFirstFilename()) == TEXT("abc"))
+			FString FileExtension = FPaths::GetExtension(Mesh->AssetImportData->GetFirstFilename());
+			const bool bIsValidFile = FileExtension.Equals(TEXT("fbx"), ESearchCase::IgnoreCase) || FileExtension.Equals("obj", ESearchCase::IgnoreCase);
+			if (!bIsValidFile)
 			{
 				return false;
 			}
