Integrating CL# 7845323 from Dev-Enterprise
Fixing crash when exporting animation after deleting bone.

diff --git a/Engine/Source/Editor/UnrealEd/Private/Fbx/FbxAnimationExport.cpp b/Engine/Source/Editor/UnrealEd/Private/Fbx/FbxAnimationExport.cpp
--- a/Engine/Source/Editor/UnrealEd/Private/Fbx/FbxAnimationExport.cpp
+++ b/Engine/Source/Editor/UnrealEd/Private/Fbx/FbxAnimationExport.cpp
@@ -83,7 +83,6 @@ void FFbxExporter::ExportAnimSequenceToFbx(const UAnimSequence* AnimSeq,
 
 	//Prepare root anim curves data to be exported
 	TArray<FName> AnimCurveNames;
-	TArray<SmartName::UID_Type> AnimCurveUIDs;
 	TMap<FName, FbxAnimCurve*> CustomCurveMap;
 	if (BoneNodes.Num() > 0)
 	{
@@ -92,7 +91,6 @@ void FFbxExporter::ExportAnimSequenceToFbx(const UAnimSequence* AnimSeq,
 		if (AnimCurveMapping)
 		{
 			AnimCurveMapping->FillNameArray(AnimCurveNames);
-			AnimCurveMapping->FillUidArray(AnimCurveUIDs);
 
 			const UFbxExportOption* ExportOptions = GetExportOptions();
 			const bool bExportMorphTargetCurvesInMesh = ExportOptions && ExportOptions->bExportPreviewMesh && ExportOptions->bExportMorphTargets;
@@ -203,7 +201,16 @@ void FFbxExporter::ExportCustomAnimCurvesToFbx(const TMap<FName, FbxAnimCurve*>&
 	}
 
 	TArray<SmartName::UID_Type> AnimCurveUIDs;
-	SmartNameMapping->FillUidArray(AnimCurveUIDs);
+	{
+		//We need to recreate the UIDs array manually so that we keep the empty entries otherwise the BlendedCurve won't have the correct mapping.
+		TArray<FName> UID_ToNameArray;
+		SmartNameMapping->FillUIDToNameArray(UID_ToNameArray);
+		AnimCurveUIDs.Reserve(UID_ToNameArray.Num());
+		for (int32 NameIndex = 0; NameIndex < UID_ToNameArray.Num(); ++NameIndex)
+		{
+			AnimCurveUIDs.Add(NameIndex);
+		}
+	}
 
 	for (auto CustomCurve : CustomCurves)
 	{
