Integrating CL# 7936361 from Dev-Enterprise
Fixing material using virtual texture at FBX import.

diff --git a/Engine/Source/Editor/UnrealEd/Private/Fbx/FbxMaterialImport.cpp b/Engine/Source/Editor/UnrealEd/Private/Fbx/FbxMaterialImport.cpp
--- a/Engine/Source/Editor/UnrealEd/Private/Fbx/FbxMaterialImport.cpp
+++ b/Engine/Source/Editor/UnrealEd/Private/Fbx/FbxMaterialImport.cpp
@@ -284,13 +284,38 @@ bool UnFbx::FFbxImporter::CreateAndLinkExpressionForMaterialProperty(
 					{
 						float ScaleU = FbxTexture->GetScaleU();
 						float ScaleV = FbxTexture->GetScaleV();
+						bool bIsVirtualTexture = UnrealTexture->VirtualTextureStreaming;
+
+#if MATERIAL_OPACITYMASK_DOESNT_SUPPORT_VIRTUALTEXTURE
+						if (bIsVirtualTexture && MaterialProperty == FbxSurfaceMaterial::sTransparencyFactor)
+						{
+							//TODO, add a tracking of the materials created during the import so we can refresh here them if the TextureFactory actually found a existing asset instead of creating a new one.
+							if (UTexture2D* UnrealTexture2D = Cast<UTexture2D>(UnrealTexture))
+							{
+								TArray<UTexture2D*> TexturesToConvert = { UnrealTexture2D };
+								FAssetToolsModule& AssetToolsModule = FModuleManager::LoadModuleChecked<FAssetToolsModule>("AssetTools");
+								AssetToolsModule.Get().ConvertVirtualTextures(TexturesToConvert, true);
+							}
+							else
+							{
+								UnrealTexture->VirtualTextureStreaming = false;
+							}
+							bIsVirtualTexture = false;
+
+							FString TextureName = UnrealTexture->GetName();
+							UE_LOG(LogFbxMaterialImport, Warning, TEXT("Texture %s could not be imported as a Virtual Texture as they are not supported in OpacityMask property (material %s)."), *TextureName, UTF8_TO_TCHAR(FbxMaterial.GetName()));
+						}
+#endif
 
 						// and link it to the material 
 						UMaterialExpressionTextureSample* UnrealTextureExpression = NewObject<UMaterialExpressionTextureSample>(UnrealMaterial);
 						UnrealMaterial->Expressions.Add( UnrealTextureExpression );
 						MaterialInput.Expression = UnrealTextureExpression;
 						UnrealTextureExpression->Texture = UnrealTexture;
-						UnrealTextureExpression->SamplerType = bSetupAsNormalMap ? SAMPLERTYPE_Normal : SAMPLERTYPE_Color;
+						UnrealTextureExpression->SamplerType = bSetupAsNormalMap ? 
+							(bIsVirtualTexture ? SAMPLERTYPE_VirtualNormal : SAMPLERTYPE_Normal) :
+							(bIsVirtualTexture ? SAMPLERTYPE_VirtualColor : SAMPLERTYPE_Color);
+
 						UnrealTextureExpression->MaterialExpressionEditorX = FMath::TruncToInt(Location.X);
 						UnrealTextureExpression->MaterialExpressionEditorY = FMath::TruncToInt(Location.Y);
 
