Redo 5811807 in its original intended way to fix duplicated movable light GUID warnings

diff --git a/Engine/Source/Editor/UnrealEd/Private/EditorServer.cpp b/Engine/Source/Editor/UnrealEd/Private/EditorServer.cpp
--- a/Engine/Source/Editor/UnrealEd/Private/EditorServer.cpp
+++ b/Engine/Source/Editor/UnrealEd/Private/EditorServer.cpp
@@ -3841,7 +3841,7 @@ bool UEditorEngine::Map_Check( UWorld* InWorld, const TCHAR* Str, FOutputDevice&
 				ULightComponent* LightComponent = LightActor->GetLightComponent();
 				// LightComponent component can be null, for example when creating a blueprint deriving from ALight.
 				// Movable light components have a light guid of 0, so skip them as well
-				if (LightComponent && !LightComponent->HasStaticShadowing()) 
+				if (LightComponent && LightComponent->HasStaticShadowing()) 
 				{
 					AActor* ExistingLightActor = LightGuidToActorMap.FindRef( LightComponent->LightGuid );
 					if( ExistingLightActor )
