- Don't display texture import size warning for 8k textures if VT is enabled
- Don't allow textures > 16k to be imported, even with VT enabled (since this overflows 32bit sizes in various UE4 image code)
- Don't allow Crunch for textures > 8k, since this overflows 32bit size within crnlib

diff --git a/Engine/Source/Editor/UnrealEd/Private/Factories/EditorFactories.cpp b/Engine/Source/Editor/UnrealEd/Private/Factories/EditorFactories.cpp
--- a/Engine/Source/Editor/UnrealEd/Private/Factories/EditorFactories.cpp
+++ b/Engine/Source/Editor/UnrealEd/Private/Factories/EditorFactories.cpp
@@ -4459,9 +4459,14 @@ void UTextureFactory::ApplyAutoImportSettings(UTexture* Texture)
 
 bool UTextureFactory::IsImportResolutionValid(int32 Width, int32 Height, bool bAllowNonPowerOfTwo, FFeedbackContext* Warn)
 {
+	static const auto CVarVirtualTexturesEnabled = IConsoleManager::Get().FindTConsoleVariableDataInt(TEXT("r.VirtualTextures")); check(CVarVirtualTexturesEnabled);
+
+	// In theory this value could be much higher, but various UE4 image code currently uses 32bit size/offset values
+	const int32 MaximumSupportedVirtualTextureResolution = 16 * 1024;
+
 	// Calculate the maximum supported resolution utilizing the global max texture mip count
 	// (Note, have to subtract 1 because 1x1 is a valid mip-size; this means a GMaxTextureMipCount of 4 means a max resolution of 8x8, not 2^4 = 16x16)
-	const int32 MaximumSupportedResolution = 1 << (GMaxTextureMipCount - 1);
+	const int32 MaximumSupportedResolution = CVarVirtualTexturesEnabled->GetValueOnAnyThread() ? MaximumSupportedVirtualTextureResolution : (1 << (GMaxTextureMipCount - 1));
 
 	bool bValid = true;
 
@@ -4474,6 +4479,12 @@ bool UTextureFactory::IsImportResolutionValid(int32 Width, int32 Height, bool bA
 		{
 			bValid = false;
 		}
+
+		if (bValid && (Width * Height) > FMath::Square(MaximumSupportedVirtualTextureResolution))
+		{
+			Warn->Log(ELogVerbosity::Error, *NSLOCTEXT("UnrealEd", "Warning_TextureSizeTooLarge", "Texture is too large to import").ToString());
+			bValid = false;
+		}
 	}
 
 	const bool bIsPowerOfTwo = FMath::IsPowerOfTwo( Width ) && FMath::IsPowerOfTwo( Height );
diff --git a/Engine/Source/Runtime/Engine/Private/VT/VirtualTextureDataBuilder.cpp b/Engine/Source/Runtime/Engine/Private/VT/VirtualTextureDataBuilder.cpp
--- a/Engine/Source/Runtime/Engine/Private/VT/VirtualTextureDataBuilder.cpp
+++ b/Engine/Source/Runtime/Engine/Private/VT/VirtualTextureDataBuilder.cpp
@@ -1078,9 +1078,14 @@ void FVirtualTextureDataBuilder::BuildSourcePixels(const FTextureSourceData& Sou
 
 		bool bUseCrunch = false;
 #if WITH_CRUNCH_COMPRESSION
-		bUseCrunch = SettingsPerLayer[0].bVirtualTextureEnableCompressCrunch &&
-			SettingsPerLayer[0].LossyCompressionAmount != TLCA_None &&
-			CrunchCompression::IsValidFormat(TextureFormatName);
+		{
+			// Crunch compressor crashes if given size larger than 8k by 8k, needs to be updated to use 64bit sizes/offsets in various places
+			const bool bCrunchSizeValid = OutData.Width * OutData.Height <= FMath::Square(8u * 1024u);
+			bUseCrunch = bCrunchSizeValid &&
+				SettingsPerLayer[0].bVirtualTextureEnableCompressCrunch &&
+				SettingsPerLayer[0].LossyCompressionAmount != TLCA_None &&
+				CrunchCompression::IsValidFormat(TextureFormatName);
+		}
 #endif // WITH_CRUNCH_COMPRESSION
 		LayerData.TextureFormatName = TextureFormatName;
 		LayerData.bUseCrunch = bUseCrunch;
