PR #5890: UE-75199: Unhide brush on convert (Contributed by [REDACTED_PERSON])

diff --git a/Engine/Source/Editor/UnrealEd/Private/EditorEngine.cpp b/Engine/Source/Editor/UnrealEd/Private/EditorEngine.cpp
--- a/Engine/Source/Editor/UnrealEd/Private/EditorEngine.cpp
+++ b/Engine/Source/Editor/UnrealEd/Private/EditorEngine.cpp
@@ -3302,6 +3302,9 @@ void UEditorEngine::ConvertSelectedBrushesToVolumes( UClass* VolumeClass )
 				NewVolume->PostEditMove( true );
 				NewVolume->Modify();
 
+				// Make the actor visible as the brush is hidden by default
+				NewVolume->SetActorHiddenInGame(false);
+
 				// Destroy the old actor.
 				if (Layers.IsValid())
 				{
