Undo //UE4/Dev-Framework/Engine/Source/Editor/UnrealEd/Private/Factories/... changelist 7058716, accidentally checked in to the wrong branch. Proper fix is on release-4.23

diff --git a/Engine/Source/Editor/UnrealEd/Private/Factories/EditorFactories.cpp b/Engine/Source/Editor/UnrealEd/Private/Factories/EditorFactories.cpp
--- a/Engine/Source/Editor/UnrealEd/Private/Factories/EditorFactories.cpp
+++ b/Engine/Source/Editor/UnrealEd/Private/Factories/EditorFactories.cpp
@@ -4893,15 +4893,11 @@ void FCustomizableTextObjectFactory::ClearObjectNameUsage(UObject* InParent, FNa
 	}
 
 	// If there is already another object in the same scope with this name, rename it.
-	while ( Found )
+	if( Found )
 	{
 		check(Found->GetOuter() == InParent);
 
 		Found->Rename(nullptr, nullptr, REN_DontCreateRedirectors);
-
-		// It's possible after undo for there to be multiple objects with the same name in the way, rename all of them
-
-		Found = FindObject<UObject>(InParent, *InName.ToString());
 	}
 }
 
