Add BlueprintEditorUtils functions for determining if functions are a part of an interface or not and for getting the override of a function based on its name.

diff --git a/Engine/Source/Editor/UnrealEd/Private/Kismet2/BlueprintEditorUtils.cpp b/Engine/Source/Editor/UnrealEd/Private/Kismet2/BlueprintEditorUtils.cpp
--- a/Engine/Source/Editor/UnrealEd/Private/Kismet2/BlueprintEditorUtils.cpp
+++ b/Engine/Source/Editor/UnrealEd/Private/Kismet2/BlueprintEditorUtils.cpp
@@ -2430,6 +2430,41 @@ void FBlueprintEditorUtils::FindImplementedInterfaces(const UBlueprint* Blueprin
 	}
 }
 
+UClass* const FBlueprintEditorUtils::GetOverrideFunctionClass(UBlueprint* Blueprint, const FName FuncName, UFunction** OutFunction)
+{
+	if (!Blueprint->SkeletonGeneratedClass)
+	{
+		return nullptr;
+	}
+
+	UFunction* OverrideFunc = FBlueprintEditorUtils::GetInterfaceFunction(Blueprint, FuncName);
+
+	if (OverrideFunc == nullptr)
+	{
+		OverrideFunc = FindField<UFunction>(Blueprint->SkeletonGeneratedClass, FuncName);
+		// search up the class hierarchy, we want to find the original declaration of the function to match FBlueprintEventNodeSpawner.
+		// Doing so ensures that we can find the existing node if there is one:
+		const UClass* Iter = Blueprint->SkeletonGeneratedClass->GetSuperClass();
+		while (Iter != nullptr && OverrideFunc == nullptr)
+		{
+			if (UFunction * F = Iter->FindFunctionByName(FuncName))
+			{
+				OverrideFunc = F;
+			}
+			else
+			{
+				break;
+			}
+			Iter = Iter->GetSuperClass();
+		}
+	}
+	if (OutFunction != nullptr)
+	{
+		*OutFunction = OverrideFunc;
+	}
+	return CastChecked<UClass>(OverrideFunc->GetOuter())->GetAuthoritativeClass();
+}
+
 void FBlueprintEditorUtils::AddMacroGraph( UBlueprint* Blueprint, class UEdGraph* Graph, bool bIsUserCreated, UClass* SignatureFromClass )
 {
 	// Give the schema a chance to fill out any required nodes (like the entry node or results node)
@@ -5106,8 +5141,14 @@ UFunction* FFunctionFromNodeHelper::FunctionFromNode(const UK2Node* Node)
 		}
 		else if (const UK2Node_Event* EventNode = Cast<UK2Node_Event>(Node))
 		{
-			// We need to search up the class hierachy by name or functions like CanAddParentNode will fail:
-			Function = SearchScope->FindFunctionByName(EventNode->EventReference.GetMemberName());
+			// We need to search up the class hierarchy by name or functions like CanAddParentNode will fail:
+			FName SearchName = EventNode->EventReference.GetMemberName();
+			// If the function member name is none, then try the custom function name
+			if (SearchName.IsNone())
+			{
+				SearchName = EventNode->CustomFunctionName;
+			}
+			Function = SearchScope->FindFunctionByName(SearchName);
 		}
 	}
 
@@ -5903,6 +5944,54 @@ void FBlueprintEditorUtils::GetInterfaceGraphs(UBlueprint* Blueprint, const FNam
 	}
 }
 
+UFunction* FBlueprintEditorUtils::GetInterfaceFunction(UBlueprint* Blueprint, const FName FuncName)
+{
+	UFunction* Function = nullptr;
+
+	// If that class is an interface class implemented by this function, then return true
+	for (const FBPInterfaceDescription& I : Blueprint->ImplementedInterfaces)
+	{
+		if (I.Interface)
+		{
+			Function = FindField<UFunction>(I.Interface, FuncName);
+			if (Function)
+			{
+				// found it, done
+				return Function;
+			}
+		}
+	}
+
+	// Check if it is in a native class or parent class
+	for (UClass* TempClass = Blueprint->ParentClass; (nullptr != TempClass) && (nullptr == Function); TempClass = TempClass->GetSuperClass())
+	{
+		for (const FImplementedInterface& I : TempClass->Interfaces)
+		{
+			Function = FindField<UFunction>(I.Class, FuncName);
+			if (Function)
+			{
+				// found it, done
+				return Function;
+			}
+		}
+	}
+
+	return nullptr;
+}
+
+bool FBlueprintEditorUtils::IsInterfaceFunction(UBlueprint* Blueprint, UFunction* Function)
+{
+	if (Blueprint == nullptr || Function == nullptr)
+	{
+		return false;
+	}
+
+	const FName FuncName = Function->GetFName();
+
+	// Will return nullptr if the function isn't found
+	return (GetInterfaceFunction(Blueprint, FuncName) != nullptr);
+}
+
 // Remove an implemented interface, and its associated member function graphs
 void FBlueprintEditorUtils::RemoveInterface(UBlueprint* Blueprint, const FName& InterfaceClassName, bool bPreserveFunctions /*= false*/)
 {
diff --git a/Engine/Source/Editor/UnrealEd/Public/Kismet2/BlueprintEditorUtils.h b/Engine/Source/Editor/UnrealEd/Public/Kismet2/BlueprintEditorUtils.h
--- a/Engine/Source/Editor/UnrealEd/Public/Kismet2/BlueprintEditorUtils.h
+++ b/Engine/Source/Editor/UnrealEd/Public/Kismet2/BlueprintEditorUtils.h
@@ -381,6 +381,17 @@ public:
 		FBlueprintEditorUtils::MarkBlueprintAsStructurallyModified(Blueprint);
 	}
 
+	/**
+	* Get the override class of a given function from its name
+	* 
+	* @param Blueprint		Blueprint to check the function on
+	* @param FuncName		Name of the function
+	* @param OutFunction	The function that has this name
+	* 
+	* @return The override class of a given function
+	*/
+	static UClass* const GetOverrideFunctionClass(UBlueprint* Blueprint, const FName FuncName, UFunction** OutFunction = nullptr);
+
 	/** Adds a macro graph to this blueprint.  If bIsUserCreated is true, the entry/exit nodes will be editable. SignatureFromClass is used to find signature for entry/exit nodes if using an existing signature. */
 	static void AddMacroGraph(UBlueprint* Blueprint, class UEdGraph* Graph,  bool bIsUserCreated, UClass* SignatureFromClass);
 
@@ -1269,6 +1280,28 @@ public:
 	/** Gets the graphs currently in the blueprint associated with the specified interface */
 	static void GetInterfaceGraphs(UBlueprint* Blueprint, const FName& InterfaceClassName, TArray<UEdGraph*>& ChildGraphs);
 
+	/**
+	* Checks if the given function is a part of an interface on this blueprint
+	* 
+	* @param Blueprint		The blueprint to consider
+	* @param Function		Function to check if it is an interface or not
+	* @return	True if the given function is implemented as part of an interface
+	*/
+	static bool IsInterfaceFunction(UBlueprint* Blueprint, UFunction* Function);
+
+	/**
+	* Get the corresponding UFunction pointer to the name given on the blueprint.
+	* Searches the given blueprints implemented interfaces first, and then looks 
+	* in the parent. 
+	* 
+	* @param Blueprint		The blueprint to consider
+	* @param FuncName		The name of the function to look for
+	*
+	* @return	Corresponding UFunction pointer to the name given; Nullptr if not 
+	*			part of any interfaces
+	*/
+	static UFunction* GetInterfaceFunction(UBlueprint* Blueprint, const FName FuncName);
+
 	/** Makes sure that all graphs for all interfaces we implement exist, and add if not */
 	static void ConformImplementedInterfaces(UBlueprint* Blueprint);
 
