Fix static analysis warning in FBlueprintEditorUtils::GetOverrideFunctionClass by doing a null check

diff --git a/Engine/Source/Editor/UnrealEd/Private/Kismet2/BlueprintEditorUtils.cpp b/Engine/Source/Editor/UnrealEd/Private/Kismet2/BlueprintEditorUtils.cpp
--- a/Engine/Source/Editor/UnrealEd/Private/Kismet2/BlueprintEditorUtils.cpp
+++ b/Engine/Source/Editor/UnrealEd/Private/Kismet2/BlueprintEditorUtils.cpp
@@ -2469,7 +2469,8 @@ UClass* const FBlueprintEditorUtils::GetOverrideFunctionClass(UBlueprint* Bluepr
 	{
 		*OutFunction = OverrideFunc;
 	}
-	return CastChecked<UClass>(OverrideFunc->GetOuter())->GetAuthoritativeClass();
+
+	return (OverrideFunc ? CastChecked<UClass>(OverrideFunc->GetOuter())->GetAuthoritativeClass() : nullptr);
 }
 
 void FBlueprintEditorUtils::AddMacroGraph( UBlueprint* Blueprint, class UEdGraph* Graph, bool bIsUserCreated, UClass* SignatureFromClass )
