Make FTransaction::Records an indirect array to avoid resizing issues when adding actors which have component instance data caches with components that have editinline subobjects.

diff --git a/Engine/Source/Editor/UnrealEd/Classes/Editor/Transactor.h b/Engine/Source/Editor/UnrealEd/Classes/Editor/Transactor.h
--- a/Engine/Source/Editor/UnrealEd/Classes/Editor/Transactor.h
+++ b/Engine/Source/Editor/UnrealEd/Classes/Editor/Transactor.h
@@ -489,7 +489,7 @@ protected:
 
 	// Transaction variables.
 	/** List of object records in this transaction */
-	TArray<FObjectRecord>	Records;
+	TIndirectArray<FObjectRecord>	Records;
 	
 	/** Unique identifier for this transaction, used to track it during its lifetime */
 	FGuid		Id;
diff --git a/Engine/Source/Editor/UnrealEd/Private/EditorTransaction.cpp b/Engine/Source/Editor/UnrealEd/Private/EditorTransaction.cpp
--- a/Engine/Source/Editor/UnrealEd/Private/EditorTransaction.cpp
+++ b/Engine/Source/Editor/UnrealEd/Private/EditorTransaction.cpp
@@ -5,6 +5,7 @@
 #include "Misc/MemStack.h"
 #include "UObject/Object.h"
 #include "UObject/Package.h"
+#include "Algo/Find.h"
 #include "Engine/Level.h"
 #include "Components/ActorComponent.h"
 #include "Model.h"
@@ -745,7 +746,7 @@ void FTransaction::SaveObject( UObject* Object )
 	{
 		ObjectMap.Add(Object,1);
 		// Save the object.
-		new( Records )FObjectRecord( this, Object, nullptr, NULL, 0, 0, 0, 0, NULL, NULL, NULL );
+		Records.Add(new FObjectRecord( this, Object, nullptr, nullptr, 0, 0, 0, 0, nullptr, nullptr, nullptr));
 	}
 	else
 	{
@@ -771,7 +772,7 @@ void FTransaction::SaveArray( UObject* Object, FScriptArray* Array, int32 Index,
 	if( Object->HasAnyFlags(RF_Transactional) && !Object->GetOutermost()->HasAnyPackageFlags(PKG_PlayInEditor))
 	{
 		// Save the array.
-		new( Records )FObjectRecord( this, Object, nullptr, Array, Index, Count, Oper, ElementSize, DefaultConstructor, Serializer, Destructor );
+		Records.Add(new FObjectRecord( this, Object, nullptr, Array, Index, Count, Oper, ElementSize, DefaultConstructor, Serializer, Destructor ));
 	}
 }
 
@@ -787,7 +788,7 @@ void FTransaction::StoreUndo( UObject* Object, TUniquePtr<FChange> UndoChange )
 	}
 
 	// Save the undo record
-	new( Records )FObjectRecord( this, Object, MoveTemp( UndoChange ), NULL, 0, 0, 0, 0, NULL, NULL, NULL );
+	Records.Add(new FObjectRecord( this, Object, MoveTemp( UndoChange ), nullptr, 0, 0, 0, 0, nullptr, nullptr, nullptr));
 }
 
 void FTransaction::SetPrimaryObject(UObject* InObject)
@@ -802,7 +803,7 @@ void FTransaction::SnapshotObject( UObject* InObject )
 {
 	if (InObject && ObjectMap.Contains(InObject))
 	{
-		FObjectRecord* FoundObjectRecord = Records.FindByPredicate([InObject](const FObjectRecord& ObjRecord)
+		FObjectRecord* FoundObjectRecord = Algo::FindByPredicate(Records, [InObject](const FObjectRecord& ObjRecord)
 		{
 			return ObjRecord.Object.Get() == InObject;
 		});
