Refresh nodes when a BP interface is removed from a class to fix BP compiler error.

+ native interfaces work with this action now
+ input and output pins are no longer duplicated when this action is taken

diff --git a/Engine/Source/Editor/UnrealEd/Private/Kismet2/BlueprintEditorUtils.cpp b/Engine/Source/Editor/UnrealEd/Private/Kismet2/BlueprintEditorUtils.cpp
--- a/Engine/Source/Editor/UnrealEd/Private/Kismet2/BlueprintEditorUtils.cpp
+++ b/Engine/Source/Editor/UnrealEd/Private/Kismet2/BlueprintEditorUtils.cpp
@@ -6022,62 +6022,121 @@ void FBlueprintEditorUtils::RemoveInterface(UBlueprint* Blueprint, const FName&
 	if( Idx != INDEX_NONE )
 	{
 		FBPInterfaceDescription& CurrentInterface = Blueprint->ImplementedInterfaces[Idx];
+		const UClass* InterfaceClass = Blueprint->ImplementedInterfaces[Idx].Interface;
 
-		// Remove all the graphs that we implemented
-		for(TArray<UEdGraph*>::TIterator it(CurrentInterface.Graphs); it; ++it)
+		// For every function and event in the interface...
+		for (TFieldIterator<UFunction> FunctionIt(InterfaceClass); FunctionIt; ++FunctionIt)
 		{
-			UEdGraph* CurrentGraph = *it;
-			if(bPreserveFunctions)
-			{
-				PromoteGraphFromInterfaceOverride(Blueprint, CurrentGraph);
-				Blueprint->FunctionGraphs.Add(CurrentGraph);
-			}
-			else
+			UFunction* Function = *FunctionIt;
+			const FName FunctionName = Function->GetFName();
+			
+			// If this function name is in the list of interface graphs, then handle it like a function
+			if (FunctionName == UEdGraphSchema_K2::FN_ExecuteUbergraphBase || 
+				FBlueprintEditorUtils::RemoveInterfaceFunction(Blueprint, CurrentInterface, Function, bPreserveFunctions))
 			{
-				FBlueprintEditorUtils::RemoveGraph(Blueprint, CurrentGraph, EGraphRemoveFlags::MarkTransient);	// Do not recompile, yet*
+				continue;
 			}
-		}
 
-		// Find all events placed in the event graph, and remove them
-		TArray<UK2Node_Event*> AllEvents;
-		FBlueprintEditorUtils::GetAllNodesOfClass(Blueprint, AllEvents);
-		const UClass* InterfaceClass = Blueprint->ImplementedInterfaces[Idx].Interface;
-		for(TArray<UK2Node_Event*>::TIterator NodeIt(AllEvents); NodeIt; ++NodeIt)
-		{
-			UK2Node_Event* EventNode = *NodeIt;
-			if( EventNode->EventReference.GetMemberParentClass(EventNode->GetBlueprintClassFromNode()) == InterfaceClass )
+			// Find all events placed in the event graph so we can check if they belong to this interface
+			TArray<UK2Node_Event*> AllEvents;
+			FBlueprintEditorUtils::GetAllNodesOfClass(Blueprint, AllEvents);
+			for (TArray<UK2Node_Event*>::TIterator NodeIt(AllEvents); NodeIt; ++NodeIt)
 			{
-				if(bPreserveFunctions)
+				UK2Node_Event* EventNode = *NodeIt;
+				if (EventNode->EventReference.GetMemberParentClass(EventNode->GetBlueprintClassFromNode()) == InterfaceClass)
 				{
-					// Create a custom event with the same name and signature
-					const FVector2D PreviousNodePos = FVector2D(EventNode->NodePosX, EventNode->NodePosY);
-					const FString PreviousNodeName = EventNode->EventReference.GetMemberName().ToString();
-					const UFunction* PreviousSignatureFunction = EventNode->FindEventSignatureFunction();
-					check(PreviousSignatureFunction);
-					
-					UK2Node_CustomEvent* NewEvent = UK2Node_CustomEvent::CreateFromFunction(PreviousNodePos, EventNode->GetGraph(), PreviousNodeName, PreviousSignatureFunction, false);
-
-					// Move the pin links from the old pin to the new pin to preserve connections
-					for (UEdGraphPin* CurrentPin : EventNode->Pins)
+					if (bPreserveFunctions)
 					{
-						UEdGraphPin* TargetPin = NewEvent->FindPinChecked(CurrentPin->PinName);
-						const UEdGraphSchema_K2* Schema = GetDefault<UEdGraphSchema_K2>();
-						Schema->MovePinLinks(*CurrentPin, *TargetPin);
+						// Create a custom event with the same name and signature
+						const FVector2D PreviousNodePos = FVector2D(EventNode->NodePosX, EventNode->NodePosY);
+						const FString PreviousNodeName = EventNode->EventReference.GetMemberName().ToString();
+						const UFunction* PreviousSignatureFunction = EventNode->FindEventSignatureFunction();
+						check(PreviousSignatureFunction);
+						UK2Node_CustomEvent* NewEvent = UK2Node_CustomEvent::CreateFromFunction(PreviousNodePos, EventNode->GetGraph(), PreviousNodeName, PreviousSignatureFunction, false);
+						// Move the pin links from the old pin to the new pin to preserve connections
+						for (UEdGraphPin* CurrentPin : EventNode->Pins)
+						{
+							UEdGraphPin* TargetPin = NewEvent->FindPinChecked(CurrentPin->PinName);
+							const UEdGraphSchema_K2* Schema = GetDefault<UEdGraphSchema_K2>();
+							Schema->MovePinLinks(*CurrentPin, *TargetPin);
+						}
 					}
+
+					EventNode->GetGraph()->RemoveNode(EventNode);
+					break;
 				}
-			
-				EventNode->GetGraph()->RemoveNode(EventNode);
 			}
 		}
 
 		// Then remove the interface from the list
 		Blueprint->ImplementedInterfaces.RemoveAt(Idx, 1);
-	
-		// *Now recompile the blueprint (this needs to be done outside of RemoveGraph, after it's been removed from ImplementedInterfaces - otherwise it'll re-add it)
+		
+		// Refresh all the nodes to make sure that the references to "Self" are updated appropriately on any function calls  @see UE-78253
+		FBlueprintEditorUtils::RefreshAllNodes(Blueprint);
+
+		// Now recompile the blueprint (this needs to be done outside of RemoveGraph, after it's been removed from ImplementedInterfaces - otherwise it'll re-add it)
 		FBlueprintEditorUtils::MarkBlueprintAsStructurallyModified(Blueprint);
 	}
 }
 
+bool FBlueprintEditorUtils::RemoveInterfaceFunction(class UBlueprint* Blueprint, FBPInterfaceDescription& Interface, UFunction* Function, bool bPreserveFunction)
+{
+	for (TArray<UEdGraph*>::TIterator it(Interface.Graphs); it; ++it)
+	{
+		UEdGraph* CurrentGraph = *it;
+		if (Function->GetFName() == CurrentGraph->GetFName())
+		{
+			CurrentGraph->Modify();
+			Blueprint->Modify();
+			UEdGraph* NewGraph = nullptr;
+			FString OriginalName;
+			CurrentGraph->GetName(OriginalName);
+
+			// Create a new one with it's signature if the user wants to preserve it
+			if (bPreserveFunction)
+			{
+				// Create a new function graph with this signature
+				static const FString TempName = "TEMP_INTERFACE_GRAPH_NAME";
+				const FName TempFuncGraphName = FBlueprintEditorUtils::GenerateUniqueGraphName(Blueprint, TempName);
+				UEdGraph* EventGraph = FBlueprintEditorUtils::FindEventGraph(Blueprint);
+				const UEdGraphSchema_K2* Schema = GetDefault<UEdGraphSchema_K2>();
+				EventGraph->Modify();
+
+				if (EventGraph)
+				{
+					NewGraph = FBlueprintEditorUtils::CreateNewGraph(
+						Blueprint,
+						TempFuncGraphName,
+						EventGraph->GetClass(),
+						EventGraph->GetSchema() ? EventGraph->GetSchema()->GetClass() : Schema->GetClass()
+					);
+
+					// Add the function graph like this to avoid the MarkBlueprintAsStructurallyModified call
+					CreateFunctionGraph(Blueprint, NewGraph, /*bIsUserCreated=*/true, Function);
+					Blueprint->FunctionGraphs.Add(NewGraph);
+					// Potentially adjust variable names for any child blueprints
+					ValidateBlueprintChildVariables(Blueprint, NewGraph->GetFName());
+				}
+			}
+
+			FBlueprintEditorUtils::UpdateTransactionalFlags(Blueprint);
+
+			// Remove the interface graph
+			FBlueprintEditorUtils::RemoveGraph(Blueprint, CurrentGraph, EGraphRemoveFlags::MarkTransient);	// Do not recompile, yet
+			
+			// Rename the new graph to the original interface name
+			if (NewGraph)
+			{
+				FBlueprintEditorUtils::RenameGraph(NewGraph, OriginalName);
+			}
+
+			return true;
+		}
+	}
+
+	return false;
+}
+
 void FBlueprintEditorUtils::PromoteGraphFromInterfaceOverride(UBlueprint* InBlueprint, UEdGraph* InInterfaceGraph)
 {
 	InInterfaceGraph->bAllowDeletion = true;
diff --git a/Engine/Source/Editor/UnrealEd/Public/Kismet2/BlueprintEditorUtils.h b/Engine/Source/Editor/UnrealEd/Public/Kismet2/BlueprintEditorUtils.h
--- a/Engine/Source/Editor/UnrealEd/Public/Kismet2/BlueprintEditorUtils.h
+++ b/Engine/Source/Editor/UnrealEd/Public/Kismet2/BlueprintEditorUtils.h
@@ -1268,6 +1268,14 @@ public:
 
 	/** Remove an implemented interface, and its associated member function graphs.  If bPreserveFunctions is true, then the interface will move its functions to be normal implemented blueprint functions */
 	static void RemoveInterface(UBlueprint* Blueprint, const FName& InterfaceClassName, bool bPreserveFunctions = false);
+	
+	/**
+	* Attempt to remove a function from an interfaces list of function graphs.
+	* Note that this will NOT remove interface events (i.e. functions with no outputs)
+	* 
+	* @return	True if the function was removed from the blueprint
+	*/
+	static bool RemoveInterfaceFunction(class UBlueprint* Blueprint, struct FBPInterfaceDescription& Interface, class UFunction* Function, bool bPreserveFunction);
 
 	/**
 	* Promotes a Graph from being an Interface Override to a full member function
