Look not just in the archetype, but also the owner for a property that references the component for the scs node
#rnx

diff --git a/Engine/Source/Editor/Kismet/Private/SKismetInspector.cpp b/Engine/Source/Editor/Kismet/Private/SKismetInspector.cpp
--- a/Engine/Source/Editor/Kismet/Private/SKismetInspector.cpp
+++ b/Engine/Source/Editor/Kismet/Private/SKismetInspector.cpp
@@ -710,36 +710,51 @@ void SKismetInspector::UpdateFromObjects(const TArray<UObject*>& PropertyObjects
 
 				if(Object != EditableComponentTemplate)
 				{
-					UObjectProperty* ObjectProperty = FindField<UObjectProperty>(Object->GetClass(), EditableComponentTemplate->GetFName());
-					if(ObjectProperty != nullptr)
+					if (UObjectProperty* ObjectProperty = FindField<UObjectProperty>(Object->GetClass(), EditableComponentTemplate->GetFName()))
 					{
 						SelectedObjectProperties.Add(ObjectProperty);
 					}
-					else if(UActorComponent* Archetype = Cast<UActorComponent>(EditableComponentTemplate->GetArchetype()))
+					else 
 					{
-						if(AActor* Owner = Archetype->GetOwner())
+						auto FindPropertyReferencingComponent = [Object](const UActorComponent* Component) -> UProperty*
 						{
-							if(UClass* OwnerClass = Owner->GetClass())
+							if (AActor* Owner = Component->GetOwner())
 							{
-								AActor* OwnerCDO = CastChecked<AActor>(OwnerClass->GetDefaultObject());
-								for(TFieldIterator<UObjectProperty> ObjPropIt(OwnerClass, EFieldIteratorFlags::IncludeSuper); ObjPropIt; ++ObjPropIt)
+								if (UClass* OwnerClass = Owner->GetClass())
 								{
-									ObjectProperty = *ObjPropIt;
-									check(ObjectProperty != nullptr);
-
-									// If the property value matches the current archetype, add it as a selected property for filtering
-									if(Archetype->GetClass()->IsChildOf(ObjectProperty->PropertyClass)
-										&& Archetype == ObjectProperty->GetObjectPropertyValue_InContainer(OwnerCDO))
+									AActor* OwnerCDO = CastChecked<AActor>(OwnerClass->GetDefaultObject());
+									for (TFieldIterator<UObjectProperty> ObjPropIt(OwnerClass, EFieldIteratorFlags::IncludeSuper); ObjPropIt; ++ObjPropIt)
 									{
-										ObjectProperty = FindField<UObjectProperty>(Object->GetClass(), ObjectProperty->GetFName());
-										if(ObjectProperty != nullptr)
+										UObjectProperty* ObjectProperty = *ObjPropIt;
+										check(ObjectProperty != nullptr);
+
+										// If the property value matches the current archetype, add it as a selected property for filtering
+										if (Component->GetClass()->IsChildOf(ObjectProperty->PropertyClass)
+											&& Component == ObjectProperty->GetObjectPropertyValue_InContainer(OwnerCDO))
 										{
-											SelectedObjectProperties.Add(ObjectProperty);
-											break;
+											ObjectProperty = FindField<UObjectProperty>(Object->GetClass(), ObjectProperty->GetFName());
+											if (ObjectProperty != nullptr)
+											{
+												return ObjectProperty;
+											}
 										}
 									}
 								}
 							}
+							return nullptr;
+						};
+
+						UProperty* ReferencingProperty = FindPropertyReferencingComponent(EditableComponentTemplate);
+						if (ReferencingProperty == nullptr)
+						{
+							if (UActorComponent* Archetype = Cast<UActorComponent>(EditableComponentTemplate->GetArchetype()))
+							{
+								ReferencingProperty = FindPropertyReferencingComponent(Archetype);
+							}
+						}
+						if (ReferencingProperty)
+						{
+							SelectedObjectProperties.Add(ReferencingProperty);
 						}
 					}
 				}
diff --git a/Engine/Source/Editor/UnrealEd/Private/Kismet2/ComponentEditorUtils.cpp b/Engine/Source/Editor/UnrealEd/Private/Kismet2/ComponentEditorUtils.cpp
--- a/Engine/Source/Editor/UnrealEd/Private/Kismet2/ComponentEditorUtils.cpp
+++ b/Engine/Source/Editor/UnrealEd/Private/Kismet2/ComponentEditorUtils.cpp
@@ -852,23 +852,10 @@ FName FComponentEditorUtils::FindVariableNameGivenComponentInstance(const UActor
 {
 	check(ComponentInstance != nullptr);
 
-	// First see if the name just works
-	if (AActor* OwnerActor = ComponentInstance->GetOwner())
+	// When names mismatch, try finding a differently named variable pointing to the the component (the mismatch should only be possible for native components)
+	auto FindPropertyReferencingComponent = [](const UActorComponent* Component) -> UProperty*
 	{
-		UClass* OwnerActorClass = OwnerActor->GetClass();
-		if (UObjectProperty* TestProperty = FindField<UObjectProperty>(OwnerActorClass, ComponentInstance->GetFName()))
-		{
-			if (ComponentInstance->GetClass()->IsChildOf(TestProperty->PropertyClass))
-			{
-				return TestProperty->GetFName();
-			}
-		}
-	}
-
-	// Name mismatch, try finding a differently named variable pointing to the the component (the mismatch should only be possible for native components)
-	if (UActorComponent* Archetype = Cast<UActorComponent>(ComponentInstance->GetArchetype()))
-	{
-		if (AActor* OwnerActor = Archetype->GetOwner())
+		if (AActor* OwnerActor = Component->GetOwner())
 		{
 			UClass* OwnerClass = OwnerActor->GetClass();
 			AActor* OwnerCDO = CastChecked<AActor>(OwnerClass->GetDefaultObject());
@@ -877,14 +864,14 @@ FName FComponentEditorUtils::FindVariableNameGivenComponentInstance(const UActor
 			for (TFieldIterator<UObjectProperty> PropIt(OwnerClass, EFieldIteratorFlags::IncludeSuper); PropIt; ++PropIt)
 			{
 				UObjectProperty* TestProperty = *PropIt;
-				if (Archetype->GetClass()->IsChildOf(TestProperty->PropertyClass))
+				if (Component->GetClass()->IsChildOf(TestProperty->PropertyClass))
 				{
 					void* TestPropertyInstanceAddress = TestProperty->ContainerPtrToValuePtr<void>(OwnerCDO);
 					UObject* ObjectPointedToByProperty = TestProperty->GetObjectPropertyValue(TestPropertyInstanceAddress);
-					if (ObjectPointedToByProperty == Archetype)
+					if (ObjectPointedToByProperty == Component)
 					{
 						// This property points to the component archetype, so it's an anchor even if it was named wrong
-						return TestProperty->GetFName();
+						return TestProperty;
 					}
 				}
 			}
@@ -904,13 +891,41 @@ FName FComponentEditorUtils::FindVariableNameGivenComponentInstance(const UActor
 				for (int32 ComponentIndex = 0; ComponentIndex < ArrayHelper.Num(); ++ComponentIndex)
 				{
 					UObject* ArrayElement = ArrayEntryProp->GetObjectPropertyValue(ArrayHelper.GetRawPtr(ComponentIndex));
-					if (ArrayElement == Archetype)
+					if (ArrayElement == Component)
 					{
-						return TestProperty->GetFName();
+						return TestProperty;
 					}
 				}
 			}
 		}
+
+		return nullptr;
+	};
+
+	// First see if the name just works
+	if (AActor* OwnerActor = ComponentInstance->GetOwner())
+	{
+		UClass* OwnerActorClass = OwnerActor->GetClass();
+		if (UObjectProperty* TestProperty = FindField<UObjectProperty>(OwnerActorClass, ComponentInstance->GetFName()))
+		{
+			if (ComponentInstance->GetClass()->IsChildOf(TestProperty->PropertyClass))
+			{
+				return TestProperty->GetFName();
+			}
+		}
+
+		if (UProperty* ReferencingProp = FindPropertyReferencingComponent(ComponentInstance))
+		{
+			return ReferencingProp->GetFName();
+		}
+	}
+
+	if (UActorComponent* Archetype = Cast<UActorComponent>(ComponentInstance->GetArchetype()))
+	{
+		if (UProperty* ReferencingProp = FindPropertyReferencingComponent(Archetype))
+		{
+			return ReferencingProp->GetFName();
+		}
 	}
 
 	return NAME_None;
