Cleanup forward declarations in EditorUtils

#rnx

diff --git a/Engine/Source/Editor/UnrealEd/Public/Kismet2/BlueprintEditorUtils.h b/Engine/Source/Editor/UnrealEd/Public/Kismet2/BlueprintEditorUtils.h
--- a/Engine/Source/Editor/UnrealEd/Public/Kismet2/BlueprintEditorUtils.h
+++ b/Engine/Source/Editor/UnrealEd/Public/Kismet2/BlueprintEditorUtils.h
@@ -30,6 +30,9 @@ class UTimelineTemplate;
 struct FBlueprintCookedComponentInstancingData;
 struct FComponentKey;
 class UAnimGraphNode_Root;
+class UBlueprint;
+struct FBPInterfaceDescription;
+class UFunction;
 
 /** 
   * Flags describing how to handle graph removal
@@ -1275,7 +1278,7 @@ public:
 	* 
 	* @return	True if the function was removed from the blueprint
 	*/
-	static bool RemoveInterfaceFunction(class UBlueprint* Blueprint, struct FBPInterfaceDescription& Interface, class UFunction* Function, bool bPreserveFunction);
+	static bool RemoveInterfaceFunction(UBlueprint* Blueprint, FBPInterfaceDescription& Interface, UFunction* Function, bool bPreserveFunction);
 
 	/**
 	* Promotes a Graph from being an Interface Override to a full member function
