Fix undo functionality for removing blueprint interfaces

#rnx

diff --git a/Engine/Source/Editor/UnrealEd/Private/Kismet2/BlueprintEditorUtils.cpp b/Engine/Source/Editor/UnrealEd/Private/Kismet2/BlueprintEditorUtils.cpp
--- a/Engine/Source/Editor/UnrealEd/Private/Kismet2/BlueprintEditorUtils.cpp
+++ b/Engine/Source/Editor/UnrealEd/Private/Kismet2/BlueprintEditorUtils.cpp
@@ -6023,6 +6023,8 @@ void FBlueprintEditorUtils::RemoveInterface(UBlueprint* Blueprint, const FName&
 	{
 		FBPInterfaceDescription& CurrentInterface = Blueprint->ImplementedInterfaces[Idx];
 		const UClass* InterfaceClass = Blueprint->ImplementedInterfaces[Idx].Interface;
+		const FScopedTransaction Transaction(LOCTEXT("RemoveInterface", "Remove interface"));
+		Blueprint->Modify();
 
 		// For every function and event in the interface...
 		for (TFieldIterator<UFunction> FunctionIt(InterfaceClass); FunctionIt; ++FunctionIt)
@@ -6045,6 +6047,10 @@ void FBlueprintEditorUtils::RemoveInterface(UBlueprint* Blueprint, const FName&
 				UK2Node_Event* EventNode = *NodeIt;
 				if (EventNode->EventReference.GetMemberParentClass(EventNode->GetBlueprintClassFromNode()) == InterfaceClass)
 				{
+					UEdGraph* EventNodeGraph = EventNode->GetGraph();
+					check(EventNodeGraph);
+					EventNodeGraph->Modify();
+
 					if (bPreserveFunctions)
 					{
 						// Create a custom event with the same name and signature
@@ -6062,7 +6068,7 @@ void FBlueprintEditorUtils::RemoveInterface(UBlueprint* Blueprint, const FName&
 						}
 					}
 
-					EventNode->GetGraph()->RemoveNode(EventNode);
+					EventNodeGraph->RemoveNode(EventNode);
 					break;
 				}
 			}
@@ -6079,7 +6085,7 @@ void FBlueprintEditorUtils::RemoveInterface(UBlueprint* Blueprint, const FName&
 	}
 }
 
-bool FBlueprintEditorUtils::RemoveInterfaceFunction(class UBlueprint* Blueprint, FBPInterfaceDescription& Interface, UFunction* Function, bool bPreserveFunction)
+bool FBlueprintEditorUtils::RemoveInterfaceFunction(UBlueprint* Blueprint, FBPInterfaceDescription& Interface, UFunction* Function, bool bPreserveFunction)
 {
 	for (TArray<UEdGraph*>::TIterator it(Interface.Graphs); it; ++it)
 	{
