Only disable Stereo Rendering after leaving game sessions when not in VR Mode
#rnx

diff --git a/Engine/Source/Editor/UnrealEd/Private/PlayLevel.cpp b/Engine/Source/Editor/UnrealEd/Private/PlayLevel.cpp
--- a/Engine/Source/Editor/UnrealEd/Private/PlayLevel.cpp
+++ b/Engine/Source/Editor/UnrealEd/Private/PlayLevel.cpp
@@ -529,8 +529,9 @@ void UEditorEngine::EndPlayMap()
 		GEngine->PendingDroppedNotes.Empty();
 	}
 
-	//ensure stereo rendering is disabled in case we need to re-enable next PIE run.
-	if (GEngine->StereoRenderingDevice)
+	//ensure stereo rendering is disabled in case we need to re-enable next PIE run (except when the editor is running in VR)
+	bool bInVRMode = IVREditorModule::Get().IsVREditorModeActive();
+	if (GEngine->StereoRenderingDevice && !bInVRMode)
 	{
 		GEngine->StereoRenderingDevice->EnableStereo(false);
 	}
