Don't apply auto VT size threshold on texture import, if VT is disabled for the project

diff --git a/Engine/Source/Editor/UnrealEd/Private/Factories/EditorFactories.cpp b/Engine/Source/Editor/UnrealEd/Private/Factories/EditorFactories.cpp
--- a/Engine/Source/Editor/UnrealEd/Private/Factories/EditorFactories.cpp
+++ b/Engine/Source/Editor/UnrealEd/Private/Factories/EditorFactories.cpp
@@ -4283,20 +4283,24 @@ UObject* UTextureFactory::FactoryCreateBinary
 		// If the texture is larger than a certain threshold make it VT. This is explicitly done after the
 		// application of the existing settings above, so if a texture gets reimported at a larger size it will
 		// still be properly flagged as a VT (note: What about reimporting at a lower resolution?)
-
-		int virtualTextureAutoEnableThreshold = GetDefault<UTextureImportSettings>()->AutoVTSize;
-		int virtualTextureAutoEnableThresholdPixels = virtualTextureAutoEnableThreshold*virtualTextureAutoEnableThreshold;
+		static const auto CVarVirtualTexturesEnabled = IConsoleManager::Get().FindTConsoleVariableDataInt(TEXT("r.VirtualTextures")); check(CVarVirtualTexturesEnabled);
 		
-		// We do this in pixels so a 8192 x 128 texture won't get VT enabled 
-		// We use the Source size instead of simple Texture2D->GetSizeX() as this uses the size of the platform data
-		// however for a new texture platform data may not be generated yet, and for an reimport of a texture this is the size of the
-		// old texture. 
-		// Using source size gives one small caveat. It looks at the size before mipmap power of two padding adjustment.
-		// Textures with more than 1 block (UDIM textures) must be imported as VT
-		if (Texture->Source.GetNumBlocks() > 1 ||
-			Texture2D->Source.GetSizeX()*Texture2D->Source.GetSizeY() >= virtualTextureAutoEnableThresholdPixels)
-		{
-			Texture2D->VirtualTextureStreaming = true;
+		if (CVarVirtualTexturesEnabled->GetValueOnAnyThread())
+		{
+			int virtualTextureAutoEnableThreshold = GetDefault<UTextureImportSettings>()->AutoVTSize;
+			int virtualTextureAutoEnableThresholdPixels = virtualTextureAutoEnableThreshold * virtualTextureAutoEnableThreshold;
+
+			// We do this in pixels so a 8192 x 128 texture won't get VT enabled 
+			// We use the Source size instead of simple Texture2D->GetSizeX() as this uses the size of the platform data
+			// however for a new texture platform data may not be generated yet, and for an reimport of a texture this is the size of the
+			// old texture. 
+			// Using source size gives one small caveat. It looks at the size before mipmap power of two padding adjustment.
+			// Textures with more than 1 block (UDIM textures) must be imported as VT
+			if (Texture->Source.GetNumBlocks() > 1 ||
+				Texture2D->Source.GetSizeX()*Texture2D->Source.GetSizeY() >= virtualTextureAutoEnableThresholdPixels)
+			{
+				Texture2D->VirtualTextureStreaming = true;
+			}
 		}
 	}
 
