Disable arcball by default

diff --git a/Engine/Source/Editor/UnrealEd/Private/Settings/SettingsClasses.cpp b/Engine/Source/Editor/UnrealEd/Private/Settings/SettingsClasses.cpp
--- a/Engine/Source/Editor/UnrealEd/Private/Settings/SettingsClasses.cpp
+++ b/Engine/Source/Editor/UnrealEd/Private/Settings/SettingsClasses.cpp
@@ -673,8 +673,8 @@ ULevelEditorViewportSettings::ULevelEditorViewportSettings( const FObjectInitial
 	BillboardScale = 1.0f;
 	TransformWidgetSizeAdjustment = 0.0f;
 	MeasuringToolUnits = MeasureUnits_Centimeters;
-	bAllowArcballRotate = true;
-	bAllowScreenRotate = true;
+	bAllowArcballRotate = false;
+	bAllowScreenRotate = false;
 	// Set a default preview mesh
 	PreviewMeshes.Add(FSoftObjectPath("/Engine/EditorMeshes/ColorCalibrator/SM_ColorCalibrator.SM_ColorCalibrator"));
 }
