On FBX animation export bone space transforms maybe empty, so try to recalc them.

diff --git a/Engine/Source/Editor/UnrealEd/Private/Fbx/FbxAnimationExport.cpp b/Engine/Source/Editor/UnrealEd/Private/Fbx/FbxAnimationExport.cpp
--- a/Engine/Source/Editor/UnrealEd/Private/Fbx/FbxAnimationExport.cpp
+++ b/Engine/Source/Editor/UnrealEd/Private/Fbx/FbxAnimationExport.cpp
@@ -554,9 +554,20 @@ void FFbxExporter::ExportAnimTrack(IAnimTrackAdapter& AnimTrackAdapter, AActor*
 
 	if ( FindSkeleton(SkeletalMeshComponent, BoneNodes)==false )
 	{
-		// error
-		return;
+		UE_LOG(LogFbx, Warning, TEXT("Error FBX Animation Export, no root skeleton found."));
+		return;		
+	}
+	//if we have no allocated bone space transforms something wrong so try to recalc them
+	if (SkeletalMeshComponent->BoneSpaceTransforms.Num() <= 0 )
+	{
+		SkeletalMeshComponent->RecalcRequiredBones(0);
+		if (SkeletalMeshComponent->BoneSpaceTransforms.Num() <= 0)
+		{
+			UE_LOG(LogFbx, Warning, TEXT("Error FBX Animation Export, no bone transforms."));
+			return;
+		}
 	}
+	
 
 	FTransform InitialInvParentTransform;
 
