CIS Fix trivial

diff --git a/Engine/Source/Editor/UnrealEd/Private/Fbx/FbxAnimationExport.cpp b/Engine/Source/Editor/UnrealEd/Private/Fbx/FbxAnimationExport.cpp
--- a/Engine/Source/Editor/UnrealEd/Private/Fbx/FbxAnimationExport.cpp
+++ b/Engine/Source/Editor/UnrealEd/Private/Fbx/FbxAnimationExport.cpp
@@ -558,10 +558,10 @@ void FFbxExporter::ExportAnimTrack(IAnimTrackAdapter& AnimTrackAdapter, AActor*
 		return;		
 	}
 	//if we have no allocated bone space transforms something wrong so try to recalc them
-	if (SkeletalMeshComponent->BoneSpaceTransforms.Num() <= 0 )
+	if (SkeletalMeshComponent->GetBoneSpaceTransforms().Num() <= 0 )
 	{
 		SkeletalMeshComponent->RecalcRequiredBones(0);
-		if (SkeletalMeshComponent->BoneSpaceTransforms.Num() <= 0)
+		if (SkeletalMeshComponent->GetBoneSpaceTransforms().Num() <= 0)
 		{
 			UE_LOG(LogFbx, Warning, TEXT("Error FBX Animation Export, no bone transforms."));
 			return;
