Fixing MAP IMPORTADD cmd crash.

diff --git a/Engine/Source/Editor/UnrealEd/Private/EditorServer.cpp b/Engine/Source/Editor/UnrealEd/Private/EditorServer.cpp
--- a/Engine/Source/Editor/UnrealEd/Private/EditorServer.cpp
+++ b/Engine/Source/Editor/UnrealEd/Private/EditorServer.cpp
@@ -138,6 +138,7 @@
 #include "Logging/MessageLog.h"
 #include "Misc/UObjectToken.h"
 #include "Misc/MapErrors.h"
+#include "Misc/ScopedSlowTask.h"
 
 
 #include "ComponentReregisterContext.h"
@@ -2789,29 +2790,85 @@ bool UEditorEngine::Map_Load(const TCHAR* Str, FOutputDevice& Ar)
 
 bool UEditorEngine::Map_Import( UWorld* InWorld, const TCHAR* Str, FOutputDevice& Ar )
 {
-	FString TempFname;
-	if( FParse::Value( Str, TEXT("FILE="), TempFname ) )
+	FString FileName;
+	if( FParse::Value( Str, TEXT("FILE="), FileName ) )
 	{
 		const FScopedBusyCursor BusyCursor;
 
 		FFormatNamedArguments Args;
-		Args.Add( TEXT("MapFilename"), FText::FromString( FPaths::GetCleanFilename(*TempFname) ) );
+		Args.Add( TEXT("MapFilename"), FText::FromString( FPaths::GetCleanFilename(*FileName) ) );
 		const FText LocalizedImportingMap = FText::Format( NSLOCTEXT("UnrealEd", "ImportingMap_F", "Importing map: {MapFilename}..." ), Args );
 		
 		ResetTransaction( LocalizedImportingMap );
-		GWarn->BeginSlowTask( LocalizedImportingMap, true );
-		InWorld->ClearWorldComponents();
-		InWorld->CleanupWorld();
-		ImportObject<UWorld>(InWorld->GetOuter(), InWorld->GetFName(), RF_Transactional, *TempFname );
-		GWarn->EndSlowTask();
+		FScopedSlowTask SlowTask(0, LocalizedImportingMap);
 
-		// Importing content into a map will likely cause the list of actors in the level to change,
-		// so we'll trigger an event to notify other systems
-		FEditorDelegates::MapChange.Broadcast( MapChangeEventFlags::NewMap );
-		GEngine->BroadcastLevelActorListChanged();
+		UClass* WorldClass = UWorld::StaticClass();
 
-		NoteSelectionChange();
-		Cleanse( false, 1, NSLOCTEXT("UnrealEd", "ImportingActors", "Importing actors") );
+		// Hotfix for 4.23 this should be refactored later
+		TArray<UFactory*> Factories;
+		{
+			auto TransientPackage = GetTransientPackage();
+			FString Extension = FPaths::GetExtension(Str);
+
+			// try all automatic factories, sorted by priority
+			for (TObjectIterator<UClass> It; It; ++It)
+			{
+				if (It->IsChildOf(UFactory::StaticClass()))
+				{
+					UFactory* Default = It->GetDefaultObject<UFactory>();
+
+					if (WorldClass->IsChildOf(Default->SupportedClass) && Default->ImportPriority >= 0)
+					{
+						TArray<FString> FactoryExtension;
+						Default->GetSupportedFileExtensions(FactoryExtension);
+
+						if (Extension.IsEmpty() || (FactoryExtension.Contains(Extension) && Default->FactoryCanImport(FileName)))
+						{
+							Factories.Add(NewObject<UFactory>(TransientPackage, *It));
+						}
+					}
+				}
+			}
+		}
+		Factories.Sort(&UFactory::SortFactoriesByPriority);
+
+
+		if ( Factories.Num() > 0 )
+		{
+			InWorld->ClearWorldComponents();
+			InWorld->CleanupWorld();
+			UWorld* NewWorld = nullptr;
+
+			for ( UFactory* Factory : Factories )
+			{
+				bool bTemp;
+				NewWorld = Cast<UWorld>( Factory->ImportObject(WorldClass, InWorld->GetOuter(), InWorld->GetFName(), RF_Transactional, FileName, nullptr, bTemp) );
+				if ( NewWorld )
+				{
+					break;
+				}
+			}
+
+			if ( NewWorld )
+			{
+				if ( !GWorld )
+				{
+					GWorld = NewWorld;
+				}
+
+				// Importing content into a map will likely cause the list of actors in the level to change,
+				// so we'll trigger an event to notify other systems
+				FEditorDelegates::MapChange.Broadcast( MapChangeEventFlags::NewMap );
+				GEngine->BroadcastLevelActorListChanged();
+
+				NoteSelectionChange();
+				Cleanse( false, 1, NSLOCTEXT("UnrealEd", "ImportingActors", "Importing actors") );
+			}
+		}
+		else
+		{
+			UE_SUPPRESS(LogExec, Warning, Ar.Log( TEXT("Unsupported file format") ));
+		}
 	}
 	else
 	{
