#rnx

diff --git a/Engine/Source/Editor/UnrealEd/Private/EditorModeManager.cpp b/Engine/Source/Editor/UnrealEd/Private/EditorModeManager.cpp
--- a/Engine/Source/Editor/UnrealEd/Private/EditorModeManager.cpp
+++ b/Engine/Source/Editor/UnrealEd/Private/EditorModeManager.cpp
@@ -161,42 +161,38 @@ bool FEditorModeTools::SelectionHasSceneComponent() const
 
 void FEditorModeTools::UpdateModeWidgetLocation()
 {
-	bool Reached = false;
-
-	double starttime = FPlatformTime::Seconds();
-
-	USelection* SelectedActors = GetSelectedActors();
-
-	if (SelectedActors && SelectedActors->Num() > 0)
+	if (!bIsTracking)
 	{
-		AActor* Actor = Cast<AActor>(SelectedActors->GetSelectedObject(SelectedActors->Num() - 1));
+		USelection* SelectedActors = GetSelectedActors();
 
-		if (Actor)
+		if (SelectedActors && SelectedActors->Num() > 0)
 		{
-			FVector CurrentActorLocation = Actor->GetActorLocation();
+			AActor* Actor = Cast<AActor>(SelectedActors->GetSelectedObject(SelectedActors->Num() - 1));
 
-			if (Actor == LastSelectedActor)
+			if (Actor)
 			{
-				if (!(LastSelectedActorLocation - CurrentActorLocation).IsNearlyZero())
+				FVector CurrentActorLocation = Actor->GetActorLocation();
+
+				if (Actor == LastSelectedActor)
 				{
-					for (const auto& Mode : ActiveModes)
+					if (!(LastSelectedActorLocation - CurrentActorLocation).IsNearlyZero())
 					{
-						if (Mode->UsesTransformWidget())
+						for (const auto& Mode : ActiveModes)
 						{
-							SetPivotLocation(CurrentActorLocation, false);
-							LastSelectedActorLocation = CurrentActorLocation;
-
-							Reached = true;
-
-							break;
+							if (Mode->UsesTransformWidget())
+							{
+								SetPivotLocation(CurrentActorLocation, false);
+								LastSelectedActorLocation = CurrentActorLocation;
+								break;
+							}
 						}
-					}
 
+					}
+				}
+				else
+				{
+					LastSelectedActor = Actor;
 				}
-			}
-			else
-			{
-				LastSelectedActor = Actor;
 			}
 		}
 	}
