GenerateMipChain() and AdjustImageColors() are made public and changed to ITextureCompressorModule::GenerateMipChain() and ITextureCompressorModule::AdjustImageColors(). This allows us to export the specific MipMap levels and the textures with adjustments to PNG.

Fix SRGB of 16bit cubemap

Import the existing texture asset without Dialog box
New parameter to SuppressImportOverwriteDialog which can suppress the dialog box that, when importing over an existing texture, force to overwrite its settings.

PR #6123: Texture changes for NVIDIA USD plugin.

diff --git a/Engine/Source/Developer/TextureCompressor/Private/TextureCompressorModule.cpp b/Engine/Source/Developer/TextureCompressor/Private/TextureCompressorModule.cpp
--- a/Engine/Source/Developer/TextureCompressor/Private/TextureCompressorModule.cpp
+++ b/Engine/Source/Developer/TextureCompressor/Private/TextureCompressorModule.cpp
@@ -692,18 +692,11 @@ static void GenerateTopMip(const FImage& SrcImage, FImage& DestImage, const FTex
 	}
 }
 
-/**
- * Generate a full mip chain. The input mip chain must have one or more mips.
- * @param Settings - Preprocess settings.
- * @param BaseImage - An image that will serve as the source for the generation of the mip chain.
- * @param OutMipChain - An array that will contain the resultant mip images. Generated mip levels are appended to the array.
- * @param MipChainDepth - number of mip images to produce. Mips chain is finished when either a 1x1 mip is produced or 'MipChainDepth' images have been produced.
- */
-static void GenerateMipChain(
+void ITextureCompressorModule::GenerateMipChain(
 	const FTextureBuildSettings& Settings,
 	const FImage& BaseImage,
 	TArray<FImage> &OutMipChain,
-	uint32 MipChainDepth = MAX_uint32
+	uint32 MipChainDepth 
 	)
 {
 	check(BaseImage.Format == ERawImageFormat::RGBA32F);
@@ -1380,17 +1373,7 @@ static void GenerateAngularFilteredMips(TArray<FImage>& InOutMipChain, int32 Num
 	}
 }
 
-/*------------------------------------------------------------------------------
-	Image Processing.
-------------------------------------------------------------------------------*/
-
-/**
- * Adjusts the colors of the image using the specified settings
- *
- * @param	Image			Image to adjust
- * @param	InBuildSettings	Image build settings
- */
-static void AdjustImageColors( FImage& Image, const FTextureBuildSettings& InBuildSettings )
+void ITextureCompressorModule::AdjustImageColors(FImage& Image, const FTextureBuildSettings& InBuildSettings)
 {
 	const FColorAdjustmentParameters& InParams = InBuildSettings.ColorAdjustment;
 	check( Image.SizeX > 0 && Image.SizeY > 0 );
@@ -2028,6 +2011,11 @@ private:
 
 		int32 NumSourceMips = InSourceMips.Num();
 
+		if (BuildSettings.MipGenSettings == TMGS_LeaveExistingMips)
+		{
+			NumOutputMips = InSourceMips.Num();
+		}
+
 		if (BuildSettings.MipGenSettings != TMGS_LeaveExistingMips || bLongLatCubemap)
 		{
 			NumSourceMips = 1;
diff --git a/Engine/Source/Developer/TextureCompressor/Public/TextureCompressorModule.h b/Engine/Source/Developer/TextureCompressor/Public/TextureCompressorModule.h
--- a/Engine/Source/Developer/TextureCompressor/Public/TextureCompressorModule.h
+++ b/Engine/Source/Developer/TextureCompressor/Public/TextureCompressorModule.h
@@ -234,4 +234,27 @@ public:
 		const FTextureBuildSettings& BuildSettings,
 		TArray<FCompressedImage2D>& OutTextureMips
 		) = 0;
+
+	
+	/**
+	 * Generate a full mip chain. The input mip chain must have one or more mips.
+	 * @param Settings - Preprocess settings.
+	 * @param BaseImage - An image that will serve as the source for the generation of the mip chain.
+	 * @param OutMipChain - An array that will contain the resultant mip images. Generated mip levels are appended to the array.
+	 * @param MipChainDepth - number of mip images to produce. Mips chain is finished when either a 1x1 mip is produced or 'MipChainDepth' images have been produced.
+	 */
+	TEXTURECOMPRESSOR_API static void GenerateMipChain(
+		const FTextureBuildSettings& Settings,
+		const FImage& BaseImage,
+		TArray<FImage> &OutMipChain,
+		uint32 MipChainDepth = MAX_uint32
+		);
+
+	/**
+     * Adjusts the colors of the image using the specified settings
+     *
+     * @param	Image			Image to adjust
+     * @param	InBuildSettings	Image build settings
+     */
+	TEXTURECOMPRESSOR_API static void AdjustImageColors(FImage& Image, const FTextureBuildSettings& InBuildSettings);
 };
diff --git a/Engine/Source/Editor/UnrealEd/Classes/Factories/TextureFactory.h b/Engine/Source/Editor/UnrealEd/Classes/Factories/TextureFactory.h
--- a/Engine/Source/Editor/UnrealEd/Classes/Factories/TextureFactory.h
+++ b/Engine/Source/Editor/UnrealEd/Classes/Factories/TextureFactory.h
@@ -148,7 +148,7 @@ public:
 	 * Suppresses the dialog box that, when importing over an existing texture, asks if the users wishes to overwrite its settings.
 	 * This is primarily for reimporting textures.
 	 */
-	static void SuppressImportOverwriteDialog();
+	static void SuppressImportOverwriteDialog(bool bOverwriteExistingSettings = false);
 
 	/**
 	 *	Initializes the given texture from the TextureData text block supplied.
@@ -175,6 +175,9 @@ private:
 	/** This variable is static because in StaticImportObject() the type of the factory is not known. */
 	static bool bSuppressImportOverwriteDialog;
 
+    /** force overwriting the existing texture without the dialog box */
+	static bool bForceOverwriteExistingSettings;
+
 	/**
 	*	Tests if the given height and width specify a supported texture resolution to import; Can optionally check if the height/width are powers of two
 	*
diff --git a/Engine/Source/Editor/UnrealEd/Private/Factories/EditorFactories.cpp b/Engine/Source/Editor/UnrealEd/Private/Factories/EditorFactories.cpp
--- a/Engine/Source/Editor/UnrealEd/Private/Factories/EditorFactories.cpp
+++ b/Engine/Source/Editor/UnrealEd/Private/Factories/EditorFactories.cpp
@@ -2891,6 +2891,7 @@ bool DecompressTGA(
 }
 
 bool UTextureFactory::bSuppressImportOverwriteDialog = false;
+bool UTextureFactory::bForceOverwriteExistingSettings = false;
 
 UTextureFactory::UTextureFactory(const FObjectInitializer& ObjectInitializer)
 	: Super(ObjectInitializer)
@@ -2955,9 +2956,10 @@ UTextureCube* UTextureFactory::CreateTextureCube( UObject* InParent, FName Name,
 	return NewObject ? CastChecked<UTextureCube>(NewObject) : nullptr;
 }
 
-void UTextureFactory::SuppressImportOverwriteDialog()
+void UTextureFactory::SuppressImportOverwriteDialog(bool bOverwriteExistingSettings)
 {
 	bSuppressImportOverwriteDialog = true;
+    bForceOverwriteExistingSettings = bOverwriteExistingSettings;
 }
 
 /**
@@ -3807,6 +3809,7 @@ UTexture* UTextureFactory::ImportTexture(UClass* Class, UObject* InParent, FName
 			if(Format == TSF_RGBA16F)
 			{
 				TextureCube->CompressionSettings = TC_HDR;
+				TextureCube->SRGB = false;
 			}
 
 			uint8* DestMipData[MAX_TEXTURE_MIP_COUNT] = {0};
@@ -4061,43 +4064,51 @@ UObject* UTextureFactory::FactoryCreateBinary
 	TextureMipGenSettings				ExistingMipGenSettings = TextureMipGenSettings(0);
 	bool								ExistingVirtualTextureStreaming = false;
 
-	bUsingExistingSettings = bSuppressImportOverwriteDialog;
-
-	if(ExistingTexture && !bSuppressImportOverwriteDialog)
+	if (bForceOverwriteExistingSettings)
+	{
+		bUsingExistingSettings = false;
+	}
+	else
 	{
-		DisplayOverwriteOptionsDialog(FText::Format(
-			NSLOCTEXT("TextureFactory", "ImportOverwriteWarning", "You are about to import '{0}' over an existing texture."),
-			FText::FromName(TextureName)));
+		bUsingExistingSettings = bSuppressImportOverwriteDialog;
 
-		switch( OverwriteYesOrNoToAllState )
+		if (ExistingTexture && !bSuppressImportOverwriteDialog)
 		{
+			DisplayOverwriteOptionsDialog(FText::Format(
+				NSLOCTEXT("TextureFactory", "ImportOverwriteWarning", "You are about to import '{0}' over an existing texture."),
+				FText::FromName(TextureName)));
 
-		case EAppReturnType::Yes:
-		case EAppReturnType::YesAll:
-			{
-				// Overwrite existing settings
-				bUsingExistingSettings = false;
-				break;
-			}
-		case EAppReturnType::No:
-		case EAppReturnType::NoAll:
-			{
-				// Preserve existing settings
-				bUsingExistingSettings = true;
-				break;
-			}
-		case EAppReturnType::Cancel:
-		default:
+			switch (OverwriteYesOrNoToAllState)
 			{
-				GEditor->GetEditorSubsystem<UImportSubsystem>()->BroadcastAssetPostImport( this, nullptr );
-				return nullptr;
+
+			case EAppReturnType::Yes:
+			case EAppReturnType::YesAll:
+				{
+					// Overwrite existing settings
+					bUsingExistingSettings = false;
+					break;
+				}
+				case EAppReturnType::No:
+				case EAppReturnType::NoAll:
+				{
+					// Preserve existing settings
+					bUsingExistingSettings = true;
+					break;
+				}
+				case EAppReturnType::Cancel:
+				default:
+				{
+					GEditor->GetEditorSubsystem<UImportSubsystem>()->BroadcastAssetPostImport(this, nullptr);
+					return nullptr;
+				}
 			}
 		}
 	}
 
 	// Don't suppress future textures from checking for overwrites unless the calling code explicitly asks for it
 	bSuppressImportOverwriteDialog = false;
-	
+	bForceOverwriteExistingSettings = false;
+
 	if (ExistingTexture && bUsingExistingSettings)
 	{
 		// save settings
diff --git a/Engine/Source/Runtime/Engine/Classes/Engine/TextureCube.h b/Engine/Source/Runtime/Engine/Classes/Engine/TextureCube.h
--- a/Engine/Source/Runtime/Engine/Classes/Engine/TextureCube.h
+++ b/Engine/Source/Runtime/Engine/Classes/Engine/TextureCube.h
@@ -62,6 +62,15 @@ public:
 		return PF_Unknown;
 	}
 
+	/**
+	 * Get mip data starting with the specified mip index.
+	 * @param FirstMipToLoad - The first mip index to cache.
+	 * @param OutMipData -	Must point to an array of pointers with at least
+	 *						Mips.Num() - FirstMipToLoad + 1 entries. Upon
+	 *						return those pointers will contain mip data.
+	 */
+	ENGINE_API void GetMipData(int32 FirstMipToLoad, void** OutMipData);
+	
 	//~ Begin UTexture Interface
 	virtual float GetSurfaceWidth() const override { return GetSizeX(); }
 	virtual float GetSurfaceHeight() const override { return GetSizeY(); }
diff --git a/Engine/Source/Runtime/Engine/Private/TextureDerivedData.cpp b/Engine/Source/Runtime/Engine/Private/TextureDerivedData.cpp
--- a/Engine/Source/Runtime/Engine/Private/TextureDerivedData.cpp
+++ b/Engine/Source/Runtime/Engine/Private/TextureDerivedData.cpp
@@ -1491,6 +1491,26 @@ void UTexture2D::GetMipData(int32 FirstMipToLoad, void** OutMipData)
 	}
 }
 
+void UTextureCube::GetMipData(int32 FirstMipToLoad, void** OutMipData)
+{
+	if (PlatformData->TryLoadMips(FirstMipToLoad, OutMipData) == false)
+	{
+		// Unable to load mips from the cache. Rebuild the texture and try again.
+		UE_LOG(LogTexture,Warning,TEXT("GetMipData failed for %s (%s)"),
+			*GetPathName(), GPixelFormats[GetPixelFormat()].Name);
+#if WITH_EDITOR
+		if (!GetOutermost()->bIsCookedForEditor)
+		{
+			ForceRebuildPlatformData();
+			if (PlatformData->TryLoadMips(FirstMipToLoad, OutMipData) == false)
+			{
+				UE_LOG(LogTexture, Error, TEXT("Failed to build texture %s."), *GetPathName());
+			}
+		}
+#endif // #if WITH_EDITOR
+	}
+}
+
 void UTexture::UpdateCachedLODBias()
 {
 	CachedCombinedLODBias = UDeviceProfileManager::Get().GetActiveProfile()->GetTextureLODSettings()->CalculateLODBias(this);
