Add support for building runtime virtual texture low mip streaming data from the editor build menu

diff --git a/Engine/Source/Editor/LevelEditor/Private/LevelEditor.cpp b/Engine/Source/Editor/LevelEditor/Private/LevelEditor.cpp
--- a/Engine/Source/Editor/LevelEditor/Private/LevelEditor.cpp
+++ b/Engine/Source/Editor/LevelEditor/Private/LevelEditor.cpp
@@ -1478,6 +1478,9 @@ void FLevelEditorModule::BindGlobalLevelEditorCommands()
 	ActionList.MapAction(Commands.BuildTextureStreamingOnly,
 		FExecuteAction::CreateStatic(&FLevelEditorActionCallbacks::BuildTextureStreamingOnly_Execute));
 
+	ActionList.MapAction(Commands.BuildVirtualTextureOnly,
+		FExecuteAction::CreateStatic(&FLevelEditorActionCallbacks::BuildVirtualTextureOnly_Execute));
+
 	ActionList.MapAction( 
 		Commands.LightingQuality_Production, 
 		FExecuteAction::CreateStatic( &FLevelEditorActionCallbacks::SetLightingQuality, (ELightingBuildQuality)Quality_Production ),
diff --git a/Engine/Source/Editor/LevelEditor/Private/LevelEditorActions.cpp b/Engine/Source/Editor/LevelEditor/Private/LevelEditorActions.cpp
--- a/Engine/Source/Editor/LevelEditor/Private/LevelEditorActions.cpp
+++ b/Engine/Source/Editor/LevelEditor/Private/LevelEditorActions.cpp
@@ -772,6 +772,12 @@ void FLevelEditorActionCallbacks::BuildTextureStreamingOnly_Execute()
 	GEngine->DeferredCommands.AddUnique(TEXT("MAP CHECK NOTIFYRESULTS"));
 }
 
+void FLevelEditorActionCallbacks::BuildVirtualTextureOnly_Execute()
+{
+	FEditorBuildUtils::EditorBuildVirtualTexture(GetWorld());
+	GEngine->DeferredCommands.AddUnique(TEXT("MAP CHECK NOTIFYRESULTS"));
+}
+
 bool FLevelEditorActionCallbacks::IsLightingQualityChecked( ELightingBuildQuality TestQuality )
 {
 	int32 CurrentQualityLevel;
@@ -3167,7 +3173,8 @@ void FLevelEditorCommands::RegisterCommands()
 	UI_COMMAND( BuildPathsOnly, "Build Paths", "Only builds paths (all levels.)", EUserInterfaceActionType::Button, FInputChord() );
 	UI_COMMAND( BuildLODsOnly, "Build LODs", "Only builds LODs (all levels.)", EUserInterfaceActionType::Button, FInputChord() );
 	UI_COMMAND( BuildTextureStreamingOnly, "Build Texture Streaming", "Build texture streaming data", EUserInterfaceActionType::Button, FInputChord() );
-	
+	UI_COMMAND(BuildVirtualTextureOnly, "Build Virtual Textures", "Build runtime virtual texture low mips streaming data", EUserInterfaceActionType::Button, FInputChord());
+
 	UI_COMMAND( LightingQuality_Production, "Production", "Sets precomputed lighting quality to highest possible quality (slowest computation time.)", EUserInterfaceActionType::RadioButton, FInputChord() );
 	UI_COMMAND( LightingQuality_High, "High", "Sets precomputed lighting quality to high quality", EUserInterfaceActionType::RadioButton, FInputChord() );
 	UI_COMMAND( LightingQuality_Medium, "Medium", "Sets precomputed lighting quality to medium quality", EUserInterfaceActionType::RadioButton, FInputChord() );
diff --git a/Engine/Source/Editor/LevelEditor/Private/LevelEditorToolBar.cpp b/Engine/Source/Editor/LevelEditor/Private/LevelEditorToolBar.cpp
--- a/Engine/Source/Editor/LevelEditor/Private/LevelEditorToolBar.cpp
+++ b/Engine/Source/Editor/LevelEditor/Private/LevelEditorToolBar.cpp
@@ -1780,13 +1780,15 @@ TSharedRef< SWidget > FLevelEditorToolBar::GenerateBuildMenuContent( TSharedRef<
 	MenuBuilder.EndSection();
 
 	MenuBuilder.BeginSection("LevelEditorTextureStreaming", LOCTEXT("TextureStreamingHeading", "Texture Streaming"));
-	if (CVarStreamingUseNewMetrics.GetValueOnAnyThread() != 0) // There is not point of in building texture streaming data with the old system.
 	{
-		MenuBuilder.AddMenuEntry(FLevelEditorCommands::Get().BuildTextureStreamingOnly);
+		if (CVarStreamingUseNewMetrics.GetValueOnAnyThread() != 0) // There is not point of in building texture streaming data with the old system.
+		{
+			MenuBuilder.AddMenuEntry(FLevelEditorCommands::Get().BuildTextureStreamingOnly);
+		}
+		MenuBuilder.AddMenuEntry(FLevelEditorCommands::Get().BuildVirtualTextureOnly);
 	}
 	MenuBuilder.EndSection();
 
-
 	MenuBuilder.BeginSection("LevelEditorAutomation", LOCTEXT( "AutomationHeading", "Automation" ) );
 	{
 		MenuBuilder.AddMenuEntry( 
diff --git a/Engine/Source/Editor/LevelEditor/Public/LevelEditorActions.h b/Engine/Source/Editor/LevelEditor/Public/LevelEditorActions.h
--- a/Engine/Source/Editor/LevelEditor/Public/LevelEditorActions.h
+++ b/Engine/Source/Editor/LevelEditor/Public/LevelEditorActions.h
@@ -87,6 +87,7 @@ public:
 	TSharedPtr< FUICommandInfo > BuildPathsOnly;
 	TSharedPtr< FUICommandInfo > BuildLODsOnly;
 	TSharedPtr< FUICommandInfo > BuildTextureStreamingOnly;
+	TSharedPtr< FUICommandInfo > BuildVirtualTextureOnly;
 	TSharedPtr< FUICommandInfo > LightingQuality_Production;
 	TSharedPtr< FUICommandInfo > LightingQuality_High;
 	TSharedPtr< FUICommandInfo > LightingQuality_Medium;
@@ -774,6 +775,7 @@ public:
 	static void BuildPathsOnly_Execute();
 	static void BuildLODsOnly_Execute();
 	static void BuildTextureStreamingOnly_Execute();
+	static void BuildVirtualTextureOnly_Execute();
 	static void SetLightingQuality( ELightingBuildQuality NewQuality );
 	static bool IsLightingQualityChecked( ELightingBuildQuality TestQuality );
 	static float GetLightingDensityIdeal();
diff --git a/Engine/Source/Editor/UnrealEd/Private/EditorBuildUtils.cpp b/Engine/Source/Editor/UnrealEd/Private/EditorBuildUtils.cpp
--- a/Engine/Source/Editor/UnrealEd/Private/EditorBuildUtils.cpp
+++ b/Engine/Source/Editor/UnrealEd/Private/EditorBuildUtils.cpp
@@ -41,6 +41,8 @@
 #include "DebugViewModeHelpers.h"
 #include "MaterialStatsCommon.h"
 #include "Materials/MaterialInstance.h"
+#include "VirtualTexturingEditorModule.h"
+#include "Components/RuntimeVirtualTextureComponent.h"
 
 DEFINE_LOG_CATEGORY_STATIC(LogEditorBuildUtils, Log, All);
 
@@ -60,6 +62,7 @@ const FName FBuildOptions::BuildAllSubmit(TEXT("BuildAllSubmit"));
 const FName FBuildOptions::BuildAllOnlySelectedPaths(TEXT("BuildAllOnlySelectedPaths"));
 const FName FBuildOptions::BuildHierarchicalLOD(TEXT("BuildHierarchicalLOD"));
 const FName FBuildOptions::BuildTextureStreaming(TEXT("BuildTextureStreaming"));
+const FName FBuildOptions::BuildVirtualTexture(TEXT("BuildVirtualTexture"));
 
 bool FEditorBuildUtils::bBuildingNavigationFromUserRequest = false;
 TMap<FName, FEditorBuildUtils::FCustomBuildType> FEditorBuildUtils::CustomBuildTypes;
@@ -288,6 +291,10 @@ bool FEditorBuildUtils::EditorBuild( UWorld* InWorld, FName Id, const bool bAllo
 	{
 		BuildType = SBuildProgressWidget::BUILDTYPE_TextureStreaming;
 	}
+	else if (Id == FBuildOptions::BuildVirtualTexture)
+	{
+		BuildType = SBuildProgressWidget::BUILDTYPE_VirtualTexture;
+	}
 	else
 	{
 		BuildType = SBuildProgressWidget::BUILDTYPE_Unknown;	
@@ -991,6 +998,11 @@ void FBuildAllHandler::ProcessBuild(const TWeakPtr<SBuildProgressWidget>& BuildP
 			BuildProgressWidget.Pin()->SetBuildType(SBuildProgressWidget::BUILDTYPE_TextureStreaming);
 			FEditorBuildUtils::EditorBuildTextureStreaming(CurrentWorld);
 		}
+		else if (StepId == FBuildOptions::BuildVirtualTexture)
+		{
+			BuildProgressWidget.Pin()->SetBuildType(SBuildProgressWidget::BUILDTYPE_VirtualTexture);
+			FEditorBuildUtils::EditorBuildVirtualTexture(CurrentWorld);
+		}
 		else if (StepId == FBuildOptions::BuildAIPaths)
 		{
 			BuildProgressWidget.Pin()->SetBuildType(SBuildProgressWidget::BUILDTYPE_Paths);
@@ -1341,6 +1353,53 @@ bool FEditorBuildUtils::EditorBuildMaterialTextureStreamingData(UPackage* Packag
 	return bAnyPackagesDirtied;
 }
 
+bool FEditorBuildUtils::EditorBuildVirtualTexture(UWorld* InWorld)
+{
+	if (InWorld == nullptr)
+	{
+		return true;
+	}
+
+	IVirtualTexturingEditorModule* Module = FModuleManager::Get().GetModulePtr<IVirtualTexturingEditorModule>("VirtualTexturingEditor");
+	if (Module == nullptr)
+	{
+		return false;
+	}
+
+	CollectGarbage(GARBAGE_COLLECTION_KEEPFLAGS);
+
+	TArray<URuntimeVirtualTextureComponent*> Components;
+	for (TObjectIterator<URuntimeVirtualTextureComponent> It; It; ++It)
+	{
+		if (Module->HasStreamedMips(*It))
+		{
+			Components.Add(*It);
+		}
+	}
+
+	if (Components.Num() == 0)
+	{
+		return true;
+	}
+
+	FScopedSlowTask BuildTask(Components.Num(), LOCTEXT("VirtualTextureBuild", "Building Virtual Textures"));
+	BuildTask.MakeDialog(true);
+
+	for (URuntimeVirtualTextureComponent* Component : Components)
+	{
+		BuildTask.EnterProgressFrame();
+
+		if (BuildTask.ShouldCancel() || !Module->BuildStreamedMips(Component))
+		{
+			return false;
+		}
+	}
+
+	CollectGarbage(GARBAGE_COLLECTION_KEEPFLAGS);
+
+	return true;
+}
+
 /** classed used to compile shaders for a specific (mobile) platform and copy the number of instruction to the editor-emulated (mobile) platform */
 class FMaterialOfflineCompilation : public FMaterialResource
 {
diff --git a/Engine/Source/Editor/UnrealEd/Private/EditorModeRegistry.cpp b/Engine/Source/Editor/UnrealEd/Private/EditorModeRegistry.cpp
--- a/Engine/Source/Editor/UnrealEd/Private/EditorModeRegistry.cpp
+++ b/Engine/Source/Editor/UnrealEd/Private/EditorModeRegistry.cpp
@@ -66,7 +66,7 @@ void FEditorModeRegistry::Initialize()
 	FModuleManager::LoadModuleChecked<IMeshPaintModule>(TEXT("MeshPaintMode"));
 	FModuleManager::LoadModuleChecked<ILandscapeEditorModule>(TEXT("LandscapeEditor"));
 	FModuleManager::LoadModuleChecked<IFoliageEditModule>(TEXT("FoliageEdit"));
-	FModuleManager::LoadModuleChecked<FVirtualTexturingEditorModule>(TEXT("VirtualTexturingEditor"));
+	FModuleManager::LoadModuleChecked<IVirtualTexturingEditorModule>(TEXT("VirtualTexturingEditor"));
 }
 
 void FEditorModeRegistry::Shutdown()
diff --git a/Engine/Source/Editor/UnrealEd/Public/Dialogs/SBuildProgress.h b/Engine/Source/Editor/UnrealEd/Public/Dialogs/SBuildProgress.h
--- a/Engine/Source/Editor/UnrealEd/Public/Dialogs/SBuildProgress.h
+++ b/Engine/Source/Editor/UnrealEd/Public/Dialogs/SBuildProgress.h
@@ -44,6 +44,8 @@ public:
 		BUILDTYPE_LODs, 
 		/** Texture streaming data is being built */
 		BUILDTYPE_TextureStreaming, 
+		/** Virtual texture data is being built */
+		BUILDTYPE_VirtualTexture,
 	};
 
 	/** The various issues that can occur. */
diff --git a/Engine/Source/Editor/UnrealEd/Public/EditorBuildUtils.h b/Engine/Source/Editor/UnrealEd/Public/EditorBuildUtils.h
--- a/Engine/Source/Editor/UnrealEd/Public/EditorBuildUtils.h
+++ b/Engine/Source/Editor/UnrealEd/Public/EditorBuildUtils.h
@@ -32,8 +32,10 @@ struct FBuildOptions
 	UNREALED_API static const FName BuildAllOnlySelectedPaths;
 	/** Build Hierarchical LOD system - need WorldSetting setup*/
 	UNREALED_API static const FName BuildHierarchicalLOD;
-	/** Build Hierarchical LOD system - need WorldSetting setup*/
+	/** Build texture streaming */
 	UNREALED_API static const FName BuildTextureStreaming;
+	/** Build virtual textures */
+	UNREALED_API static const FName BuildVirtualTexture;
 };
 /**
  * Result of a custom editor build.
@@ -151,6 +153,15 @@ public:
 	*/
 	static UNREALED_API bool EditorBuildMaterialTextureStreamingData(UPackage* Package);
 
+	/**
+	 * Perform an editor build for virtual textures
+	 *
+	 * @param	InWorld				WorldContext
+	 *
+	 * @return	true if the build completed successfully; false if it did not (or was manually canceled)
+	 */
+	static UNREALED_API bool EditorBuildVirtualTexture(UWorld* InWorld);
+
 	/** 
 	* check if navigation build was was triggered from editor as user request
 	*
diff --git a/Engine/Source/Editor/VirtualTexturingEditor/Private/VirtualTexturingEditorModule.cpp b/Engine/Source/Editor/VirtualTexturingEditor/Private/VirtualTexturingEditorModule.cpp
--- a/Engine/Source/Editor/VirtualTexturingEditor/Private/VirtualTexturingEditorModule.cpp
+++ b/Engine/Source/Editor/VirtualTexturingEditor/Private/VirtualTexturingEditorModule.cpp
@@ -5,6 +5,7 @@
 #include "IPlacementModeModule.h"
 #include "PropertyEditorModule.h"
 #include "RuntimeVirtualTextureAssetTypeActions.h"
+#include "RuntimeVirtualTextureBuild.h"
 #include "RuntimeVirtualTextureDetailsCustomization.h"
 #include "RuntimeVirtualTextureThumbnailRenderer.h"
 #include "ThumbnailRendering/ThumbnailManager.h"
@@ -13,6 +14,25 @@
 
 #define LOCTEXT_NAMESPACE "VirtualTexturingEditorModule"
 
+/** Concrete implementation of the IVirtualTexturingEditorModule interface. */
+class FVirtualTexturingEditorModule : public IVirtualTexturingEditorModule
+{
+public:
+	//~ Begin IModuleInterface Interface.
+	virtual void StartupModule() override;
+	virtual void ShutdownModule() override;
+	virtual bool SupportsDynamicReloading() override;
+	//~ End IModuleInterface Interface.
+
+	//~ Begin IVirtualTexturingEditorModule Interface.
+	virtual bool HasStreamedMips(URuntimeVirtualTextureComponent* InComponent) const override;
+	virtual bool BuildStreamedMips(URuntimeVirtualTextureComponent* InComponent) const override;
+	//~ End IVirtualTexturingEditorModule Interface.
+
+private:
+	void OnPlacementModeRefresh(FName CategoryName);
+};
+
 IMPLEMENT_MODULE(FVirtualTexturingEditorModule, VirtualTexturingEditor);
 
 void FVirtualTexturingEditorModule::StartupModule()
@@ -53,4 +73,14 @@ void FVirtualTexturingEditorModule::OnPlacementModeRefresh(FName CategoryName)
 	}
 }
 
+bool FVirtualTexturingEditorModule::HasStreamedMips(URuntimeVirtualTextureComponent* InComponent) const
+{
+	return RuntimeVirtualTexture::HasStreamedMips(InComponent);
+}
+
+bool FVirtualTexturingEditorModule::BuildStreamedMips(URuntimeVirtualTextureComponent* InComponent) const
+{
+	return RuntimeVirtualTexture::BuildStreamedMips(InComponent, ERuntimeVirtualTextureDebugType::None);
+}
+
 #undef LOCTEXT_NAMESPACE
diff --git a/Engine/Source/Editor/VirtualTexturingEditor/Public/VirtualTexturingEditorModule.h b/Engine/Source/Editor/VirtualTexturingEditor/Public/VirtualTexturingEditorModule.h
--- a/Engine/Source/Editor/VirtualTexturingEditor/Public/VirtualTexturingEditorModule.h
+++ b/Engine/Source/Editor/VirtualTexturingEditor/Public/VirtualTexturingEditorModule.h
@@ -4,13 +4,11 @@
 #include "Modules/ModuleInterface.h"
 
 /** Module for virtual texturing editor extensions. */
-class FVirtualTexturingEditorModule	: public IModuleInterface
+class IVirtualTexturingEditorModule : public IModuleInterface
 {
 public:
-	virtual void StartupModule() override;
-	virtual void ShutdownModule() override;
-	virtual bool SupportsDynamicReloading() override;
-
-private:
-	void OnPlacementModeRefresh(FName CategoryName);
+	/** Returns true if the component describes a runtime virtual texture that has streaming low mips. */
+	virtual bool HasStreamedMips(class URuntimeVirtualTextureComponent* InComponent) const = 0;
+	/** Build the low mips and store in the component's associated URuntimeVirtualTexture object. */
+	virtual bool BuildStreamedMips(class URuntimeVirtualTextureComponent* InComponent) const = 0;
 };
