ChildActors will no longer be marked transactional as part of level load

diff --git a/Engine/Source/Editor/UnrealEd/Private/EditorServer.cpp b/Engine/Source/Editor/UnrealEd/Private/EditorServer.cpp
--- a/Engine/Source/Editor/UnrealEd/Private/EditorServer.cpp
+++ b/Engine/Source/Editor/UnrealEd/Private/EditorServer.cpp
@@ -2689,7 +2689,13 @@ bool UEditorEngine::Map_Load(const TCHAR* Str, FOutputDevice& Ar)
 					for( FActorIterator It(Context.World()); It; ++It )
 					{
 						AActor* Actor = *It;
-						Actor->SetFlags( RF_Transactional );
+
+						// Child actors of non-transactional components should not be transactional
+						UChildActorComponent* CAC = Actor->GetParentComponent();
+						if (CAC == nullptr || CAC->HasAnyFlags(RF_Transactional))
+						{
+							Actor->SetFlags(RF_Transactional);
+						}
 					}
 
 					InitializingFeedback.EnterProgressFrame();
