Fix name used for search matching with runtime virtual texture volume

diff --git a/Engine/Source/Editor/UnrealEd/Private/Factories/ActorFactory.cpp b/Engine/Source/Editor/UnrealEd/Private/Factories/ActorFactory.cpp
--- a/Engine/Source/Editor/UnrealEd/Private/Factories/ActorFactory.cpp
+++ b/Engine/Source/Editor/UnrealEd/Private/Factories/ActorFactory.cpp
@@ -1907,14 +1907,14 @@ UActorFactoryRuntimeVirtualTextureVolume
 UActorFactoryRuntimeVirtualTextureVolume::UActorFactoryRuntimeVirtualTextureVolume(const FObjectInitializer& ObjectInitializer)
 	: Super(ObjectInitializer)
 {
-	DisplayName = LOCTEXT("VirtualTextureVolume_DisplayName", "Virtual Texture Volume");
+	DisplayName = LOCTEXT("VirtualTextureVolume_DisplayName", "Runtime Virtual Texture Volume");
 	NewActorClass = ARuntimeVirtualTextureVolume::StaticClass();
 	bShowInEditorQuickMenu = 1;
 }
 
 void UActorFactoryRuntimeVirtualTextureVolume::PostSpawnActor(UObject* Asset, AActor* NewActor)
 {
-	FText ActorName = LOCTEXT("VirtualTextureVolume_DefaultActorName", "Virtual Texture Volume");
+	FText ActorName = LOCTEXT("VirtualTextureVolume_DefaultActorName", "Runtime Virtual Texture Volume");
 	NewActor->SetActorLabel(ActorName.ToString());
 
 	// Good default size to see object in editor
