Modify the resave package commandlet to only save packages that been saved by a licensee. Add -OnlyLicenseed

diff --git a/Engine/Source/Editor/UnrealEd/Classes/Commandlets/ResavePackagesCommandlet.h b/Engine/Source/Editor/UnrealEd/Classes/Commandlets/ResavePackagesCommandlet.h
--- a/Engine/Source/Editor/UnrealEd/Classes/Commandlets/ResavePackagesCommandlet.h
+++ b/Engine/Source/Editor/UnrealEd/Classes/Commandlets/ResavePackagesCommandlet.h
@@ -95,6 +95,9 @@ protected:
 	/** Only save packages with a changelist of zero **/
 	bool bOnlyUnversioned;
 
+	/** Only save packages that been saved by a licensee **/
+	bool bOnlyLicenseed;
+
 	/** Ignore package version changelist **/
 	bool bIgnoreChangelist;
 
diff --git a/Engine/Source/Editor/UnrealEd/Private/Commandlets/ContentCommandlets.cpp b/Engine/Source/Editor/UnrealEd/Private/Commandlets/ContentCommandlets.cpp
--- a/Engine/Source/Editor/UnrealEd/Private/Commandlets/ContentCommandlets.cpp
+++ b/Engine/Source/Editor/UnrealEd/Private/Commandlets/ContentCommandlets.cpp
@@ -898,6 +898,8 @@ int32 UResavePackagesCommandlet::Main( const FString& Params )
 	bIgnoreChangelist = Switches.Contains(TEXT("IgnoreChangelist"));
 	/** whether we should only save packages with changelist zero */
 	bOnlyUnversioned = Switches.Contains(TEXT("OnlyUnversioned"));
+	/** whether we should only save packages saved by licenseed */
+	bOnlyLicenseed = Switches.Contains(TEXT("OnlyLicenseed"));
 	/** only process packages containing materials */
 	bOnlyMaterials = Switches.Contains(TEXT("onlymaterials"));
 
@@ -1188,39 +1190,49 @@ void UResavePackagesCommandlet::PerformPreloadOperations( FLinkerLoad* PackageLi
 	if ( MinResaveUE4Version != IGNORE_PACKAGE_VERSION && UE4PackageVersion < MinResaveUE4Version )
 	{
 		bSavePackage = false;
+		return;
 	}
 
 	// Check if this package meets the maximum requirements.
-	bool bNoLimitation = MaxResaveUE4Version == IGNORE_PACKAGE_VERSION && MaxResaveLicenseeUE4Version == IGNORE_PACKAGE_VERSION;
-	bool bAllowResave = bNoLimitation ||
+	const bool bNoLimitation = MaxResaveUE4Version == IGNORE_PACKAGE_VERSION && MaxResaveLicenseeUE4Version == IGNORE_PACKAGE_VERSION;
+	const bool bAllowResave = bNoLimitation ||
 						 (MaxResaveUE4Version != IGNORE_PACKAGE_VERSION && UE4PackageVersion <= MaxResaveUE4Version) ||
 						 (MaxResaveLicenseeUE4Version != IGNORE_PACKAGE_VERSION && LicenseeUE4PackageVersion <= MaxResaveLicenseeUE4Version);
+	// If not, don't resave it.
+	if ( !bAllowResave )
+	{
+		bSavePackage = false;
+		return;
+	}
 
 	// If the package was saved with a higher engine version do not try to resave it. This also addresses problem with people 
 	// building editor locally and resaving content with a 0 CL version (e.g. BUILD_FROM_CL == 0)
 	if (!bIgnoreChangelist && PackageLinker->Summary.SavedByEngineVersion.GetChangelist() > FEngineVersion::Current().GetChangelist())
 	{
-		UE_LOG(LogContentCommandlet, Warning, TEXT("Skipping resave of %s due to engine version mismatch (Package:%d, Editor:%d "), 
+		UE_LOG(LogContentCommandlet, Warning, TEXT("Skipping resave of %s due to engine version mismatch (Package:%d, Editor:%d) "), 
 			*PackageLinker->GetArchiveName(),
 			PackageLinker->Summary.SavedByEngineVersion.GetChangelist(), 
 			FEngineVersion::Current().GetChangelist());
 		bSavePackage = false;
+		return;
 	}
 
 	// Check if the changelist number is zero
 	if (bOnlyUnversioned && PackageLinker->Summary.SavedByEngineVersion.GetChangelist() != 0)
 	{
 		bSavePackage = false;
+		return;
 	}
 
-	// If not, don't resave it.
-	if ( !bAllowResave )
+	// Check if the package was saved by licensees
+	if ( bOnlyLicenseed && !PackageLinker->Summary.SavedByEngineVersion.IsLicenseeVersion() )
 	{
 		bSavePackage = false;
+		return;
 	}
 
 	// Check if the package contains any instances of the class that needs to be resaved.
-	if ( bSavePackage && ResaveClasses.Num() > 0 )
+	if ( ResaveClasses.Num() > 0 )
 	{
 		bSavePackage = false;
 		for (int32 ExportIndex = 0; !bSavePackage && ExportIndex < PackageLinker->ExportMap.Num(); ExportIndex++)
