PR #6147: Show the selected IDE instead of the suggested IDE in the tooltip of code hyperlinks. (Contributed by [REDACTED_PERSON])

diff --git a/Engine/Source/Editor/UnrealEd/Private/EditorClassUtils.cpp b/Engine/Source/Editor/UnrealEd/Private/EditorClassUtils.cpp
--- a/Engine/Source/Editor/UnrealEd/Private/EditorClassUtils.cpp
+++ b/Engine/Source/Editor/UnrealEd/Private/EditorClassUtils.cpp
@@ -132,7 +132,7 @@ TSharedRef<SWidget> FEditorClassUtils::GetSourceLinkFormatted(const UClass* Clas
 			.Style(FEditorStyle::Get(), "Common.GotoNativeCodeHyperlink")
 			.OnNavigate_Static(&Local::OnEditCodeClicked, Class)
 			.Text(FText::Format(CodeFormat, FText::FromString((Class->GetName()))))
-			.ToolTipText(FText::Format(NSLOCTEXT("SourceHyperlink", "GoToCode_ToolTip", "Click to open this source file in {0}"), FSourceCodeNavigation::GetSuggestedSourceCodeIDE()));
+			.ToolTipText(FText::Format(NSLOCTEXT("SourceHyperlink", "GoToCode_ToolTip", "Click to open this source file in {0}"), FSourceCodeNavigation::GetSelectedSourceCodeIDE()));
 	}
 
 	return SourceHyperlink;
