Fix startup crash due to incorrect assumption that all editor viewports are scene viewports.

The editor viewport list can contain objects which are not instances of FSceneViewport (for example, FCompositingViewport), so we can't blindly cast InViewportClient->Viewport. FSceneViewport is the only type which sets the bIsSlateViewport flag in the FViewport base class, so we can use IsSlateViewport() to know when the cast is valid.

diff --git a/Engine/Source/Editor/UnrealEd/Private/EditorEngine.cpp b/Engine/Source/Editor/UnrealEd/Private/EditorEngine.cpp
--- a/Engine/Source/Editor/UnrealEd/Private/EditorEngine.cpp
+++ b/Engine/Source/Editor/UnrealEd/Private/EditorEngine.cpp
@@ -1996,12 +1996,15 @@ bool UEditorEngine::UpdateSingleViewportClient(FEditorViewportClient* InViewport
 {
 	bool bUpdatedNonRealtimeViewport = false;
 
-	// When rendering the viewport we need to know whether the final result will be shown on a HDR display. This affects the final post processing step
-	FSceneViewport *SceneViewport = static_cast<FSceneViewport*>(InViewportClient->Viewport);
-	TSharedPtr<SWindow> Window = SceneViewport->FindWindow();
-	if (Window)
+	if (InViewportClient->Viewport->IsSlateViewport())
 	{
-		InViewportClient->Viewport->SetHDRMode(Window->GetIsHDR());
+		// When rendering the viewport we need to know whether the final result will be shown on a HDR display. This affects the final post processing step
+		FSceneViewport *SceneViewport = static_cast<FSceneViewport*>(InViewportClient->Viewport);
+		TSharedPtr<SWindow> Window = SceneViewport->FindWindow();
+		if (Window)
+		{
+			InViewportClient->Viewport->SetHDRMode(Window->GetIsHDR());
+		}
 	}
 
 	// Always submit view information for content streaming 
