Add an exception for trying to instantiate UnrealEd in non-editor targets.

#rnx

diff --git a/Engine/Source/Editor/UnrealEd/UnrealEd.Build.cs b/Engine/Source/Editor/UnrealEd/UnrealEd.Build.cs
--- a/Engine/Source/Editor/UnrealEd/UnrealEd.Build.cs
+++ b/Engine/Source/Editor/UnrealEd/UnrealEd.Build.cs
@@ -7,6 +7,11 @@ public class UnrealEd : ModuleRules
 {
 	public UnrealEd(ReadOnlyTargetRules Target) : base(Target)
 	{
+		if(Target.Type != TargetType.Editor)
+		{
+			throw new BuildException("Unable to instantiate UnrealEd module for non-editor targets.");
+		}
+
 		PrivatePCHHeaderFile = "Private/UnrealEdPrivatePCH.h";
 
 		SharedPCHHeaderFile = "Public/UnrealEdSharedPCH.h";
