Fixed a bug where FKismetEditorUtilities::GetIBlueprintEditorForObject would fail to open a blueprint if passed the blueprint generated class directly (it was trying to handle the class, but only worked for things inside the class)
#rnx

diff --git a/Engine/Source/Editor/UnrealEd/Private/Kismet2/Kismet2.cpp b/Engine/Source/Editor/UnrealEd/Private/Kismet2/Kismet2.cpp
--- a/Engine/Source/Editor/UnrealEd/Private/Kismet2/Kismet2.cpp
+++ b/Engine/Source/Editor/UnrealEd/Private/Kismet2/Kismet2.cpp
@@ -2043,30 +2043,21 @@ TSharedPtr<class IBlueprintEditor> FKismetEditorUtilities::GetIBlueprintEditorFo
 	check(ObjectToFocusOn);
 
 	// Find the associated blueprint
-	UBlueprint* TargetBP = Cast<UBlueprint>(const_cast<UObject*>(ObjectToFocusOn));
-	if (TargetBP == NULL)
+	UBlueprint* TargetBP = nullptr;
+	for (UObject* TestObject = const_cast<UObject*>(ObjectToFocusOn); (TestObject != nullptr) && (TargetBP == nullptr); TestObject = TestObject->GetOuter())
 	{
-		for (UObject* TestOuter = ObjectToFocusOn->GetOuter(); TestOuter; TestOuter = TestOuter->GetOuter())
+		if (UBlueprintGeneratedClass* BPGeneratedClass = Cast<UBlueprintGeneratedClass>(TestObject))
 		{
-			TargetBP = Cast<UBlueprint>(TestOuter);
-
-			if(TargetBP == nullptr)
-			{
-				if(UBlueprintGeneratedClass* BPGeneratedClass = Cast<UBlueprintGeneratedClass>(TestOuter))
-				{
-					TargetBP = Cast<UBlueprint>(BPGeneratedClass->ClassGeneratedBy);
-				}
-			}
-
-			if (TargetBP != NULL)
-			{
-				break;
-			}
+			TargetBP = Cast<UBlueprint>(BPGeneratedClass->ClassGeneratedBy);
+		}
+		else
+		{
+			TargetBP = Cast<UBlueprint>(TestObject);
 		}
 	}
 
 	TSharedPtr<IBlueprintEditor> BlueprintEditor;
-	if (TargetBP != NULL)
+	if (TargetBP != nullptr)
 	{
 		if (bOpenEditor)
 		{
