Fixed a bug in SGraphEditor::GetNumberOfSelectedNodes and added a new method SGraphEditor::FindGraphEditorForGraph

diff --git a/Engine/Source/Editor/UnrealEd/Private/GraphEditor.cpp b/Engine/Source/Editor/UnrealEd/Private/GraphEditor.cpp
--- a/Engine/Source/Editor/UnrealEd/Private/GraphEditor.cpp
+++ b/Engine/Source/Editor/UnrealEd/Private/GraphEditor.cpp
@@ -155,3 +155,17 @@ void SGraphEditor::FocusCommentNodes(TArray<UEdGraphNode*> &CommentNodes, TArray
 		}
 	}
 }
+
+TSharedPtr<SGraphEditor> SGraphEditor::FindGraphEditorForGraph(const UEdGraph* Graph)
+{
+	for (TWeakPtr<SGraphEditor>& WeakWidget : AllInstances)
+	{
+		TSharedPtr<SGraphEditor> WidgetPtr = WeakWidget.Pin();
+		if (WidgetPtr.IsValid() && (WidgetPtr->GetCurrentGraph() == Graph))
+		{
+			return WidgetPtr;
+		}
+	}
+
+	return nullptr;
+}
diff --git a/Engine/Source/Editor/UnrealEd/Public/GraphEditor.h b/Engine/Source/Editor/UnrealEd/Public/GraphEditor.h
--- a/Engine/Source/Editor/UnrealEd/Public/GraphEditor.h
+++ b/Engine/Source/Editor/UnrealEd/Public/GraphEditor.h
@@ -546,7 +546,7 @@ public:
 	{
 		if (Implementation.IsValid())
 		{
-			Implementation->GetNumberOfSelectedNodes();
+			return Implementation->GetNumberOfSelectedNodes();
 		}
 		return 0;
 	}
@@ -562,6 +562,8 @@ public:
 		return nullptr;
 	}
 
+	// Returns the first graph editor that is viewing the specified graph
+	UNREALED_API static TSharedPtr<SGraphEditor> FindGraphEditorForGraph(const UEdGraph* Graph);
 
 protected:
 	/** Invoked when the underlying Graph is being changed. */
