Added a display/hyperlink in the top right of data asset editors to their native class (or the most-common-ancestor if editing multiple different ones at once)

diff --git a/Engine/Source/Editor/UnrealEd/Private/Toolkits/SimpleAssetEditor.cpp b/Engine/Source/Editor/UnrealEd/Private/Toolkits/SimpleAssetEditor.cpp
--- a/Engine/Source/Editor/UnrealEd/Private/Toolkits/SimpleAssetEditor.cpp
+++ b/Engine/Source/Editor/UnrealEd/Private/Toolkits/SimpleAssetEditor.cpp
@@ -7,6 +7,8 @@
 #include "PropertyEditorModule.h"
 #include "IDetailsView.h"
 #include "Editor.h"
+#include "Widgets/Input/SHyperlink.h"
+#include "SourceCodeNavigation.h"
 
 
 #define LOCTEXT_NAMESPACE "GenericEditor"
@@ -84,6 +86,7 @@ void FSimpleAssetEditor::InitEditor( const EToolkitMode::Type Mode, const TShare
 	const bool bCreateDefaultToolbar = true;
 	FAssetEditorToolkit::InitAssetEditor( Mode, InitToolkitHost, FSimpleAssetEditor::SimpleEditorAppIdentifier, StandaloneDefaultLayout, bCreateDefaultStandaloneMenu, bCreateDefaultToolbar, ObjectsToEdit );
 
+	RegenerateMenusAndToolbars();
 	// @todo toolkit world centric editing
 	// Setup our tool's layout
 	/*if( IsWorldCentricAssetEditor() && !PropertiesTab.IsValid() )
@@ -299,4 +302,59 @@ TSharedRef<FSimpleAssetEditor> FSimpleAssetEditor::CreateEditor( const EToolkitM
 	return NewEditor;
 }
 
+void FSimpleAssetEditor::PostRegenerateMenusAndToolbars()
+{
+	// Find the common denominator class of the assets we're editing
+	TArray<UClass*> ClassList;
+	for (UObject* Obj : EditingObjects)
+	{
+		check(Obj);
+		ClassList.Add(Obj->GetClass());
+	}
+
+	UClass* CommonDenominatorClass = UClass::FindCommonBase(ClassList);
+	const bool bNotAllSame = (EditingObjects.Num() > 0) && (EditingObjects[0]->GetClass() != CommonDenominatorClass);
+
+	// Provide a hyperlink to view that native class
+	if (CommonDenominatorClass)
+	{
+		TWeakObjectPtr<UClass> WeakClassPtr(CommonDenominatorClass);
+		auto OnNavigateToClassCode = [WeakClassPtr]()
+		{
+			if (UClass* StrongClassPtr = WeakClassPtr.Get())
+			{
+				if (FSourceCodeNavigation::CanNavigateToClass(StrongClassPtr))
+				{
+					FSourceCodeNavigation::NavigateToClass(StrongClassPtr);
+				}
+			}
+		};
+
+		// build and attach the menu overlay
+		TSharedRef<SHorizontalBox> MenuOverlayBox = SNew(SHorizontalBox)
+			+ SHorizontalBox::Slot()
+			.AutoWidth()
+			.VAlign(VAlign_Center)
+			[
+				SNew(STextBlock)
+				.ColorAndOpacity(FSlateColor::UseSubduedForeground())
+				.ShadowOffset(FVector2D::UnitVector)
+				.Text(bNotAllSame ? LOCTEXT("SimpleAssetEditor_AssetType_Varied", "Common Asset Type: ") : LOCTEXT("SimpleAssetEditor_AssetType", "Asset Type: "))
+			]
+			+SHorizontalBox::Slot()
+			.AutoWidth()
+			.VAlign(VAlign_Center)
+			.Padding(0.0f, 0.0f, 8.0f, 0.0f)
+			[
+				SNew(SHyperlink)
+				.Style(FEditorStyle::Get(), "Common.GotoNativeCodeHyperlink")
+				.OnNavigate_Lambda(OnNavigateToClassCode)
+				.Text(FText::FromName(CommonDenominatorClass->GetFName()))
+				.ToolTipText(FText::Format(LOCTEXT("GoToCode_ToolTip", "Click to open this source file in {0}"), FSourceCodeNavigation::GetSelectedSourceCodeIDE()))
+			];
+	
+		SetMenuOverlay(MenuOverlayBox);
+	}
+}
+
 #undef LOCTEXT_NAMESPACE
diff --git a/Engine/Source/Editor/UnrealEd/Public/Toolkits/SimpleAssetEditor.h b/Engine/Source/Editor/UnrealEd/Public/Toolkits/SimpleAssetEditor.h
--- a/Engine/Source/Editor/UnrealEd/Public/Toolkits/SimpleAssetEditor.h
+++ b/Engine/Source/Editor/UnrealEd/Public/Toolkits/SimpleAssetEditor.h
@@ -41,7 +41,10 @@ public:
 	virtual FString GetWorldCentricTabPrefix() const override;
 	virtual FLinearColor GetWorldCentricTabColorScale() const override;
 	virtual bool IsPrimaryEditor() const override { return true; }
-	
+
+	/** FAssetEditorToolkit interface */
+	virtual void PostRegenerateMenusAndToolbars() override;
+
 	/** Used to show or hide certain properties */
 	void SetPropertyVisibilityDelegate(FIsPropertyVisible InVisibilityDelegate);
 
