Fix shadowing + remove additional GWorld
#rnx

diff --git a/Engine/Source/Editor/UnrealEd/Private/Commandlets/ContentCommandlets.cpp b/Engine/Source/Editor/UnrealEd/Private/Commandlets/ContentCommandlets.cpp
--- a/Engine/Source/Editor/UnrealEd/Private/Commandlets/ContentCommandlets.cpp
+++ b/Engine/Source/Editor/UnrealEd/Private/Commandlets/ContentCommandlets.cpp
@@ -1531,7 +1531,7 @@ void UResavePackagesCommandlet::PerformAdditionalOperations(class UWorld* World,
 			// If we are (minimally) rebuilding HLOD, set the visible streamed-in levels packages to clean (as FlushLevelStreaming will dirty their packages in this commandlet context)
 			if(bShouldBuildHLOD)
 			{
-				for (const ULevel* Level : GWorld->GetLevels())
+				for (const ULevel* Level : World->GetLevels())
 				{
 					if (Level->bIsVisible)
 					{
@@ -1550,7 +1550,6 @@ void UResavePackagesCommandlet::PerformAdditionalOperations(class UWorld* World,
 
 			if (bShouldBuildHLOD)
 			{
-				UWorld* World = GWorld;
 				UE_LOG( LogContentCommandlet, Display, TEXT( "Generating HLOD data for %s" ), *World->GetOutermost()->GetName() );
 
 				if( !ForceHLODSetupAsset.IsEmpty() )
