Add a helper function for checking if a graph is an interface graph

#rnx

diff --git a/Engine/Source/Editor/UnrealEd/Private/Kismet2/BlueprintEditorUtils.cpp b/Engine/Source/Editor/UnrealEd/Private/Kismet2/BlueprintEditorUtils.cpp
--- a/Engine/Source/Editor/UnrealEd/Private/Kismet2/BlueprintEditorUtils.cpp
+++ b/Engine/Source/Editor/UnrealEd/Private/Kismet2/BlueprintEditorUtils.cpp
@@ -3238,6 +3238,11 @@ bool FBlueprintEditorUtils::IsInterfaceBlueprint(const UBlueprint* Blueprint)
 	return (Blueprint && Blueprint->BlueprintType == BPTYPE_Interface);
 }
 
+bool FBlueprintEditorUtils::IsInterfaceGraph(const UEdGraph* Graph)
+{
+	return IsInterfaceBlueprint(FindBlueprintForGraph(Graph));
+}
+
 bool FBlueprintEditorUtils::IsLevelScriptBlueprint(const UBlueprint* Blueprint)
 {
 	return (Blueprint && Blueprint->BlueprintType == BPTYPE_LevelScript);
diff --git a/Engine/Source/Editor/UnrealEd/Public/Kismet2/BlueprintEditorUtils.h b/Engine/Source/Editor/UnrealEd/Public/Kismet2/BlueprintEditorUtils.h
--- a/Engine/Source/Editor/UnrealEd/Public/Kismet2/BlueprintEditorUtils.h
+++ b/Engine/Source/Editor/UnrealEd/Public/Kismet2/BlueprintEditorUtils.h
@@ -555,6 +555,13 @@ public:
 	 */
 	static bool IsInterfaceBlueprint(const UBlueprint* Blueprint);
 
+	/**
+	* Whether or not this graph is an interface graph (i.e. is from an interface blueprint)
+	* 
+	* @return	Whether or not this is an interface graph
+	*/
+	static bool IsInterfaceGraph(const UEdGraph* Graph);
+
 	/**
 	 * Whether or not this blueprint is an interface, used only for defining functions to implement
 	 *
