Allow drag to re-order local variables in MyBlueprint  #jira UE-12398 #rb Dan.Oconnor

diff --git a/Engine/Source/Editor/BlueprintGraph/Classes/EdGraphSchema_K2_Actions.h b/Engine/Source/Editor/BlueprintGraph/Classes/EdGraphSchema_K2_Actions.h
--- a/Engine/Source/Editor/BlueprintGraph/Classes/EdGraphSchema_K2_Actions.h
+++ b/Engine/Source/Editor/BlueprintGraph/Classes/EdGraphSchema_K2_Actions.h
@@ -569,6 +569,7 @@ public:
 	// Simple type info
 	static FName StaticGetTypeId() {static FName Type("FEdGraphSchemaAction_K2LocalVar"); return Type;}
 	virtual FName GetTypeId() const override { return StaticGetTypeId(); } 
+	virtual int32 GetReorderIndexInContainer() const override;
 
 	FEdGraphSchemaAction_K2LocalVar() 
 		: FEdGraphSchemaAction_BlueprintVariableBase()
diff --git a/Engine/Source/Editor/BlueprintGraph/Private/EdGraphSchema_K2_Actions.cpp b/Engine/Source/Editor/BlueprintGraph/Private/EdGraphSchema_K2_Actions.cpp
--- a/Engine/Source/Editor/BlueprintGraph/Private/EdGraphSchema_K2_Actions.cpp
+++ b/Engine/Source/Editor/BlueprintGraph/Private/EdGraphSchema_K2_Actions.cpp
@@ -57,16 +57,16 @@ int32 FEdGraphSchemaAction_BlueprintVariableBase::GetReorderIndexInContainer() c
 
 bool FEdGraphSchemaAction_BlueprintVariableBase::ReorderToBeforeAction(TSharedRef<FEdGraphSchemaAction> OtherAction)
 {
-	if ((OtherAction->GetTypeId() == StaticGetTypeId()) && (OtherAction->GetPersistentItemDefiningObject() == GetPersistentItemDefiningObject()))
+	if (OtherAction->GetPersistentItemDefiningObject() == GetPersistentItemDefiningObject())
 	{
 		FEdGraphSchemaAction_BlueprintVariableBase* VarAction = (FEdGraphSchemaAction_BlueprintVariableBase*)&OtherAction.Get();
 
 		// Only let you drag and drop if variables are from same BP class, and not onto itself
 		UBlueprint* BP = GetSourceBlueprint();
 		FName TargetVarName = VarAction->GetVariableName();
-		if ((BP != nullptr) && (VarName != TargetVarName) && (VariableSource == VarAction->GetVariableClass()))
+		if ((BP != nullptr) && (VarName != TargetVarName) && (VariableSource == VarAction->GetVariableScope()))
 		{
-			if (FBlueprintEditorUtils::MoveVariableBeforeVariable(BP, VarName, TargetVarName, true))
+			if (FBlueprintEditorUtils::MoveVariableBeforeVariable(BP, VarAction->GetVariableScope(), VarName, TargetVarName, true))
 			{
 				// Change category of var to match the one we dragged on to as well
 				FText TargetVarCategory = FBlueprintEditorUtils::GetBlueprintVariableCategory(BP, TargetVarName, GetVariableScope());
@@ -107,6 +107,14 @@ UBlueprint* FEdGraphSchemaAction_BlueprintVariableBase::GetSourceBlueprint() con
 	return UBlueprint::GetBlueprintFromClass(ClassToCheck);
 }
 
+/////////////////////////////////////////////////////
+// FEdGraphSchemaAction_K2LocalVar
+
+int32 FEdGraphSchemaAction_K2LocalVar::GetReorderIndexInContainer() const
+{
+	return FBlueprintEditorUtils::FindLocalVariableIndex(GetSourceBlueprint(), GetVariableScope(), GetVariableName());
+}
+
 /////////////////////////////////////////////////////
 // FEdGraphSchemaAction_K2Graph
 
@@ -132,7 +140,7 @@ int32 FEdGraphSchemaAction_K2Graph::GetReorderIndexInContainer() const
 
 bool FEdGraphSchemaAction_K2Graph::ReorderToBeforeAction(TSharedRef<FEdGraphSchemaAction> OtherAction)
 {
-	if ((OtherAction->GetTypeId() == GetTypeId()) && (OtherAction->GetPersistentItemDefiningObject() == GetPersistentItemDefiningObject()))
+	if (OtherAction->GetTypeId() == GetTypeId())
 	{
 		const int32 OldIndex = GetReorderIndexInContainer();
 		const int32 NewIndexToGoBefore = OtherAction->GetReorderIndexInContainer();
@@ -708,3 +716,4 @@ UEdGraphNode* FEdGraphSchemaAction_K2PasteHere::PerformAction( class UEdGraph* P
 	return nullptr;
 }
 
+
diff --git a/Engine/Source/Editor/UnrealEd/Private/Kismet2/BlueprintEditorUtils.cpp b/Engine/Source/Editor/UnrealEd/Private/Kismet2/BlueprintEditorUtils.cpp
--- a/Engine/Source/Editor/UnrealEd/Private/Kismet2/BlueprintEditorUtils.cpp
+++ b/Engine/Source/Editor/UnrealEd/Private/Kismet2/BlueprintEditorUtils.cpp
@@ -3466,26 +3466,81 @@ int32 FBlueprintEditorUtils::FindNewVariableIndex(const UBlueprint* Blueprint, c
 	return INDEX_NONE;
 }
 
-bool FBlueprintEditorUtils::MoveVariableBeforeVariable(UBlueprint* Blueprint, FName VarNameToMove, FName TargetVarName, bool bDontRecompile)
+int32 FBlueprintEditorUtils::FindLocalVariableIndex(const UBlueprint* Blueprint, UStruct* VariableScope, const FName& InVariableName)
 {
+	UK2Node_FunctionEntry* FunctionEntryNode = nullptr;
+	FindLocalVariable(Blueprint, VariableScope, InVariableName, &FunctionEntryNode);
+
+	if (FunctionEntryNode != nullptr)
+	{
+		for (int32 i = 0; i < FunctionEntryNode->LocalVariables.Num(); i++)
+		{
+			if (FunctionEntryNode->LocalVariables[i].VarName == InVariableName)
+			{
+				return i;
+			}
+		}
+	}
+	return INDEX_NONE;
+}
+
+bool FBlueprintEditorUtils::MoveVariableBeforeVariable(UBlueprint* Blueprint, UStruct* VariableScope, FName VarNameToMove, FName TargetVarName, bool bDontRecompile)
+{
+	check(Blueprint && VariableScope);
+
 	bool bMoved = false;
-	int32 VarIndexToMove = FindNewVariableIndex(Blueprint, VarNameToMove);
-	int32 TargetVarIndex = FindNewVariableIndex(Blueprint, TargetVarName);
-	if(VarIndexToMove != INDEX_NONE && TargetVarIndex != INDEX_NONE)
+	int32 VarIndexToMove = INDEX_NONE;
+	int32 TargetVarIndex = INDEX_NONE;
+
+	//Get the indices of the variables to be re-ordered
+	if (VariableScope->IsA(UFunction::StaticClass()))
+	{
+		VarIndexToMove = FindLocalVariableIndex(Blueprint, VariableScope, VarNameToMove);
+		TargetVarIndex = FindLocalVariableIndex(Blueprint, VariableScope, TargetVarName);
+	}
+	else
+	{
+		VarIndexToMove = FindNewVariableIndex(Blueprint, VarNameToMove);
+		TargetVarIndex = FindNewVariableIndex(Blueprint, TargetVarName);
+	}
+
+	if (VarIndexToMove != INDEX_NONE && TargetVarIndex != INDEX_NONE)
 	{
-		// Copy var we want to move
-		FBPVariableDescription MoveVar = Blueprint->NewVariables[VarIndexToMove];
 		// When we remove item, will back all items after it. If your target is after it, need to adjust
-		if(TargetVarIndex > VarIndexToMove)
+		if (TargetVarIndex > VarIndexToMove)
 		{
 			TargetVarIndex--;
 		}
-		// Remove var we are moving
-		Blueprint->NewVariables.RemoveAt(VarIndexToMove);
-		// Add in before target variable
-		Blueprint->NewVariables.Insert(MoveVar, TargetVarIndex);
 
-		if(!bDontRecompile)
+		//Handle Local vs Class scope 
+		if (VariableScope->IsA(UFunction::StaticClass()))
+		{
+			UK2Node_FunctionEntry* FunctionEntryNode = nullptr;
+			FindLocalVariable(Blueprint, VariableScope, VarNameToMove, &FunctionEntryNode);
+
+			if (FunctionEntryNode != nullptr)
+			{
+				// Copy var we want to move
+				FBPVariableDescription MoveVar = FunctionEntryNode->LocalVariables[VarIndexToMove];
+
+				// Remove var we are moving
+				FunctionEntryNode->LocalVariables.RemoveAt(VarIndexToMove);
+				// Add in before target variable
+				FunctionEntryNode->LocalVariables.Insert(MoveVar, TargetVarIndex);				
+			}
+		}
+		else
+		{			
+			// Copy var we want to move
+			FBPVariableDescription MoveVar = Blueprint->NewVariables[VarIndexToMove];
+			
+			// Remove var we are moving
+			Blueprint->NewVariables.RemoveAt(VarIndexToMove);
+			// Add in before target variable
+			Blueprint->NewVariables.Insert(MoveVar, TargetVarIndex);
+		}
+
+		if (!bDontRecompile)
 		{
 			FBlueprintEditorUtils::MarkBlueprintAsStructurallyModified(Blueprint);
 		}
diff --git a/Engine/Source/Editor/UnrealEd/Public/Kismet2/BlueprintEditorUtils.h b/Engine/Source/Editor/UnrealEd/Public/Kismet2/BlueprintEditorUtils.h
--- a/Engine/Source/Editor/UnrealEd/Public/Kismet2/BlueprintEditorUtils.h
+++ b/Engine/Source/Editor/UnrealEd/Public/Kismet2/BlueprintEditorUtils.h
@@ -1152,8 +1152,19 @@ public:
 	 */
 	static int32 FindNewVariableIndex(const UBlueprint* Blueprint, const FName& InName);
 
+	/**
+	 * Find the index of a local variable declared in this blueprint. Returns INDEX_NONE if not found.
+	 *
+	 * @param	VariableScope	Struct of owning function.
+	 *
+	 * @param	InVariableName	Name of the variable to find.
+	 *
+	 * @return	The index of the variable, or INDEX_NONE if it wasn't introduced in this blueprint.
+	 */
+	static int32 FindLocalVariableIndex(const UBlueprint* Blueprint, UStruct* VariableScope, const FName& InVariableName);
+
 	/** Change the order of variables in the Blueprint */
-	static bool MoveVariableBeforeVariable(UBlueprint* Blueprint, FName VarNameToMove, FName TargetVarName, bool bDontRecompile);
+	static bool MoveVariableBeforeVariable(UBlueprint* Blueprint, UStruct* VariableScope, FName VarNameToMove, FName TargetVarName, bool bDontRecompile);
 
 	/**
 	 * Find the index of a timeline first declared in this blueprint. Returns INDEX_NONE if not found.
