PIE Network Emulation settings will now be disabled instead of hidden when Play settings are configured for a single instance

diff --git a/Engine/Source/Editor/UnrealEd/Private/Settings/LevelEditorPlaySettingsCustomization.h b/Engine/Source/Editor/UnrealEd/Private/Settings/LevelEditorPlaySettingsCustomization.h
--- a/Engine/Source/Editor/UnrealEd/Private/Settings/LevelEditorPlaySettingsCustomization.h
+++ b/Engine/Source/Editor/UnrealEd/Private/Settings/LevelEditorPlaySettingsCustomization.h
@@ -701,7 +701,7 @@ public:
 				.Visibility(TAttribute<EVisibility>::Create(TAttribute<EVisibility>::FGetter::CreateSP(this, &FLevelEditorPlaySettingsCustomization::HandleCmdLineVisibility)));
 
 			NetworkCategory.AddProperty("NetworkEmulationSettings")
-				.Visibility(TAttribute<EVisibility>::Create(TAttribute<EVisibility>::FGetter::CreateSP(this, &FLevelEditorPlaySettingsCustomization::HandleNetworkEmulationVisibility)));
+				.IsEnabled(TAttribute<bool>(this, &FLevelEditorPlaySettingsCustomization::HandleNetworkEmulationIsEnabled));
 
 			NetworkCategory.AddCustomRow(LOCTEXT("PlayInNetworkViewportSize", "Multiplayer Viewport Size"), false)
 				.NameContent()
@@ -902,9 +902,9 @@ private:
 	}
 
 	// Callback for checking if the Network Emulation can be used
-	EVisibility HandleNetworkEmulationVisibility() const
+	bool HandleNetworkEmulationIsEnabled() const
 	{
-		return GetDefault<ULevelEditorPlaySettings>()->GetNetworkEmulationVisibility();
+		return GetDefault<ULevelEditorPlaySettings>()->GetNetworkEmulationVisibility() == EVisibility::Visible;
 	}
 
 	// Callback for getting the visibility of the StandaloneServerMapName property.
diff --git a/Engine/Source/Runtime/Engine/Classes/Engine/NetConnection.h b/Engine/Source/Runtime/Engine/Classes/Engine/NetConnection.h
--- a/Engine/Source/Runtime/Engine/Classes/Engine/NetConnection.h
+++ b/Engine/Source/Runtime/Engine/Classes/Engine/NetConnection.h
@@ -632,7 +632,6 @@ public:
 
 #if DO_ENABLE_NET_TEST
 
-	// For development.
 	/** Packet settings for testing lag, net errors, etc */
 	FPacketSimulationSettings PacketSimulationSettings;
 
