Change socket draw size to match size in static mesh editor for better visibility.

diff --git a/Engine/Source/Editor/UnrealEd/Private/LevelEditorViewport.cpp b/Engine/Source/Editor/UnrealEd/Private/LevelEditorViewport.cpp
--- a/Engine/Source/Editor/UnrealEd/Private/LevelEditorViewport.cpp
+++ b/Engine/Source/Editor/UnrealEd/Private/LevelEditorViewport.cpp
@@ -4381,7 +4381,7 @@ void FLevelEditorViewportClient::Draw(const FSceneView* View,FPrimitiveDrawInter
 						{
 							const FTransform SocketTransform = SceneComponent->GetSocketTransform(Socket.Name);
 
-							const float DiamondSize = 2.0f;
+							const float DiamondSize = 5.0f;
 							const FColor DiamondColor(255,128,128);
 
 							PDI->SetHitProxy( new HLevelSocketProxy( *It, SceneComponent, Socket.Name ) );
