[CrashReport] FHierarchicalLODBuilder::BuildMeshesForLODActors(bool)
Replaced check by error message

diff --git a/Engine/Source/Editor/UnrealEd/Private/HierarchicalLOD.cpp b/Engine/Source/Editor/UnrealEd/Private/HierarchicalLOD.cpp
--- a/Engine/Source/Editor/UnrealEd/Private/HierarchicalLOD.cpp
+++ b/Engine/Source/Editor/UnrealEd/Private/HierarchicalLOD.cpp
@@ -699,7 +699,6 @@ void FHierarchicalLODBuilder::BuildMeshesForLODActors(bool bForceAll)
 			}
 
 			// If there are any available process them
-			bool bBuildSuccessful = true;
 			if (NumLODActors)
 			{
 				// Only create the outer package if we are going to save something to it (otherwise we end up with an empty HLOD folder)
@@ -723,13 +722,22 @@ void FHierarchicalLODBuilder::BuildMeshesForLODActors(bool bForceAll)
 					for (ALODActor* Actor : LODLevel)
 					{
 						SlowTask.EnterProgressFrame(100.0f / (float)NumLODActors, FText::Format(LOCTEXT("HierarchicalLOD_BuildLODActorMeshesProgress", "Building LODActor Mesh {1} of {2} (LOD Level {0})"), FText::AsNumber(LODIndex + 1), FText::AsNumber(LODActorIndex), FText::AsNumber(LODLevelActors[CurrentLODLevel].Num())));
-						bBuildSuccessful &= Utilities->BuildStaticMeshForLODActor(Actor, AssetsOuter, BuildLODLevelSettings[CurrentLODLevel], BaseMaterial);
+
+						bool bBuildSuccessful = Utilities->BuildStaticMeshForLODActor(Actor, AssetsOuter, BuildLODLevelSettings[CurrentLODLevel], BaseMaterial);
+
+						// Report an error if the build failed
+						if (!bBuildSuccessful)
+						{
+							FMessageLog("HLODResults").Error()
+								->AddToken(FTextToken::Create(LOCTEXT("HLODError_MeshNotBuildOne", "Cannot create proxy mesh for ")))
+								->AddToken(FUObjectToken::Create(Actor))
+								->AddToken(FTextToken::Create(LOCTEXT("HLODError_MeshNotBuildTwo", " this could be caused by incorrect mesh components in the sub actors")));
+						}
+
 						++LODActorIndex;
 					}
 				}
 			}
-
-			check(bBuildSuccessful);
 		}
 	}
 
