- On a skeletal mesh reimport save / restore OnMeshChanged delegate information to ensure all subscribers are notified about a reimport
- This fixes an issue around the Niagara Skel Mesh DI not being notified on a reimport of a skeletal mesh

diff --git a/Engine/Source/Editor/UnrealEd/Private/Factories/SkeletalMeshImport.cpp b/Engine/Source/Editor/UnrealEd/Private/Factories/SkeletalMeshImport.cpp
--- a/Engine/Source/Editor/UnrealEd/Private/Factories/SkeletalMeshImport.cpp
+++ b/Engine/Source/Editor/UnrealEd/Private/Factories/SkeletalMeshImport.cpp
@@ -570,6 +570,9 @@ ExistingSkelMeshData* SaveExistingSkelMeshData(USkeletalMesh* ExistingSkelMesh,
 			}
 		}
 	}
+	//Store mesh changed delegate data
+	ExistingMeshDataPtr->ExistingOnMeshChanged = ExistingSkelMesh->GetOnMeshChanged();
+
 	return ExistingMeshDataPtr;
 }
 
@@ -1194,5 +1197,7 @@ void RestoreExistingSkelMeshData(ExistingSkelMeshData* MeshData, USkeletalMesh*
 		UserDataObject->Rename(nullptr, SkeletalMesh, REN_DontCreateRedirectors | REN_DoNotDirty);
 		SkeletalMesh->AddAssetUserData(UserDataObject);
 	}
+	//Copy mesh changed delegate data
+	SkeletalMesh->GetOnMeshChanged() = MeshData->ExistingOnMeshChanged;
 }
 #undef LOCTEXT_NAMESPACE
diff --git a/Engine/Source/Editor/UnrealEd/Public/SkelImport.h b/Engine/Source/Editor/UnrealEd/Public/SkelImport.h
--- a/Engine/Source/Editor/UnrealEd/Public/SkelImport.h
+++ b/Engine/Source/Editor/UnrealEd/Public/SkelImport.h
@@ -88,6 +88,8 @@ struct ExistingSkelMeshData
 	FPerPlatformInt							MaxNumOptionalLODs;
 
 	TMap<UAssetUserData*, bool>				ExistingAssetUserData;
+
+	USkeletalMesh::FOnMeshChanged			ExistingOnMeshChanged;
 };
 
 /** 
