Edition support for RVT parameters

diff --git a/Engine/Source/Editor/MaterialEditor/Private/MaterialEditorDetailCustomization.cpp b/Engine/Source/Editor/MaterialEditor/Private/MaterialEditorDetailCustomization.cpp
--- a/Engine/Source/Editor/MaterialEditor/Private/MaterialEditorDetailCustomization.cpp
+++ b/Engine/Source/Editor/MaterialEditor/Private/MaterialEditorDetailCustomization.cpp
@@ -23,12 +23,13 @@
 #include "Modules/ModuleManager.h"
 #include "Factories/MaterialInstanceConstantFactoryNew.h"
 #include "Materials/MaterialInstanceConstant.h"
-#include "MaterialEditor/DEditorScalarParameterValue.h"
 #include "MaterialEditor/DEditorFontParameterValue.h"
-#include "MaterialEditor/DEditorTextureParameterValue.h"
-#include "MaterialEditor/DEditorVectorParameterValue.h"
+#include "MaterialEditor/DEditorRuntimeVirtualTextureParameterValue.h"
+#include "MaterialEditor/DEditorScalarParameterValue.h"
 #include "MaterialEditor/DEditorStaticSwitchParameterValue.h"
 #include "MaterialEditor/DEditorStaticComponentMaskParameterValue.h"
+#include "MaterialEditor/DEditorTextureParameterValue.h"
+#include "MaterialEditor/DEditorVectorParameterValue.h"
 #include "MaterialPropertyHelpers.h"
 #include "MaterialEditor/DEditorMaterialLayersParameterValue.h"
 #include "PropertyCustomizationHelpers.h"
@@ -185,6 +186,7 @@ void FMaterialEditorParameterDetails::CreateSingleGroupWidget(FEditorParameterGr
 		UDEditorStaticComponentMaskParameterValue* CompMaskParam = Cast<UDEditorStaticComponentMaskParameterValue>(Parameter);
 		UDEditorStaticSwitchParameterValue* SwitchParam = Cast<UDEditorStaticSwitchParameterValue>(Parameter);
 		UDEditorTextureParameterValue* TextureParam = Cast<UDEditorTextureParameterValue>(Parameter);
+		UDEditorRuntimeVirtualTextureParameterValue* RuntimeVirtualTextureParam = Cast<UDEditorRuntimeVirtualTextureParameterValue>(Parameter);
 		UDEditorVectorParameterValue* VectorParam = Cast<UDEditorVectorParameterValue>(Parameter);
 		UDEditorMaterialLayersParameterValue* LayersParam = Cast<UDEditorMaterialLayersParameterValue>(Parameter);
 
@@ -196,7 +198,7 @@ void FMaterialEditorParameterDetails::CreateSingleGroupWidget(FEditorParameterGr
 		{
 			CreateScalarAtlasPositionParameterValueWidget(Parameter, ParameterProperty, DetailGroup);
 		}
-		else if (ScalarParam || SwitchParam || TextureParam || VectorParam || FontParam)
+		else if (ScalarParam || SwitchParam || TextureParam || RuntimeVirtualTextureParam || VectorParam || FontParam)
 		{
 			if (ScalarParam && ScalarParam->SliderMax > ScalarParam->SliderMin)
 			{
diff --git a/Engine/Source/Editor/MaterialEditor/Private/MaterialEditorInstanceDetailCustomization.cpp b/Engine/Source/Editor/MaterialEditor/Private/MaterialEditorInstanceDetailCustomization.cpp
--- a/Engine/Source/Editor/MaterialEditor/Private/MaterialEditorInstanceDetailCustomization.cpp
+++ b/Engine/Source/Editor/MaterialEditor/Private/MaterialEditorInstanceDetailCustomization.cpp
@@ -13,6 +13,7 @@
 #include "Materials/MaterialInterface.h"
 #include "MaterialEditor/DEditorFontParameterValue.h"
 #include "MaterialEditor/DEditorMaterialLayersParameterValue.h"
+#include "MaterialEditor/DEditorRuntimeVirtualTextureParameterValue.h"
 #include "MaterialEditor/DEditorScalarParameterValue.h"
 #include "MaterialEditor/DEditorStaticComponentMaskParameterValue.h"
 #include "MaterialEditor/DEditorStaticSwitchParameterValue.h"
@@ -416,6 +417,7 @@ void FMaterialInstanceParameterDetails::CreateSingleGroupWidget(FEditorParameter
 		UDEditorStaticComponentMaskParameterValue* CompMaskParam = Cast<UDEditorStaticComponentMaskParameterValue>(Parameter);
 		UDEditorStaticSwitchParameterValue* SwitchParam = Cast<UDEditorStaticSwitchParameterValue>(Parameter);
 		UDEditorTextureParameterValue* TextureParam = Cast<UDEditorTextureParameterValue>(Parameter);
+		UDEditorRuntimeVirtualTextureParameterValue* RuntimeVirtualTextureParam = Cast<UDEditorRuntimeVirtualTextureParameterValue>(Parameter);
 		UDEditorVectorParameterValue* VectorParam = Cast<UDEditorVectorParameterValue>(Parameter);
 
 		if (Parameter->ParameterInfo.Association == EMaterialParameterAssociation::GlobalParameter)
@@ -428,7 +430,7 @@ void FMaterialInstanceParameterDetails::CreateSingleGroupWidget(FEditorParameter
 			{
 				CreateScalarAtlasPositionParameterValueWidget(Parameter, ParameterProperty, DetailGroup);
 			}
-			else if (ScalarParam || SwitchParam || TextureParam || VectorParam || FontParam)
+			else if (ScalarParam || SwitchParam || TextureParam || RuntimeVirtualTextureParam || VectorParam || FontParam)
 			{
 				if (ScalarParam && ScalarParam->SliderMax > ScalarParam->SliderMin)
 				{
diff --git a/Engine/Source/Editor/MaterialEditor/Private/MaterialEditorUtilities.cpp b/Engine/Source/Editor/MaterialEditor/Private/MaterialEditorUtilities.cpp
--- a/Engine/Source/Editor/MaterialEditor/Private/MaterialEditorUtilities.cpp
+++ b/Engine/Source/Editor/MaterialEditor/Private/MaterialEditorUtilities.cpp
@@ -16,6 +16,7 @@
 #include "Materials/MaterialExpressionComment.h"
 #include "Materials/MaterialExpressionTextureSample.h"
 #include "Materials/MaterialExpressionTextureSampleParameter.h"
+#include "Materials/MaterialExpressionRuntimeVirtualTextureSampleParameter.h"
 #include "Materials/MaterialExpressionFontSampleParameter.h"
 #include "Materials/MaterialExpressionMaterialFunctionCall.h"
 #include "Materials/MaterialExpressionScalarParameter.h"
@@ -407,6 +408,7 @@ void FMaterialEditorUtilities::GetVisibleMaterialParametersFromExpression(
 	// If it's a material parameter it must be visible so add it to the list
 	UMaterialExpressionParameter* Param = Cast<UMaterialExpressionParameter>( MaterialExpressionKey.Expression );
 	UMaterialExpressionTextureSampleParameter* TexParam = Cast<UMaterialExpressionTextureSampleParameter>( MaterialExpressionKey.Expression );
+	UMaterialExpressionRuntimeVirtualTextureSampleParameter* RuntimeVirtualTexParam = Cast<UMaterialExpressionRuntimeVirtualTextureSampleParameter>(MaterialExpressionKey.Expression);
 	UMaterialExpressionFontSampleParameter* FontParam = Cast<UMaterialExpressionFontSampleParameter>( MaterialExpressionKey.Expression );
 
 	if (Param)
@@ -417,12 +419,16 @@ void FMaterialEditorUtilities::GetVisibleMaterialParametersFromExpression(
 	{
 		ParameterInfo.Name = TexParam->ParameterName;
 	}
+	else if (RuntimeVirtualTexParam)
+	{
+		ParameterInfo.Name = RuntimeVirtualTexParam->ParameterName;
+	}
 	else if (FontParam)
 	{
 		ParameterInfo.Name = FontParam->ParameterName;
 	}
 		
-	if (Param || TexParam || FontParam)
+	if (Param || TexParam || FontParam || RuntimeVirtualTexParam)
 	{
 		VisibleExpressions.AddUnique(ParameterInfo);
 	}
diff --git a/Engine/Source/Editor/MaterialEditor/Private/MaterialInstanceEditor.cpp b/Engine/Source/Editor/MaterialEditor/Private/MaterialInstanceEditor.cpp
--- a/Engine/Source/Editor/MaterialEditor/Private/MaterialInstanceEditor.cpp
+++ b/Engine/Source/Editor/MaterialEditor/Private/MaterialInstanceEditor.cpp
@@ -11,6 +11,7 @@
 #include "EditorStyleSet.h"
 #include "Styling/CoreStyle.h"
 #include "MaterialEditor/DEditorTextureParameterValue.h"
+#include "MaterialEditor/DEditorRuntimeVirtualTextureParameterValue.h"
 #include "Materials/Material.h"
 #include "MaterialEditor/MaterialEditorInstanceConstant.h"
 #include "ThumbnailRendering/SceneThumbnailInfoWithPrimitive.h"
@@ -22,6 +23,7 @@
 
 #include "Materials/MaterialExpressionTextureBase.h"
 #include "Materials/MaterialExpressionTextureSampleParameter.h"
+#include "Materials/MaterialExpressionRuntimeVirtualTextureSampleParameter.h"
 
 #include "MaterialEditor.h"
 #include "MaterialEditorActions.h"
@@ -44,6 +46,7 @@
 #include "MaterialEditingLibrary.h"
 #include "Widgets/Layout/SScrollBox.h"
 #include "DebugViewModeHelpers.h"
+#include "VT/RuntimeVirtualTexture.h"
 
 #define LOCTEXT_NAMESPACE "MaterialInstanceEditor"
 
@@ -488,6 +491,7 @@ void FMaterialInstanceEditor::ReInitMaterialFunctionProxies()
 		TArray<FScalarParameterValue> ScalarParameterValues = FunctionInstanceProxy->ScalarParameterValues;
 		TArray<FVectorParameterValue> VectorParameterValues = FunctionInstanceProxy->VectorParameterValues;
 		TArray<FTextureParameterValue> TextureParameterValues = FunctionInstanceProxy->TextureParameterValues;
+		TArray<FRuntimeVirtualTextureParameterValue> RuntimeVirtualTextureParameterValues = FunctionInstanceProxy->RuntimeVirtualTextureParameterValues;
 		TArray<FFontParameterValue> FontParameterValues = FunctionInstanceProxy->FontParameterValues;
 
 		const FStaticParameterSet& OldStaticParameters = FunctionInstanceProxy->GetStaticParameters();
@@ -539,6 +543,18 @@ void FMaterialInstanceEditor::ReInitMaterialFunctionProxies()
 			}
 		}
 
+		FunctionInstanceProxy->GetAllRuntimeVirtualTextureParameterInfo(OutParameterInfo, Guids);
+		FunctionInstanceProxy->RuntimeVirtualTextureParameterValues.Empty();
+		for (FRuntimeVirtualTextureParameterValue& RuntimeVirtualTextureParameter : RuntimeVirtualTextureParameterValues)
+		{
+			int32 Index = Guids.Find(RuntimeVirtualTextureParameter.ExpressionGUID);
+			if (Index != INDEX_NONE)
+			{
+				FunctionInstanceProxy->RuntimeVirtualTextureParameterValues.Add(RuntimeVirtualTextureParameter);
+				FunctionInstanceProxy->RuntimeVirtualTextureParameterValues.Last().ParameterInfo = OutParameterInfo[Index];
+			}
+		}
+
 		FunctionInstanceProxy->GetAllFontParameterInfo(OutParameterInfo, Guids);
 		FunctionInstanceProxy->FontParameterValues.Empty();
 		for (FFontParameterValue& FontParameter : FontParameterValues)
@@ -1256,6 +1272,8 @@ void FMaterialInstanceEditor::DrawSamplerWarningStrings(FCanvas* Canvas, int32&
 			const int32 SpacingBetweenLines = 13;
 			UEnum* SamplerTypeEnum = StaticEnum<EMaterialSamplerType>();
 			check( SamplerTypeEnum );
+			UEnum* MaterialTypeEnum = StaticEnum<ERuntimeVirtualTextureMaterialType>();
+			check(MaterialTypeEnum);
 
 			const int32 GroupCount = MaterialEditorInstance->ParameterGroups.Num();
 			for ( int32 GroupIndex = 0; GroupIndex < GroupCount; ++GroupIndex )
@@ -1323,6 +1341,34 @@ void FMaterialInstanceEditor::DrawSamplerWarningStrings(FCanvas* Canvas, int32&
 							}
 						}
 					}
+
+					UDEditorRuntimeVirtualTextureParameterValue* RuntimeVirtualTextureParameterValue = Cast<UDEditorRuntimeVirtualTextureParameterValue>(Group.Parameters[ParameterIndex]);
+					if (RuntimeVirtualTextureParameterValue && RuntimeVirtualTextureParameterValue->ExpressionId.IsValid())
+					{
+						URuntimeVirtualTexture* RuntimeVirtualTexture = NULL;
+						MaterialEditorInstance->SourceInstance->GetRuntimeVirtualTextureParameterValue(RuntimeVirtualTextureParameterValue->ParameterInfo, RuntimeVirtualTexture);
+						if (RuntimeVirtualTexture)
+						{
+							UMaterialExpressionRuntimeVirtualTextureSampleParameter* Expression = BaseMaterial->FindExpressionByGUID<UMaterialExpressionRuntimeVirtualTextureSampleParameter>(RuntimeVirtualTextureParameterValue->ExpressionId);
+							if (Expression->MaterialType != RuntimeVirtualTexture->GetMaterialType())
+							{
+								FString BaseMaterialTypeDisplayName = MaterialTypeEnum->GetDisplayNameTextByValue((int64)(Expression->MaterialType)).ToString();
+								FString OverrideMaterialTypeDisplayName = MaterialTypeEnum->GetDisplayNameTextByValue((int64)(RuntimeVirtualTexture->GetMaterialType())).ToString();
+
+								Canvas->DrawShadowedString(
+									5, DrawPositionY,
+									*FString::Printf(TEXT("Warning: '%s' interprets the virtual texture as '%s' not '%s'"),
+										*RuntimeVirtualTextureParameterValue->ParameterInfo.Name.ToString(),
+										*BaseMaterialTypeDisplayName,
+										*OverrideMaterialTypeDisplayName,
+										*RuntimeVirtualTexture->GetPathName()),
+									FontToUse,
+									FLinearColor(1, 1, 0));
+
+								DrawPositionY += SpacingBetweenLines;
+							}
+						}
+					}
 				}
 			}
 		}
diff --git a/Engine/Source/Editor/MaterialEditor/Private/MaterialPropertyHelpers.cpp b/Engine/Source/Editor/MaterialEditor/Private/MaterialPropertyHelpers.cpp
--- a/Engine/Source/Editor/MaterialEditor/Private/MaterialPropertyHelpers.cpp
+++ b/Engine/Source/Editor/MaterialEditor/Private/MaterialPropertyHelpers.cpp
@@ -14,6 +14,7 @@
 #include "Materials/MaterialInterface.h"
 #include "MaterialEditor/DEditorFontParameterValue.h"
 #include "MaterialEditor/DEditorMaterialLayersParameterValue.h"
+#include "MaterialEditor/DEditorRuntimeVirtualTextureParameterValue.h"
 #include "MaterialEditor/DEditorScalarParameterValue.h"
 #include "MaterialEditor/DEditorStaticComponentMaskParameterValue.h"
 #include "MaterialEditor/DEditorStaticSwitchParameterValue.h"
@@ -328,6 +329,17 @@ void FMaterialPropertyHelpers::CopyMaterialToInstance(UMaterialInstanceConstant*
 							continue;
 						}
 					}
+
+					UDEditorRuntimeVirtualTextureParameterValue* RuntimeVirtualTextureParameterValue = Cast<UDEditorRuntimeVirtualTextureParameterValue>(Group.Parameters[ParameterIdx]);
+					if (RuntimeVirtualTextureParameterValue)
+					{
+						if (RuntimeVirtualTextureParameterValue->bOverride)
+						{
+							ChildInstance->SetRuntimeVirtualTextureParameterValueEditorOnly(RuntimeVirtualTextureParameterValue->ParameterInfo, RuntimeVirtualTextureParameterValue->ParameterValue);
+							continue;
+						}
+					}
+
 					UDEditorVectorParameterValue* VectorParameterValue = Cast<UDEditorVectorParameterValue>(Group.Parameters[ParameterIdx]);
 					if (VectorParameterValue)
 					{
@@ -572,6 +584,7 @@ FReply FMaterialPropertyHelpers::OnClickedSaveNewFunctionInstance(class UMateria
 					ChildInstance->ScalarParameterValues = EditedInstance->ScalarParameterValues;
 					ChildInstance->VectorParameterValues = EditedInstance->VectorParameterValues;
 					ChildInstance->TextureParameterValues = EditedInstance->TextureParameterValues;
+					ChildInstance->RuntimeVirtualTextureParameterValues = EditedInstance->RuntimeVirtualTextureParameterValues; 
 					ChildInstance->FontParameterValues = EditedInstance->FontParameterValues;
 
 					const FStaticParameterSet& StaticParameters = EditedInstance->GetStaticParameters();
@@ -662,6 +675,7 @@ FReply FMaterialPropertyHelpers::OnClickedSaveNewLayerInstance(class UMaterialFu
 					ChildInstance->ScalarParameterValues = EditedInstance->ScalarParameterValues;
 					ChildInstance->VectorParameterValues = EditedInstance->VectorParameterValues;
 					ChildInstance->TextureParameterValues = EditedInstance->TextureParameterValues;
+					ChildInstance->RuntimeVirtualTextureParameterValues = EditedInstance->RuntimeVirtualTextureParameterValues;
 					ChildInstance->FontParameterValues = EditedInstance->FontParameterValues;
 
 					const FStaticParameterSet& StaticParameters = EditedInstance->GetStaticParameters();
@@ -771,6 +785,7 @@ void FMaterialPropertyHelpers::ResetToDefault(TSharedPtr<IPropertyHandle> Proper
 	UDEditorScalarParameterValue* ScalarParam = Cast<UDEditorScalarParameterValue>(Parameter);
 	UDEditorVectorParameterValue* VectorParam = Cast<UDEditorVectorParameterValue>(Parameter);
 	UDEditorTextureParameterValue* TextureParam = Cast<UDEditorTextureParameterValue>(Parameter);
+	UDEditorRuntimeVirtualTextureParameterValue* RuntimeVirtualTextureParam = Cast<UDEditorRuntimeVirtualTextureParameterValue>(Parameter);
 	UDEditorFontParameterValue* FontParam = Cast<UDEditorFontParameterValue>(Parameter);
 	UDEditorStaticSwitchParameterValue* SwitchParam = Cast<UDEditorStaticSwitchParameterValue>(Parameter);
 	UDEditorStaticComponentMaskParameterValue* CompMaskParam = Cast<UDEditorStaticComponentMaskParameterValue>(Parameter);
@@ -803,6 +818,15 @@ void FMaterialPropertyHelpers::ResetToDefault(TSharedPtr<IPropertyHandle> Proper
 			MaterialEditorInstance->CopyToSourceInstance();
 		}
 	}
+	else if (RuntimeVirtualTextureParam)
+	{
+		URuntimeVirtualTexture* OutValue;
+		if (MaterialEditorInstance->SourceInstance->GetRuntimeVirtualTextureParameterDefaultValue(ParameterInfo, OutValue))
+		{
+			RuntimeVirtualTextureParam->ParameterValue = OutValue;
+			MaterialEditorInstance->CopyToSourceInstance();
+		}
+	}
 	else if (FontParam)
 	{
 		UFont* OutFontValue;
@@ -953,6 +977,7 @@ bool FMaterialPropertyHelpers::ShouldShowResetToDefault(TSharedPtr<IPropertyHand
 	UDEditorStaticComponentMaskParameterValue* CompMaskParam = Cast<UDEditorStaticComponentMaskParameterValue>(InParameter);
 	UDEditorStaticSwitchParameterValue* SwitchParam = Cast<UDEditorStaticSwitchParameterValue>(InParameter);
 	UDEditorTextureParameterValue* TextureParam = Cast<UDEditorTextureParameterValue>(InParameter);
+	UDEditorRuntimeVirtualTextureParameterValue* RuntimeVirtualTextureParam = Cast<UDEditorRuntimeVirtualTextureParameterValue>(InParameter);
 	UDEditorVectorParameterValue* VectorParam = Cast<UDEditorVectorParameterValue>(InParameter);
 
 	if (ScalarParam)
@@ -990,6 +1015,17 @@ bool FMaterialPropertyHelpers::ShouldShowResetToDefault(TSharedPtr<IPropertyHand
 			}
 		}
 	}
+	else if (RuntimeVirtualTextureParam)
+	{
+		URuntimeVirtualTexture* OutValue;
+		if (MaterialEditorInstance->SourceInstance->GetRuntimeVirtualTextureParameterDefaultValue(ParameterInfo, OutValue))
+		{
+			if (RuntimeVirtualTextureParam->ParameterValue != OutValue)
+			{
+				return true;
+			}
+		}
+	}
 	else if (VectorParam)
 	{
 		FLinearColor OutValue;
diff --git a/Engine/Source/Editor/MaterialEditor/Private/SMaterialLayersFunctionsTree.cpp b/Engine/Source/Editor/MaterialEditor/Private/SMaterialLayersFunctionsTree.cpp
--- a/Engine/Source/Editor/MaterialEditor/Private/SMaterialLayersFunctionsTree.cpp
+++ b/Engine/Source/Editor/MaterialEditor/Private/SMaterialLayersFunctionsTree.cpp
@@ -3,6 +3,7 @@
 #include "SMaterialLayersFunctionsTree.h"
 #include "MaterialEditor/DEditorFontParameterValue.h"
 #include "MaterialEditor/DEditorMaterialLayersParameterValue.h"
+#include "MaterialEditor/DEditorRuntimeVirtualTextureParameterValue.h"
 #include "MaterialEditor/DEditorScalarParameterValue.h"
 #include "MaterialEditor/DEditorStaticComponentMaskParameterValue.h"
 #include "MaterialEditor/DEditorStaticSwitchParameterValue.h"
@@ -242,13 +243,18 @@ FReply SMaterialLayersFunctionsInstanceTreeItem::OnLayerDrop(const FDragDropEven
 				Param->ParameterInfo.Index = GetNewParamIndex(Param->ParameterInfo.Association, Param->ParameterInfo.Index, OriginalIndex, NewIndex, OriginalBlendIndex, NewBlendIndex);
 			}
 
+			for (int32 ParamIt = 0; ParamIt < Tree->MaterialEditorInstance->SourceInstance->RuntimeVirtualTextureParameterValues.Num(); ParamIt++)
+			{
+				FRuntimeVirtualTextureParameterValue* Param = &Tree->MaterialEditorInstance->SourceInstance->RuntimeVirtualTextureParameterValues[ParamIt];
+				Param->ParameterInfo.Index = GetNewParamIndex(Param->ParameterInfo.Association, Param->ParameterInfo.Index, OriginalIndex, NewIndex, OriginalBlendIndex, NewBlendIndex);
+			}
+
 			for (int32 ParamIt = 0; ParamIt < Tree->MaterialEditorInstance->SourceInstance->FontParameterValues.Num(); ParamIt++)
 			{
 				FFontParameterValue* Param = &Tree->MaterialEditorInstance->SourceInstance->FontParameterValues[ParamIt];
 				Param->ParameterInfo.Index = GetNewParamIndex(Param->ParameterInfo.Association, Param->ParameterInfo.Index, OriginalIndex, NewIndex, OriginalBlendIndex, NewBlendIndex);
 			}
 
-
 			if (NewIndex > OriginalIndex)
 			{
 				NewIndex += 1;
diff --git a/Engine/Source/Editor/UnrealEd/Classes/MaterialEditor/DEditorRuntimeVirtualTextureParameterValue.h b/Engine/Source/Editor/UnrealEd/Classes/MaterialEditor/DEditorRuntimeVirtualTextureParameterValue.h
new file mode 100644
--- /dev/null
+++ b/Engine/Source/Editor/UnrealEd/Classes/MaterialEditor/DEditorRuntimeVirtualTextureParameterValue.h
@@ -0,0 +1,18 @@
+// Copyright 1998-2019 Epic Games, Inc. All Rights Reserved.
+
+#pragma once
+
+#include "CoreMinimal.h"
+#include "UObject/ObjectMacros.h"
+#include "MaterialEditor/DEditorParameterValue.h"
+
+#include "DEditorRuntimeVirtualTextureParameterValue.generated.h"
+
+UCLASS(hidecategories=Object, collapsecategories)
+class UNREALED_API UDEditorRuntimeVirtualTextureParameterValue : public UDEditorParameterValue
+{
+	GENERATED_UCLASS_BODY()
+
+	UPROPERTY(EditAnywhere, Category=DEditorTextureParameterValue)
+	class URuntimeVirtualTexture* ParameterValue;
+};
diff --git a/Engine/Source/Editor/UnrealEd/Private/EditorParameters.cpp b/Engine/Source/Editor/UnrealEd/Private/EditorParameters.cpp
--- a/Engine/Source/Editor/UnrealEd/Private/EditorParameters.cpp
+++ b/Engine/Source/Editor/UnrealEd/Private/EditorParameters.cpp
@@ -4,6 +4,7 @@
 #include "MaterialEditor/DEditorParameterValue.h"
 #include "MaterialEditor/DEditorFontParameterValue.h"
 #include "MaterialEditor/DEditorMaterialLayersParameterValue.h"
+#include "MaterialEditor/DEditorRuntimeVirtualTextureParameterValue.h"
 #include "MaterialEditor/DEditorScalarParameterValue.h"
 #include "MaterialEditor/DEditorStaticComponentMaskParameterValue.h"
 #include "MaterialEditor/DEditorStaticSwitchParameterValue.h"
@@ -49,3 +50,8 @@ UDEditorVectorParameterValue::UDEditorVectorParameterValue(const FObjectInitiali
 	: Super(ObjectInitializer)
 {
 }
+
+UDEditorRuntimeVirtualTextureParameterValue::UDEditorRuntimeVirtualTextureParameterValue(const FObjectInitializer& ObjectInitializer)
+	: Super(ObjectInitializer)
+{
+}
diff --git a/Engine/Source/Editor/UnrealEd/Private/PreviewMaterial.cpp b/Engine/Source/Editor/UnrealEd/Private/PreviewMaterial.cpp
--- a/Engine/Source/Editor/UnrealEd/Private/PreviewMaterial.cpp
+++ b/Engine/Source/Editor/UnrealEd/Private/PreviewMaterial.cpp
@@ -5,6 +5,7 @@
 #include "MaterialEditor/DEditorParameterValue.h"
 #include "MaterialEditor/DEditorFontParameterValue.h"
 #include "MaterialEditor/DEditorMaterialLayersParameterValue.h"
+#include "MaterialEditor/DEditorRuntimeVirtualTextureParameterValue.h"
 #include "MaterialEditor/DEditorScalarParameterValue.h"
 #include "MaterialEditor/DEditorStaticComponentMaskParameterValue.h"
 #include "MaterialEditor/DEditorStaticSwitchParameterValue.h"
@@ -417,6 +418,34 @@ void UMaterialEditorPreviewParameters::RegenerateArrays()
 			AssignParameterToGroup(ParentMaterial, &ParameterValue);
 		}
 
+		// Runtime Virtual Texture Parameters.
+		ParentMaterial->GetAllTextureParameterInfo(ParameterInfo, Guids);
+		for (int32 ParameterIdx = 0; ParameterIdx < ParameterInfo.Num(); ParameterIdx++)
+		{
+			UDEditorRuntimeVirtualTextureParameterValue& ParameterValue = *(NewObject<UDEditorRuntimeVirtualTextureParameterValue>());
+			FName ParameterName = ParameterInfo[ParameterIdx].Name;
+			URuntimeVirtualTexture* Value;
+			int32 SortPriority;
+
+			ParameterValue.bOverride = true;
+			ParameterValue.ParameterInfo = ParameterInfo[ParameterIdx];
+			ParameterValue.ExpressionId = Guids[ParameterIdx];
+
+			if (PreviewMaterial->GetRuntimeVirtualTextureParameterValue(ParameterValue.ParameterInfo, Value))
+			{
+				ParameterValue.ParameterValue = Value;
+			}
+			if (ParentMaterial->GetParameterSortPriority(ParameterName, SortPriority))
+			{
+				ParameterValue.SortPriority = SortPriority;
+			}
+			else
+			{
+				ParameterValue.SortPriority = 0;
+			}
+			AssignParameterToGroup(ParentMaterial, &ParameterValue);
+		}
+
 		// Font Parameters.
 		ParentMaterial->GetAllFontParameterInfo(ParameterInfo, Guids);
 		for (int32 ParameterIdx = 0; ParameterIdx < ParameterInfo.Num(); ParameterIdx++)
@@ -668,13 +697,18 @@ void UMaterialEditorPreviewParameters::CopyToSourceInstance()
 					PreviewMaterial->SetFontParameterValueEditorOnly(FontParameterValue->ParameterInfo.Name, FontParameterValue->ParameterValue.FontValue, FontParameterValue->ParameterValue.FontPage);
 					continue;
 				}
-
 				UDEditorTextureParameterValue* TextureParameterValue = Cast<UDEditorTextureParameterValue>(Group.Parameters[ParameterIdx]);
 				if (TextureParameterValue)
 				{
 					PreviewMaterial->SetTextureParameterValueEditorOnly(TextureParameterValue->ParameterInfo.Name, TextureParameterValue->ParameterValue);
 					continue;
 				}
+				UDEditorRuntimeVirtualTextureParameterValue* RuntimeVirtualTextureParameterValue = Cast<UDEditorRuntimeVirtualTextureParameterValue>(Group.Parameters[ParameterIdx]);
+				if (RuntimeVirtualTextureParameterValue)
+				{
+					PreviewMaterial->SetRuntimeVirtualTextureParameterValueEditorOnly(RuntimeVirtualTextureParameterValue->ParameterInfo.Name, RuntimeVirtualTextureParameterValue->ParameterValue);
+					continue;
+				}
 				UDEditorVectorParameterValue* VectorParameterValue = Cast<UDEditorVectorParameterValue>(Group.Parameters[ParameterIdx]);
 				if (VectorParameterValue)
 				{
@@ -735,13 +769,18 @@ void UMaterialEditorPreviewParameters::ApplySourceFunctionChanges()
 					OriginalFunction->SetFontParameterValueEditorOnly(FontParameterValue->ParameterInfo.Name, FontParameterValue->ParameterValue.FontValue, FontParameterValue->ParameterValue.FontPage);
 					continue;
 				}
-
 				UDEditorTextureParameterValue * TextureParameterValue = Cast<UDEditorTextureParameterValue>(Group.Parameters[ParameterIdx]);
 				if (TextureParameterValue)
 				{
 					OriginalFunction->SetTextureParameterValueEditorOnly(TextureParameterValue->ParameterInfo.Name, TextureParameterValue->ParameterValue);
 					continue;
 				}
+				UDEditorRuntimeVirtualTextureParameterValue * RuntimeVirtualTextureParameterValue = Cast<UDEditorRuntimeVirtualTextureParameterValue>(Group.Parameters[ParameterIdx]);
+				if (RuntimeVirtualTextureParameterValue)
+				{
+					OriginalFunction->SetRuntimeVirtualTextureParameterValueEditorOnly(RuntimeVirtualTextureParameterValue->ParameterInfo.Name, RuntimeVirtualTextureParameterValue->ParameterValue);
+					continue;
+				}
 				UDEditorVectorParameterValue * VectorParameterValue = Cast<UDEditorVectorParameterValue>(Group.Parameters[ParameterIdx]);
 				if (VectorParameterValue)
 				{
@@ -873,6 +912,10 @@ void  UMaterialEditorInstanceConstant::AssignParameterToGroup(UMaterial* ParentM
 			{
 				ParameterGroupName = TEXT("Texture Parameter Values");
 			}
+			else if (Cast<UDEditorRuntimeVirtualTextureParameterValue>(ParameterValue))
+			{
+				ParameterGroupName = TEXT("Texture Parameter Values");
+			}
 			else if (Cast<UDEditorScalarParameterValue>(ParameterValue))
 			{
 				ParameterGroupName = TEXT("Scalar Parameter Values");
@@ -1068,6 +1111,43 @@ void UMaterialEditorInstanceConstant::RegenerateArrays()
 			AssignParameterToGroup(ParentMaterial, &ParameterValue);
 		}
 
+		// Runtime Virtual Texture Parameters.
+		SourceInstance->GetAllRuntimeVirtualTextureParameterInfo(OutParameterInfo, Guids);
+		for (int32 ParameterIdx = 0; ParameterIdx < OutParameterInfo.Num(); ParameterIdx++)
+		{
+			UDEditorRuntimeVirtualTextureParameterValue& ParameterValue = *(NewObject<UDEditorRuntimeVirtualTextureParameterValue>());
+			const FMaterialParameterInfo& ParameterInfo = OutParameterInfo[ParameterIdx];
+
+			ParameterValue.bOverride = false;
+			ParameterValue.ParameterInfo = ParameterInfo;
+			ParameterValue.ExpressionId = Guids[ParameterIdx];
+
+			ParameterValue.ParameterValue = nullptr;
+			SourceInstance->GetRuntimeVirtualTextureParameterValue(ParameterInfo, ParameterValue.ParameterValue);
+
+			// @todo: This is kind of slow, maybe store these in a map for lookup?
+			// See if this keyname exists in the source instance.
+			for (int32 TextureParameterIdx = 0; TextureParameterIdx < SourceInstance->RuntimeVirtualTextureParameterValues.Num(); TextureParameterIdx++)
+			{
+				FRuntimeVirtualTextureParameterValue& SourceParam = SourceInstance->RuntimeVirtualTextureParameterValues[TextureParameterIdx];
+				if (ParameterInfo == SourceParam.ParameterInfo)
+				{
+					ParameterValue.bOverride = true;
+					ParameterValue.ParameterValue = SourceParam.ParameterValue;
+				}
+				if (ParameterInfo.Name.IsEqual(SourceParam.ParameterInfo.Name) && ParameterInfo.Association == SourceParam.ParameterInfo.Association && ParameterInfo.Index == SourceParam.ParameterInfo.Index)
+				{
+					ParameterValue.bOverride = true;
+					ParameterValue.ParameterValue = SourceParam.ParameterValue;
+				}
+			}
+
+			ParameterValue.SortPriority = 0;
+			SourceInstance->GetParameterSortPriority(ParameterInfo, ParameterValue.SortPriority);
+
+			AssignParameterToGroup(ParentMaterial, &ParameterValue);
+		}
+
 		// Font Parameters.
 		SourceInstance->GetAllFontParameterInfo(OutParameterInfo, Guids);
 		for(int32 ParameterIdx=0; ParameterIdx<OutParameterInfo.Num(); ParameterIdx++)
@@ -1262,6 +1342,7 @@ void UMaterialEditorInstanceConstant::ResetOverrides(int32 Index, EMaterialParam
 				UDEditorScalarParameterValue* ScalarParameterValue = Cast<UDEditorScalarParameterValue>(Parameter);
 				UDEditorVectorParameterValue* VectorParameterValue = Cast<UDEditorVectorParameterValue>(Parameter);
 				UDEditorTextureParameterValue* TextureParameterValue = Cast<UDEditorTextureParameterValue>(Parameter);
+				UDEditorRuntimeVirtualTextureParameterValue* RuntimeVirtualTextureParameterValue = Cast<UDEditorRuntimeVirtualTextureParameterValue>(Parameter);
 				UDEditorFontParameterValue* FontParameterValue = Cast<UDEditorFontParameterValue>(Parameter);
 				UDEditorStaticSwitchParameterValue* StaticSwitchParameterValue = Cast<UDEditorStaticSwitchParameterValue>(Parameter);
 				UDEditorStaticComponentMaskParameterValue* StaticMaskParameterValue = Cast<UDEditorStaticComponentMaskParameterValue>(Parameter);
@@ -1280,6 +1361,11 @@ void UMaterialEditorInstanceConstant::ResetOverrides(int32 Index, EMaterialParam
 					UTexture* Value;
 					Parameter->bOverride = SourceInstance->GetTextureParameterValue(Parameter->ParameterInfo, Value, true);
 				}
+				if (RuntimeVirtualTextureParameterValue)
+				{
+					URuntimeVirtualTexture* Value;
+					Parameter->bOverride = SourceInstance->GetRuntimeVirtualTextureParameterValue(Parameter->ParameterInfo, Value, true);
+				}
 				if (FontParameterValue)
 				{
 					UFont* FontValue;
@@ -1337,6 +1423,7 @@ void UMaterialEditorInstanceConstant::CopyToSourceInstance(const bool bForceStat
 				UDEditorScalarParameterValue* ScalarParameterValue = Cast<UDEditorScalarParameterValue>(Group.Parameters[ParameterIdx]);
 				UDEditorVectorParameterValue* VectorParameterValue = Cast<UDEditorVectorParameterValue>(Group.Parameters[ParameterIdx]);
 				UDEditorTextureParameterValue* TextureParameterValue = Cast<UDEditorTextureParameterValue>(Group.Parameters[ParameterIdx]);
+				UDEditorRuntimeVirtualTextureParameterValue* RuntimeVirtualTextureParameterValue = Cast<UDEditorRuntimeVirtualTextureParameterValue>(Group.Parameters[ParameterIdx]);
 				UDEditorFontParameterValue* FontParameterValue = Cast<UDEditorFontParameterValue>(Group.Parameters[ParameterIdx]);
 				if (ScalarParameterValue && ScalarParameterValue->bOverride)
 				{
@@ -1359,6 +1446,10 @@ void UMaterialEditorInstanceConstant::CopyToSourceInstance(const bool bForceStat
 				{
 					SourceInstance->SetTextureParameterValueEditorOnly(TextureParameterValue->ParameterInfo, TextureParameterValue->ParameterValue);
 				}
+				else if (RuntimeVirtualTextureParameterValue && RuntimeVirtualTextureParameterValue->bOverride)
+				{
+					SourceInstance->SetRuntimeVirtualTextureParameterValueEditorOnly(RuntimeVirtualTextureParameterValue->ParameterInfo, RuntimeVirtualTextureParameterValue->ParameterValue);
+				}
 				else if (FontParameterValue && FontParameterValue->bOverride)
 				{
 					SourceInstance->SetFontParameterValueEditorOnly(FontParameterValue->ParameterInfo, FontParameterValue->ParameterValue.FontValue, FontParameterValue->ParameterValue.FontPage);
@@ -1414,6 +1505,7 @@ void UMaterialEditorInstanceConstant::ApplySourceFunctionChanges()
 		SourceFunction->ScalarParameterValues = SourceInstance->ScalarParameterValues;
 		SourceFunction->VectorParameterValues = SourceInstance->VectorParameterValues;
 		SourceFunction->TextureParameterValues = SourceInstance->TextureParameterValues;
+		SourceFunction->RuntimeVirtualTextureParameterValues = SourceInstance->RuntimeVirtualTextureParameterValues;
 		SourceFunction->FontParameterValues = SourceInstance->FontParameterValues;
 
 		const FStaticParameterSet& StaticParameters = SourceInstance->GetStaticParameters();
