Fix static analysis issue

diff --git a/Engine/Source/Editor/UnrealEd/Private/MaterialGraphNode_Base.cpp b/Engine/Source/Editor/UnrealEd/Private/MaterialGraphNode_Base.cpp
--- a/Engine/Source/Editor/UnrealEd/Private/MaterialGraphNode_Base.cpp
+++ b/Engine/Source/Editor/UnrealEd/Private/MaterialGraphNode_Base.cpp
@@ -114,7 +114,6 @@ void UMaterialGraphNode_Base::ReplaceNode(UMaterialGraphNode_Base* OldNode)
 	{
 		// Try to find an equivalent output in this node
 		int32 FoundPinIndex = -1;
-		if (FoundPinIndex == -1)
 		{
 			// First check names
 			for (int32 NewPinIndex = 0; NewPinIndex < NewOutputPins.Num(); NewPinIndex++)
