Editor - Updating the UIMaterial thumbnail renderer to render a checkboard behind transparent UI Materials.

Dave.Belanger

diff --git a/Engine/Source/Editor/UnrealEd/Private/ThumbnailRendering/MaterialInstanceThumbnailRenderer.cpp b/Engine/Source/Editor/UnrealEd/Private/ThumbnailRendering/MaterialInstanceThumbnailRenderer.cpp
--- a/Engine/Source/Editor/UnrealEd/Private/ThumbnailRendering/MaterialInstanceThumbnailRenderer.cpp
+++ b/Engine/Source/Editor/UnrealEd/Private/ThumbnailRendering/MaterialInstanceThumbnailRenderer.cpp
@@ -10,6 +10,7 @@
 #include "Slate/WidgetRenderer.h"
 #include "Widgets/Images/SImage.h"
 #include "SlateMaterialBrush.h"
+#include "ThumbnailRendering/ThumbnailManager.h"
 
 UMaterialInstanceThumbnailRenderer::UMaterialInstanceThumbnailRenderer(const FObjectInitializer& ObjectInitializer)
 	: Super(ObjectInitializer)
@@ -42,10 +43,32 @@ void UMaterialInstanceThumbnailRenderer::Draw(UObject* Object, int32 X, int32 Y,
 		
 			UIMaterialBrush->SetMaterial(MatInst);
 
-			TSharedPtr<SImage> UIMaterialImage = SNew(SImage).Image(UIMaterialBrush);
+			UTexture2D* CheckerboardTexture = UThumbnailManager::Get().CheckerboardTexture;
+			
+			FSlateBrush CheckerboardBrush;
+			CheckerboardBrush.SetResourceObject(CheckerboardTexture);
+			CheckerboardBrush.ImageSize = FVector2D(CheckerboardTexture->GetSizeX(), CheckerboardTexture->GetSizeY());
+			CheckerboardBrush.Tiling = ESlateBrushTileType::Both;
+
+			TSharedRef<SWidget> Thumbnail =
+				SNew(SOverlay)
+
+				// Checkerboard
+				+ SOverlay::Slot()
+				[
+					SNew(SImage)
+					.Image(&CheckerboardBrush)
+				]
+
+				+ SOverlay::Slot()
+				[
+					SNew(SImage)
+					.Image(UIMaterialBrush)
+				];
+
 			const FVector2D DrawSize((float)Width, (float)Height);
 			const float DeltaTime = 0.f;
-			WidgetRenderer->DrawWidget(RenderTarget, UIMaterialImage.ToSharedRef(), DrawSize, DeltaTime);
+			WidgetRenderer->DrawWidget(RenderTarget, Thumbnail, DrawSize, DeltaTime);
 
 			UIMaterialBrush->SetMaterial(nullptr);
 		}
