Fix ICWYU error.

diff --git a/Engine/Source/Editor/UnrealEd/Private/ThumbnailRendering/MaterialInstanceThumbnailRenderer.cpp b/Engine/Source/Editor/UnrealEd/Private/ThumbnailRendering/MaterialInstanceThumbnailRenderer.cpp
--- a/Engine/Source/Editor/UnrealEd/Private/ThumbnailRendering/MaterialInstanceThumbnailRenderer.cpp
+++ b/Engine/Source/Editor/UnrealEd/Private/ThumbnailRendering/MaterialInstanceThumbnailRenderer.cpp
@@ -11,6 +11,7 @@
 #include "Widgets/Images/SImage.h"
 #include "SlateMaterialBrush.h"
 #include "ThumbnailRendering/ThumbnailManager.h"
+#include "Engine/Texture2D.h"
 
 UMaterialInstanceThumbnailRenderer::UMaterialInstanceThumbnailRenderer(const FObjectInitializer& ObjectInitializer)
 	: Super(ObjectInitializer)
