Fix GC crash on level unload if a post process volume had  a MID created from a world actor

diff --git a/Engine/Source/Editor/UnrealEd/Private/LevelEditorViewport.cpp b/Engine/Source/Editor/UnrealEd/Private/LevelEditorViewport.cpp
--- a/Engine/Source/Editor/UnrealEd/Private/LevelEditorViewport.cpp
+++ b/Engine/Source/Editor/UnrealEd/Private/LevelEditorViewport.cpp
@@ -4921,6 +4921,14 @@ void FLevelEditorViewportClient::ClearHoverFromObjects()
 void FLevelEditorViewportClient::OnEditorCleanse()
 {
 	ClearHoverFromObjects();
+
+	FSceneViewStateInterface* ViewStateInterface = ViewState.GetReference();
+	if (ViewStateInterface)
+	{
+		// The view state can reference materials from the world being cleaned up. (Example post process materials)
+		ViewStateInterface->ClearMIDPool();
+	}
+
 }
 
 void FLevelEditorViewportClient::OnPreBeginPIE(const bool bIsSimulating)
diff --git a/Engine/Source/Runtime/Engine/Public/SceneManagement.h b/Engine/Source/Runtime/Engine/Public/SceneManagement.h
--- a/Engine/Source/Runtime/Engine/Public/SceneManagement.h
+++ b/Engine/Source/Runtime/Engine/Public/SceneManagement.h
@@ -174,6 +174,11 @@ public:
 	 */
 	virtual UMaterialInstanceDynamic* GetReusableMID(class UMaterialInterface* InSource) = 0;
 
+	/**
+	 * Clears the pool of mids being referenced by this view state 
+	 */
+	virtual void ClearMIDPool() = 0;
+
 #if !(UE_BUILD_SHIPPING || UE_BUILD_TEST)
 	/** If frozen view matrices are available, set those as active on the SceneView */
 	virtual void ActivateFrozenViewMatrices(FSceneView& SceneView) = 0;
diff --git a/Engine/Source/Runtime/Renderer/Private/ScenePrivate.h b/Engine/Source/Runtime/Renderer/Private/ScenePrivate.h
--- a/Engine/Source/Runtime/Renderer/Private/ScenePrivate.h
+++ b/Engine/Source/Runtime/Renderer/Private/ScenePrivate.h
@@ -1330,6 +1330,12 @@ public:
 		return NewMID;
 	}
 
+	virtual void ClearMIDPool() override
+	{
+		check(IsInGameThread());
+		MIDPool.Empty();
+	}
+
 #if !(UE_BUILD_SHIPPING || UE_BUILD_TEST)
 	virtual void ActivateFrozenViewMatrices(FSceneView& SceneView) override
 	{
