Show "Client & Server" as a single option in the project packaging menu.

#rnx

diff --git a/Engine/Source/Editor/MainFrame/Private/Frame/MainFrameActions.cpp b/Engine/Source/Editor/MainFrame/Private/Frame/MainFrameActions.cpp
--- a/Engine/Source/Editor/MainFrame/Private/Frame/MainFrameActions.cpp
+++ b/Engine/Source/Editor/MainFrame/Private/Frame/MainFrameActions.cpp
@@ -737,23 +737,22 @@ void FMainFrameActionCallbacks::PackageProject( const FName InPlatformInfoName )
 		OptionalParams += FString::Printf(TEXT(" -NumCookersToSpawn=%d"), NumCookers); 
 	}
 
-	if (ConfigurationInfo.TargetType == EBuildTargetType::Client)
+	if (ConfigurationInfo.bClientAndServer)
 	{
-		OptionalParams += TEXT(" -client");
+		OptionalParams += FString::Printf(TEXT(" -client -clientconfig=%s -server -serverconfig=%s"), LexToString(ConfigurationInfo.Configuration), LexToString(ConfigurationInfo.Configuration));
 	}
-	else if (ConfigurationInfo.TargetType == EBuildTargetType::Server)
+	else
 	{
-		OptionalParams += TEXT(" -server");
+		OptionalParams += FString::Printf(TEXT(" -clientconfig=%s"), LexToString(ConfigurationInfo.Configuration));
 	}
 
 	FString ProjectPath = FPaths::IsProjectFilePathSet() ? FPaths::ConvertRelativePathToFull(FPaths::GetProjectFilePath()) : FPaths::RootDir() / FApp::GetProjectName() / FApp::GetProjectName() + TEXT(".uproject");
-	FString CommandLine = FString::Printf(TEXT("-ScriptsForProject=\"%s\" BuildCookRun %s%s -nop4 -project=\"%s\" -cook -stage -archive -archivedirectory=\"%s\" -package -clientconfig=%s -ue4exe=\"%s\" %s -utf8output"),
+	FString CommandLine = FString::Printf(TEXT("-ScriptsForProject=\"%s\" BuildCookRun %s%s -nop4 -project=\"%s\" -cook -stage -archive -archivedirectory=\"%s\" -package -ue4exe=\"%s\" %s -utf8output"),
 		*ProjectPath,
 		GetUATCompilationFlags(),
 		FApp::IsEngineInstalled() ? TEXT(" -installed") : TEXT(""),
 		*ProjectPath,
 		*PackagingSettings->StagingDirectory.Path,
-		LexToString(ConfigurationInfo.Configuration),
 		*FUnrealEdMisc::Get().GetExecutableForCommandlets(),
 		*OptionalParams
 	);
diff --git a/Engine/Source/Editor/MainFrame/Private/Menus/PackageProjectMenu.h b/Engine/Source/Editor/MainFrame/Private/Menus/PackageProjectMenu.h
--- a/Engine/Source/Editor/MainFrame/Private/Menus/PackageProjectMenu.h
+++ b/Engine/Source/Editor/MainFrame/Private/Menus/PackageProjectMenu.h
@@ -169,48 +169,10 @@ protected:
 	 */
 	static void MakeBuildConfigurationsMenu(FMenuBuilder& MenuBuilder)
 	{
-		// Check if the project has code
-		FProjectStatus ProjectStatus;
-		bool bHasCode = IProjectManager::Get().QueryStatusForCurrentProject(ProjectStatus) && ProjectStatus.bCodeBasedProject;
-
-		// If if does, find all the targets
-		const TArray<FTargetInfo>* Targets = nullptr;
-		if (bHasCode)
-		{
-			Targets = &(FDesktopPlatformModule::Get()->GetTargetsForCurrentProject());
-		}
-
-		// Set up all the configurations
-		for (int32 Idx = 0; Idx < PPBC_Max; Idx++)
+		TArray<EProjectPackagingBuildConfigurations> PackagingConfigurations = UProjectPackagingSettings::GetValidPackageConfigurations();
+		for(EProjectPackagingBuildConfigurations PackagingConfiguration : PackagingConfigurations)
 		{
-			EProjectPackagingBuildConfigurations PackagingConfiguration = (EProjectPackagingBuildConfigurations)Idx;
-			const UProjectPackagingSettings::FConfigurationInfo& Info = UProjectPackagingSettings::ConfigurationInfo[Idx];
-
-			// Make sure the selected configuration is supported
-			EProjectType ProjectType = bHasCode? EProjectType::Code : EProjectType::Content;
-			if(!FInstalledPlatformInfo::Get().IsValid(Info.TargetType, TOptional<FString>(), Info.Configuration, ProjectType, EInstalledPlatformState::Downloaded))
-			{
-				continue;
-			}
-
-			// Check the target type is valid
-			if(bHasCode)
-			{
-				EBuildTargetType TargetType = Info.TargetType;
-				if(!Targets->ContainsByPredicate([TargetType](const FTargetInfo& Target) -> bool { return Target.Type == TargetType; }))
-				{
-					continue;
-				}
-			}
-			else
-			{
-				if(Info.Configuration == EBuildConfiguration::DebugGame)
-				{
-					continue;
-				}
-			}
-
-			// Add the menu item
+			const UProjectPackagingSettings::FConfigurationInfo& Info = UProjectPackagingSettings::ConfigurationInfo[(int)PackagingConfiguration];
 			MenuBuilder.AddMenuEntry(
 				Info.Name, 
 				Info.ToolTip,
diff --git a/Engine/Source/Editor/UnrealEd/Classes/Settings/ProjectPackagingSettings.h b/Engine/Source/Editor/UnrealEd/Classes/Settings/ProjectPackagingSettings.h
--- a/Engine/Source/Editor/UnrealEd/Classes/Settings/ProjectPackagingSettings.h
+++ b/Engine/Source/Editor/UnrealEd/Classes/Settings/ProjectPackagingSettings.h
@@ -17,50 +17,35 @@ enum EProjectPackagingBuildConfigurations
 	/** Debug configuration. */
 	PPBC_Debug UMETA(DisplayName="Debug"),
 
-	/** Debug Client configuration. */
-	PPBC_DebugClient UMETA(DisplayName = "Debug Client"),
-
-	/** Debug Server configuration. */
-	PPBC_DebugServer UMETA(DisplayName = "Debug Server"),
+	/** Debug client & server configuration. */
+	PPBC_DebugClientAndServer UMETA(DisplayName = "Debug (Client & Server)"),
 
 	/** DebugGame configuration. */
 	PPBC_DebugGame UMETA(DisplayName="DebugGame"),
 
-	/** DebugGame Client configuration. */
-	PPBC_DebugGameClient UMETA(DisplayName = "DebugGame Client"),
-
-	/** DebugGame Server configuration. */
-	PPBC_DebugGameServer UMETA(DisplayName = "DebugGame Server"),
+	/** DebugGame client & server configuration. */
+	PPBC_DebugGameClientAndServer UMETA(DisplayName = "DebugGame (Client & Server)"),
 
 	/** Development configuration. */
 	PPBC_Development UMETA(DisplayName="Development"),
 
-	/** Development Client configuration. */
-	PPBC_DevelopmentClient UMETA(DisplayName = "Development Client"),
-
-	/** Development Server configuration. */
-	PPBC_DevelopmentServer UMETA(DisplayName = "Development Server"),
+	/** Development client & server configuration. */
+	PPBC_DevelopmentClientAndServer UMETA(DisplayName = "Development (Client & Server)"),
 
 	/** Test configuration. */
 	PPBC_Test UMETA(DisplayName="Test"),
 
-	/** Test Client configuration. */
-	PPBC_TestClient UMETA(DisplayName = "Test Client"),
-
-	/** Test Server configuration. */
-	PPBC_TestServer UMETA(DisplayName = "Test Server"),
+	/** Test client & server configuration. */
+	PPBC_TestClientAndServer UMETA(DisplayName = "Test (Client & Server)"),
 
 	/** Shipping configuration. */
 	PPBC_Shipping UMETA(DisplayName="Shipping"),
 
-	/** Shipping Client configuration. */
-	PPBC_ShippingClient UMETA(DisplayName = "Shipping Client"),
-
-	/** Shipping Server configuration. */
-	PPBC_ShippingServer UMETA(DisplayName = "Shipping Server"),
+	/** Shipping client & server configuration. */
+	PPBC_ShippingClientAndServer UMETA(DisplayName = "Shipping (Client & Server)"),
 
 	/** Number of entries in the enum. */
-	PPBC_Max
+	PPBC_MAX
 };
 
 /**
@@ -136,7 +121,7 @@ public:
 	struct FConfigurationInfo
 	{
 		EBuildConfiguration Configuration;
-		EBuildTargetType TargetType;
+		bool bClientAndServer;
 		FText Name;
 		FText ToolTip;
 	};
@@ -144,7 +129,7 @@ public:
 	/**
 	 * Static array of information about each configuration
 	 */
-	static const FConfigurationInfo ConfigurationInfo[PPBC_Max];
+	static const FConfigurationInfo ConfigurationInfo[PPBC_MAX];
 
 	/** Specifies whether to build the game executable during packaging. */
 	UPROPERTY(config, EditAnywhere, Category=Project)
@@ -466,6 +451,9 @@ public:
 	/** Determines if the specified Blueprint is already saved for exclusive nativization. */
 	bool IsBlueprintAssetInNativizationList(const class UBlueprint* InBlueprint) const { return FindBlueprintInNativizationList(InBlueprint) >= 0; }
 
+	/** Gets a list of all valid packaging configurations for the current project */
+	static TArray<EProjectPackagingBuildConfigurations> GetValidPackageConfigurations();
+
 private:
 	/** Returns the index of the specified Blueprint in the exclusive nativization list (otherwise INDEX_NONE) */
 	int32 FindBlueprintInNativizationList(const UBlueprint* InBlueprint) const;
diff --git a/Engine/Source/Editor/UnrealEd/Private/Settings/ProjectPackagingSettingsCustomization.h b/Engine/Source/Editor/UnrealEd/Private/Settings/ProjectPackagingSettingsCustomization.h
--- a/Engine/Source/Editor/UnrealEd/Private/Settings/ProjectPackagingSettingsCustomization.h
+++ b/Engine/Source/Editor/UnrealEd/Private/Settings/ProjectPackagingSettingsCustomization.h
@@ -27,6 +27,7 @@
 #include "Settings/ProjectPackagingSettings.h"
 #include "PropertyRestriction.h"
 #include "Widgets/Views/SMultipleOptionTable.h"
+#include "DesktopPlatformModule.h"
 
 #define LOCTEXT_NAMESPACE "FProjectPackagingSettingsCustomization"
 
@@ -144,40 +145,22 @@ protected:
 	 */
 	void CustomizeProjectCategory( IDetailLayoutBuilder& LayoutBuilder )
 	{
-		// Hide the DebugGame configurations for content-only games
-		TArray<FString> TargetFileNames;
-		IFileManager::Get().FindFiles(TargetFileNames, *(FPaths::GameSourceDir() / TEXT("*.target.cs")), true, false);
+		TArray<EProjectPackagingBuildConfigurations> PackagingConfigurations = UProjectPackagingSettings::GetValidPackageConfigurations();
 
-		if (TargetFileNames.Num() == 0)
-		{
-			IDetailCategoryBuilder& ProjectCategory = LayoutBuilder.EditCategory("Project");
-			{
-				TSharedRef<FPropertyRestriction> BuildConfigurationRestriction = MakeShareable(new FPropertyRestriction(LOCTEXT("ContentOnlyRestrictionReason", "The DebugGame and Client build configurations are not available in content-only projects.")));
-				const UEnum* const ProjectPackagingBuildConfigurationsEnum = StaticEnum<EProjectPackagingBuildConfigurations>();		
-				BuildConfigurationRestriction->AddDisabledValue(ProjectPackagingBuildConfigurationsEnum->GetNameStringByValue((uint8)EProjectPackagingBuildConfigurations::PPBC_DebugGame));
-				BuildConfigurationRestriction->AddDisabledValue(ProjectPackagingBuildConfigurationsEnum->GetNameStringByValue((uint8)EProjectPackagingBuildConfigurations::PPBC_DebugGameClient));
+		TSharedRef<FPropertyRestriction> BuildConfigurationRestriction = MakeShareable(new FPropertyRestriction(LOCTEXT("ConfigurationRestrictionReason", "This configuration is not valid for this project. DebugGame configurations require a code project, and client/server configurations require the appropriate targets.")));
 
-				TSharedRef<IPropertyHandle> BuildConfigurationHandle = LayoutBuilder.GetProperty("BuildConfiguration");
-				BuildConfigurationHandle->AddRestriction(BuildConfigurationRestriction);
-			}
-		}
-		else
+		const UEnum* const ProjectPackagingBuildConfigurationsEnum = StaticEnum<EProjectPackagingBuildConfigurations>();
+		for (int Idx = 0; Idx < PPBC_MAX; Idx++)
 		{
-			// Hide the Client configurations if there is no {ProjectName}Client.Target.cs
-			TArray<FString> ClientTargetFileNames;
-			IFileManager::Get().FindFiles(ClientTargetFileNames, *(FPaths::GameSourceDir() / TEXT("*client.target.cs")), true, false);
-			if (ClientTargetFileNames.Num() == 0)
+			EProjectPackagingBuildConfigurations Configuration = (EProjectPackagingBuildConfigurations)Idx;
+			if (!PackagingConfigurations.Contains(Configuration))
 			{
-				TSharedRef<FPropertyRestriction> BuildConfigurationRestriction = MakeShareable(new FPropertyRestriction(LOCTEXT("ClientRestrictionReason", "The Client build configurations require a {ProjectName}Client.Target.cs file in your Project/Source folder.")));
-				const UEnum* const ProjectPackagingBuildConfigurationsEnum = StaticEnum<EProjectPackagingBuildConfigurations>();
-				BuildConfigurationRestriction->AddDisabledValue(ProjectPackagingBuildConfigurationsEnum->GetNameStringByValue((uint8)EProjectPackagingBuildConfigurations::PPBC_DebugGameClient));
-				BuildConfigurationRestriction->AddDisabledValue(ProjectPackagingBuildConfigurationsEnum->GetNameStringByValue((uint8)EProjectPackagingBuildConfigurations::PPBC_DevelopmentClient));
-				BuildConfigurationRestriction->AddDisabledValue(ProjectPackagingBuildConfigurationsEnum->GetNameStringByValue((uint8)EProjectPackagingBuildConfigurations::PPBC_ShippingClient));
-
-				TSharedRef<IPropertyHandle> BuildConfigurationHandle = LayoutBuilder.GetProperty("BuildConfiguration");
-				BuildConfigurationHandle->AddRestriction(BuildConfigurationRestriction);
+				BuildConfigurationRestriction->AddDisabledValue(ProjectPackagingBuildConfigurationsEnum->GetNameStringByValue(Idx));
 			}
 		}
+
+		TSharedRef<IPropertyHandle> BuildConfigurationHandle = LayoutBuilder.GetProperty("BuildConfiguration");
+		BuildConfigurationHandle->AddRestriction(BuildConfigurationRestriction);
 	}
 
 	/**
diff --git a/Engine/Source/Editor/UnrealEd/Private/Settings/SettingsClasses.cpp b/Engine/Source/Editor/UnrealEd/Private/Settings/SettingsClasses.cpp
--- a/Engine/Source/Editor/UnrealEd/Private/Settings/SettingsClasses.cpp
+++ b/Engine/Source/Editor/UnrealEd/Private/Settings/SettingsClasses.cpp
@@ -36,6 +36,8 @@
 #include "Settings/SkeletalMeshEditorSettings.h"
 #include "DeviceProfiles/DeviceProfile.h"
 #include "DeviceProfiles/DeviceProfileManager.h"
+#include "DesktopPlatformModule.h"
+#include "InstalledPlatformInfo.h"
 
 #define LOCTEXT_NAMESPACE "SettingsClasses"
 
@@ -809,27 +811,22 @@ void ULevelEditorViewportSettings::PostEditChangeProperty( struct FPropertyChang
 /* UProjectPackagingSettings interface
  *****************************************************************************/
 
-const UProjectPackagingSettings::FConfigurationInfo UProjectPackagingSettings::ConfigurationInfo[PPBC_Max] = 
+const UProjectPackagingSettings::FConfigurationInfo UProjectPackagingSettings::ConfigurationInfo[PPBC_MAX] = 
 {
-	/* PPBC_Debug */             { EBuildConfiguration::Debug, EBuildTargetType::Game, LOCTEXT("DebugConfiguration", "Debug"), LOCTEXT("DebugConfigurationTooltip", "Package the game in Debug configuration") },
-	/* PPBC_DebugClient */       { EBuildConfiguration::Debug, EBuildTargetType::Client, LOCTEXT("DebugClientConfiguration", "Debug Client"), LOCTEXT("DebugClientConfigurationTooltip", "Package the client in Debug configuration") },
-	/* PPBC_DebugServer */       { EBuildConfiguration::Debug, EBuildTargetType::Server, LOCTEXT("DebugServerConfiguration", "Debug Server"), LOCTEXT("DebugServerConfigurationTooltip", "Package the server in Debug configuration") },
+	/* PPBC_Debug */                      { EBuildConfiguration::Debug, false, LOCTEXT("DebugConfiguration", "Debug"), LOCTEXT("DebugConfigurationTooltip", "Package the game in Debug configuration") },
+	/* PPBC_DebugClientAndServer */       { EBuildConfiguration::Debug, true, LOCTEXT("DebugClientConfiguration", "Debug (Client & Server)"), LOCTEXT("DebugClientConfigurationTooltip", "Package a client and server in Debug configuration") },
 
-	/* PPBC_DebugGame */         { EBuildConfiguration::DebugGame, EBuildTargetType::Game, LOCTEXT("DebugGameConfiguration", "DebugGame"), LOCTEXT("DebugGameConfigurationTooltip", "Package the game in DebugGame configuration") },
-	/* PPBC_DebugGameClient */   { EBuildConfiguration::DebugGame, EBuildTargetType::Client, LOCTEXT("DebugGameClientConfiguration", "DebugGame Client"), LOCTEXT("DebugGameClientConfigurationTooltip", "Package the client in DebugGame configuration") },
-	/* PPBC_DebugGameServer */   { EBuildConfiguration::DebugGame, EBuildTargetType::Server, LOCTEXT("DebugGameServerConfiguration", "DebugGame Server"), LOCTEXT("DebugGameServerConfigurationTooltip", "Package the server in DebugGame configuration") },
+	/* PPBC_DebugGame */                  { EBuildConfiguration::DebugGame, false, LOCTEXT("DebugGameConfiguration", "DebugGame"), LOCTEXT("DebugGameConfigurationTooltip", "Package the game in DebugGame configuration") },
+	/* PPBC_DebugGameClientAndServer */   { EBuildConfiguration::DebugGame, true, LOCTEXT("DebugGameClientConfiguration", "DebugGame (Client & Server)"), LOCTEXT("DebugGameClientConfigurationTooltip", "Package a client and server in DebugGame configuration") },
 
-	/* PPBC_Development */       { EBuildConfiguration::Development, EBuildTargetType::Game, LOCTEXT("DevelopmentConfiguration", "Development"), LOCTEXT("DevelopmentConfigurationTooltip", "Package the game in Development configuration") },
-	/* PPBC_DevelopmentClient */ { EBuildConfiguration::Development, EBuildTargetType::Client, LOCTEXT("DevelopmentClientConfiguration", "Development Client"), LOCTEXT("DevelopmentClientConfigurationTooltip", "Package the client in Development configuration") },
-	/* PPBC_DevelopmentServer */ { EBuildConfiguration::Development, EBuildTargetType::Server, LOCTEXT("DevelopmentServerConfiguration", "Development Server"), LOCTEXT("DevelopmentServerConfigurationTooltip", "Package the server in Development configuration") },
+	/* PPBC_Development */                { EBuildConfiguration::Development, false, LOCTEXT("DevelopmentConfiguration", "Development"), LOCTEXT("DevelopmentConfigurationTooltip", "Package the game in Development configuration") },
+	/* PPBC_DevelopmentClientAndServer */ { EBuildConfiguration::Development, true, LOCTEXT("DevelopmentClientConfiguration", "Development (Client & Server)"), LOCTEXT("DevelopmentClientConfigurationTooltip", "Package a client and server in Development configuration") },
 
-	/* PPBC_Test */              { EBuildConfiguration::Test, EBuildTargetType::Game, LOCTEXT("TestConfiguration", "Test"), LOCTEXT("TestConfigurationTooltip", "Package the game in Test configuration") },
-	/* PPBC_TestClient */        { EBuildConfiguration::Test, EBuildTargetType::Client, LOCTEXT("TestClientConfiguration", "Test Client"), LOCTEXT("TestClientConfigurationTooltip", "Package the client in Test configuration") },
-	/* PPBC_TestServer */        { EBuildConfiguration::Test, EBuildTargetType::Server, LOCTEXT("TestServerConfiguration", "Test Server"), LOCTEXT("TestServerConfigurationTooltip", "Package the server in Test configuration") },
+	/* PPBC_Test */                       { EBuildConfiguration::Test, false, LOCTEXT("TestConfiguration", "Test"), LOCTEXT("TestConfigurationTooltip", "Package the game in Test configuration") },
+	/* PPBC_TestClientAndServer */        { EBuildConfiguration::Test, true, LOCTEXT("TestClientConfiguration", "Test (Client & Server)"), LOCTEXT("TestClientConfigurationTooltip", "Package a client and server in Test configuration") },
 
-	/* PPBC_Shipping */          { EBuildConfiguration::Shipping, EBuildTargetType::Game, LOCTEXT("ShippingConfiguration", "Shipping"), LOCTEXT("ShippingConfigurationTooltip", "Package the game in Shipping configuration") },
-	/* PPBC_ShippingClient */    { EBuildConfiguration::Shipping, EBuildTargetType::Client, LOCTEXT("ShippingClientConfiguration", "Shipping Client"), LOCTEXT("ShippingClientConfigurationTooltip", "Package the client in Shipping configuration") },
-	/* PPBC_ShippingServer */    { EBuildConfiguration::Shipping, EBuildTargetType::Server, LOCTEXT("ShippingServerConfiguration", "Shipping Server"), LOCTEXT("ShippingServerConfigurationTooltip", "Package the server in Shipping configuration") }
+	/* PPBC_Shipping */                   { EBuildConfiguration::Shipping, false, LOCTEXT("ShippingConfiguration", "Shipping"), LOCTEXT("ShippingConfigurationTooltip", "Package the game in Shipping configuration") },
+	/* PPBC_ShippingClientAndServer */    { EBuildConfiguration::Shipping, true, LOCTEXT("ShippingClientConfiguration", "Shipping (Client & Server)"), LOCTEXT("ShippingClientConfigurationTooltip", "Package a client and server in Shipping configuration") },
 };
 
 UProjectPackagingSettings::UProjectPackagingSettings( const FObjectInitializer& ObjectInitializer )
@@ -908,7 +905,7 @@ void UProjectPackagingSettings::PostEditChangeProperty( FPropertyChangedEvent& P
 	}
 	else if (Name == FName(TEXT("ForDistribution")))
 	{
-		if (ForDistribution && BuildConfiguration != EProjectPackagingBuildConfigurations::PPBC_Shipping && BuildConfiguration != EProjectPackagingBuildConfigurations::PPBC_ShippingClient)
+		if (ForDistribution && BuildConfiguration != EProjectPackagingBuildConfigurations::PPBC_Shipping && BuildConfiguration != EProjectPackagingBuildConfigurations::PPBC_ShippingClientAndServer)
 		{
 			BuildConfiguration = EProjectPackagingBuildConfigurations::PPBC_Shipping;
 			// force serialization for "Build COnfiguration"
@@ -1122,6 +1119,82 @@ bool UProjectPackagingSettings::RemoveBlueprintAssetFromNativizationList(const c
 	return false;
 }
 
+TArray<EProjectPackagingBuildConfigurations> UProjectPackagingSettings::GetValidPackageConfigurations()
+{
+	// Check if the project has code
+	FProjectStatus ProjectStatus;
+	bool bHasCode = IProjectManager::Get().QueryStatusForCurrentProject(ProjectStatus) && ProjectStatus.bCodeBasedProject;
+
+	// If if does, find all the targets
+	const TArray<FTargetInfo>* Targets = nullptr;
+	if (bHasCode)
+	{
+		Targets = &(FDesktopPlatformModule::Get()->GetTargetsForCurrentProject());
+	}
+
+	// Set up all the configurations
+	TArray<EProjectPackagingBuildConfigurations> Configurations;
+	for (int32 Idx = 0; Idx < PPBC_MAX; Idx++)
+	{
+		EProjectPackagingBuildConfigurations PackagingConfiguration = (EProjectPackagingBuildConfigurations)Idx;
+		const UProjectPackagingSettings::FConfigurationInfo& Info = UProjectPackagingSettings::ConfigurationInfo[Idx];
+
+		// Make sure the selected configuration is supported
+		EProjectType ProjectType = bHasCode? EProjectType::Code : EProjectType::Content;
+		if (Info.bClientAndServer)
+		{
+			if(!FInstalledPlatformInfo::Get().IsValid(EBuildTargetType::Client, TOptional<FString>(), Info.Configuration, ProjectType, EInstalledPlatformState::Downloaded))
+			{
+				continue;
+			}
+			if(!FInstalledPlatformInfo::Get().IsValid(EBuildTargetType::Server, TOptional<FString>(), Info.Configuration, ProjectType, EInstalledPlatformState::Downloaded))
+			{
+				continue;
+			}
+		}
+		else
+		{
+			if(!FInstalledPlatformInfo::Get().IsValid(EBuildTargetType::Game, TOptional<FString>(), Info.Configuration, ProjectType, EInstalledPlatformState::Downloaded))
+			{
+				continue;
+			}
+		}
+
+		// Check the target type is valid
+		if(bHasCode)
+		{
+			if (Info.bClientAndServer)
+			{
+				if (!Targets->ContainsByPredicate([](const FTargetInfo& Target) -> bool { return Target.Type == EBuildTargetType::Client; }))
+				{
+					continue;
+				}
+				if (!Targets->ContainsByPredicate([](const FTargetInfo& Target) -> bool { return Target.Type == EBuildTargetType::Server; }))
+				{
+					continue;
+				}
+			}
+			else
+			{
+				if (!Targets->ContainsByPredicate([](const FTargetInfo& Target) -> bool { return Target.Type == EBuildTargetType::Game; }))
+				{
+					continue;
+				}
+			}
+		}
+		else
+		{
+			if(Info.Configuration == EBuildConfiguration::DebugGame)
+			{
+				continue;
+			}
+		}
+
+		Configurations.Add(PackagingConfiguration);
+	}
+	return Configurations;
+}
+
 int32 UProjectPackagingSettings::FindBlueprintInNativizationList(const UBlueprint* InBlueprint) const
 {
 	int32 ListIndex = INDEX_NONE;
