Keep the skeletal mesh asset user data when we reimport

diff --git a/Engine/Source/Editor/UnrealEd/Private/Factories/EditorFactories.cpp b/Engine/Source/Editor/UnrealEd/Private/Factories/EditorFactories.cpp
--- a/Engine/Source/Editor/UnrealEd/Private/Factories/EditorFactories.cpp
+++ b/Engine/Source/Editor/UnrealEd/Private/Factories/EditorFactories.cpp
@@ -5475,12 +5475,13 @@ EReimportResult::Type UReimportFbxStaticMeshFactory::Reimport( UObject* Obj )
 				{
 					if ((*UserData)[Idx] != nullptr)
 					{
-						bool bAddDupToRoot = !((*UserData)[Idx]->IsRooted());
+						UAssetUserData* DupObject = (UAssetUserData*)StaticDuplicateObject((*UserData)[Idx], GetTransientPackage());
+						bool bAddDupToRoot = !(DupObject->IsRooted());
 						if (bAddDupToRoot)
 						{
-							(*UserData)[Idx]->AddToRoot();
+							DupObject->AddToRoot();
 						}
-						UserDataCopy.Add((UAssetUserData*)StaticDuplicateObject((*UserData)[Idx], GetTransientPackage()), bAddDupToRoot);
+						UserDataCopy.Add(DupObject, bAddDupToRoot);
 					}
 				}
 			}
diff --git a/Engine/Source/Editor/UnrealEd/Private/Factories/SkeletalMeshImport.cpp b/Engine/Source/Editor/UnrealEd/Private/Factories/SkeletalMeshImport.cpp
--- a/Engine/Source/Editor/UnrealEd/Private/Factories/SkeletalMeshImport.cpp
+++ b/Engine/Source/Editor/UnrealEd/Private/Factories/SkeletalMeshImport.cpp
@@ -30,6 +30,7 @@
 #include "Factories/FbxSkeletalMeshImportData.h"
 #include "IMeshReductionManagerModule.h"
 #include "Rendering/SkeletalMeshModel.h"
+#include "Engine/AssetUserData.h"
 
 DEFINE_LOG_CATEGORY_STATIC(LogSkeletalMeshImport, Log, All);
 
@@ -538,6 +539,25 @@ ExistingSkelMeshData* SaveExistingSkelMeshData(USkeletalMesh* ExistingSkelMesh,
 		}
 	}
 
+	//Store the user asset data
+	const TArray<UAssetUserData*>* UserData = ExistingSkelMesh->GetAssetUserDataArray();
+	if (UserData)
+	{
+		for (int32 Idx = 0; Idx < UserData->Num(); Idx++)
+		{
+			if ((*UserData)[Idx] != nullptr)
+			{
+				UAssetUserData* DupObject = (UAssetUserData*)StaticDuplicateObject((*UserData)[Idx], GetTransientPackage());
+				bool bAddDupToRoot = !(DupObject->IsRooted());
+				if (bAddDupToRoot)
+				{
+					DupObject->AddToRoot();
+				}
+				ExistingMeshDataPtr->ExistingAssetUserData.Add(DupObject, bAddDupToRoot);
+			}
+		}
+	}
+
 	return ExistingMeshDataPtr;
 }
 
@@ -1191,5 +1211,18 @@ void RestoreExistingSkelMeshData(ExistingSkelMeshData* MeshData, USkeletalMesh*
 			SkeletalMeshImportedModel->OriginalReductionSourceMeshData[ReimportLODIndex]->EmptyBulkData();
 		}
 	}
+
+	// Copy user data to newly created mesh
+	for (auto Kvp : MeshData->ExistingAssetUserData)
+	{
+		UAssetUserData* UserDataObject = Kvp.Key;
+		if (Kvp.Value)
+		{
+			//if the duplicated temporary UObject was add to root, we must remove it from the root
+			UserDataObject->RemoveFromRoot();
+		}
+		UserDataObject->Rename(nullptr, SkeletalMesh, REN_DontCreateRedirectors | REN_DoNotDirty);
+		SkeletalMesh->AddAssetUserData(UserDataObject);
+	}
 }
 #undef LOCTEXT_NAMESPACE
diff --git a/Engine/Source/Editor/UnrealEd/Private/Fbx/ReimportFbxSceneFactory.cpp b/Engine/Source/Editor/UnrealEd/Private/Fbx/ReimportFbxSceneFactory.cpp
--- a/Engine/Source/Editor/UnrealEd/Private/Fbx/ReimportFbxSceneFactory.cpp
+++ b/Engine/Source/Editor/UnrealEd/Private/Fbx/ReimportFbxSceneFactory.cpp
@@ -1226,6 +1226,25 @@ EReimportResult::Type UReimportFbxSceneFactory::ReimportSkeletalMesh(void* VoidF
 		return EReimportResult::Failed;
 	}
 
+	const TArray<UAssetUserData*>* UserData = Mesh->GetAssetUserDataArray();
+	TMap<UAssetUserData*, bool> UserDataCopy;
+	if (UserData)
+	{
+		for (int32 Idx = 0; Idx < UserData->Num(); Idx++)
+		{
+			if ((*UserData)[Idx] != nullptr)
+			{
+				UAssetUserData* DupObject = (UAssetUserData*)StaticDuplicateObject((*UserData)[Idx], GetTransientPackage());
+				bool bAddDupToRoot = !(DupObject->IsRooted());
+				if (bAddDupToRoot)
+				{
+					DupObject->AddToRoot();
+				}
+				UserDataCopy.Add(DupObject, bAddDupToRoot);
+			}
+		}
+	}
+
 	ApplyMeshInfoFbxOptions(MeshInfo);
 	//TODO support bBakePivotInVertex
 	bool Old_bBakePivotInVertex = GlobalImportSettings->bBakePivotInVertex;
@@ -1241,6 +1260,19 @@ EReimportResult::Type UReimportFbxSceneFactory::ReimportSkeletalMesh(void* VoidF
 	{
 		Mesh->AssetImportData->Update(FbxImportFileName);
 
+		// Copy user data to newly created mesh
+		for (auto Kvp : UserDataCopy)
+		{
+			UAssetUserData* UserDataObject = Kvp.Key;
+			if (Kvp.Value)
+			{
+				//if the duplicated temporary UObject was add to root, we must remove it from the root
+				UserDataObject->RemoveFromRoot();
+			}
+			UserDataObject->Rename(nullptr, Mesh, REN_DontCreateRedirectors | REN_DoNotDirty);
+			Mesh->AddAssetUserData(UserDataObject);
+		}
+
 		// Try to find the outer package so we can dirty it up
 		if (Mesh->GetOuter())
 		{
@@ -1432,14 +1464,20 @@ EReimportResult::Type UReimportFbxSceneFactory::ReimportStaticMesh(void* VoidFbx
 
 	FbxImporter->ApplyTransformSettingsToFbxNode(FbxImporter->Scene->GetRootNode(), StaticMeshImportData);
 	const TArray<UAssetUserData*>* UserData = Mesh->GetAssetUserDataArray();
-	TArray<UAssetUserData*> UserDataCopy;
+	TMap<UAssetUserData*, bool> UserDataCopy;
 	if (UserData)
 	{
 		for (int32 Idx = 0; Idx < UserData->Num(); Idx++)
 		{
 			if ((*UserData)[Idx] != nullptr)
 			{
-				UserDataCopy.Add((UAssetUserData*)StaticDuplicateObject((*UserData)[Idx], GetTransientPackage()));
+				UAssetUserData* DupObject = (UAssetUserData*)StaticDuplicateObject((*UserData)[Idx], GetTransientPackage());
+				bool bAddDupToRoot = !(DupObject->IsRooted());
+				if (bAddDupToRoot)
+				{
+					DupObject->AddToRoot();
+				}
+				UserDataCopy.Add(DupObject, bAddDupToRoot);
 			}
 		}
 	}
@@ -1466,10 +1504,16 @@ EReimportResult::Type UReimportFbxSceneFactory::ReimportStaticMesh(void* VoidFbx
 		Mesh->AssetImportData = StaticMeshImportData;
 
 		// Copy user data to newly created mesh
-		for (int32 Idx = 0; Idx < UserDataCopy.Num(); Idx++)
+		for (auto Kvp : UserDataCopy)
 		{
-			UserDataCopy[Idx]->Rename(nullptr, Mesh, REN_DontCreateRedirectors | REN_DoNotDirty);
-			Mesh->AddAssetUserData(UserDataCopy[Idx]);
+			UAssetUserData* UserDataObject = Kvp.Key;
+			if (Kvp.Value)
+			{
+				//if the duplicated temporary UObject was add to root, we must remove it from the root
+				UserDataObject->RemoveFromRoot();
+			}
+			UserDataObject->Rename(nullptr, Mesh, REN_DontCreateRedirectors | REN_DoNotDirty);
+			Mesh->AddAssetUserData(UserDataObject);
 		}
 
 		if (NavCollision)
diff --git a/Engine/Source/Editor/UnrealEd/Public/SkelImport.h b/Engine/Source/Editor/UnrealEd/Public/SkelImport.h
--- a/Engine/Source/Editor/UnrealEd/Public/SkelImport.h
+++ b/Engine/Source/Editor/UnrealEd/Public/SkelImport.h
@@ -82,6 +82,8 @@ struct ExistingSkelMeshData
 
 	FSkeletalMeshSamplingInfo				ExistingSamplingInfo;
 	FPerPlatformInt							MinLOD;
+
+	TMap<UAssetUserData*, bool>				ExistingAssetUserData;
 };
 
 /** 
