Fix static analysis

diff --git a/Engine/Source/Editor/UnrealEd/Private/Factories/SkeletalMeshImport.cpp b/Engine/Source/Editor/UnrealEd/Private/Factories/SkeletalMeshImport.cpp
--- a/Engine/Source/Editor/UnrealEd/Private/Factories/SkeletalMeshImport.cpp
+++ b/Engine/Source/Editor/UnrealEd/Private/Factories/SkeletalMeshImport.cpp
@@ -381,164 +381,165 @@ bool SkeletalMeshIsUsingMaterialSlotNameWorkflow(UAssetImportData* AssetImportDa
 
 ExistingSkelMeshData* SaveExistingSkelMeshData(USkeletalMesh* ExistingSkelMesh, bool bSaveMaterials, int32 ReimportLODIndex)
 {
-	struct ExistingSkelMeshData* ExistingMeshDataPtr = NULL;
-	if(ExistingSkelMesh)
+	ExistingSkelMeshData* ExistingMeshDataPtr = nullptr;
+
+	if (!ExistingSkelMesh)
 	{
-		bool ReimportSpecificLOD = (ReimportLODIndex > 0) && ExistingSkelMesh->GetLODNum() > ReimportLODIndex;
+		return nullptr;
+	}
+	bool ReimportSpecificLOD = (ReimportLODIndex > 0) && ExistingSkelMesh->GetLODNum() > ReimportLODIndex;
 
-		ExistingMeshDataPtr = new ExistingSkelMeshData();
+	ExistingMeshDataPtr = new ExistingSkelMeshData();
 		
-		ExistingMeshDataPtr->UseMaterialNameSlotWorkflow = SkeletalMeshIsUsingMaterialSlotNameWorkflow(ExistingSkelMesh->AssetImportData);
-		ExistingMeshDataPtr->MinLOD = ExistingSkelMesh->MinLod;
+	ExistingMeshDataPtr->UseMaterialNameSlotWorkflow = SkeletalMeshIsUsingMaterialSlotNameWorkflow(ExistingSkelMesh->AssetImportData);
+	ExistingMeshDataPtr->MinLOD = ExistingSkelMesh->MinLod;
 
-		FSkeletalMeshModel* ImportedResource = ExistingSkelMesh->GetImportedModel();
+	FSkeletalMeshModel* ImportedResource = ExistingSkelMesh->GetImportedModel();
 
-		//Add the existing Material slot name data
-		for (int32 MaterialIndex = 0; MaterialIndex < ExistingSkelMesh->Materials.Num(); ++MaterialIndex)
-		{
-			ExistingMeshDataPtr->ExistingImportMaterialOriginalNameData.Add(ExistingSkelMesh->Materials[MaterialIndex].ImportedMaterialSlotName);
-		}
+	//Add the existing Material slot name data
+	for (int32 MaterialIndex = 0; MaterialIndex < ExistingSkelMesh->Materials.Num(); ++MaterialIndex)
+	{
+		ExistingMeshDataPtr->ExistingImportMaterialOriginalNameData.Add(ExistingSkelMesh->Materials[MaterialIndex].ImportedMaterialSlotName);
+	}
 
-		for (int32 LodIndex = 0; LodIndex < ImportedResource->LODModels.Num(); ++LodIndex)
+	for (int32 LodIndex = 0; LodIndex < ImportedResource->LODModels.Num(); ++LodIndex)
+	{
+		ExistingMeshDataPtr->ExistingImportMeshLodSectionMaterialData.AddZeroed();
+		for (int32 SectionIndex = 0; SectionIndex < ImportedResource->LODModels[LodIndex].Sections.Num(); ++SectionIndex)
 		{
-			ExistingMeshDataPtr->ExistingImportMeshLodSectionMaterialData.AddZeroed();
-			for (int32 SectionIndex = 0; SectionIndex < ImportedResource->LODModels[LodIndex].Sections.Num(); ++SectionIndex)
+			int32 SectionMaterialIndex = ImportedResource->LODModels[LodIndex].Sections[SectionIndex].MaterialIndex;
+			bool SectionCastShadow = ImportedResource->LODModels[LodIndex].Sections[SectionIndex].bCastShadow;
+			bool SectionRecomputeTangents = ImportedResource->LODModels[LodIndex].Sections[SectionIndex].bRecomputeTangent;
+			int32 GenerateUpTo = ImportedResource->LODModels[LodIndex].Sections[SectionIndex].GenerateUpToLodIndex;
+			bool bDisabled = ImportedResource->LODModels[LodIndex].Sections[SectionIndex].bDisabled;
+			if (ExistingMeshDataPtr->ExistingImportMaterialOriginalNameData.IsValidIndex(SectionMaterialIndex))
 			{
-				int32 SectionMaterialIndex = ImportedResource->LODModels[LodIndex].Sections[SectionIndex].MaterialIndex;
-				bool SectionCastShadow = ImportedResource->LODModels[LodIndex].Sections[SectionIndex].bCastShadow;
-				bool SectionRecomputeTangents = ImportedResource->LODModels[LodIndex].Sections[SectionIndex].bRecomputeTangent;
-				int32 GenerateUpTo = ImportedResource->LODModels[LodIndex].Sections[SectionIndex].GenerateUpToLodIndex;
-				bool bDisabled = ImportedResource->LODModels[LodIndex].Sections[SectionIndex].bDisabled;
-				if (ExistingMeshDataPtr->ExistingImportMaterialOriginalNameData.IsValidIndex(SectionMaterialIndex))
-				{
-					ExistingMeshDataPtr->ExistingImportMeshLodSectionMaterialData[LodIndex].Add(ExistingMeshLodSectionData(ExistingMeshDataPtr->ExistingImportMaterialOriginalNameData[SectionMaterialIndex], SectionCastShadow, SectionRecomputeTangents, GenerateUpTo, bDisabled));
-				}
+				ExistingMeshDataPtr->ExistingImportMeshLodSectionMaterialData[LodIndex].Add(ExistingMeshLodSectionData(ExistingMeshDataPtr->ExistingImportMaterialOriginalNameData[SectionMaterialIndex], SectionCastShadow, SectionRecomputeTangents, GenerateUpTo, bDisabled));
 			}
 		}
+	}
+
+	ExistingMeshDataPtr->ExistingSockets = ExistingSkelMesh->GetMeshOnlySocketList();
+	ExistingMeshDataPtr->bSaveRestoreMaterials = bSaveMaterials;
+	if (ExistingMeshDataPtr->bSaveRestoreMaterials)
+	{
+		ExistingMeshDataPtr->ExistingMaterials = ExistingSkelMesh->Materials;
+	}
+	ExistingMeshDataPtr->ExistingRetargetBasePose = ExistingSkelMesh->RetargetBasePose;
 
-		ExistingMeshDataPtr->ExistingSockets = ExistingSkelMesh->GetMeshOnlySocketList();
-		ExistingMeshDataPtr->bSaveRestoreMaterials = bSaveMaterials;
-		if (ExistingMeshDataPtr->bSaveRestoreMaterials)
+	if( ImportedResource->LODModels.Num() > 0 &&
+		ExistingSkelMesh->GetLODNum() == ImportedResource->LODModels.Num() )
+	{
+		int32 OffsetReductionLODIndex = 0;
+		FSkeletalMeshLODInfo* LODInfo = ExistingSkelMesh->GetLODInfo( ReimportLODIndex < 0 ? 0 : ReimportLODIndex);
+		ExistingMeshDataPtr->bIsReimportLODReduced = (LODInfo && LODInfo->bHasBeenSimplified);
+		if (ExistingMeshDataPtr->bIsReimportLODReduced)
 		{
-			ExistingMeshDataPtr->ExistingMaterials = ExistingSkelMesh->Materials;
+			//Save the imported LOD reduction settings
+			ExistingMeshDataPtr->ExistingReimportLODReductionSettings = LODInfo->ReductionSettings;
 		}
-		ExistingMeshDataPtr->ExistingRetargetBasePose = ExistingSkelMesh->RetargetBasePose;
 
-		if( ImportedResource->LODModels.Num() > 0 &&
-			ExistingSkelMesh->GetLODNum() == ImportedResource->LODModels.Num() )
+		// Remove the zero'th LOD (ie: the LOD being reimported).
+		if (!ReimportSpecificLOD)
 		{
-			int32 OffsetReductionLODIndex = 0;
-			FSkeletalMeshLODInfo* LODInfo = ExistingSkelMesh->GetLODInfo( ReimportLODIndex < 0 ? 0 : ReimportLODIndex);
-			ExistingMeshDataPtr->bIsReimportLODReduced = (LODInfo && LODInfo->bHasBeenSimplified);
-			if (ExistingMeshDataPtr->bIsReimportLODReduced)
-			{
-				//Save the imported LOD reduction settings
-				ExistingMeshDataPtr->ExistingReimportLODReductionSettings = LODInfo->ReductionSettings;
-			}
-
-			// Remove the zero'th LOD (ie: the LOD being reimported).
-			if (!ReimportSpecificLOD)
-			{
-				ImportedResource->LODModels.RemoveAt(0);
-				ExistingSkelMesh->RemoveLODInfo(0);
-				OffsetReductionLODIndex = 1;
-			}
+			ImportedResource->LODModels.RemoveAt(0);
+			ExistingSkelMesh->RemoveLODInfo(0);
+			OffsetReductionLODIndex = 1;
+		}
 
-			// Copy off the remaining LODs.
-			for ( int32 LODModelIndex = 0 ; LODModelIndex < ImportedResource->LODModels.Num() ; ++LODModelIndex )
+		// Copy off the remaining LODs.
+		for ( int32 LODModelIndex = 0 ; LODModelIndex < ImportedResource->LODModels.Num() ; ++LODModelIndex )
+		{
+			FSkeletalMeshLODModel& LODModel = ImportedResource->LODModels[LODModelIndex];
+			LODModel.RawPointIndices.Lock( LOCK_READ_ONLY );
+			LODModel.LegacyRawPointIndices.Lock( LOCK_READ_ONLY );
+			LODModel.RawSkeletalMeshBulkData.GetBulkData().Lock( LOCK_READ_ONLY );
+			int32 ReductionLODIndex = LODModelIndex + OffsetReductionLODIndex;
+			if (ImportedResource->OriginalReductionSourceMeshData.IsValidIndex(ReductionLODIndex) && !ImportedResource->OriginalReductionSourceMeshData[ReductionLODIndex]->IsEmpty())
 			{
-				FSkeletalMeshLODModel& LODModel = ImportedResource->LODModels[LODModelIndex];
-				LODModel.RawPointIndices.Lock( LOCK_READ_ONLY );
-				LODModel.LegacyRawPointIndices.Lock( LOCK_READ_ONLY );
-				LODModel.RawSkeletalMeshBulkData.GetBulkData().Lock( LOCK_READ_ONLY );
-				int32 ReductionLODIndex = LODModelIndex + OffsetReductionLODIndex;
-				if (ImportedResource->OriginalReductionSourceMeshData.IsValidIndex(ReductionLODIndex) && !ImportedResource->OriginalReductionSourceMeshData[ReductionLODIndex]->IsEmpty())
+				FSkeletalMeshLODModel BaseLODModel;
+				TMap<FString, TArray<FMorphTargetDelta>> BaseLODMorphTargetData;
+				ImportedResource->OriginalReductionSourceMeshData[ReductionLODIndex]->LoadReductionData(BaseLODModel, BaseLODMorphTargetData);
+				FReductionBaseSkeletalMeshBulkData* ReductionLODData = new FReductionBaseSkeletalMeshBulkData();
+				ReductionLODData->SaveReductionData(BaseLODModel, BaseLODMorphTargetData);
+				//Add necessary empty slot
+				while (ExistingMeshDataPtr->ExistingOriginalReductionSourceMeshData.Num() < LODModelIndex)
 				{
-					FSkeletalMeshLODModel BaseLODModel;
-					TMap<FString, TArray<FMorphTargetDelta>> BaseLODMorphTargetData;
-					ImportedResource->OriginalReductionSourceMeshData[ReductionLODIndex]->LoadReductionData(BaseLODModel, BaseLODMorphTargetData);
-					FReductionBaseSkeletalMeshBulkData* ReductionLODData = new FReductionBaseSkeletalMeshBulkData();
-					ReductionLODData->SaveReductionData(BaseLODModel, BaseLODMorphTargetData);
-					//Add necessary empty slot
-					while (ExistingMeshDataPtr->ExistingOriginalReductionSourceMeshData.Num() < LODModelIndex)
-					{
-						FReductionBaseSkeletalMeshBulkData* EmptyReductionLODData = new FReductionBaseSkeletalMeshBulkData();
-						ExistingMeshDataPtr->ExistingOriginalReductionSourceMeshData.Add(EmptyReductionLODData);
-					}
-					ExistingMeshDataPtr->ExistingOriginalReductionSourceMeshData.Add(ReductionLODData);
+					FReductionBaseSkeletalMeshBulkData* EmptyReductionLODData = new FReductionBaseSkeletalMeshBulkData();
+					ExistingMeshDataPtr->ExistingOriginalReductionSourceMeshData.Add(EmptyReductionLODData);
 				}
+				ExistingMeshDataPtr->ExistingOriginalReductionSourceMeshData.Add(ReductionLODData);
 			}
-			ExistingMeshDataPtr->ExistingLODModels = ImportedResource->LODModels;
-			for (int32 LODModelIndex = 0; LODModelIndex < ImportedResource->LODModels.Num(); ++LODModelIndex)
-			{
-				FSkeletalMeshLODModel& LODModel = ImportedResource->LODModels[LODModelIndex];
-				LODModel.RawPointIndices.Unlock();
-				LODModel.LegacyRawPointIndices.Unlock();
-				LODModel.RawSkeletalMeshBulkData.GetBulkData().Unlock();
-			}
+		}
+		ExistingMeshDataPtr->ExistingLODModels = ImportedResource->LODModels;
+		for (int32 LODModelIndex = 0; LODModelIndex < ImportedResource->LODModels.Num(); ++LODModelIndex)
+		{
+			FSkeletalMeshLODModel& LODModel = ImportedResource->LODModels[LODModelIndex];
+			LODModel.RawPointIndices.Unlock();
+			LODModel.LegacyRawPointIndices.Unlock();
+			LODModel.RawSkeletalMeshBulkData.GetBulkData().Unlock();
+		}
 
-			ExistingMeshDataPtr->ExistingLODInfo = ExistingSkelMesh->GetLODInfoArray();
-			ExistingMeshDataPtr->ExistingRefSkeleton = ExistingSkelMesh->RefSkeleton;
+		ExistingMeshDataPtr->ExistingLODInfo = ExistingSkelMesh->GetLODInfoArray();
+		ExistingMeshDataPtr->ExistingRefSkeleton = ExistingSkelMesh->RefSkeleton;
 		
-		}
+	}
 
-		// First asset should be the one that the skeletal mesh should point too
-		ExistingMeshDataPtr->ExistingPhysicsAssets.Empty();
-		ExistingMeshDataPtr->ExistingPhysicsAssets.Add( ExistingSkelMesh->PhysicsAsset );
-		for (TObjectIterator<UPhysicsAsset> It; It; ++It)
+	// First asset should be the one that the skeletal mesh should point too
+	ExistingMeshDataPtr->ExistingPhysicsAssets.Empty();
+	ExistingMeshDataPtr->ExistingPhysicsAssets.Add( ExistingSkelMesh->PhysicsAsset );
+	for (TObjectIterator<UPhysicsAsset> It; It; ++It)
+	{
+		UPhysicsAsset* PhysicsAsset = *It;
+		if ( PhysicsAsset->PreviewSkeletalMesh == ExistingSkelMesh && ExistingSkelMesh->PhysicsAsset != PhysicsAsset )
 		{
-			UPhysicsAsset* PhysicsAsset = *It;
-			if ( PhysicsAsset->PreviewSkeletalMesh == ExistingSkelMesh && ExistingSkelMesh->PhysicsAsset != PhysicsAsset )
-			{
-				ExistingMeshDataPtr->ExistingPhysicsAssets.Add( PhysicsAsset );
-			}
+			ExistingMeshDataPtr->ExistingPhysicsAssets.Add( PhysicsAsset );
 		}
+	}
 
-		ExistingMeshDataPtr->ExistingShadowPhysicsAsset = ExistingSkelMesh->ShadowPhysicsAsset;
+	ExistingMeshDataPtr->ExistingShadowPhysicsAsset = ExistingSkelMesh->ShadowPhysicsAsset;
 
-		ExistingMeshDataPtr->ExistingSkeleton = ExistingSkelMesh->Skeleton;
-		// since copying back original skeleton, this shoudl be safe to do
-		ExistingMeshDataPtr->ExistingPostProcessAnimBlueprint = ExistingSkelMesh->PostProcessAnimBlueprint;
+	ExistingMeshDataPtr->ExistingSkeleton = ExistingSkelMesh->Skeleton;
+	// since copying back original skeleton, this shoudl be safe to do
+	ExistingMeshDataPtr->ExistingPostProcessAnimBlueprint = ExistingSkelMesh->PostProcessAnimBlueprint;
 
-		ExistingMeshDataPtr->ExistingLODSettings = ExistingSkelMesh->LODSettings;
+	ExistingMeshDataPtr->ExistingLODSettings = ExistingSkelMesh->LODSettings;
 
-		ExistingSkelMesh->ExportMirrorTable(ExistingMeshDataPtr->ExistingMirrorTable);
+	ExistingSkelMesh->ExportMirrorTable(ExistingMeshDataPtr->ExistingMirrorTable);
 
-		ExistingMeshDataPtr->ExistingMorphTargets.Empty(ExistingSkelMesh->MorphTargets.Num());
-		ExistingMeshDataPtr->ExistingMorphTargets.Append(ExistingSkelMesh->MorphTargets);
+	ExistingMeshDataPtr->ExistingMorphTargets.Empty(ExistingSkelMesh->MorphTargets.Num());
+	ExistingMeshDataPtr->ExistingMorphTargets.Append(ExistingSkelMesh->MorphTargets);
 	
-		ExistingMeshDataPtr->bExistingUseFullPrecisionUVs = ExistingSkelMesh->bUseFullPrecisionUVs;
-		ExistingMeshDataPtr->bExistingUseHighPrecisionTangentBasis = ExistingSkelMesh->bUseHighPrecisionTangentBasis;
+	ExistingMeshDataPtr->bExistingUseFullPrecisionUVs = ExistingSkelMesh->bUseFullPrecisionUVs;
+	ExistingMeshDataPtr->bExistingUseHighPrecisionTangentBasis = ExistingSkelMesh->bUseHighPrecisionTangentBasis;
 
-		ExistingMeshDataPtr->ExistingAssetImportData = ExistingSkelMesh->AssetImportData;
-		ExistingMeshDataPtr->ExistingThumbnailInfo = ExistingSkelMesh->ThumbnailInfo;
+	ExistingMeshDataPtr->ExistingAssetImportData = ExistingSkelMesh->AssetImportData;
+	ExistingMeshDataPtr->ExistingThumbnailInfo = ExistingSkelMesh->ThumbnailInfo;
 
-		ExistingMeshDataPtr->ExistingClothingAssets = ExistingSkelMesh->MeshClothingAssets;
+	ExistingMeshDataPtr->ExistingClothingAssets = ExistingSkelMesh->MeshClothingAssets;
 
-		ExistingMeshDataPtr->ExistingSamplingInfo = ExistingSkelMesh->GetSamplingInfo();
+	ExistingMeshDataPtr->ExistingSamplingInfo = ExistingSkelMesh->GetSamplingInfo();
 
-		//Add the last fbx import data
-		UFbxSkeletalMeshImportData* ImportData = Cast<UFbxSkeletalMeshImportData>(ExistingSkelMesh->AssetImportData);
-		if (ImportData && ExistingMeshDataPtr->UseMaterialNameSlotWorkflow)
+	//Add the last fbx import data
+	UFbxSkeletalMeshImportData* ImportData = Cast<UFbxSkeletalMeshImportData>(ExistingSkelMesh->AssetImportData);
+	if (ImportData && ExistingMeshDataPtr->UseMaterialNameSlotWorkflow)
+	{
+		for (int32 ImportMaterialOriginalNameDataIndex = 0; ImportMaterialOriginalNameDataIndex < ImportData->ImportMaterialOriginalNameData.Num(); ++ImportMaterialOriginalNameDataIndex)
 		{
-			for (int32 ImportMaterialOriginalNameDataIndex = 0; ImportMaterialOriginalNameDataIndex < ImportData->ImportMaterialOriginalNameData.Num(); ++ImportMaterialOriginalNameDataIndex)
-			{
-				FName MaterialName = ImportData->ImportMaterialOriginalNameData[ImportMaterialOriginalNameDataIndex];
-				ExistingMeshDataPtr->LastImportMaterialOriginalNameData.Add(MaterialName);
-			}
-			for (int32 LodIndex = 0; LodIndex < ImportData->ImportMeshLodData.Num(); ++LodIndex)
+			FName MaterialName = ImportData->ImportMaterialOriginalNameData[ImportMaterialOriginalNameDataIndex];
+			ExistingMeshDataPtr->LastImportMaterialOriginalNameData.Add(MaterialName);
+		}
+		for (int32 LodIndex = 0; LodIndex < ImportData->ImportMeshLodData.Num(); ++LodIndex)
+		{
+			ExistingMeshDataPtr->LastImportMeshLodSectionMaterialData.AddZeroed();
+			const FImportMeshLodSectionsData &ImportMeshLodSectionsData = ImportData->ImportMeshLodData[LodIndex];
+			for (int32 SectionIndex = 0; SectionIndex < ImportMeshLodSectionsData.SectionOriginalMaterialName.Num(); ++SectionIndex)
 			{
-				ExistingMeshDataPtr->LastImportMeshLodSectionMaterialData.AddZeroed();
-				const FImportMeshLodSectionsData &ImportMeshLodSectionsData = ImportData->ImportMeshLodData[LodIndex];
-				for (int32 SectionIndex = 0; SectionIndex < ImportMeshLodSectionsData.SectionOriginalMaterialName.Num(); ++SectionIndex)
-				{
-					FName MaterialName = ImportMeshLodSectionsData.SectionOriginalMaterialName[SectionIndex];
-					ExistingMeshDataPtr->LastImportMeshLodSectionMaterialData[LodIndex].Add(MaterialName);
-				}
+				FName MaterialName = ImportMeshLodSectionsData.SectionOriginalMaterialName[SectionIndex];
+				ExistingMeshDataPtr->LastImportMeshLodSectionMaterialData[LodIndex].Add(MaterialName);
 			}
 		}
 	}
-
 	//Store the user asset data
 	const TArray<UAssetUserData*>* UserData = ExistingSkelMesh->GetAssetUserDataArray();
 	if (UserData)
@@ -557,7 +558,6 @@ ExistingSkelMeshData* SaveExistingSkelMeshData(USkeletalMesh* ExistingSkelMesh,
 			}
 		}
 	}
-
 	return ExistingMeshDataPtr;
 }
 
