Non-unity CIS fix.

#rnx

diff --git a/Engine/Source/Editor/UnrealEd/Classes/Settings/LevelEditorPlaySettings.h b/Engine/Source/Editor/UnrealEd/Classes/Settings/LevelEditorPlaySettings.h
--- a/Engine/Source/Editor/UnrealEd/Classes/Settings/LevelEditorPlaySettings.h
+++ b/Engine/Source/Editor/UnrealEd/Classes/Settings/LevelEditorPlaySettings.h
@@ -7,6 +7,7 @@
 #include "UObject/Object.h"
 #include "Layout/Margin.h"
 #include "Layout/Visibility.h"
+#include "Misc/App.h"
 #include "Settings/LevelEditorPlayNetworkEmulationSettings.h"
 #include "LevelEditorPlaySettings.generated.h"
 
