Slate - Fixing border and image to properly invalidate with the underlying resource or brush data changes when the brush is changed.  Moving that code into a shared area to make this bit reuseable, this caused some fallout fixing up places that subclassed SBorder.  Likely to experience other cases in plugins or other projects.

Matt.Kuhlenschmidt

diff --git a/Engine/Source/Editor/UnrealEd/Private/ConsolidateWindow.cpp b/Engine/Source/Editor/UnrealEd/Private/ConsolidateWindow.cpp
--- a/Engine/Source/Editor/UnrealEd/Private/ConsolidateWindow.cpp
+++ b/Engine/Source/Editor/UnrealEd/Private/ConsolidateWindow.cpp
@@ -364,7 +364,7 @@ void SConsolidateToolWidget::Construct( const FArguments& InArgs )
 	SelectedListItem = NULL;
 	bSavePackagesChecked = ISourceControlModule::Get().IsEnabled();
 
-	this->BorderImage = FEditorStyle::GetBrush("NoBorder");
+	this->BorderImage = FInvalidatableBrushAttribute(FEditorStyle::GetBrush("NoBorder"));
 
 	ChildSlot
 	[
diff --git a/Engine/Source/Editor/UnrealEd/Private/Dialogs/SBuildProgress.cpp b/Engine/Source/Editor/UnrealEd/Private/Dialogs/SBuildProgress.cpp
--- a/Engine/Source/Editor/UnrealEd/Private/Dialogs/SBuildProgress.cpp
+++ b/Engine/Source/Editor/UnrealEd/Private/Dialogs/SBuildProgress.cpp
@@ -20,7 +20,8 @@ SBuildProgressWidget::~SBuildProgressWidget()
 
 void SBuildProgressWidget::Construct( const FArguments& InArgs )
 {
-	BorderImage = FEditorStyle::GetBrush("Menu.Background");
+	BorderImage = FInvalidatableBrushAttribute(FEditorStyle::GetBrush("Menu.Background"));
+
 	this->ChildSlot
 	.VAlign(VAlign_Center)
 	[
