Added the level name in the warning message upon failure to load when level is already loaded.

diff --git a/Engine/Source/Editor/UnrealEd/Private/EditorLevelUtils.cpp b/Engine/Source/Editor/UnrealEd/Private/EditorLevelUtils.cpp
--- a/Engine/Source/Editor/UnrealEd/Private/EditorLevelUtils.cpp
+++ b/Engine/Source/Editor/UnrealEd/Private/EditorLevelUtils.cpp
@@ -376,7 +376,9 @@ ULevelStreaming* UEditorLevelUtils::AddLevelToWorld_Internal(UWorld* InWorld, co
 	if (bIsPersistentLevel || FLevelUtils::FindStreamingLevel(InWorld, LevelPackageName))
 	{
 		// Do nothing if the level already exists in the world.
-		FMessageDialog::Open(EAppMsgType::Ok, NSLOCTEXT("UnrealEd", "LevelAlreadyExistsInWorld", "A level with that name already exists in the world."));
+		const FString LevelName(LevelPackageName);
+		const FText MessageText = FText::Format(NSLOCTEXT("UnrealEd", "LevelAlreadyExistsInWorld", "A level with that name ({0}) already exists in the world."), FText::FromString(LevelName));
+		FMessageDialog::Open(EAppMsgType::Ok, MessageText);
 	}
 	else
 	{
