- Save / restore OnMeshChanged delegate to avoid a crash in Nigara Skel Mesh DI on reimport

#rnx

diff --git a/Engine/Source/Editor/UnrealEd/Private/StaticMeshEdit.cpp b/Engine/Source/Editor/UnrealEd/Private/StaticMeshEdit.cpp
--- a/Engine/Source/Editor/UnrealEd/Private/StaticMeshEdit.cpp
+++ b/Engine/Source/Editor/UnrealEd/Private/StaticMeshEdit.cpp
@@ -1011,6 +1011,8 @@ struct ExistingStaticMeshData
 	bool						ExistingAllowCpuAccess;
 	FVector						ExistingPositiveBoundsExtension;
 	FVector						ExistingNegativeBoundsExtension;
+
+	UStaticMesh::FOnMeshChanged	ExistingOnMeshChanged;
 };
 
 bool IsUsingMaterialSlotNameWorkflow(UAssetImportData* AssetImportData)
@@ -1192,6 +1194,7 @@ ExistingStaticMeshData* SaveExistingStaticMeshData(UStaticMesh* ExistingMesh, Un
 				}
 			}
 		}
+		ExistingMeshDataPtr->ExistingOnMeshChanged = ExistingMesh->OnMeshChanged;
 	}
 
 	return ExistingMeshDataPtr;
@@ -1337,6 +1340,8 @@ void RestoreExistingMeshSettings(ExistingStaticMeshData* ExistingMesh, UStaticMe
 		}
 		NewMesh->MaterialRemapIndexPerImportVersion.Add(FMaterialRemapIndex(MaterialMapKey, ImportRemapMaterial));
 	}
+	// Copy mesh changed delegate data
+	NewMesh->OnMeshChanged = ExistingMesh->ExistingOnMeshChanged;
 }
 
 void UpdateSomeLodsImportMeshData(UStaticMesh* NewMesh, TArray<int32> *ReimportLodList)
