Fixing crash caused by skeletal mesh morphtarget export when clothing is not in the last section

diff --git a/Engine/Source/Editor/UnrealEd/Private/Fbx/FbxSkeletalMeshExport.cpp b/Engine/Source/Editor/UnrealEd/Private/Fbx/FbxSkeletalMeshExport.cpp
--- a/Engine/Source/Editor/UnrealEd/Private/Fbx/FbxSkeletalMeshExport.cpp
+++ b/Engine/Source/Editor/UnrealEd/Private/Fbx/FbxSkeletalMeshExport.cpp
@@ -207,12 +207,14 @@ FbxNode* FFbxExporter::CreateMesh(const USkeletalMesh* SkelMesh, const TCHAR* Me
 	// Create the per-material polygons sets.
 	int32 SectionCount = SourceModel.Sections.Num();
 	int32 ClothSectionVertexRemoveOffset = 0;
+	TArray<TPair<uint32, uint32>> VertexIndexOffsetArray{ TPair<uint32, uint32>(0,0) };
 	for (int32 SectionIndex = 0; SectionIndex < SectionCount; ++SectionIndex)
 	{
 		const FSkelMeshSection& Section = SourceModel.Sections[SectionIndex];
 		if (Section.HasClothingData())
 		{
 			ClothSectionVertexRemoveOffset += Section.GetNumVertices();
+			VertexIndexOffsetArray.Emplace(Section.BaseIndex, ClothSectionVertexRemoveOffset);
 			continue;
 		}
 		int32 MatIndex = Section.MaterialIndex;
@@ -283,6 +285,16 @@ FbxNode* FFbxExporter::CreateMesh(const USkeletalMesh* SkelMesh, const TCHAR* Me
 				{
 					// Apply the morph target deltas to the control points.
 					FMorphTargetDelta& CurrentDelta = MorphTargetDeltas[MorphTargetDeltaIndex];
+					uint32 remappedSourceIndex = CurrentDelta.SourceIdx;
+
+					if (VertexIndexOffsetArray.Num() > 1)
+					{
+						//If the skeletal mesh contains clothing we need to remap the morph target index too.
+						int32 upperBoundIndex = Algo::UpperBoundBy(VertexIndexOffsetArray, remappedSourceIndex,
+							[](const auto& CurrentPair) { return CurrentPair.Key; }); //Value functor
+						remappedSourceIndex -= VertexIndexOffsetArray[upperBoundIndex - 1].Value;
+					}
+
 					ShapeControlPoints[CurrentDelta.SourceIdx] = Converter.ConvertToFbxPos(Vertices[CurrentDelta.SourceIdx].Position + CurrentDelta.PositionDelta);
 				}
 
