Fix crash at skeletal mesh export

diff --git a/Engine/Source/Editor/UnrealEd/Private/Fbx/FbxSkeletalMeshExport.cpp b/Engine/Source/Editor/UnrealEd/Private/Fbx/FbxSkeletalMeshExport.cpp
--- a/Engine/Source/Editor/UnrealEd/Private/Fbx/FbxSkeletalMeshExport.cpp
+++ b/Engine/Source/Editor/UnrealEd/Private/Fbx/FbxSkeletalMeshExport.cpp
@@ -295,7 +295,7 @@ FbxNode* FFbxExporter::CreateMesh(const USkeletalMesh* SkelMesh, const TCHAR* Me
 						RemappedSourceIndex -= VertexIndexOffsetPairArray[UpperBoundIndex - 1].Value;
 					}
 
-					ShapeControlPoints[CurrentDelta.SourceIdx] = Converter.ConvertToFbxPos(Vertices[CurrentDelta.SourceIdx].Position + CurrentDelta.PositionDelta);
+					ShapeControlPoints[RemappedSourceIndex] = Converter.ConvertToFbxPos(Vertices[RemappedSourceIndex].Position + CurrentDelta.PositionDelta);
 				}
 
 				BlendShapeChannel->AddTargetShape(Shape);
