Force Deleting an asset now closes all asset editors that reference that object.

diff --git a/Engine/Source/Editor/UnrealEd/Private/ObjectTools.cpp b/Engine/Source/Editor/UnrealEd/Private/ObjectTools.cpp
--- a/Engine/Source/Editor/UnrealEd/Private/ObjectTools.cpp
+++ b/Engine/Source/Editor/UnrealEd/Private/ObjectTools.cpp
@@ -2386,7 +2386,23 @@ namespace ObjectTools
 		return true;
 	}
 
-	int32 ForceDeleteObjects( const TArray< UObject* >& InObjectsToDelete, bool ShowConfirmation )
+	static void RecursiveRetrieveReferencers(UObject* Object, TSet<UObject*>& ReferencingObjects)
+	{
+		TArray<FReferencerInformation> ExternalReferencers;
+		Object->RetrieveReferencers(nullptr /* internal refs */, &ExternalReferencers);
+
+		for (const FReferencerInformation& Referencer : ExternalReferencers)
+		{
+			bool bAlreadyIn = false;
+			ReferencingObjects.Add(Referencer.Referencer, &bAlreadyIn);
+			if (!bAlreadyIn)
+			{
+				RecursiveRetrieveReferencers(Referencer.Referencer, ReferencingObjects);
+			}
+		}
+	}
+
+	int32 ForceDeleteObjects(const TArray< UObject* >& InObjectsToDelete, bool ShowConfirmation)
 	{
 		int32 NumDeletedObjects = 0;
 
@@ -2403,22 +2419,33 @@ namespace ObjectTools
 		}
 
 		// Confirm that the delete was intentional
-		if ( ShowConfirmation && !ShowDeleteConfirmationDialog(ShownObjectsToDelete) )
+		if (ShowConfirmation && !ShowDeleteConfirmationDialog(ShownObjectsToDelete))
 		{
 			return 0;
 		}
 
-		// Attempt to close all editors referencing this asset.
-		bool bClosedAllEditors = true;
+		// Recursively find all references to objects being deleted
+		TSet<UObject*> ReferencingObjects;
+		for (UObject* ToDelete : InObjectsToDelete)
+		{
+			ReferencingObjects.Add(ToDelete);
+
+			RecursiveRetrieveReferencers(ToDelete, ReferencingObjects);
+		}
 
-		for (UObject* ObjectToDelete : InObjectsToDelete)
+		// Attempt to close all editors referencing any of the deleted objects
+		bool bClosedAllEditors = true;
+		for (UObject* Object : ReferencingObjects)
 		{
-			const TArray<IAssetEditorInstance*> ObjectEditors = GEditor->GetEditorSubsystem<UAssetEditorSubsystem>()->FindEditorsForAsset(ObjectToDelete);
-			for (IAssetEditorInstance* ObjectEditorInstance : ObjectEditors)
+			if (Object->IsAsset())
 			{
-				if (!ObjectEditorInstance->CloseWindow())
+				const TArray<IAssetEditorInstance*> ObjectEditors = GEditor->GetEditorSubsystem<UAssetEditorSubsystem>()->FindEditorsForAsset(Object);
+				for (IAssetEditorInstance* ObjectEditorInstance : ObjectEditors)
 				{
-					bClosedAllEditors = false;
+					if (!ObjectEditorInstance->CloseWindow())
+					{
+						bClosedAllEditors = false;
+					}
 				}
 			}
 		}
@@ -2766,7 +2793,6 @@ namespace ObjectTools
 		return NumDeletedObjects;
 	}	
 
-	
 	/**
 	 * Utility function to compose a string list of referencing objects
 	 *
