Fixed compiler errors caused by mismerge in CL 8686367.

diff --git a/Engine/Source/Editor/UnrealEd/Private/ObjectTools.cpp b/Engine/Source/Editor/UnrealEd/Private/ObjectTools.cpp
--- a/Engine/Source/Editor/UnrealEd/Private/ObjectTools.cpp
+++ b/Engine/Source/Editor/UnrealEd/Private/ObjectTools.cpp
@@ -72,7 +72,7 @@
 #include "ReferencedAssetsUtils.h"
 #include "AssetRegistryModule.h"
 #include "PackagesDialog.h"
-
+#include "Toolkits/AssetEditorManager.h"
 #include "PropertyEditorModule.h"
 #include "Kismet2/KismetEditorUtilities.h"
 #include "Kismet2/KismetReinstanceUtilities.h"
@@ -85,14 +85,13 @@
 #include "ComponentRecreateRenderStateContext.h"
 #include "Framework/Notifications/NotificationManager.h"
 #include "Widgets/Notifications/SNotificationList.h"
-#include "Layers/LayersSubsystem.h"
+#include "Layers/ILayers.h"
 #include "Engine/SCS_Node.h"
 #include "ShaderCompiler.h"
 #include "Templates/UniquePtr.h"
 #include "Engine/MapBuildDataRegistry.h"
 #include "HAL/PlatformApplicationMisc.h"
 #include "Kismet2/BlueprintEditorUtils.h"
-#include "Subsystems/AssetEditorSubsystem.h"
 
 DEFINE_LOG_CATEGORY_STATIC(LogObjectTools, Log, All);
 
@@ -923,7 +922,7 @@ namespace ObjectTools
 		if ( ObjectToConsolidateTo )
 		{
 			// Close all editors to avoid changing references to temporary objects used by the editor
-			if ( !GEditor->GetEditorSubsystem<UAssetEditorSubsystem>()->CloseAllAssetEditors() )
+			if (!FAssetEditorManager::Get().CloseAllAssetEditors())
 			{
 				// Failed to close at least one editor. It is possible that this editor has in-memory object references
 				// which are not prepared to be changed dynamically so it is not safe to continue
@@ -2439,7 +2438,7 @@ namespace ObjectTools
 		{
 			if (Object->IsAsset())
 			{
-				const TArray<IAssetEditorInstance*> ObjectEditors = GEditor->GetEditorSubsystem<UAssetEditorSubsystem>()->FindEditorsForAsset(Object);
+				const TArray<IAssetEditorInstance*> ObjectEditors = FAssetEditorManager::Get().FindEditorsForAsset(Object);
 				for (IAssetEditorInstance* ObjectEditorInstance : ObjectEditors)
 				{
 					if (!ObjectEditorInstance->CloseWindow())
@@ -2601,7 +2600,6 @@ namespace ObjectTools
 		// Destroy all Actor instances
 		if ( ActorsToDelete.Num() > 0 )
 		{
-			ULayersSubsystem* Layers = GEditor->GetEditorSubsystem<ULayersSubsystem>();
 			for ( TArray<AActor*>::TConstIterator ActorItr( ActorsToDelete ); ActorItr; ++ActorItr )
 			{
 				AActor* CurActor = *ActorItr;
@@ -2619,7 +2617,7 @@ namespace ObjectTools
 					}
 
 					// Destroy the Actor instance. This is similar to edactDeleteSelected(), but we don't request user confirmation here.
-					Layers->DisassociateActorFromLayers( CurActor );
+					GEditor->Layers->DisassociateActorFromLayers( CurActor );
 					if( CurActor->GetWorld() )
 					{
 						CurActor->GetWorld()->EditorDestroyActor( CurActor, false );
