HLOD System - HLOD 0 is visible while HLOD 1 is visible

Fixed a few HLOD cluster generation options combinations

diff --git a/Engine/Source/Editor/UnrealEd/Private/HierarchicalLOD.cpp b/Engine/Source/Editor/UnrealEd/Private/HierarchicalLOD.cpp
--- a/Engine/Source/Editor/UnrealEd/Private/HierarchicalLOD.cpp
+++ b/Engine/Source/Editor/UnrealEd/Private/HierarchicalLOD.cpp
@@ -164,11 +164,11 @@ void FHierarchicalLODBuilder::BuildClusters(ULevel* InLevel, const bool bCreateM
 		}		
 		else
 		{
-			// Handle HierachicalLOD volumes first
-			HandleHLODVolumes(InLevel);
-
 			for (int32 LODId = 0; LODId < NumHLODLevels; ++LODId)
 			{
+				// Handle HierachicalLOD volumes first
+				HandleHLODVolumes(InLevel);
+
 				// Reuse clusters from previous HLOD level (only works for HLOD level 1 and beyond)
 				if (BuildLODLevelSettings[LODId].bReusePreviousLevelClusters && LODId > 0)
 				{
@@ -347,40 +347,41 @@ void FHierarchicalLODBuilder::GenerateAsSingleCluster(const int32 NumHLODLevels,
 void FHierarchicalLODBuilder::InitializeClusters(ULevel* InLevel, const int32 LODIdx, float CullCost, const bool bPreviewBuild, bool const bVolumesOnly)
 {
 	SCOPE_LOG_TIME(TEXT("STAT_HLOD_InitializeClusters"), nullptr);
+
+	// Check whether or not this actor falls within a HierarchicalLODVolume, if so add to the Volume's cluster and exclude from normal process
+	auto ProcessVolumeClusters = [this](AActor* InActor) -> bool
+	{
+		for (TPair<AHierarchicalLODVolume*, FLODCluster>& Cluster : HLODVolumeClusters)
+		{
+			if (Cluster.Key->EncompassesPoint(InActor->GetActorLocation(), Cluster.Key->bIncludeOverlappingActors ? InActor->GetComponentsBoundingBox().GetSize().Size() : 0.0f, nullptr))
+			{
+				FBox BoundingBox = InActor->GetComponentsBoundingBox(true);
+				FBox VolumeBox = Cluster.Key->GetComponentsBoundingBox(true);
+
+				if (VolumeBox.IsInside(BoundingBox) || (Cluster.Key->bIncludeOverlappingActors && VolumeBox.Intersect(BoundingBox)))
+				{
+					FLODCluster ActorCluster(InActor);
+					Cluster.Value += ActorCluster;
+					return true;
+				}
+			}
+		}
+
+		return false;
+	};
+
+	Clusters.Empty();
+
 	if (InLevel->Actors.Num() > 0)
 	{
 		if (LODIdx == 0)
 		{
-			Clusters.Empty();
-
 			for (int32 ActorId = 0; ActorId < InLevel->Actors.Num(); ++ActorId)
 			{
 				AActor* Actor = InLevel->Actors[ActorId];
 				const bool bShouldGenerate = ShouldGenerateCluster(Actor, bPreviewBuild, LODIdx);
 				if (bShouldGenerate)
 				{
-					// Check whether or not this actor falls within a HierarchicalLODVolume, if so add to the Volume's cluster and exclude from normal process
-					auto ProcessVolumeClusters = [this](AActor* InActor) -> bool
-					{
-						for (TPair<AHierarchicalLODVolume*, FLODCluster>& Cluster : HLODVolumeClusters)
-						{
-							if (Cluster.Key->EncompassesPoint(InActor->GetActorLocation(), Cluster.Key->bIncludeOverlappingActors ? InActor->GetComponentsBoundingBox().GetSize().Size() : 0.0f, nullptr))
-							{
-								FBox BoundingBox = InActor->GetComponentsBoundingBox(true);
-								FBox VolumeBox = Cluster.Key->GetComponentsBoundingBox(true);
-
-								if (VolumeBox.IsInside(BoundingBox) || (Cluster.Key->bIncludeOverlappingActors && VolumeBox.Intersect(BoundingBox)))
-								{
-									FLODCluster ActorCluster(InActor);
-									Cluster.Value += ActorCluster;
-									return true;
-								}
-							}
-						}
-
-						return false;
-					};
-
 					if (bVolumesOnly)
 					{
 						ProcessVolumeClusters(Actor);
@@ -391,7 +392,7 @@ void FHierarchicalLODBuilder::InitializeClusters(ULevel* InLevel, const int32 LO
 						{
 							ValidStaticMeshActorsInLevel.Add(Actor);
 						}
-					}					
+					}
 				}
 				else
 				{
@@ -421,15 +422,24 @@ void FHierarchicalLODBuilder::InitializeClusters(ULevel* InLevel, const int32 LO
 				}
 			}
 		}
-		else // at this point we only care for LODActors
+		else
 		{
-			Clusters.Empty();
-
 			// we filter the LOD index first
 			TArray<AActor*> Actors;
 
 			Actors.Append(LODLevelLODActors[LODIdx - 1]);
-						
+
+			// Re-evaluate level actors
+			for (int32 Idx = 0; Idx < ValidStaticMeshActorsInLevel.Num(); ++Idx)
+			{
+				AActor* Actor = ValidStaticMeshActorsInLevel[Idx];
+				if (!ShouldGenerateCluster(Actor, bPreviewBuild, LODIdx))
+				{
+					ValidStaticMeshActorsInLevel.RemoveAt(Idx);
+					--Idx;
+				}
+			}
+
 			// Re-evaluate rejected actors
 			for (AActor* Actor : RejectedActorsInLevel)
 			{
@@ -445,6 +455,8 @@ void FHierarchicalLODBuilder::InitializeClusters(ULevel* InLevel, const int32 LO
 
 			Actors.Append(ValidStaticMeshActorsInLevel);
 
+			Actors.RemoveAll(ProcessVolumeClusters);
+
 			// first we generate graph with 2 pair nodes
 			// this is very expensive when we have so many actors
 			// so we'll need to optimize later @todo
@@ -1009,12 +1021,9 @@ void FHierarchicalLODBuilder::MergeClustersAndBuildActors(ULevel* InLevel, const
 			}
 		}
 
-		if (LODIdx == 0)
+		for (TPair<AHierarchicalLODVolume*, FLODCluster>& Cluster : HLODVolumeClusters)
 		{
-			  for (TPair<AHierarchicalLODVolume*, FLODCluster>& Cluster : HLODVolumeClusters)
-			{
-				Clusters.Add(Cluster.Value);
-			}
+			Clusters.Add(Cluster.Value);
 		}
 
 
