HLOD CalculateOverlap crashes when spheres barely overlap

diff --git a/Engine/Source/Editor/UnrealEd/Private/LODCluster.cpp b/Engine/Source/Editor/UnrealEd/Private/LODCluster.cpp
--- a/Engine/Source/Editor/UnrealEd/Private/LODCluster.cpp
+++ b/Engine/Source/Editor/UnrealEd/Private/LODCluster.cpp
@@ -64,16 +64,12 @@ const float CalculateOverlap(const FSphere& ASphere, const float AFillingFactor,
 
 	float OverlapRadius1 = ((ARadius + BRadius)*(ARadius + BRadius) - Distance*Distance) * (Distance*Distance - (ARadius - BRadius)*(ARadius - BRadius));
 	float OverlapRadius2 = 2 * Distance;
-	float OverlapRedius = FMath::Sqrt(OverlapRadius1) / OverlapRadius2;
-	float OverlapRediusSq = OverlapRedius*OverlapRedius;
-
-	check(!FMath::IsNearlyZero(OverlapRadius1));
+	float OverlapRadius = FMath::Sqrt(OverlapRadius1) / OverlapRadius2;
+	float OverlapRadiusSq = FMath::Square(OverlapRadius);
 
 	float ConstPI = PI/6.0f;
-	float AVolume = ConstPI*(3*OverlapRediusSq + ACapHeight*ACapHeight) * ACapHeight;
-	float BVolume = ConstPI*(3*OverlapRediusSq + BCapHeight*BCapHeight) * BCapHeight;
-
-	check(!FMath::IsNearlyZero(AVolume) && !FMath::IsNearlyZero(BVolume));
+	float AVolume = ConstPI*(3*OverlapRadiusSq + ACapHeight*ACapHeight) * ACapHeight;
+	float BVolume = ConstPI*(3*OverlapRadiusSq + BCapHeight*BCapHeight) * BCapHeight;
 
 	float TotalVolume = AFillingFactor*AVolume + BFillingFactor*BVolume;
 	return TotalVolume;
