AssetRegUtilCommandlet now supports generate partially updated file open order

diff --git a/Engine/Source/Editor/UnrealEd/Classes/Commandlets/AssetRegUtilCommandlet.h b/Engine/Source/Editor/UnrealEd/Classes/Commandlets/AssetRegUtilCommandlet.h
--- a/Engine/Source/Editor/UnrealEd/Classes/Commandlets/AssetRegUtilCommandlet.h
+++ b/Engine/Source/Editor/UnrealEd/Classes/Commandlets/AssetRegUtilCommandlet.h
@@ -35,4 +35,11 @@ protected:
 		const int32& DepSet, int32& DepOrder, int32 DepHierarchy, TSet<FName>& ProcessedFiles, const TSet<FName>& OriginalSet, const FName& FilePath, const FName& PackageFName, const TArray<FName>& FilterByClass);
 
 	void ReorderOrderFile(const FString& OrderFilePath, const FString& ReorderFileOutPath);
+
+	// Generate a new FOO which only takes partial update from new FOO files with input of an old FOO file and a new FOO file
+	bool GeneratePartiallyUpdatedOrderFile(const FString& OldOrderFilePath, const FString& NewOrderFilePath, const FString& OutOrderFilePath, const float PatchSizePerfBalanceFactor);
+
+private:
+
+	bool LoadOrderFiles(const FString& OrderFilePath, TSet<FName>& OrderFiles);
 };
diff --git a/Engine/Source/Editor/UnrealEd/Private/Commandlets/AssetRegUtilCommandlet.cpp b/Engine/Source/Editor/UnrealEd/Private/Commandlets/AssetRegUtilCommandlet.cpp
--- a/Engine/Source/Editor/UnrealEd/Private/Commandlets/AssetRegUtilCommandlet.cpp
+++ b/Engine/Source/Editor/UnrealEd/Private/Commandlets/AssetRegUtilCommandlet.cpp
@@ -253,172 +253,255 @@ void UAssetRegUtilCommandlet::RecursivelyGrabDependencies(TArray<FSortableDepend
 
 void UAssetRegUtilCommandlet::ReorderOrderFile(const FString& OrderFilePath, const FString& ReorderFileOutPath)
 {
-	UE_LOG(LogAssetRegUtil, Display, TEXT("Parsing order file: %s"), *OrderFilePath);
-	FString Text;
-	if (FFileHelper::LoadFileToString(Text, *OrderFilePath))
+	TSet<FName> OriginalEntrySet;
+	if (!LoadOrderFiles(OrderFilePath, OriginalEntrySet))
 	{
-		//just parsing the list into a set while discarding order values, as we can rely on iteration order matching added order
-		TSet<FName> OriginalEntrySet;
-
-		TArray<FString> Lines;
-		Text.ParseIntoArray(Lines, TEXT("\n"), true);
-
-		//parse each entry out of the order list. adapted from UnrealPak code, might want to unify somewhere.
-		OriginalEntrySet.Reserve(Lines.Num());
-		for (int32 EntryIndex = 0; EntryIndex < Lines.Num(); ++EntryIndex)
-		{
-			Lines[EntryIndex].ReplaceInline(TEXT("\r"), TEXT(""));
-			Lines[EntryIndex].ReplaceInline(TEXT("\n"), TEXT(""));
-			//discard the order number, assuming the list is in-order and has no special bits in use at this stage.
-			int32 QuoteIndex;
-			if (Lines[EntryIndex].FindLastChar('"', QuoteIndex))
-			{
-				FString ReadNum = Lines[EntryIndex].RightChop(QuoteIndex + 1);
-				Lines[EntryIndex] = Lines[EntryIndex].Left(QuoteIndex + 1);
-				//verify our expectations about the order, just in case something changes on the OpenOrder generation side
-				ReadNum.TrimStartInline();
-				if (ReadNum.IsNumeric())
-				{
-					const int32 ExplicitOrder = FCString::Atoi(*ReadNum);
-					if (ExplicitOrder != EntryIndex + 1)
-					{
-						UE_LOG(LogAssetRegUtil, Warning, TEXT("Unexpected order: %i vs %i"), ExplicitOrder, EntryIndex + 1);
-					}
-				}
-			}
-
-			Lines[EntryIndex] = Lines[EntryIndex].TrimQuotes();
-
-			const FString EntryPath = *Lines[EntryIndex].ToLower();
-			const FName EntryFName = FName(*EntryPath);
-			OriginalEntrySet.Add(EntryFName);
-		}
+		UE_LOG(LogAssetRegUtil, Warning, TEXT("Could not load specified order file."));
+		return;
+	}
 
-		UE_LOG(LogAssetRegUtil, Display, TEXT("Generating new file order via Asset Registry."));
+	UE_LOG(LogAssetRegUtil, Display, TEXT("Generating new file order via Asset Registry."));
 
-		TArray<FSortableDependencyEntry> UnsortedEntries;
-		UnsortedEntries.Empty(OriginalEntrySet.Num());
+	TArray<FSortableDependencyEntry> UnsortedEntries;
+	UnsortedEntries.Empty(OriginalEntrySet.Num());
 
-		// quick elimination tset
-		TSet<FName> ProcessedFiles;
-		ProcessedFiles.Reserve(OriginalEntrySet.Num());
+	// quick elimination tset
+	TSet<FName> ProcessedFiles;
+	ProcessedFiles.Reserve(OriginalEntrySet.Num());
 
-		TSet<FName> AssetExtensions;
-		AssetExtensions.Add(NAME_uasset);
-		AssetExtensions.Add(NAME_umap);
+	TSet<FName> AssetExtensions;
+	AssetExtensions.Add(NAME_uasset);
+	AssetExtensions.Add(NAME_umap);
 
-		TSet<FName> ExtraAssetExtensions;
-		ExtraAssetExtensions.Add(NAME_uexp);
-		ExtraAssetExtensions.Add(NAME_ubulk);
-		ExtraAssetExtensions.Add(NAME_uptnl);
+	TSet<FName> ExtraAssetExtensions;
+	ExtraAssetExtensions.Add(NAME_uexp);
+	ExtraAssetExtensions.Add(NAME_ubulk);
+	ExtraAssetExtensions.Add(NAME_uptnl);
 
-		TArray<FName> FilterByClasses;
-		FilterByClasses.Add(FName(TEXT("Material")));
-		FilterByClasses.Add(FName(TEXT("MaterialFunction")));
-		FilterByClasses.Add(FName(TEXT("MaterialInstance")));
-		FilterByClasses.Add(FName(TEXT("BlueprintCore")));
-		FilterByClasses.Add(FName(TEXT("ParticleEmitter")));
-		FilterByClasses.Add(FName(TEXT("ParticleModule")));
+	TArray<FName> FilterByClasses;
+	FilterByClasses.Add(FName(TEXT("Material")));
+	FilterByClasses.Add(FName(TEXT("MaterialFunction")));
+	FilterByClasses.Add(FName(TEXT("MaterialInstance")));
+	FilterByClasses.Add(FName(TEXT("BlueprintCore")));
+	FilterByClasses.Add(FName(TEXT("ParticleEmitter")));
+	FilterByClasses.Add(FName(TEXT("ParticleModule")));
 
 
-		int32 DepSet = 0; // this is the root set for the dependency (i.e files with a number probably came from the same core dependency)
-		for (const FName& FilePath : OriginalEntrySet)
+	int32 DepSet = 0; // this is the root set for the dependency (i.e files with a number probably came from the same core dependency)
+	for (const FName& FilePath : OriginalEntrySet)
+	{
+		++DepSet;
+		if (!ProcessedFiles.Contains(FilePath))
 		{
-			++DepSet;
-			if (!ProcessedFiles.Contains(FilePath))
+			const FString FilePathExtension = FPaths::GetExtension(FilePath.ToString());
+			const FName FNameExtension = FName(*FilePathExtension);
+			if (AssetExtensions.Contains( FNameExtension) )
 			{
-				const FString FilePathExtension = FPaths::GetExtension(FilePath.ToString());
-				const FName FNameExtension = FName(*FilePathExtension);
-				if (AssetExtensions.Contains( FNameExtension) )
+				FString PackageName;
+				if (FPackageName::TryConvertFilenameToLongPackageName(FilePath.ToString(), PackageName))
 				{
-					FString PackageName;
-					if (FPackageName::TryConvertFilenameToLongPackageName(FilePath.ToString(), PackageName))
-					{
-						FName PackageFName(*PackageName);
-						int32 DependencyOrderIndex = 0;
-						RecursivelyGrabDependencies(UnsortedEntries, DepSet, DependencyOrderIndex, 0, ProcessedFiles, OriginalEntrySet, FilePath, PackageFName, FilterByClasses);
-					}
-					else
-					{
-						//special case for packages outside of our mounted paths, pick up the header and the export without any dependency-gathering.
-						ProcessedFiles.Add(FilePath);
-						UnsortedEntries.Add(FSortableDependencyEntry(NAME_UnresolvedPackageName, FilePath, FNameExtension, DepSet, 0, 0, false, TSet<FName>()));
-					}
+					FName PackageFName(*PackageName);
+					int32 DependencyOrderIndex = 0;
+					RecursivelyGrabDependencies(UnsortedEntries, DepSet, DependencyOrderIndex, 0, ProcessedFiles, OriginalEntrySet, FilePath, PackageFName, FilterByClasses);
 				}
-				else if ( ExtraAssetExtensions.Contains(FNameExtension) == false )
+				else
 				{
-					//not a package, no need to do special sorting/handling for headers and exports
-					UnsortedEntries.Add(FSortableDependencyEntry(FilePath, FNameExtension, DepSet));
+					//special case for packages outside of our mounted paths, pick up the header and the export without any dependency-gathering.
 					ProcessedFiles.Add(FilePath);
+					UnsortedEntries.Add(FSortableDependencyEntry(NAME_UnresolvedPackageName, FilePath, FNameExtension, DepSet, 0, 0, false, TSet<FName>()));
 				}
 			}
+			else if ( ExtraAssetExtensions.Contains(FNameExtension) == false )
+			{
+				//not a package, no need to do special sorting/handling for headers and exports
+				UnsortedEntries.Add(FSortableDependencyEntry(FilePath, FNameExtension, DepSet));
+				ProcessedFiles.Add(FilePath);
+			}
 		}
+	}
 
-		for (int32 I = UnsortedEntries.Num()-1; I >= 0; --I)
+	for (int32 I = UnsortedEntries.Num()-1; I >= 0; --I)
+	{
+		const FSortableDependencyEntry& DependencyEntry = UnsortedEntries[I];
+		if (DependencyEntry.bIsAsset)
 		{
-			const FSortableDependencyEntry& DependencyEntry = UnsortedEntries[I];
-			if (DependencyEntry.bIsAsset)
+			// find all the uexp files and ubulk files and uptnl files
+			FString StringPath = DependencyEntry.FilePath.ToString();
+			for (const FName& ExtraAssetExtension : ExtraAssetExtensions)
 			{
-				// find all the uexp files and ubulk files and uptnl files
-				FString StringPath = DependencyEntry.FilePath.ToString();
-				for (const FName& ExtraAssetExtension : ExtraAssetExtensions)
+				FString ExtraAssetPath = FPaths::ChangeExtension(StringPath, ExtraAssetExtension.ToString());
+				FName ExtraAssetPathFName = FName(*ExtraAssetPath);
+				if (OriginalEntrySet.Contains(ExtraAssetPathFName))
 				{
-					FString ExtraAssetPath = FPaths::ChangeExtension(StringPath, ExtraAssetExtension.ToString());
-					FName ExtraAssetPathFName = FName(*ExtraAssetPath);
-					if (OriginalEntrySet.Contains(ExtraAssetPathFName))
-					{
-						check(!ProcessedFiles.Contains(ExtraAssetPathFName));
-						ProcessedFiles.Add(ExtraAssetPathFName);
-						TSet<FName> Classes = DependencyEntry.Classes;
-						UnsortedEntries.Add(FSortableDependencyEntry(DependencyEntry.LongPackageName, ExtraAssetPathFName, ExtraAssetExtension, DependencyEntry.DepSet, DependencyEntry.DepHierarchy, DependencyEntry.DepOrder, DependencyEntry.bHasDependencies, MoveTemp(Classes)));
-					}
+					check(!ProcessedFiles.Contains(ExtraAssetPathFName));
+					ProcessedFiles.Add(ExtraAssetPathFName);
+					TSet<FName> Classes = DependencyEntry.Classes;
+					UnsortedEntries.Add(FSortableDependencyEntry(DependencyEntry.LongPackageName, ExtraAssetPathFName, ExtraAssetExtension, DependencyEntry.DepSet, DependencyEntry.DepHierarchy, DependencyEntry.DepOrder, DependencyEntry.bHasDependencies, MoveTemp(Classes)));
 				}
 			}
 		}
+	}
 
-		//if this were to fire, first guess would be that there's somehow a rogue export without a header
-		check(OriginalEntrySet.Num() == ProcessedFiles.Num() && ProcessedFiles.Num() == UnsortedEntries.Num());
+	//if this were to fire, first guess would be that there's somehow a rogue export without a header
+	check(OriginalEntrySet.Num() == ProcessedFiles.Num() && ProcessedFiles.Num() == UnsortedEntries.Num());
 
 
-		TArray<FName> ShouldGroupExtensions;
-		ShouldGroupExtensions.Add(NAME_ubulk);
+	TArray<FName> ShouldGroupExtensions;
+	ShouldGroupExtensions.Add(NAME_ubulk);
 
 
-		TMap<FName, int32> ExtensionPriority;
-		ExtensionPriority.Add(NAME_umap, 0); 
-		ExtensionPriority.Add(NAME_uasset, 0);
-		ExtensionPriority.Add(NAME_uexp, 1);
-		ExtensionPriority.Add(NAME_uptnl, 1);
-		ExtensionPriority.Add(NAME_ubulk, 1);
-		FSortableDependencySort DependencySortClass(ShouldGroupExtensions, FilterByClasses, ExtensionPriority);
-		UnsortedEntries.Sort(DependencySortClass);
+	TMap<FName, int32> ExtensionPriority;
+	ExtensionPriority.Add(NAME_umap, 0); 
+	ExtensionPriority.Add(NAME_uasset, 0);
+	ExtensionPriority.Add(NAME_uexp, 1);
+	ExtensionPriority.Add(NAME_uptnl, 1);
+	ExtensionPriority.Add(NAME_ubulk, 1);
+	FSortableDependencySort DependencySortClass(ShouldGroupExtensions, FilterByClasses, ExtensionPriority);
+	UnsortedEntries.Sort(DependencySortClass);
 
 
 		
-		UE_LOG(LogAssetRegUtil, Display, TEXT("Writing output: %s"), *ReorderFileOutPath);
-		FArchive* OutArc = IFileManager::Get().CreateFileWriter(*ReorderFileOutPath);
-		if (OutArc)
+	UE_LOG(LogAssetRegUtil, Display, TEXT("Writing output: %s"), *ReorderFileOutPath);
+	FArchive* OutArc = IFileManager::Get().CreateFileWriter(*ReorderFileOutPath);
+	if (OutArc)
+	{
+		//base from 1, to match existing order list convention
+		uint64 NewOrderIndex = 1;
+		for (const FSortableDependencyEntry& SortedEntry : UnsortedEntries)
 		{
-			//base from 1, to match existing order list convention
-			uint64 NewOrderIndex = 1;
-			for (const FSortableDependencyEntry& SortedEntry : UnsortedEntries)
+			const FString& FilePath = SortedEntry.FilePath.ToString();
+			FString OutputLine = FString::Printf(TEXT("\"%s\" %llu\n"), *FilePath, NewOrderIndex++);
+			OutArc->Serialize(const_cast<ANSICHAR*>(StringCast<ANSICHAR>(*OutputLine).Get()), OutputLine.Len());
+		}
+		OutArc->Close();
+		delete OutArc;
+	}
+	else
+	{
+		UE_LOG(LogAssetRegUtil, Warning, TEXT("Could not open specified output file."));
+	}
+}
+
+bool UAssetRegUtilCommandlet::LoadOrderFiles(const FString& OrderFilePath, TSet<FName>& OrderFiles)
+{
+	UE_LOG(LogAssetRegUtil, Display, TEXT("Parsing order file: %s"), *OrderFilePath);
+	FString Text;
+	if (!FFileHelper::LoadFileToString(Text, *OrderFilePath))
+	{
+		UE_LOG(LogAssetRegUtil, Error, TEXT("Could not open file %s"), *OrderFilePath);
+		return false;
+	}
+
+	TArray<FString> Lines;
+	Text.ParseIntoArray(Lines, TEXT("\n"), true);
+
+	// Parse each entry out of the order list. adapted from UnrealPak code, might want to unify somewhere.
+	OrderFiles.Reserve(Lines.Num());
+	for (int32 EntryIndex = 0; EntryIndex < Lines.Num(); ++EntryIndex)
+	{
+		Lines[EntryIndex].ReplaceInline(TEXT("\r"), TEXT(""));
+		Lines[EntryIndex].ReplaceInline(TEXT("\n"), TEXT(""));
+		//discard the order number, assuming the list is in-order and has no special bits in use at this stage.
+		int32 QuoteIndex;
+		if (Lines[EntryIndex].FindLastChar('"', QuoteIndex))
+		{
+			FString ReadNum = Lines[EntryIndex].RightChop(QuoteIndex + 1);
+			Lines[EntryIndex] = Lines[EntryIndex].Left(QuoteIndex + 1);
+			//verify our expectations about the order, just in case something changes on the OpenOrder generation side
+			ReadNum.TrimStartInline();
+			if (ReadNum.IsNumeric())
 			{
-				const FString& FilePath = SortedEntry.FilePath.ToString();
-				FString OutputLine = FString::Printf(TEXT("\"%s\" %llu\n"), *FilePath, NewOrderIndex++);
-				OutArc->Serialize(const_cast<ANSICHAR*>(StringCast<ANSICHAR>(*OutputLine).Get()), OutputLine.Len());
+				const int32 ExplicitOrder = FCString::Atoi(*ReadNum);
+				if (ExplicitOrder != EntryIndex + 1)
+				{
+					UE_LOG(LogAssetRegUtil, Warning, TEXT("Unexpected order: %i vs %i"), ExplicitOrder, EntryIndex + 1);
+				}
 			}
-			OutArc->Close();
-			delete OutArc;
+		}
+
+		Lines[EntryIndex] = Lines[EntryIndex].TrimQuotes();
+
+		const FString EntryPath = *Lines[EntryIndex].ToLower();
+		const FName EntryFName = FName(*EntryPath);
+		OrderFiles.Add(EntryFName);
+	}
+
+	return true;
+}
+
+bool UAssetRegUtilCommandlet::GeneratePartiallyUpdatedOrderFile(const FString& OldOrderFilePath, const FString& NewOrderFilePath, const FString& OutOrderFilePath, const float PatchSizePerfBalanceFactor)
+{
+	TSet<FName> OldOrderFileLineSet;
+	TSet<FName> NewOrderFileLineSet;
+	if (!LoadOrderFiles(OldOrderFilePath, OldOrderFileLineSet) ||
+		!LoadOrderFiles(NewOrderFilePath, NewOrderFileLineSet))
+	{
+		return false;
+	}
+
+	// Remove deleted files from old order file lines
+	for (auto OldOrderFileIter = OldOrderFileLineSet.CreateIterator(); OldOrderFileIter; ++OldOrderFileIter)
+	{
+		if (!NewOrderFileLineSet.Contains(*OldOrderFileIter))
+		{
+			OldOrderFileIter.RemoveCurrent();
+		}
+	}
+	OldOrderFileLineSet.CompactStable();
+
+	// Add new files to old order file lines
+	int LastFoundIndexInOldFileLines = -1;
+	TArray<FName> OldOrderFileLineArray = OldOrderFileLineSet.Array();
+	for (auto NewOrderFileIter = NewOrderFileLineSet.CreateConstIterator(); NewOrderFileIter; ++NewOrderFileIter)
+	{
+		int IndexInOldFileLines = OldOrderFileLineArray.IndexOfByKey(*NewOrderFileIter);
+
+		if (IndexInOldFileLines != INDEX_NONE)
+		{
+			LastFoundIndexInOldFileLines = IndexInOldFileLines;
 		}
 		else
 		{
-			UE_LOG(LogAssetRegUtil, Warning, TEXT("Could not open specified output file."));
+			OldOrderFileLineArray.Insert(*NewOrderFileIter, LastFoundIndexInOldFileLines + 1);
+			LastFoundIndexInOldFileLines++;
+		}
+	}
+	OldOrderFileLineSet = TSet<FName>(OldOrderFileLineArray);
+
+	// Write new file order to out file path
+	FArchive* OutArc = IFileManager::Get().CreateFileWriter(*OutOrderFilePath);
+	if (OutArc)
+	{
+		// Order number starts from 1
+		int CurrentOrderNumber = 1;
+
+		for (int i = 0; i < NewOrderFileLineSet.Num() * PatchSizePerfBalanceFactor; i++)
+		{
+			const FName OrderFileLine = NewOrderFileLineSet[FSetElementId::FromInteger(i)];
+			FString OutputLine = FString::Printf(TEXT("\"%s\" %d\r\n"), *OrderFileLine.ToString(), CurrentOrderNumber);
+			OutArc->Serialize(const_cast<ANSICHAR*>(StringCast<ANSICHAR>(*OutputLine).Get()), OutputLine.Len());
+			CurrentOrderNumber++;
+
+			OldOrderFileLineSet.Remove(OrderFileLine);
+		}
+
+		for (auto OldOrderFileIter = OldOrderFileLineSet.CreateIterator(); OldOrderFileIter; ++OldOrderFileIter)
+		{
+			const FName OrderFileLine = *OldOrderFileIter;
+			FString OutputLine = FString::Printf(TEXT("\"%s\" %d\r\n"), *OrderFileLine.ToString(), CurrentOrderNumber);
+			OutArc->Serialize(const_cast<ANSICHAR*>(StringCast<ANSICHAR>(*OutputLine).Get()), OutputLine.Len());
+			CurrentOrderNumber++;
 		}
+
+		OutArc->Close();
+		delete OutArc;
 	}
 	else
 	{
-		UE_LOG(LogAssetRegUtil, Warning, TEXT("Could not load specified order file."));
+		UE_LOG(LogAssetRegUtil, Warning, TEXT("Could not open specified output file."))
 	}
+
+	return true;
 }
 
 int32 UAssetRegUtilCommandlet::Main(const FString& CmdLineParams)
@@ -440,6 +523,28 @@ int32 UAssetRegUtilCommandlet::Main(const FString& CmdLineParams)
 		}
 
 		ReorderOrderFile(ReorderFile, ReorderOutput);
+
+		// Generate partially-updated order file
+		float PatchSizePerfBalanceFactor;	// Set the value close to 0.0 to favor patch size and close to 1.0 to favor streaming performance
+		FParse::Value(*CmdLineParams, TEXT("PatchSizePerfBalanceFactor="), PatchSizePerfBalanceFactor);
+		PatchSizePerfBalanceFactor = FMath::Clamp(PatchSizePerfBalanceFactor, 0.f, 1.f);
+
+		FString OldOrderFilePath;
+		if (FParse::Value(*CmdLineParams, TEXT("OldFileOpenOrderFile="), OldOrderFilePath, false))
+		{
+			FString OutPartialOrderFilePath;
+			if (!FParse::Value(*CmdLineParams, TEXT("PartialFileOpenOrderOutput="), OutPartialOrderFilePath, false))
+			{
+				//if nothing specified, base it on the input name
+				OutPartialOrderFilePath = FPaths::SetExtension(FPaths::SetExtension(ReorderFile, TEXT("")) + TEXT("PartialUpdate"), FPaths::GetExtension(ReorderFile));
+			}
+
+			GeneratePartiallyUpdatedOrderFile(OldOrderFilePath, ReorderOutput, OutPartialOrderFilePath, PatchSizePerfBalanceFactor);
+		}
+		else
+		{
+			UE_LOG(LogAssetRegUtil, Error, TEXT("Invalid or missing arguments.  Skipping generating partially-updated order file."));
+		}
 	}
 
 	return 0;
