Switched from FAssetEditorManager and GEditor->Layers to subsystem versions of the same.

Fixes CIS warnings about deprecated functionality.

diff --git a/Engine/Source/Editor/UnrealEd/Private/ObjectTools.cpp b/Engine/Source/Editor/UnrealEd/Private/ObjectTools.cpp
--- a/Engine/Source/Editor/UnrealEd/Private/ObjectTools.cpp
+++ b/Engine/Source/Editor/UnrealEd/Private/ObjectTools.cpp
@@ -72,7 +72,6 @@
 #include "ReferencedAssetsUtils.h"
 #include "AssetRegistryModule.h"
 #include "PackagesDialog.h"
-#include "Toolkits/AssetEditorManager.h"
 #include "PropertyEditorModule.h"
 #include "Kismet2/KismetEditorUtilities.h"
 #include "Kismet2/KismetReinstanceUtilities.h"
@@ -85,13 +84,14 @@
 #include "ComponentRecreateRenderStateContext.h"
 #include "Framework/Notifications/NotificationManager.h"
 #include "Widgets/Notifications/SNotificationList.h"
-#include "Layers/ILayers.h"
+#include "Layers/LayersSubsystem.h"
 #include "Engine/SCS_Node.h"
 #include "ShaderCompiler.h"
 #include "Templates/UniquePtr.h"
 #include "Engine/MapBuildDataRegistry.h"
 #include "HAL/PlatformApplicationMisc.h"
 #include "Kismet2/BlueprintEditorUtils.h"
+#include "Subsystems/AssetEditorSubsystem.h"
 
 DEFINE_LOG_CATEGORY_STATIC(LogObjectTools, Log, All);
 
@@ -922,7 +922,7 @@ namespace ObjectTools
 		if ( ObjectToConsolidateTo )
 		{
 			// Close all editors to avoid changing references to temporary objects used by the editor
-			if (!FAssetEditorManager::Get().CloseAllAssetEditors())
+			if (!GEditor->GetEditorSubsystem<UAssetEditorSubsystem>()->CloseAllAssetEditors())
 			{
 				// Failed to close at least one editor. It is possible that this editor has in-memory object references
 				// which are not prepared to be changed dynamically so it is not safe to continue
@@ -2438,7 +2438,7 @@ namespace ObjectTools
 		{
 			if (Object->IsAsset())
 			{
-				const TArray<IAssetEditorInstance*> ObjectEditors = FAssetEditorManager::Get().FindEditorsForAsset(Object);
+				const TArray<IAssetEditorInstance*> ObjectEditors = GEditor->GetEditorSubsystem<UAssetEditorSubsystem>()->FindEditorsForAsset(Object);
 				for (IAssetEditorInstance* ObjectEditorInstance : ObjectEditors)
 				{
 					if (!ObjectEditorInstance->CloseWindow())
@@ -2600,6 +2600,7 @@ namespace ObjectTools
 		// Destroy all Actor instances
 		if ( ActorsToDelete.Num() > 0 )
 		{
+			ULayersSubsystem* Layers = GEditor->GetEditorSubsystem<ULayersSubsystem>();
 			for ( TArray<AActor*>::TConstIterator ActorItr( ActorsToDelete ); ActorItr; ++ActorItr )
 			{
 				AActor* CurActor = *ActorItr;
@@ -2617,7 +2618,7 @@ namespace ObjectTools
 					}
 
 					// Destroy the Actor instance. This is similar to edactDeleteSelected(), but we don't request user confirmation here.
-					GEditor->Layers->DisassociateActorFromLayers( CurActor );
+					Layers->DisassociateActorFromLayers( CurActor );
 					if( CurActor->GetWorld() )
 					{
 						CurActor->GetWorld()->EditorDestroyActor( CurActor, false );
