Build fix after previous bad merge in 8790675

diff --git a/Engine/Source/Editor/UnrealEd/Private/ThumbnailRendering/TextureThumbnailRenderer.cpp b/Engine/Source/Editor/UnrealEd/Private/ThumbnailRendering/TextureThumbnailRenderer.cpp
--- a/Engine/Source/Editor/UnrealEd/Private/ThumbnailRendering/TextureThumbnailRenderer.cpp
+++ b/Engine/Source/Editor/UnrealEd/Private/ThumbnailRendering/TextureThumbnailRenderer.cpp
@@ -82,7 +82,7 @@ void UTextureThumbnailRenderer::Draw(UObject* Object, int32 X, int32 Y, uint32 W
 		{
 			bool bIsNormalMap = Texture2DArray->IsNormalMap();
 			bool bIsSingleChannel = true;
-			BatchedElementParameters = new FBatchedElementTexture2DPreviewParameters(0, 0, bIsNormalMap, bIsSingleChannel, false, true);
+			BatchedElementParameters = new FBatchedElementTexture2DPreviewParameters(0, 0, bIsNormalMap, bIsSingleChannel, false, false, true);
 		}
 		else if (TextureLightProfile)
 		{
diff --git a/Engine/Source/Runtime/Renderer/Private/VT/AllocatedVirtualTexture.h b/Engine/Source/Runtime/Renderer/Private/VT/AllocatedVirtualTexture.h
--- a/Engine/Source/Runtime/Renderer/Private/VT/AllocatedVirtualTexture.h
+++ b/Engine/Source/Runtime/Renderer/Private/VT/AllocatedVirtualTexture.h
@@ -79,7 +79,7 @@ private:
 	// We create a page table layer for each physical space group in each producer.
 	struct FPageTableLayerDesc
 	{
-		FVirtualTexturePhysicalSpace* PhysicalSpace;
+		TRefCountPtr<FVirtualTexturePhysicalSpace> PhysicalSpace;
 		uint8 UniqueProducerIndex;
 		uint8 ProducerPhysicalGroupIndex;
 		uint16 ProducerTextureLayerMask;
