Fixing crash at skeletal mesh export

diff --git a/Engine/Source/Editor/UnrealEd/Private/Fbx/FbxSkeletalMeshExport.cpp b/Engine/Source/Editor/UnrealEd/Private/Fbx/FbxSkeletalMeshExport.cpp
--- a/Engine/Source/Editor/UnrealEd/Private/Fbx/FbxSkeletalMeshExport.cpp
+++ b/Engine/Source/Editor/UnrealEd/Private/Fbx/FbxSkeletalMeshExport.cpp
@@ -279,7 +279,7 @@ FbxNode* FFbxExporter::CreateMesh(const USkeletalMesh* SkelMesh, const TCHAR* Me
 					ShapeControlPoints[VertIndex] = Converter.ConvertToFbxPos(Position);
 				}
 				
-				int32 NumberOfDeltas;
+				int32 NumberOfDeltas = 0;
 				FMorphTargetDelta* MorphTargetDeltas = MorphTarget->GetMorphTargetDelta(LODIndex, NumberOfDeltas);
 				for (int32 MorphTargetDeltaIndex = 0; MorphTargetDeltaIndex < NumberOfDeltas; ++MorphTargetDeltaIndex)
 				{
diff --git a/Engine/Source/Runtime/Engine/Private/MorphTools.cpp b/Engine/Source/Runtime/Engine/Private/MorphTools.cpp
--- a/Engine/Source/Runtime/Engine/Private/MorphTools.cpp
+++ b/Engine/Source/Runtime/Engine/Private/MorphTools.cpp
@@ -29,6 +29,7 @@ FMorphTargetDelta* UMorphTarget::GetMorphTargetDelta(int32 LODIndex, int32& OutN
 		return MorphModel.Vertices.GetData();
 	}
 
+	OutNumDeltas = 0;
 	return NULL;
 }
 
