FStructureEditorUtils::GetVariableFriendlyNameForProperty now doesn't cause a check() if Struct->EditorData is null.

Source: 8739676 (//UE4/Dev-Editor)

diff --git a/Engine/Source/Editor/UnrealEd/Private/Kismet2/StructureEditorUtils.cpp b/Engine/Source/Editor/UnrealEd/Private/Kismet2/StructureEditorUtils.cpp
--- a/Engine/Source/Editor/UnrealEd/Private/Kismet2/StructureEditorUtils.cpp
+++ b/Engine/Source/Editor/UnrealEd/Private/Kismet2/StructureEditorUtils.cpp
@@ -490,8 +490,19 @@ FString FStructureEditorUtils::GetVariableFriendlyName(const UUserDefinedStruct*
 
 FString FStructureEditorUtils::GetVariableFriendlyNameForProperty(const UUserDefinedStruct* Struct, const UProperty* Property)
 {
-	const FStructVariableDescription* VarDesc = Struct && Property ? GetVarDesc(Struct).FindByPredicate(FFindByNameHelper<FStructVariableDescription>(Property->GetFName())) : nullptr;
-	return VarDesc ? VarDesc->FriendlyName : FString();
+	if (Struct && Property)
+	{
+		const TArray<FStructVariableDescription>* VarDescArray = GetVarDescPtr(Struct);
+		if (VarDescArray)
+		{
+			const FStructVariableDescription* VarDesc = VarDescArray->FindByPredicate(FFindByNameHelper<FStructVariableDescription>(Property->GetFName()));
+			if (VarDesc)
+			{
+				return VarDesc->FriendlyName;
+			}
+		}
+	}
+	return FString();
 }
 
 UProperty* FStructureEditorUtils::GetPropertyByFriendlyName(const UUserDefinedStruct* Struct, FString DisplayName)
@@ -650,7 +661,7 @@ const FStructVariableDescription* FStructureEditorUtils::GetVarDescByGuid(const
 
 FString FStructureEditorUtils::GetTooltip(const UUserDefinedStruct* Struct)
 {
-	const UUserDefinedStructEditorData* StructEditorData = Struct ? Cast<const UUserDefinedStructEditorData>(Struct->EditorData) : NULL;
+	const UUserDefinedStructEditorData* StructEditorData = Struct ? Cast<const UUserDefinedStructEditorData>(Struct->EditorData) : nullptr;
 	return StructEditorData ? StructEditorData->ToolTip : FString();
 }
 
