Fix an issue that was introduced in 4.23 that caused the UI to stun when building lighting via the 'Build All' button

The blocking wait for the lighting build has been moved to a higher level to reudce the number of code paths that can reach it and added a check against the BuildSettings so that the blocking wait is only invoked for the command line build.

diff --git a/Engine/Source/Editor/UnrealEd/Private/EditorBuildUtils.cpp b/Engine/Source/Editor/UnrealEd/Private/EditorBuildUtils.cpp
--- a/Engine/Source/Editor/UnrealEd/Private/EditorBuildUtils.cpp
+++ b/Engine/Source/Editor/UnrealEd/Private/EditorBuildUtils.cpp
@@ -142,6 +142,15 @@ bool FEditorBuildUtils::EditorAutomatedBuildAndSubmit( const FEditorAutomatedBui
 		{
 			LogErrorMessage( NSLOCTEXT("UnrealEd", "AutomatedBuild_Error_BuildFailed", "The map build failed or was canceled."), OutErrorMessages );
 		}
+
+		// If we are going to shutdown after this has run (ie running from the cmdline) then we should wait for the distributed lighting build to complete.
+		if (BuildSettings.bShutdownEditorOnCompletion)
+		{
+			while (GUnrealEd->IsLightingBuildCurrentlyRunning())
+			{
+				GUnrealEd->UpdateBuildLighting();
+			}
+		}	
 	}
 
 	// If any map errors resulted from the build, process them according to the behavior specified in the build settings
@@ -1030,11 +1039,7 @@ void FBuildAllHandler::ProcessBuild(const TWeakPtr<SBuildProgressWidget>& BuildP
 			LightingOptions.QualityLevel = (ELightingBuildQuality)QualityLevel;
 
 			GUnrealEd->BuildLighting(LightingOptions);
-			while (GUnrealEd->IsLightingBuildCurrentlyRunning())
-			{
-				GUnrealEd->UpdateBuildLighting();
-			}
-			
+
 			// TODO!
 			//bShouldMapCheck = false;
 
