SetParentEditorOnly must always be followed by a PostEditChange

diff --git a/Engine/Source/Editor/UnrealEd/Private/PreviewMaterial.cpp b/Engine/Source/Editor/UnrealEd/Private/PreviewMaterial.cpp
--- a/Engine/Source/Editor/UnrealEd/Private/PreviewMaterial.cpp
+++ b/Engine/Source/Editor/UnrealEd/Private/PreviewMaterial.cpp
@@ -1608,6 +1608,7 @@ void UMaterialEditorInstanceConstant::UpdateSourceInstanceParent()
 	}
 
 	SourceInstance->SetParentEditorOnly( Parent );
+	SourceInstance->PostEditChange();
 }
 
 
