Fixed crash when editing a notify property

The timing view was not getting properly refreshed, so could end up out of sync with the notify tracks.
Also fixed a division-by-zero when notifies are dragged around a montage with no sections.

diff --git a/Engine/Source/Editor/Persona/Private/SAnimNotifyPanel.cpp b/Engine/Source/Editor/Persona/Private/SAnimNotifyPanel.cpp
--- a/Engine/Source/Editor/Persona/Private/SAnimNotifyPanel.cpp
+++ b/Engine/Source/Editor/Persona/Private/SAnimNotifyPanel.cpp
@@ -4645,8 +4645,7 @@ void SAnimNotifyPanel::OnPropertyChanged(UObject* ChangedObject, FPropertyChange
 			if(Event.Notify == ChangedObject || Event.NotifyStateClass == ChangedObject)
 			{
 				// If we've changed a notify present in the sequence, refresh our tracks.
-				Sequence->RefreshCacheData();
-				RefreshNotifyTracks();
+				Update();
 			}
 		}
 	}
diff --git a/Engine/Source/Editor/UnrealEd/Public/SCurveEditor.h b/Engine/Source/Editor/UnrealEd/Public/SCurveEditor.h
--- a/Engine/Source/Editor/UnrealEd/Public/SCurveEditor.h
+++ b/Engine/Source/Editor/UnrealEd/Public/SCurveEditor.h
@@ -69,7 +69,7 @@ struct FTrackScaleInfo
 	/** Local Widget Space -> Curve Input domain. */
 	float LocalXToInput(float ScreenX) const
 	{
-		return (ScreenX/PixelsPerInput) + ViewMinInput;
+		return PixelsPerInput != 0.0f ? (ScreenX/PixelsPerInput) + ViewMinInput : 0.0f;
 	}
 
 	/** Curve Input domain -> local Widget Space */
@@ -81,7 +81,7 @@ struct FTrackScaleInfo
 	/** Local Widget Space -> Curve Output domain. */
 	float LocalYToOutput(float ScreenY) const
 	{
-		return (ViewOutputRange + ViewMinOutput) - (ScreenY/PixelsPerOutput);
+		return PixelsPerOutput != 0.0f ? (ViewOutputRange + ViewMinOutput) - (ScreenY/PixelsPerOutput) : 0.0f;
 	}
 
 	/** Curve Output domain -> local Widget Space */
