Prevent out of bound array access when proxy meshes are rebuilt.
Output an error instead: "Invalid LODActor found (invalid LOD level with regards to the current HLOD build settings). You must rebuild the HLOD clusters"

diff --git a/Engine/Source/Editor/UnrealEd/Private/HierarchicalLOD.cpp b/Engine/Source/Editor/UnrealEd/Private/HierarchicalLOD.cpp
--- a/Engine/Source/Editor/UnrealEd/Private/HierarchicalLOD.cpp
+++ b/Engine/Source/Editor/UnrealEd/Private/HierarchicalLOD.cpp
@@ -690,12 +690,18 @@ void FHierarchicalLODBuilder::BuildMeshesForLODActors(bool bForceAll)
 
 			// Retrieve LOD actors from the level
 			uint32 NumLODActors = 0;
-			for (int32 ActorId = 0; ActorId < LevelIter->Actors.Num(); ++ActorId)
+			for (AActor* Actor : LevelIter->Actors)
 			{
-				AActor* Actor = LevelIter->Actors[ActorId];
-				if (Actor && Actor->IsA<ALODActor>())
+				if (ALODActor* LODActor = Cast<ALODActor>(Actor))
 				{
-					ALODActor* LODActor = CastChecked<ALODActor>(Actor);
+					// Ensure the LODActor we found is valid for our current HLOD build settings.
+					if (LODActor->LODLevel - 1 >= LODLevelActors.Num())
+					{
+						FMessageLog("HLODResults").Error()
+							->AddToken(FTextToken::Create(LOCTEXT("HLODError_ClusterRebuildNeeded", "Invalid LODActor found (invalid LOD level with regards to the current HLOD build settings). You must rebuild the HLOD clusters!")))
+							->AddToken(FUObjectToken::Create(LODActor));
+						continue;
+					}
 
 					if (bForceAll || (!LODActor->IsBuilt(true) && LODActor->HasValidSubActors()))
 					{
