Fixed crash when undoing material changes in static mesh editor when object is also instanced in the loaded level
Partition ActorComponent to the beginning of the map using sort like it used to be before the depth sort

diff --git a/Engine/Source/Editor/UnrealEd/Private/EditorTransaction.cpp b/Engine/Source/Editor/UnrealEd/Private/EditorTransaction.cpp
--- a/Engine/Source/Editor/UnrealEd/Private/EditorTransaction.cpp
+++ b/Engine/Source/Editor/UnrealEd/Private/EditorTransaction.cpp
@@ -885,16 +885,7 @@ void FTransaction::Apply()
 	// so do a quick sort on Outer depth, component will deeper than their owner
 	ChangedObjects.KeySort([](UObject& A, UObject& B)
 	{
-		auto GetObjectDepth = [](UObject* InObj)
-		{
-			int32 Depth = 0;
-			for (UObject* Outer = InObj; Outer; Outer = Outer->GetOuter())
-			{
-				++Depth;
-			}
-			return Depth;
-		};
-		return GetObjectDepth(&A) > GetObjectDepth(&B);
+		return Cast<UActorComponent>(&A) != nullptr;
 	});
 
 	TArray<ULevel*> LevelsToCommitModelSurface;
@@ -976,16 +967,7 @@ void FTransaction::Finalize()
 	// so do a quick sort on Outer depth, component will deeper than their owner
 	ChangedObjects.KeySort([](UObject& A, UObject& B)
 	{
-		auto GetObjectDepth = [](UObject* InObj)
-		{
-			int32 Depth = 0;
-			for (UObject* Outer = InObj; Outer; Outer = Outer->GetOuter())
-			{
-				++Depth;
-			}
-			return Depth;
-		};
-		return GetObjectDepth(&A) > GetObjectDepth(&B);
+		return Cast<UActorComponent>(&A) != nullptr;
 	});
 
 	for (auto ChangedObjectIt : ChangedObjects)
