Allow controlling the GAllowUnversionedContentInEditor setting via a cvar.

diff --git a/Engine/Source/Editor/UnrealEd/Private/Commandlets/PackageUtilities.cpp b/Engine/Source/Editor/UnrealEd/Private/Commandlets/PackageUtilities.cpp
--- a/Engine/Source/Editor/UnrealEd/Private/Commandlets/PackageUtilities.cpp
+++ b/Engine/Source/Editor/UnrealEd/Private/Commandlets/PackageUtilities.cpp
@@ -1562,7 +1562,7 @@ int32 UPkgInfoCommandlet::Main( const FString& Params )
 
 		if (!bDumpProperties)
 		{
-			TGuardValue<bool> GuardAllowUnversionedContentInEditor(GAllowUnversionedContentInEditor, true);
+			TGuardValue<int32> GuardAllowUnversionedContentInEditor(GAllowUnversionedContentInEditor, 1);
 			TGuardValue<int32> GuardAllowCookedContentInEditor(GAllowCookedDataInEditorBuilds, 1);
 			TRefCountPtr<FUObjectSerializeContext> LoadContext(FUObjectThreadContext::Get().GetSerializeContext());
 			BeginLoad(LoadContext);
@@ -1588,7 +1588,7 @@ int32 UPkgInfoCommandlet::Main( const FString& Params )
 			Reader = FArchiveStackTraceReader::CreateFromFile(*Filename);
 			if (Reader)
 			{
-				TGuardValue<bool> GuardAllowUnversionedContentInEditor(GAllowUnversionedContentInEditor, true);
+				TGuardValue<int32> GuardAllowUnversionedContentInEditor(GAllowUnversionedContentInEditor, 1);
 				TGuardValue<int32> GuardAllowCookedContentInEditor(GAllowCookedDataInEditorBuilds, 1);
 				UPackage* LoadedPackage = LoadPackage(Package, *Filename, LOAD_NoVerify, Reader);
 				if (LoadedPackage)
diff --git a/Engine/Source/Runtime/CoreUObject/Private/Serialization/ArchiveStackTrace.cpp b/Engine/Source/Runtime/CoreUObject/Private/Serialization/ArchiveStackTrace.cpp
--- a/Engine/Source/Runtime/CoreUObject/Private/Serialization/ArchiveStackTrace.cpp
+++ b/Engine/Source/Runtime/CoreUObject/Private/Serialization/ArchiveStackTrace.cpp
@@ -1311,7 +1311,7 @@ void FArchiveStackTrace::DumpPackageHeaderDiffs(const FPackageData& SourcePackag
 	FString DestAssetPackageName = FPaths::Combine(TEXT("/Memory"), TEXT("/DestForDiff"), *AssetPathName);
 
 	TGuardValue<bool> GuardIsSavingPackage(GIsSavingPackage, false);
-	TGuardValue<bool> GuardAllowUnversionedContentInEditor(GAllowUnversionedContentInEditor, true);
+	TGuardValue<int32> GuardAllowUnversionedContentInEditor(GAllowUnversionedContentInEditor, 1);
 	TGuardValue<int32> GuardAllowCookedDataInEditorBuilds(GAllowCookedDataInEditorBuilds, 1);
 
 	FLinkerLoad* SourceLinker = nullptr;
diff --git a/Engine/Source/Runtime/CoreUObject/Private/UObject/UObjectGlobals.cpp b/Engine/Source/Runtime/CoreUObject/Private/UObject/UObjectGlobals.cpp
--- a/Engine/Source/Runtime/CoreUObject/Private/UObject/UObjectGlobals.cpp
+++ b/Engine/Source/Runtime/CoreUObject/Private/UObject/UObjectGlobals.cpp
@@ -62,7 +62,15 @@ DEFINE_LOG_CATEGORY(LogUObjectGlobals);
 #endif
 
 bool GIsSavingPackage = false;
-bool GAllowUnversionedContentInEditor = false;
+
+int32 GAllowUnversionedContentInEditor = 0;
+
+static FAutoConsoleVariableRef CVarAllowUnversionedContentInEditor(
+	TEXT("s.AllowUnversionedContentInEditor"),
+	GAllowUnversionedContentInEditor,
+	TEXT("If true, allows unversioned content to be loaded by the editor."),
+	ECVF_Default
+);
 
 /** Object annotation used by the engine to keep track of which objects are selected */
 FUObjectAnnotationSparseBool GSelectedObjectAnnotation;
diff --git a/Engine/Source/Runtime/CoreUObject/Public/UObject/UObjectGlobals.h b/Engine/Source/Runtime/CoreUObject/Public/UObject/UObjectGlobals.h
--- a/Engine/Source/Runtime/CoreUObject/Public/UObject/UObjectGlobals.h
+++ b/Engine/Source/Runtime/CoreUObject/Public/UObject/UObjectGlobals.h
@@ -43,7 +43,7 @@ DECLARE_DWORD_COUNTER_STAT_EXTERN(TEXT("FindObjectFast"),STAT_FindObjectFast,STA
 /** Set while in SavePackage() to detect certain operations that are illegal while saving */
 extern COREUOBJECT_API bool					GIsSavingPackage;
 /** This allows loading unversioned cooked content in the editor */
-extern COREUOBJECT_API bool					GAllowUnversionedContentInEditor;
+extern COREUOBJECT_API int32				GAllowUnversionedContentInEditor;
 /** This allows loading cooked content in the editor */
 extern COREUOBJECT_API int32				GAllowCookedDataInEditorBuilds;
 
