Nightly HLOD Rebuild: skip levels that fail to check-out/submit

diff --git a/Engine/Source/Editor/UnrealEd/Classes/Commandlets/ResavePackagesCommandlet.h b/Engine/Source/Editor/UnrealEd/Classes/Commandlets/ResavePackagesCommandlet.h
--- a/Engine/Source/Editor/UnrealEd/Classes/Commandlets/ResavePackagesCommandlet.h
+++ b/Engine/Source/Editor/UnrealEd/Classes/Commandlets/ResavePackagesCommandlet.h
@@ -207,7 +207,7 @@ protected:
 
 	bool CanCheckoutFile(const FString& Filename, FString& CheckedOutUser);
 
-	void CheckoutAndSavePackage(UPackage* Package, TArray<FString>& SublevelFilenames);
+	void CheckoutAndSavePackage(UPackage* Package, TArray<FString>& SublevelFilenames, bool bIgnoreAlreadyCheckedOut = false);
 
 	// Print out a message only if running in very verbose mode
 	void VerboseMessage(const FString& Message);
diff --git a/Engine/Source/Editor/UnrealEd/Private/Commandlets/ContentCommandlets.cpp b/Engine/Source/Editor/UnrealEd/Private/Commandlets/ContentCommandlets.cpp
--- a/Engine/Source/Editor/UnrealEd/Private/Commandlets/ContentCommandlets.cpp
+++ b/Engine/Source/Editor/UnrealEd/Private/Commandlets/ContentCommandlets.cpp
@@ -1266,13 +1266,27 @@ bool UResavePackagesCommandlet::CheckoutFile(const FString& Filename, bool bAddF
 	if (SourceControlState.IsValid())
 	{
 		FString CurrentlyCheckedOutUser;
-		if (SourceControlState->IsCheckedOutOther(&CurrentlyCheckedOutUser) && !bIgnoreAlreadyCheckedOut)
+		if (SourceControlState->IsCheckedOutOther(&CurrentlyCheckedOutUser))
 		{
-			UE_LOG(LogContentCommandlet, Error, TEXT("[REPORT] %s level is already checked out by someone else (%s), can not submit!"), *Filename, *CurrentlyCheckedOutUser);
+			if (!bIgnoreAlreadyCheckedOut)
+			{
+				UE_LOG(LogContentCommandlet, Error, TEXT("[REPORT] %s level is already checked out by someone else (%s), can not submit!"), *Filename, *CurrentlyCheckedOutUser);
+			}
+			else
+			{
+				UE_LOG(LogContentCommandlet, Warning, TEXT("[REPORT] %s level is already checked out by someone else (%s), can not submit!"), *Filename, *CurrentlyCheckedOutUser);
+			}
 		}
 		else if (!SourceControlState->IsCurrent())
 		{
-			UE_LOG(LogContentCommandlet, Error, TEXT("[REPORT] %s is not synced to head, can not submit"), *Filename);
+			if (!bIgnoreAlreadyCheckedOut)
+			{
+				UE_LOG(LogContentCommandlet, Error, TEXT("[REPORT] %s is not synced to head, can not submit"), *Filename);
+			}
+			else
+			{
+				UE_LOG(LogContentCommandlet, Warning, TEXT("[REPORT] %s is not synced to head, can not submit"), *Filename);
+			}
 		}
 		else if ( SourceControlState->IsSourceControlled() == false )
 		{
@@ -1310,7 +1324,6 @@ bool UResavePackagesCommandlet::CheckoutFile(const FString& Filename, bool bAddF
 	return false;
 }
 
-
 bool UResavePackagesCommandlet::RevertFile(const FString& Filename)
 {
 	ISourceControlProvider& SourceControlProvider = ISourceControlModule::Get().GetProvider();
@@ -1333,8 +1346,6 @@ bool UResavePackagesCommandlet::RevertFile(const FString& Filename)
 	return bSuccesfullyReverted;
 }
 
-
-
 bool UResavePackagesCommandlet::CanCheckoutFile(const FString& Filename, FString& CheckedOutUser)
 {
 	ISourceControlProvider& SourceControlProvider = ISourceControlModule::Get().GetProvider();
@@ -1342,8 +1353,13 @@ bool UResavePackagesCommandlet::CanCheckoutFile(const FString& Filename, FString
 	bool bCanCheckout = true;
 	if (SourceControlState.IsValid())
 	{
-		if (!SourceControlState->IsCheckedOut() && SourceControlState->IsCheckedOutOther(&CheckedOutUser))
+		if (!SourceControlState->CanCheckout())
 		{
+			if (!SourceControlState->IsCheckedOutOther(&CheckedOutUser))
+			{
+				CheckedOutUser = "";
+			}
+
 			bCanCheckout = false;
 		}
 	}
@@ -1351,7 +1367,7 @@ bool UResavePackagesCommandlet::CanCheckoutFile(const FString& Filename, FString
 	return bCanCheckout;
 }
 
-void UResavePackagesCommandlet::CheckoutAndSavePackage(UPackage* Package, TArray<FString>& SublevelFilenames)
+void UResavePackagesCommandlet::CheckoutAndSavePackage(UPackage* Package, TArray<FString>& SublevelFilenames, bool bIgnoreAlreadyCheckedOut)
 {
 	check(Package);
 
@@ -1360,18 +1376,20 @@ void UResavePackagesCommandlet::CheckoutAndSavePackage(UPackage* Package, TArray
 	{
 		if (IFileManager::Get().FileExists(*PackageFilename))
 		{
-			if (CheckoutFile(PackageFilename, true))
+			if (CheckoutFile(PackageFilename, true, bIgnoreAlreadyCheckedOut))
 			{
 				SublevelFilenames.Add(PackageFilename);
+				SavePackageHelper(Package, PackageFilename);
 			}
-			SavePackageHelper(Package, PackageFilename);
 		}
 		else
 		{
-			SavePackageHelper(Package, PackageFilename);
-			if (CheckoutFile(PackageFilename, true))
+			if (SavePackageHelper(Package, PackageFilename))
 			{
-				SublevelFilenames.Add(PackageFilename);
+				if (CheckoutFile(PackageFilename, true, bIgnoreAlreadyCheckedOut))
+				{
+					SublevelFilenames.Add(PackageFilename);
+				}
 			}
 		}
 	}
@@ -1488,30 +1506,42 @@ void UResavePackagesCommandlet::PerformAdditionalOperations(class UWorld* World,
 		{
 			World->LoadSecondaryLevels(true, NULL);
 
-			for (ULevelStreaming* NextStreamingLevel : World->GetStreamingLevels())
+			TArray<ULevelStreaming*> StreamingLevels = World->GetStreamingLevels();
+			for (ULevelStreaming* StreamingLevel : StreamingLevels)
 			{
+				bool bShouldBeLoaded = true;
+
 				// If we are not building HLODs or are but also rebuilding lighting we check out the level file, otherwise we don't to try and ensure a minimal HLOD rebuild
 				if (!bShouldBuildHLOD || bBuildingNonHLODData)
 				{
-					CheckOutLevelFile(NextStreamingLevel->GetLoadedLevel());
+					CheckOutLevelFile(StreamingLevel->GetLoadedLevel());
 				}
 
 				FString StreamingLevelPackageFilename;
-				const FString StreamingLevelWorldAssetPackageName = NextStreamingLevel->GetWorldAssetPackageName();
+				const FString StreamingLevelWorldAssetPackageName = StreamingLevel->GetWorldAssetPackageName();
 				if (FPackageName::DoesPackageExist(StreamingLevelWorldAssetPackageName, NULL, &StreamingLevelPackageFilename))
 				{
 					// If we are building HLODs only, we dont check out the files ahead of rebuilding the data
 					if(bShouldBuildHLOD && !bBuildingNonHLODData)
 					{
-						FString CurrentlyCheckedOutUser;
-						if (CanCheckoutFile(StreamingLevelPackageFilename, CurrentlyCheckedOutUser) || !bSkipCheckedOutFiles)
+						FString OutUser;
+						if (CanCheckoutFile(StreamingLevelPackageFilename, OutUser) || !bSkipCheckedOutFiles)
 						{
 							CheckedOutPackagesFilenames.Add(StreamingLevelPackageFilename);
 						}
 						else 
 						{
-							UE_LOG(LogContentCommandlet, Warning, TEXT("[REPORT] Skipping %s as it is checked out by %s"), *StreamingLevelPackageFilename, *CurrentlyCheckedOutUser);
-						}						
+							if (OutUser.Len())
+							{
+								UE_LOG(LogContentCommandlet, Warning, TEXT("[REPORT] Skipping %s as it is checked out by %s"), *StreamingLevelPackageFilename, *OutUser);
+							}
+							else
+							{
+								UE_LOG(LogContentCommandlet, Warning, TEXT("[REPORT] Skipping %s as it could not be checked out (not at head revision ?)"), *StreamingLevelPackageFilename);
+							}
+
+							bShouldBeLoaded = false;
+						}
 					}
 					else
 					{
@@ -1524,13 +1554,19 @@ void UResavePackagesCommandlet::PerformAdditionalOperations(class UWorld* World,
 						{
 							UE_LOG(LogContentCommandlet, Error, TEXT("[REPORT] %s is currently already checked out, cannot continue resaving"), *StreamingLevelPackageFilename);
 							bShouldProceedWithRebuild = false;
+							bShouldBeLoaded = false;
 							break;
 						}
 					}
 				}
 
-				NextStreamingLevel->SetShouldBeVisible(true);
-				NextStreamingLevel->SetShouldBeLoaded(true);
+				if (!bShouldBeLoaded)
+				{
+					World->RemoveStreamingLevel(StreamingLevel);
+				}
+				
+				StreamingLevel->SetShouldBeVisible(bShouldBeLoaded);
+				StreamingLevel->SetShouldBeLoaded(bShouldBeLoaded);
 			}
 		}
 
@@ -1638,7 +1674,7 @@ void UResavePackagesCommandlet::PerformAdditionalOperations(class UWorld* World,
 				{
 					if (Package->IsDirty())
 					{
-						CheckoutAndSavePackage(Package, CheckedOutPackagesFilenames);
+						CheckoutAndSavePackage(Package, CheckedOutPackagesFilenames, bSkipCheckedOutFiles);
 					}
 				}
 			}
