Deprecating ARRAY_COUNT and changing it to UE_ARRAY_COUNT.

diff --git a/Engine/Source/Editor/UnrealEd/Private/Lightmass/ImportVolumetricLightmap.cpp b/Engine/Source/Editor/UnrealEd/Private/Lightmass/ImportVolumetricLightmap.cpp
--- a/Engine/Source/Editor/UnrealEd/Private/Lightmass/ImportVolumetricLightmap.cpp
+++ b/Engine/Source/Editor/UnrealEd/Private/Lightmass/ImportVolumetricLightmap.cpp
@@ -143,7 +143,7 @@ bool CopyFromBrickmapTexel(
 
 		*(FFloat3Packed*)&BrickData.AmbientVector.Data[LinearDestCellIndex * sizeof(FFloat3Packed)] = FilteredVolumeLookupReconverted<FFloat3Packed>(BrickTextureCoordinate, CurrentLevelData.BrickDataDimensions, (const FFloat3Packed*)BrickData.AmbientVector.Data.GetData());
 		
-		for (int32 i = 0; i < ARRAY_COUNT(BrickData.SHCoefficients); i++)
+		for (int32 i = 0; i < UE_ARRAY_COUNT(BrickData.SHCoefficients); i++)
 		{
 			*(FColor*)&BrickData.SHCoefficients[i].Data[LinearDestCellIndex * sizeof(FColor)] = FilteredVolumeLookupReconverted<FColor>(BrickTextureCoordinate, CurrentLevelData.BrickDataDimensions, (const FColor*)BrickData.SHCoefficients[i].Data.GetData());
 		}
@@ -196,7 +196,7 @@ void FLightmassProcessor::ImportIrradianceTasks(bool& bGenerateSkyShadowing, TAr
 				Swarm.ReadChannel(Channel, &NewBrick.AverageClosestGeometryDistance, sizeof(NewBrick.AverageClosestGeometryDistance));
 				ReadArray(Channel, NewBrick.AmbientVector);
 
-				for (int32 i = 0; i < ARRAY_COUNT(NewBrick.SHCoefficients); i++)
+				for (int32 i = 0; i < UE_ARRAY_COUNT(NewBrick.SHCoefficients); i++)
 				{
 					ReadArray(Channel, NewBrick.SHCoefficients[i]);
 				}
@@ -232,7 +232,7 @@ int32 NumDilateOverEmbeddedVoxelsPasses = 2;
 struct FFilteredBrickData
 {
 	FFloat3Packed AmbientVector;
-	FColor SHCoefficients[ARRAY_COUNT(FImportedVolumetricLightmapBrick::SHCoefficients)];
+	FColor SHCoefficients[UE_ARRAY_COUNT(FImportedVolumetricLightmapBrick::SHCoefficients)];
 	uint8 DirectionalLightShadowing;
 };
 
@@ -305,10 +305,10 @@ void FilterWithNeighbors(
 					{
 						//@todo - filter SkyBentNormal from neighbors too
 						FLinearColor AmbientVector = FLinearColor(0, 0, 0, 0);
-						FLinearColor SHCoefficients[ARRAY_COUNT(Brick.SHCoefficients)];
+						FLinearColor SHCoefficients[UE_ARRAY_COUNT(Brick.SHCoefficients)];
 						FLinearColor DirectionalLightShadowing = FLinearColor(0, 0, 0, 0);
 
-						for (int32 i = 0; i < ARRAY_COUNT(SHCoefficients); i++)
+						for (int32 i = 0; i < UE_ARRAY_COUNT(SHCoefficients); i++)
 						{
 							SHCoefficients[i] = FLinearColor(0, 0, 0, 0);
 						}
@@ -351,9 +351,9 @@ void FilterWithNeighbors(
 												FLinearColor NeighborAmbientVector = FilteredVolumeLookup<FFloat3Packed>(BrickTextureCoordinate, CurrentLevelData.BrickDataDimensions, (const FFloat3Packed*)CurrentLevelData.BrickData.AmbientVector.Data.GetData());
 												AmbientVector += NeighborAmbientVector * Weight;
 
-												for (int32 i = 0; i < ARRAY_COUNT(SHCoefficients); i++)
+												for (int32 i = 0; i < UE_ARRAY_COUNT(SHCoefficients); i++)
 												{
-													static_assert(ARRAY_COUNT(SHCoefficients) == 6, "Assuming 2 SHCoefficient vectors per color channel");
+													static_assert(UE_ARRAY_COUNT(SHCoefficients) == 6, "Assuming 2 SHCoefficient vectors per color channel");
 
 													// Weight by ambient before filtering, normalized SH coefficients don't filter properly
 													float AmbientCoefficient = NeighborAmbientVector.Component(i / 2);
@@ -384,9 +384,9 @@ void FilterWithNeighbors(
 							const FLinearColor FilteredAmbientColor = AmbientVector * InvTotalWeight;
 							FilteredBrickData[LinearVoxelIndex].AmbientVector = ConvertFromLinearColor<FFloat3Packed>(FilteredAmbientColor);
 
-							for (int32 i = 0; i < ARRAY_COUNT(SHCoefficients); i++)
+							for (int32 i = 0; i < UE_ARRAY_COUNT(SHCoefficients); i++)
 							{
-								static_assert(ARRAY_COUNT(SHCoefficients) == 6, "Assuming 2 SHCoefficient vectors per color channel");
+								static_assert(UE_ARRAY_COUNT(SHCoefficients) == 6, "Assuming 2 SHCoefficient vectors per color channel");
 
 								float AmbientCoefficient = FMath::Max(FilteredAmbientColor.Component(i / 2), KINDA_SMALL_NUMBER);
 
@@ -420,7 +420,7 @@ void FilterWithNeighbors(
 						FFloat3Packed* DestAmbientVector = (FFloat3Packed*)&CurrentLevelData.BrickData.AmbientVector.Data[LinearBrickLayoutCellIndex * sizeof(FFloat3Packed)];
 						*DestAmbientVector = FilteredBrickData[LinearVoxelIndex].AmbientVector;
 
-						for (int32 i = 0; i < ARRAY_COUNT(FilteredBrickData[LinearVoxelIndex].SHCoefficients); i++)
+						for (int32 i = 0; i < UE_ARRAY_COUNT(FilteredBrickData[LinearVoxelIndex].SHCoefficients); i++)
 						{
 							FColor* DestCoefficients = (FColor*)&CurrentLevelData.BrickData.SHCoefficients[i].Data[LinearBrickLayoutCellIndex * sizeof(FColor)];
 							*DestCoefficients = FilteredBrickData[LinearVoxelIndex].SHCoefficients[i];
@@ -1056,7 +1056,7 @@ void FLightmassProcessor::ImportVolumetricLightmap()
 		CurrentLevelData.BrickData.SkyBentNormal.Format = PF_B8G8R8A8;
 		CurrentLevelData.BrickData.DirectionalLightShadowing.Format = PF_G8;
 
-		for (int32 i = 0; i < ARRAY_COUNT(CurrentLevelData.BrickData.SHCoefficients); i++)
+		for (int32 i = 0; i < UE_ARRAY_COUNT(CurrentLevelData.BrickData.SHCoefficients); i++)
 		{
 			CurrentLevelData.BrickData.SHCoefficients[i].Format = PF_B8G8R8A8;
 		}
@@ -1116,7 +1116,7 @@ void FLightmassProcessor::ImportVolumetricLightmap()
 
 		CurrentLevelData.BrickData.DirectionalLightShadowing.Resize(TotalBrickDataSize * GPixelFormats[CurrentLevelData.BrickData.DirectionalLightShadowing.Format].BlockBytes);
 
-		for (int32 i = 0; i < ARRAY_COUNT(CurrentLevelData.BrickData.SHCoefficients); i++)
+		for (int32 i = 0; i < UE_ARRAY_COUNT(CurrentLevelData.BrickData.SHCoefficients); i++)
 		{
 			const int32 Stride = GPixelFormats[CurrentLevelData.BrickData.SHCoefficients[i].Format].BlockBytes;
 			CurrentLevelData.BrickData.SHCoefficients[i].Resize(TotalBrickDataSize * Stride);
@@ -1149,7 +1149,7 @@ void FLightmassProcessor::ImportVolumetricLightmap()
 				(const uint8*)Brick.AmbientVector.GetData(),
 				CurrentLevelData.BrickData.AmbientVector.Data.GetData());
 
-			for (int32 i = 0; i < ARRAY_COUNT(Brick.SHCoefficients); i++)
+			for (int32 i = 0; i < UE_ARRAY_COUNT(Brick.SHCoefficients); i++)
 			{
 				CopyBrickToAtlasVolumeTexture(
 					GPixelFormats[CurrentLevelData.BrickData.SHCoefficients[i].Format].BlockBytes,
@@ -1350,7 +1350,7 @@ void FLightmassProcessor::ImportVolumetricLightmap()
 					SubLevelData.BrickData.SkyBentNormal.Format = PF_B8G8R8A8;
 					SubLevelData.BrickData.DirectionalLightShadowing.Format = PF_G8;
 
-					for (int32 i = 0; i < ARRAY_COUNT(SubLevelData.BrickData.SHCoefficients); i++)
+					for (int32 i = 0; i < UE_ARRAY_COUNT(SubLevelData.BrickData.SHCoefficients); i++)
 					{
 						SubLevelData.BrickData.SHCoefficients[i].Format = PF_B8G8R8A8;
 					}
@@ -1371,7 +1371,7 @@ void FLightmassProcessor::ImportVolumetricLightmap()
 
 				SubLevelData.BrickData.DirectionalLightShadowing.Resize(TotalBrickDataSize * GPixelFormats[SubLevelData.BrickData.DirectionalLightShadowing.Format].BlockBytes);
 
-				for (int32 i = 0; i < ARRAY_COUNT(SubLevelData.BrickData.SHCoefficients); i++)
+				for (int32 i = 0; i < UE_ARRAY_COUNT(SubLevelData.BrickData.SHCoefficients); i++)
 				{
 					const int32 Stride = GPixelFormats[SubLevelData.BrickData.SHCoefficients[i].Format].BlockBytes;
 					SubLevelData.BrickData.SHCoefficients[i].Resize(TotalBrickDataSize * Stride);
@@ -1422,7 +1422,7 @@ void FLightmassProcessor::ImportVolumetricLightmap()
 					SubLevelBrickLayoutPosition,
 					SubLevelData.BrickData.AmbientVector.Data);
 
-				for (int32 i = 0; i < ARRAY_COUNT(Brick.SHCoefficients); i++)
+				for (int32 i = 0; i < UE_ARRAY_COUNT(Brick.SHCoefficients); i++)
 				{
 					CopyBetweenAtlasVolumeTextures(
 						GPixelFormats[CurrentLevelData.BrickData.SHCoefficients[i].Format].BlockBytes,
@@ -1554,7 +1554,7 @@ void FLightmassProcessor::ImportVolumetricLightmap()
 
 				ConvertBGRA8ToRGBA8ForLayer(SubLevelData.BrickData.SkyBentNormal);
 
-				for (int32 i = 0; i < ARRAY_COUNT(SubLevelData.BrickData.SHCoefficients); i++)
+				for (int32 i = 0; i < UE_ARRAY_COUNT(SubLevelData.BrickData.SHCoefficients); i++)
 				{
 					ConvertBGRA8ToRGBA8ForLayer(SubLevelData.BrickData.SHCoefficients[i]);
 				}
diff --git a/Engine/Source/Editor/UnrealEd/Private/Lightmass/Lightmass.cpp b/Engine/Source/Editor/UnrealEd/Private/Lightmass/Lightmass.cpp
--- a/Engine/Source/Editor/UnrealEd/Private/Lightmass/Lightmass.cpp
+++ b/Engine/Source/Editor/UnrealEd/Private/Lightmass/Lightmass.cpp
@@ -3147,13 +3147,13 @@ bool FLightmassProcessor::BeginRun()
 	{
 		JobSpecification32 = NSwarm::FJobSpecification( *LightmassExecutable32, *CommandLineParameters, ( NSwarm::TJobTaskFlags )JobFlags );
 		JobSpecification32.AddDependencies( RequiredDependencyPaths32.GetArray(), RequiredDependencyPaths32.Num(), OptionalDependencyPaths32.GetArray(), OptionalDependencyPaths32.Num() );
-		JobSpecification32.AddDescription( DescriptionKeys, DescriptionValues, ARRAY_COUNT(DescriptionKeys) );
+		JobSpecification32.AddDescription( DescriptionKeys, DescriptionValues, UE_ARRAY_COUNT(DescriptionKeys) );
 	}
 	else
 	{
 		JobSpecification64 = NSwarm::FJobSpecification( *LightmassExecutable64, *CommandLineParameters, ( NSwarm::TJobTaskFlags )JobFlags );
 		JobSpecification64.AddDependencies( RequiredDependencyPaths64.GetArray(), RequiredDependencyPaths64.Num(), OptionalDependencyPaths64.GetArray(), OptionalDependencyPaths64.Num() );
-		JobSpecification64.AddDescription( DescriptionKeys, DescriptionValues, ARRAY_COUNT(DescriptionKeys) );
+		JobSpecification64.AddDescription( DescriptionKeys, DescriptionValues, UE_ARRAY_COUNT(DescriptionKeys) );
 	}
 	int32 ErrorCode = Swarm.BeginJobSpecification( JobSpecification32, JobSpecification64 );
 	if( ErrorCode < 0 )
diff --git a/Engine/Source/Editor/UnrealEd/Private/MaterialGraph.cpp b/Engine/Source/Editor/UnrealEd/Private/MaterialGraph.cpp
--- a/Engine/Source/Editor/UnrealEd/Private/MaterialGraph.cpp
+++ b/Engine/Source/Editor/UnrealEd/Private/MaterialGraph.cpp
@@ -53,7 +53,7 @@ void UMaterialGraph::RebuildGraph()
 		MaterialInputs.Add( FMaterialInputInfo( LOCTEXT("AmbientOcclusion", "Ambient Occlusion"), MP_AmbientOcclusion, LOCTEXT( "AmbientOcclusionToolTip", "Simulate the self-shadowing that happens within crevices of a surface" ) ) );
 		MaterialInputs.Add( FMaterialInputInfo( LOCTEXT("Refraction", "Refraction"), MP_Refraction, LOCTEXT( "RefractionToolTip", "Takes in a texture or value that simulates the index of refraction of the surface" ) ) );
 
-		for (int32 UVIndex = 0; UVIndex < ARRAY_COUNT(Material->CustomizedUVs); UVIndex++)
+		for (int32 UVIndex = 0; UVIndex < UE_ARRAY_COUNT(Material->CustomizedUVs); UVIndex++)
 		{
 			//@todo - localize
 			MaterialInputs.Add( FMaterialInputInfo( FText::FromString(FString::Printf(TEXT("Customized UV%u"), UVIndex)), (EMaterialProperty)(MP_CustomizedUVs0 + UVIndex), FText::FromString(FString::Printf( TEXT( "CustomizedUV%uToolTip" ), UVIndex ) ) ) );
diff --git a/Engine/Source/Editor/UnrealEd/Private/StaticMeshEdit.cpp b/Engine/Source/Editor/UnrealEd/Private/StaticMeshEdit.cpp
--- a/Engine/Source/Editor/UnrealEd/Private/StaticMeshEdit.cpp
+++ b/Engine/Source/Editor/UnrealEd/Private/StaticMeshEdit.cpp
@@ -636,7 +636,7 @@ FMergeStaticMeshParams::FMergeStaticMeshParams()
 	, bUseUVScaleBias(false)
 {
 	// initialize some UV channel arrays
-	for (int32 Channel = 0; Channel < ARRAY_COUNT(UVChannelRemap); Channel++)
+	for (int32 Channel = 0; Channel < UE_ARRAY_COUNT(UVChannelRemap); Channel++)
 	{
 		// we can't just map channel to channel by default, because we need to know when a UV channel is
 		// actually being redirected in to, so that we can update Triangle.NumUVs
diff --git a/Engine/Source/Editor/UnrealEd/Private/Texture2DPreview.cpp b/Engine/Source/Editor/UnrealEd/Private/Texture2DPreview.cpp
--- a/Engine/Source/Editor/UnrealEd/Private/Texture2DPreview.cpp
+++ b/Engine/Source/Editor/UnrealEd/Private/Texture2DPreview.cpp
@@ -90,7 +90,7 @@ public:
 			AllocatedVT->GetPackedPageTableUniform(PageTableUniform, false);
 			AllocatedVT->GetPackedUniform(&Uniform, (uint32)LayerIndex);
 
-			SetShaderValueArray(RHICmdList, GetPixelShader(), VTPackedPageTableUniform, PageTableUniform, ARRAY_COUNT(PageTableUniform));
+			SetShaderValueArray(RHICmdList, GetPixelShader(), VTPackedPageTableUniform, PageTableUniform, UE_ARRAY_COUNT(PageTableUniform));
 			SetShaderValue(RHICmdList, GetPixelShader(), VTPackedUniform, Uniform);
 		}
 		else
diff --git a/Engine/Source/Editor/WorldBrowser/Private/Tiles/WorldTileThumbnails.cpp b/Engine/Source/Editor/WorldBrowser/Private/Tiles/WorldTileThumbnails.cpp
--- a/Engine/Source/Editor/WorldBrowser/Private/Tiles/WorldTileThumbnails.cpp
+++ b/Engine/Source/Editor/WorldBrowser/Private/Tiles/WorldTileThumbnails.cpp
@@ -119,7 +119,7 @@ FTileAtlasPage::FTileAtlasPage()
 	static int32 NextPageUniqueID = 1;
 	FName AtlasPageName = *FString::Printf(TEXT("WorldCompositionAtlasPage_%d"), NextPageUniqueID++);
 	
-	for (int32 i = 0; i < ARRAY_COUNT(AtlasSlots); ++i)
+	for (int32 i = 0; i < UE_ARRAY_COUNT(AtlasSlots); ++i)
 	{
 		FTileAtlasSlot& Slot = AtlasSlots[i];
 		
@@ -141,7 +141,7 @@ FTileAtlasPage::FTileAtlasPage()
 
 FTileAtlasPage::~FTileAtlasPage()
 {
-	for (int32 i = 0; i < ARRAY_COUNT(AtlasSlots); ++i)
+	for (int32 i = 0; i < UE_ARRAY_COUNT(AtlasSlots); ++i)
 	{
 		FTileAtlasSlot& Slot = AtlasSlots[i];
 		delete Slot.SlotBrush;
@@ -160,7 +160,7 @@ void FTileAtlasPage::SetOccupied(int32 SlotIdx, bool bOccupied)
 
 bool FTileAtlasPage::HasOccupiedSlots() const
 {
-	for (int32 i = 0; i < ARRAY_COUNT(AtlasSlots); ++i)
+	for (int32 i = 0; i < UE_ARRAY_COUNT(AtlasSlots); ++i)
 	{
 		if (AtlasSlots[i].bOccupied)
 		{
@@ -174,7 +174,7 @@ bool FTileAtlasPage::HasOccupiedSlots() const
 int32 FTileAtlasPage::GetFreeSlotIndex() const
 {
 	int32 Result = INDEX_NONE;
-	for (int32 i = 0; i < ARRAY_COUNT(AtlasSlots); ++i)
+	for (int32 i = 0; i < UE_ARRAY_COUNT(AtlasSlots); ++i)
 	{
 		if (!AtlasSlots[i].bOccupied)
 		{
diff --git a/Engine/Source/Programs/MinidumpDiagnostics/Private/Windows/MinidumpDiagnosticsMainWindows.cpp b/Engine/Source/Programs/MinidumpDiagnostics/Private/Windows/MinidumpDiagnosticsMainWindows.cpp
--- a/Engine/Source/Programs/MinidumpDiagnostics/Private/Windows/MinidumpDiagnosticsMainWindows.cpp
+++ b/Engine/Source/Programs/MinidumpDiagnostics/Private/Windows/MinidumpDiagnosticsMainWindows.cpp
@@ -19,8 +19,8 @@ int32 SimpleCrashHandler( LPEXCEPTION_POINTERS ExceptionInfo )
 	StackTrace[0] = 0;
 	FPlatformStackWalk::StackWalkAndDump( StackTrace, StackTraceSize, 0, ExceptionInfo->ContextRecord );
 
-	FCString::Strncat( GErrorHist, TEXT( "\r\n\r\n" ), ARRAY_COUNT( GErrorHist ) );
-	FCString::Strncat( GErrorHist, ANSI_TO_TCHAR( StackTrace ), ARRAY_COUNT( GErrorHist ) );
+	FCString::Strncat( GErrorHist, TEXT( "\r\n\r\n" ), UE_ARRAY_COUNT( GErrorHist ) );
+	FCString::Strncat( GErrorHist, ANSI_TO_TCHAR( StackTrace ), UE_ARRAY_COUNT( GErrorHist ) );
 
 	GMalloc->Free( StackTrace );
 
diff --git a/Engine/Source/Programs/TestPAL/Private/Main.cpp b/Engine/Source/Programs/TestPAL/Private/Main.cpp
--- a/Engine/Source/Programs/TestPAL/Private/Main.cpp
+++ b/Engine/Source/Programs/TestPAL/Private/Main.cpp
@@ -858,7 +858,7 @@ struct FMemoryAllocatingThread : public FRunnable
 			int32 NumAllocs = 0;
 			uint64 MemAllocatedThisRun = 0;
 
-			while(MemAllocatedThisRun < GMallocTestMemoryPerThreadKB * 1024 && NumAllocs < ARRAY_COUNT(Ptrs))
+			while(MemAllocatedThisRun < GMallocTestMemoryPerThreadKB * 1024 && NumAllocs < UE_ARRAY_COUNT(Ptrs))
 			{
 				// increase the probability of the smallest chunks, but allow fairly large (up to 16 MB)
 				double Uniform = FMath::FRand();
@@ -1336,7 +1336,7 @@ int32 ThreadPriorityTest(const TCHAR* CommandLine)
 	for (int32 Idx = 0, NumThreads = RunnableArray.Num(); Idx < NumThreads; ++Idx)
 	{
 		uint32 PrioIndex = static_cast<uint32>(RunnableArray[Idx]->Prio);
-		if (PrioIndex >= ARRAY_COUNT(AllThreadsCounters))
+		if (PrioIndex >= UE_ARRAY_COUNT(AllThreadsCounters))
 		{
 			UE_LOG(LogTestPAL, Fatal, TEXT("EThreadPriority enum changed and has values larger than 7 now. Revisit this code."));
 		}
@@ -1365,10 +1365,10 @@ int32 ThreadPriorityTest(const TCHAR* CommandLine)
 		static_cast<uint32>(TPri_TimeCritical)
 	};
 
-	for (uint32 IdxPrio = 0; IdxPrio < ARRAY_COUNT(OrderToPrint); ++IdxPrio)
+	for (uint32 IdxPrio = 0; IdxPrio < UE_ARRAY_COUNT(OrderToPrint); ++IdxPrio)
 	{
 		uint32 Prio = OrderToPrint[IdxPrio];
-		checkf(Prio < ARRAY_COUNT(AllThreadsCounters), TEXT("Number of thread priority enums changed, revisit this code."));
+		checkf(Prio < UE_ARRAY_COUNT(AllThreadsCounters), TEXT("Number of thread priority enums changed, revisit this code."));
 		UE_LOG(LogTestPAL, Display, TEXT("Threads with prio %24s incremented counters %14llu times (%.1f x speed of TPri_Normal) during %1.f sec total"),
 			ThreadPrioToString(static_cast<EThreadPriority>(Prio)), AllThreadsCounters[Prio],
 			static_cast<double>(AllThreadsCounters[Prio]) / BaseCountDbl,
@@ -1527,7 +1527,7 @@ struct FThreadStackTest : public FRunnable
 		for (int i = 0; i < BackTraceCount; i++)
 		{
 			ANSICHAR TempString[1024] = {0};
-			FPlatformStackWalk::ProgramCounterToHumanReadableString(i, BackTrace[i], TempString, ARRAY_COUNT(TempString));
+			FPlatformStackWalk::ProgramCounterToHumanReadableString(i, BackTrace[i], TempString, UE_ARRAY_COUNT(TempString));
 
 			UE_LOG(LogTestPAL, Warning, TEXT("%s"), ANSI_TO_TCHAR(TempString));
 		}
diff --git a/Engine/Source/Programs/UnrealFrontend/Private/Commands/StatsConvertCommand.cpp b/Engine/Source/Programs/UnrealFrontend/Private/Commands/StatsConvertCommand.cpp
--- a/Engine/Source/Programs/UnrealFrontend/Private/Commands/StatsConvertCommand.cpp
+++ b/Engine/Source/Programs/UnrealFrontend/Private/Commands/StatsConvertCommand.cpp
@@ -119,7 +119,7 @@ void FCSVStatsProfiler::WriteString( const ANSICHAR* Format, ... )
 	va_list ArgPtr;
 	va_start( ArgPtr, Format );
 	// Build the string.
-	int32 Result = FCStringAnsi::GetVarArgs( Array, ARRAY_COUNT( Array ), Format, ArgPtr );
+	int32 Result = FCStringAnsi::GetVarArgs( Array, UE_ARRAY_COUNT( Array ), Format, ArgPtr );
 	va_end(ArgPtr);
 
 	// Now write that to the file.
diff --git a/Engine/Source/Programs/UnrealHeaderTool/Private/CodeGenerator.cpp b/Engine/Source/Programs/UnrealHeaderTool/Private/CodeGenerator.cpp
--- a/Engine/Source/Programs/UnrealHeaderTool/Private/CodeGenerator.cpp
+++ b/Engine/Source/Programs/UnrealHeaderTool/Private/CodeGenerator.cpp
@@ -543,7 +543,7 @@ static FString OutputMetaDataCodeForObject(FOutputDevice& OutDeclaration, FOutpu
 		Out.Logf(TEXT("%s};\r\n"), Spaces);
 		Out.Log (TEXT("#endif\r\n"));
 
-		Result = FString::Printf(TEXT("METADATA_PARAMS(%s, ARRAY_COUNT(%s))"), MetaDataBlockName, MetaDataBlockName);
+		Result = FString::Printf(TEXT("METADATA_PARAMS(%s, UE_ARRAY_COUNT(%s))"), MetaDataBlockName, MetaDataBlockName);
 	}
 	else
 	{
@@ -1612,14 +1612,14 @@ TTuple<FString, FString> FNativeClassHeaderGenerator::OutputProperties(FOutputDe
 	{
 		return TTuple<FString, FString>(
 			FString::Printf(TEXT("IF_WITH_EDITORONLY_DATA(%sPropPointers, nullptr)"), Scope),
-			FString::Printf(TEXT("IF_WITH_EDITORONLY_DATA(ARRAY_COUNT(%sPropPointers), 0)"), Scope)
+			FString::Printf(TEXT("IF_WITH_EDITORONLY_DATA(UE_ARRAY_COUNT(%sPropPointers), 0)"), Scope)
 		);
 	}
 	else
 	{
 		return TTuple<FString, FString>(
 			FString::Printf(TEXT("%sPropPointers"), Scope),
-			FString::Printf(TEXT("ARRAY_COUNT(%sPropPointers)"), Scope)
+			FString::Printf(TEXT("UE_ARRAY_COUNT(%sPropPointers)"), Scope)
 		);
 	}
 }
@@ -1875,7 +1875,7 @@ void FNativeClassHeaderGenerator::ExportGeneratedPackageInitCode(FOutputDevice&
 		Out.Logf(TEXT("\t\t\t};\r\n"));
 
 		SingletonArray = TEXT("SingletonFuncArray");
-		SingletonCount = TEXT("ARRAY_COUNT(SingletonFuncArray)");
+		SingletonCount = TEXT("UE_ARRAY_COUNT(SingletonFuncArray)");
 	}
 	else
 	{
@@ -1991,7 +1991,7 @@ void FNativeClassHeaderGenerator::ExportNativeGeneratedInitCode(FOutputDevice& O
 			StaticDefinitions.Logf(TEXT("\t};\r\n"));
 
 			SingletonsArray = TEXT("DependentSingletons");
-			SingletonsCount = TEXT("ARRAY_COUNT(DependentSingletons)");
+			SingletonsCount = TEXT("UE_ARRAY_COUNT(DependentSingletons)");
 		}
 		else
 		{
@@ -2041,12 +2041,12 @@ void FNativeClassHeaderGenerator::ExportNativeGeneratedInitCode(FOutputDevice& O
 			if (bAllEditorOnlyFunctions)
 			{
 				FunctionsArray = TEXT("IF_WITH_EDITOR(FuncInfo, nullptr)");
-				FunctionsCount = TEXT("IF_WITH_EDITOR(ARRAY_COUNT(FuncInfo), 0)");
+				FunctionsCount = TEXT("IF_WITH_EDITOR(UE_ARRAY_COUNT(FuncInfo), 0)");
 			}
 			else
 			{
 				FunctionsArray = TEXT("FuncInfo");
-				FunctionsCount = TEXT("ARRAY_COUNT(FuncInfo)");
+				FunctionsCount = TEXT("UE_ARRAY_COUNT(FuncInfo)");
 			}
 		}
 		else
@@ -2101,7 +2101,7 @@ void FNativeClassHeaderGenerator::ExportNativeGeneratedInitCode(FOutputDevice& O
 			StaticDefinitions.Log(TEXT("\t\t};\r\n"));
 
 			InterfaceArray = TEXT("InterfaceParams");
-			InterfaceCount = TEXT("ARRAY_COUNT(InterfaceParams)");
+			InterfaceCount = TEXT("UE_ARRAY_COUNT(InterfaceParams)");
 		}
 		else
 		{
@@ -2424,7 +2424,7 @@ void FNativeClassHeaderGenerator::ExportNatives(FOutputDevice& Out, FClass* Clas
 			EditorOnly(bAllEditorOnly);
 
 			Out.Log(TEXT("\t\t};\r\n"));
-			Out.Logf(TEXT("\t\tFNativeFunctionRegistrar::RegisterFunctions(Class, Funcs, ARRAY_COUNT(Funcs));\r\n"));
+			Out.Logf(TEXT("\t\tFNativeFunctionRegistrar::RegisterFunctions(Class, Funcs, UE_ARRAY_COUNT(Funcs));\r\n"));
 		}
 	}
 
@@ -3792,7 +3792,7 @@ void FNativeClassHeaderGenerator::ExportGeneratedEnumInitCode(FOutputDevice& Out
 	GeneratedEnumRegisterFunctionText.Logf(TEXT("\t\t\t\t%s,\r\n"), *CreateUTF8LiteralString(OverriddenEnumNameCpp));
 	GeneratedEnumRegisterFunctionText.Logf(TEXT("\t\t\t\t%s,\r\n"), *CreateUTF8LiteralString(Enum->CppType));
 	GeneratedEnumRegisterFunctionText.Logf(TEXT("\t\t\t\tEnumerators,\r\n"));
-	GeneratedEnumRegisterFunctionText.Logf(TEXT("\t\t\t\tARRAY_COUNT(Enumerators),\r\n"));
+	GeneratedEnumRegisterFunctionText.Logf(TEXT("\t\t\t\tUE_ARRAY_COUNT(Enumerators),\r\n"));
 	GeneratedEnumRegisterFunctionText.Logf(TEXT("\t\t\t\t%s,\r\n"), UEnumObjectFlags);
 	GeneratedEnumRegisterFunctionText.Logf(TEXT("\t\t\t\tUE4CodeGen_Private::EDynamicType::%s,\r\n"), bIsDynamic ? TEXT("Dynamic") : TEXT("NotDynamic"));
 	GeneratedEnumRegisterFunctionText.Logf(TEXT("\t\t\t\t(uint8)%s,\r\n"), *EnumFormStr);
@@ -3916,7 +3916,7 @@ void FNativeClassHeaderGenerator::ExportDelegateDeclaration(FOutputDevice& Out,
 
 	// Add class name to beginning of function, to avoid collisions with other classes with the same delegate name in this scope
 	check(FunctionData.MarshallAndCallName.StartsWith(DelegateStr));
-	FString ShortName = *FunctionData.MarshallAndCallName + ARRAY_COUNT(DelegateStr) - 1;
+	FString ShortName = *FunctionData.MarshallAndCallName + UE_ARRAY_COUNT(DelegateStr) - 1;
 	FunctionData.MarshallAndCallName = FString::Printf( TEXT( "F%s_DelegateWrapper" ), *ShortName );
 
 	// Setup delegate parameter
@@ -3962,7 +3962,7 @@ void FNativeClassHeaderGenerator::ExportDelegateDefinition(FOutputDevice& Out, c
 
 	// Add class name to beginning of function, to avoid collisions with other classes with the same delegate name in this scope
 	check(FunctionData.MarshallAndCallName.StartsWith(DelegateStr));
-	FString ShortName = *FunctionData.MarshallAndCallName + ARRAY_COUNT(DelegateStr) - 1;
+	FString ShortName = *FunctionData.MarshallAndCallName + UE_ARRAY_COUNT(DelegateStr) - 1;
 	FunctionData.MarshallAndCallName = FString::Printf( TEXT( "F%s_DelegateWrapper" ), *ShortName );
 
 	// Setup delegate parameter
@@ -5963,17 +5963,17 @@ ECompilationResult::Type PreparseModules(const FString& ModuleInfoPath, int32& N
 						static const TCHAR PublicFolderName[]  = TEXT("Public/");
 						static const TCHAR PrivateFolderName[] = TEXT("Private/");
 						static const TCHAR ClassesFolderName[] = TEXT("Classes/");
-						if (FCString::Strnicmp(IncludePath, PublicFolderName, ARRAY_COUNT(PublicFolderName) - 1) == 0)
+						if (FCString::Strnicmp(IncludePath, PublicFolderName, UE_ARRAY_COUNT(PublicFolderName) - 1) == 0)
 						{
-							IncludePath += (ARRAY_COUNT(PublicFolderName) - 1);
+							IncludePath += (UE_ARRAY_COUNT(PublicFolderName) - 1);
 						}
-						else if (FCString::Strnicmp(IncludePath, PrivateFolderName, ARRAY_COUNT(PrivateFolderName) - 1) == 0)
+						else if (FCString::Strnicmp(IncludePath, PrivateFolderName, UE_ARRAY_COUNT(PrivateFolderName) - 1) == 0)
 						{
-							IncludePath += (ARRAY_COUNT(PrivateFolderName) - 1);
+							IncludePath += (UE_ARRAY_COUNT(PrivateFolderName) - 1);
 						}
-						else if (FCString::Strnicmp(IncludePath, ClassesFolderName, ARRAY_COUNT(ClassesFolderName) - 1) == 0)
+						else if (FCString::Strnicmp(IncludePath, ClassesFolderName, UE_ARRAY_COUNT(ClassesFolderName) - 1) == 0)
 						{
-							IncludePath += (ARRAY_COUNT(ClassesFolderName) - 1);
+							IncludePath += (UE_ARRAY_COUNT(ClassesFolderName) - 1);
 						}
 
 						// Add the include path
diff --git a/Engine/Source/Programs/UnrealHeaderTool/Private/HeaderParser.cpp b/Engine/Source/Programs/UnrealHeaderTool/Private/HeaderParser.cpp
--- a/Engine/Source/Programs/UnrealHeaderTool/Private/HeaderParser.cpp
+++ b/Engine/Source/Programs/UnrealHeaderTool/Private/HeaderParser.cpp
@@ -101,19 +101,19 @@ namespace
 
 		// Test for known delegate and event macros.
 		TCHAR MulticastDelegateStart[] = TEXT("DECLARE_MULTICAST_DELEGATE");
-		if (!FCString::Strncmp(Identifier, MulticastDelegateStart, ARRAY_COUNT(MulticastDelegateStart) - 1))
+		if (!FCString::Strncmp(Identifier, MulticastDelegateStart, UE_ARRAY_COUNT(MulticastDelegateStart) - 1))
 		{
 			return true;
 		}
 
 		TCHAR DelegateStart[] = TEXT("DECLARE_DELEGATE");
-		if (!FCString::Strncmp(Identifier, DelegateStart, ARRAY_COUNT(DelegateStart) - 1))
+		if (!FCString::Strncmp(Identifier, DelegateStart, UE_ARRAY_COUNT(DelegateStart) - 1))
 		{
 			return true;
 		}
 
 		TCHAR DelegateEvent[] = TEXT("DECLARE_EVENT");
-		if (!FCString::Strncmp(Identifier, DelegateEvent, ARRAY_COUNT(DelegateEvent) - 1))
+		if (!FCString::Strncmp(Identifier, DelegateEvent, UE_ARRAY_COUNT(DelegateEvent) - 1))
 		{
 			return true;
 		}
@@ -189,7 +189,7 @@ namespace
 			TEXT("FORCEINLINE_STATS"),
 			TEXT("SIZE_T")
 		};
-		if (Algo::FindSortedStringCaseInsensitive(Token.Identifier, Whitelist, ARRAY_COUNT(Whitelist)) >= 0)
+		if (Algo::FindSortedStringCaseInsensitive(Token.Identifier, Whitelist, UE_ARRAY_COUNT(Whitelist)) >= 0)
 		{
 			return false;
 		}
@@ -227,7 +227,7 @@ namespace
 			{
 				// It's a tag with an argument
 
-				if (FCString::Strnicmp(IdentifierPtr, IdTag, ARRAY_COUNT(IdTag) - 1) == 0)
+				if (FCString::Strnicmp(IdentifierPtr, IdTag, UE_ARRAY_COUNT(IdTag) - 1) == 0)
 				{
 					int32 TempInt = FCString::Atoi(Equals + 1);
 					if (TempInt <= 0 || TempInt > MAX_uint16)
@@ -236,8 +236,8 @@ namespace
 					}
 					FuncInfo.RPCId = TempInt;
 				}
-				else if (FCString::Strnicmp(IdentifierPtr, ResponseIdTag, ARRAY_COUNT(ResponseIdTag) - 1) == 0 ||
-					FCString::Strnicmp(IdentifierPtr, JSBridgePriTag, ARRAY_COUNT(JSBridgePriTag) - 1) == 0)
+				else if (FCString::Strnicmp(IdentifierPtr, ResponseIdTag, UE_ARRAY_COUNT(ResponseIdTag) - 1) == 0 ||
+					FCString::Strnicmp(IdentifierPtr, JSBridgePriTag, UE_ARRAY_COUNT(JSBridgePriTag) - 1) == 0)
 				{
 					int32 TempInt = FCString::Atoi(Equals + 1);
 					if (TempInt <= 0 || TempInt > MAX_uint16)
@@ -1276,7 +1276,7 @@ FClass* FHeaderParser::GetQualifiedClass(const FClasses& AllClasses, const TCHAR
 	{
 		RedirectTypeIdentifier(Token);
 
-		FCString::Strncat( ClassName, Token.Identifier, ARRAY_COUNT(ClassName) );
+		FCString::Strncat( ClassName, Token.Identifier, UE_ARRAY_COUNT(ClassName) );
 	}
 
 	if (!ClassName[0])
diff --git a/Engine/Source/Programs/UnrealHeaderTool/Private/Manifest.cpp b/Engine/Source/Programs/UnrealHeaderTool/Private/Manifest.cpp
--- a/Engine/Source/Programs/UnrealHeaderTool/Private/Manifest.cpp
+++ b/Engine/Source/Programs/UnrealHeaderTool/Private/Manifest.cpp
@@ -191,7 +191,7 @@ bool FManifestModule::NeedsRegeneration() const
 	}
 	FString Timestamp;
 	TCHAR TimestampText[] = TEXT("Timestamp");
-	Timestamp.Empty(GeneratedIncludeDirectory.Len() + ARRAY_COUNT(TimestampText));
+	Timestamp.Empty(GeneratedIncludeDirectory.Len() + UE_ARRAY_COUNT(TimestampText));
 	Timestamp += GeneratedIncludeDirectory;
 	Timestamp += TEXT("Timestamp");
 
diff --git a/Engine/Source/Programs/UnrealLightmass/Private/ImportExport/Exporter.cpp b/Engine/Source/Programs/UnrealLightmass/Private/ImportExport/Exporter.cpp
--- a/Engine/Source/Programs/UnrealLightmass/Private/ImportExport/Exporter.cpp
+++ b/Engine/Source/Programs/UnrealLightmass/Private/ImportExport/Exporter.cpp
@@ -497,7 +497,7 @@ namespace Lightmass
 				Swarm->Write(&TaskData.BrickData[BrickIndex].AverageClosestGeometryDistance, sizeof(TaskData.BrickData[BrickIndex].AverageClosestGeometryDistance));
 				WriteArray(TaskData.BrickData[BrickIndex].AmbientVector);
 
-				for (int32 i = 0; i < ARRAY_COUNT(TaskData.BrickData[BrickIndex].SHCoefficients); i++)
+				for (int32 i = 0; i < UE_ARRAY_COUNT(TaskData.BrickData[BrickIndex].SHCoefficients); i++)
 				{
 					WriteArray(TaskData.BrickData[BrickIndex].SHCoefficients[i]);
 				}
diff --git a/Engine/Source/Programs/UnrealLightmass/Private/ImportExport/LightmassScene.cpp b/Engine/Source/Programs/UnrealLightmass/Private/ImportExport/LightmassScene.cpp
--- a/Engine/Source/Programs/UnrealLightmass/Private/ImportExport/LightmassScene.cpp
+++ b/Engine/Source/Programs/UnrealLightmass/Private/ImportExport/LightmassScene.cpp
@@ -1860,7 +1860,7 @@ void FSkyLight::Import( FLightmassImporter& Importer )
 			FIntPoint(1, 1)
 		};
 
-		const float SubCellWeight = 1.0f / (float)ARRAY_COUNT(SubCellOffsets);
+		const float SubCellWeight = 1.0f / (float)UE_ARRAY_COUNT(SubCellOffsets);
 
 		for (int32 MipIndex = 1; MipIndex < NumMips; MipIndex++)
 		{
@@ -1879,7 +1879,7 @@ void FSkyLight::Import( FLightmassImporter& Importer )
 					{
 						FLinearColor FilteredValue(0, 0, 0, 0);
 
-						for (int32 OffsetIndex = 0; OffsetIndex < ARRAY_COUNT(SubCellOffsets); OffsetIndex++)
+						for (int32 OffsetIndex = 0; OffsetIndex < UE_ARRAY_COUNT(SubCellOffsets); OffsetIndex++)
 						{
 							FIntPoint ParentOffset = FIntPoint(X, Y) * 2 + SubCellOffsets[OffsetIndex];
 							int32 ParentTexelIndex = FaceIndex * ParentMipSize * ParentMipSize + ParentOffset.Y * ParentMipSize + ParentOffset.X;
@@ -2439,7 +2439,7 @@ void FMeshAreaLight::SampleDirection(FLMRandomStream& RandomStream, FLightRay& S
 			bool bAllCornersInOppositeHemisphere = true;
 			// Determine whether the cell is completely in the same hemisphere as the sample direction, completely on the other side or spanning the terminator
 			// This is done by checking each cell's corners
-			for (int32 CornerIndex = 0; CornerIndex < ARRAY_COUNT(Corners); CornerIndex++)
+			for (int32 CornerIndex = 0; CornerIndex < UE_ARRAY_COUNT(Corners); CornerIndex++)
 			{
 				const float Theta = (ThetaStep + Corners[CornerIndex].X) / (float)MeshAreaLightGridSize * (float)PI;
 				const float Phi = (PhiStep + Corners[CornerIndex].Y) / (float)MeshAreaLightGridSize * 2 * (float)PI - (float)PI;
diff --git a/Engine/Source/Programs/UnrealLightmass/Private/Lighting/AdaptiveVolumetricLightmap.cpp b/Engine/Source/Programs/UnrealLightmass/Private/Lighting/AdaptiveVolumetricLightmap.cpp
--- a/Engine/Source/Programs/UnrealLightmass/Private/Lighting/AdaptiveVolumetricLightmap.cpp
+++ b/Engine/Source/Programs/UnrealLightmass/Private/Lighting/AdaptiveVolumetricLightmap.cpp
@@ -52,7 +52,7 @@ FOverlapInterval GetInterval(const FBox& Box, const FVector& Vector)
 	FOverlapInterval Result;
 	Result.Min = Result.Max = FVector::DotProduct(Vector, BoxVertices[0]);
 
-	for (int32 i = 1; i < ARRAY_COUNT(BoxVertices); ++i) 
+	for (int32 i = 1; i < UE_ARRAY_COUNT(BoxVertices); ++i) 
 	{
 		float Projection = FVector::DotProduct(Vector, BoxVertices[i]);
 		Result.Min = FMath::Min(Result.Min, Projection);
@@ -99,7 +99,7 @@ bool IntersectTriangleAndAABB(const FTriangle& Triangle, const FBox& Box)
 		FVector::CrossProduct(BoxNormal2, TriangleEdge2)
 	};
 
-	for (int i = 0; i < ARRAY_COUNT(TestDirections); ++i) 
+	for (int i = 0; i < UE_ARRAY_COUNT(TestDirections); ++i) 
 	{
 		if (!OverlapOnAxis(Box, Triangle, TestDirections[i])) 
 		{
@@ -156,7 +156,7 @@ bool FStaticLightingSystem::DoesVoxelIntersectSceneGeometry(const FBox& CellBoun
 					{
 						TransformedCorners[Corner] = Transform.TransformPosition(Box.GetExtrema(Corner));
 					}
-					return FBox(TransformedCorners, ARRAY_COUNT(TransformedCorners));
+					return FBox(TransformedCorners, UE_ARRAY_COUNT(TransformedCorners));
 				};
 
 				if (MeshInstance->Mapping->GetVolumeMapping() == nullptr)
@@ -210,7 +210,7 @@ bool FStaticLightingSystem::DoesVoxelIntersectSceneGeometry(const FBox& CellBoun
 						Triangle.Vertices[1] = Vertices[1].WorldPosition;
 						Triangle.Vertices[2] = Vertices[2].WorldPosition;
 
-						FBox TriangleAABB(Triangle.Vertices, ARRAY_COUNT(Triangle.Vertices));
+						FBox TriangleAABB(Triangle.Vertices, UE_ARRAY_COUNT(Triangle.Vertices));
 
 						if (ExpandedCellBounds.Intersect(TriangleAABB))
 						{
@@ -589,7 +589,7 @@ void FStaticLightingSystem::ProcessVolumetricLightmapBrickTask(FVolumetricLightm
 	BrickData.DirectionalLightShadowing.Empty(TotalBrickSize);
 	BrickData.DirectionalLightShadowing.AddDefaulted(TotalBrickSize);
 
-	for (int32 i = 0; i < ARRAY_COUNT(BrickData.SHCoefficients); i++)
+	for (int32 i = 0; i < UE_ARRAY_COUNT(BrickData.SHCoefficients); i++)
 	{
 		BrickData.SHCoefficients[i].Empty(TotalBrickSize);
 		BrickData.SHCoefficients[i].AddDefaulted(TotalBrickSize);
@@ -882,7 +882,7 @@ void FStaticLightingSystem::CalculateAdaptiveVolumetricLightmap(int32 TaskIndex)
 
 void FIrradianceBrickData::SetFromVolumeLightingSample(int32 Index, const FVolumeLightingSample& Sample, bool bInsideGeometry, float MinDistanceToSurface, bool bBorderVoxel)
 {
-	static_assert(ARRAY_COUNT(Sample.HighQualityCoefficients) >= ARRAY_COUNT(SHCoefficients) + 1, "Coefficient mismatch");
+	static_assert(UE_ARRAY_COUNT(Sample.HighQualityCoefficients) >= UE_ARRAY_COUNT(SHCoefficients) + 1, "Coefficient mismatch");
 
 	AmbientVector[Index] = FFloat3Packed(FLinearColor(Sample.HighQualityCoefficients[0][0], Sample.HighQualityCoefficients[0][1], Sample.HighQualityCoefficients[0][2], 0.0f));
 
diff --git a/Engine/Source/Programs/UnrealLightmass/Private/Lighting/FinalGather.cpp b/Engine/Source/Programs/UnrealLightmass/Private/Lighting/FinalGather.cpp
--- a/Engine/Source/Programs/UnrealLightmass/Private/Lighting/FinalGather.cpp
+++ b/Engine/Source/Programs/UnrealLightmass/Private/Lighting/FinalGather.cpp
@@ -755,7 +755,7 @@ public:
 							const FLinearColor Radiance = RootElementLighting.Lighting + RootElementLighting.StationarySkyLighting;
 							const float RelativeBrightness = Radiance.ComputeLuminance() / AverageBrightness;
 
-							for (int32 NeighborIndex = 0; NeighborIndex < ARRAY_COUNT(Neighbors); NeighborIndex++)
+							for (int32 NeighborIndex = 0; NeighborIndex < UE_ARRAY_COUNT(Neighbors); NeighborIndex++)
 							{
 								int32 NeighborTheta = ThetaIndex + Neighbors[NeighborIndex].X;
 								// Wrap phi around, since it is the angle around the hemisphere axis
@@ -894,7 +894,7 @@ public:
 								const float RelativeBrightness = Radiance.ComputeLuminance() / AverageBrightness;
 
 								// Only search the axis neighbors past the first depth
-								for (int32 NeighborIndex = 0; NeighborIndex < ARRAY_COUNT(Neighbors) / 2; NeighborIndex++)
+								for (int32 NeighborIndex = 0; NeighborIndex < UE_ARRAY_COUNT(Neighbors) / 2; NeighborIndex++)
 								{
 									const float NeighborU = NodeContext.Min.X + (SubThetaIndex + Neighbors[NeighborIndex].X) * NodeContext.Size.X / 2;
 									const float NeighborV = NodeContext.Min.Y + (SubPhiIndex + Neighbors[NeighborIndex].Y) * NodeContext.Size.Y / 2;
@@ -1093,7 +1093,7 @@ private:
 		{
 			FLightingAndOcclusion FilteredValue;
 
-			for (int32 ChildIndex = 0; ChildIndex < ARRAY_COUNT(Parent->Children); ChildIndex++)
+			for (int32 ChildIndex = 0; ChildIndex < UE_ARRAY_COUNT(Parent->Children); ChildIndex++)
 			{
 				FilteredValue = FilteredValue + GetFilteredValueRecursive(Parent->Children[ChildIndex]) / 4.0f;
 			}
@@ -1115,7 +1115,7 @@ private:
 				FinalGatherHitPoints[Parent->Element.HitPointIndex].Weight = 0.0f;
 			}
 
-			for (int32 ChildIndex = 0; ChildIndex < ARRAY_COUNT(Parent->Children); ChildIndex++)
+			for (int32 ChildIndex = 0; ChildIndex < UE_ARRAY_COUNT(Parent->Children); ChildIndex++)
 			{
 				UpdateHitPointWeightsRecursive(FinalGatherHitPoints, Parent->Children[ChildIndex], ParentWeight / 4.0f);
 			}
diff --git a/Engine/Source/Programs/UnrealLightmass/Private/Lighting/LightingMesh.h b/Engine/Source/Programs/UnrealLightmass/Private/Lighting/LightingMesh.h
--- a/Engine/Source/Programs/UnrealLightmass/Private/Lighting/LightingMesh.h
+++ b/Engine/Source/Programs/UnrealLightmass/Private/Lighting/LightingMesh.h
@@ -26,7 +26,7 @@ struct FStaticLightingVertex: public FStaticLightingVertexData
 		WorldPosition = InVertex.WorldPosition;
 		WorldTangentZ = InVertex.WorldTangentZ;
 
-		for (int32 i = 0; i < ARRAY_COUNT(TextureCoordinates); i++)
+		for (int32 i = 0; i < UE_ARRAY_COUNT(TextureCoordinates); i++)
 		{
 			TextureCoordinates[i] = InVertex.TextureCoordinates[i];
 		}
diff --git a/Engine/Source/Programs/UnrealLightmass/Private/Lighting/LightingSystem.cpp b/Engine/Source/Programs/UnrealLightmass/Private/Lighting/LightingSystem.cpp
--- a/Engine/Source/Programs/UnrealLightmass/Private/Lighting/LightingSystem.cpp
+++ b/Engine/Source/Programs/UnrealLightmass/Private/Lighting/LightingSystem.cpp
@@ -1405,9 +1405,9 @@ void FStaticLightingSystem::CacheSamples()
 		CachedSamplesMaxUnoccludedLength = (CombinedVector / CachedHemisphereSamples.Num()).Size3();
 	}
 	
-	for (int32 SampleSet = 0; SampleSet < ARRAY_COUNT(CachedHemisphereSamplesForRadiosity); SampleSet++)
+	for (int32 SampleSet = 0; SampleSet < UE_ARRAY_COUNT(CachedHemisphereSamplesForRadiosity); SampleSet++)
 	{
-		float SampleSetScale = FMath::Lerp(.5f, .125f, SampleSet / ((float)ARRAY_COUNT(CachedHemisphereSamplesForRadiosity) - 1));
+		float SampleSetScale = FMath::Lerp(.5f, .125f, SampleSet / ((float)UE_ARRAY_COUNT(CachedHemisphereSamplesForRadiosity) - 1));
 		int32 TargetNumApproximateSkyLightingSamples = FMath::Max(FMath::TruncToInt(ImportanceTracingSettings.NumHemisphereSamples * SampleSetScale * GeneralSettings.IndirectLightingQuality), 12);
 		CachedHemisphereSamplesForRadiosity[SampleSet].Empty(TargetNumApproximateSkyLightingSamples);
 		CachedHemisphereSamplesForRadiosityUniforms[SampleSet].Empty(TargetNumApproximateSkyLightingSamples);
diff --git a/Engine/Source/Programs/UnrealLightmass/Private/Lighting/LightingSystem.h b/Engine/Source/Programs/UnrealLightmass/Private/Lighting/LightingSystem.h
--- a/Engine/Source/Programs/UnrealLightmass/Private/Lighting/LightingSystem.h
+++ b/Engine/Source/Programs/UnrealLightmass/Private/Lighting/LightingSystem.h
@@ -753,7 +753,7 @@ public:
 		VolumetricLightmapDirectLightingTime(0),
 		VolumetricLightmapFinalGatherTime(0)
 	{
-		for (int32 i = 0; i < ARRAY_COUNT(NumRefiningFinalGatherSamples); i++)
+		for (int32 i = 0; i < UE_ARRAY_COUNT(NumRefiningFinalGatherSamples); i++)
 		{
 			NumRefiningFinalGatherSamples[i] = 0;
 		}
@@ -825,7 +825,7 @@ public:
 		BaseFinalGatherSampleTime += B.BaseFinalGatherSampleTime;
 		RefiningFinalGatherSampleTime += B.RefiningFinalGatherSampleTime;
 
-		for (int32 i = 0; i < ARRAY_COUNT(NumRefiningFinalGatherSamples); i++)
+		for (int32 i = 0; i < UE_ARRAY_COUNT(NumRefiningFinalGatherSamples); i++)
 		{
 			NumRefiningFinalGatherSamples[i] += B.NumRefiningFinalGatherSamples[i];
 		}
@@ -1215,7 +1215,7 @@ private:
 
 	void ReturnToFreeListRecursive(FSimpleQuadTreeNode<ElementType>* Node, TArray<FSimpleQuadTreeNode<ElementType>*>& OutNodes) const
 	{
-		for (int32 ChildIndex = 0; ChildIndex < ARRAY_COUNT(Node->Children); ChildIndex++)
+		for (int32 ChildIndex = 0; ChildIndex < UE_ARRAY_COUNT(Node->Children); ChildIndex++)
 		{
 			if (Node->Children[ChildIndex])
 			{
diff --git a/Engine/Source/Programs/UnrealLightmass/Private/Lighting/PhotonMapping.cpp b/Engine/Source/Programs/UnrealLightmass/Private/Lighting/PhotonMapping.cpp
--- a/Engine/Source/Programs/UnrealLightmass/Private/Lighting/PhotonMapping.cpp
+++ b/Engine/Source/Programs/UnrealLightmass/Private/Lighting/PhotonMapping.cpp
@@ -2393,7 +2393,7 @@ FLinearColor FStaticLightingSystem::CalculatePhotonIrradiance(
 
 				// Try to determine visibility to the photon before letting it contribute
 				// This helps to prevent leaking through thin walls
-				for (int32 SearchIndex = 0; SearchIndex < ARRAY_COUNT(SearchNormalScales) && !bPhotonVisible; SearchIndex++)
+				for (int32 SearchIndex = 0; SearchIndex < UE_ARRAY_COUNT(SearchNormalScales) && !bPhotonVisible; SearchIndex++)
 				{
 					const float NormalOffset = SearchDistance * SearchNormalScales[SearchIndex];
 
diff --git a/Engine/Source/Programs/UnrealLightmass/Private/Lighting/PrecomputedVisibility.cpp b/Engine/Source/Programs/UnrealLightmass/Private/Lighting/PrecomputedVisibility.cpp
--- a/Engine/Source/Programs/UnrealLightmass/Private/Lighting/PrecomputedVisibility.cpp
+++ b/Engine/Source/Programs/UnrealLightmass/Private/Lighting/PrecomputedVisibility.cpp
@@ -275,7 +275,7 @@ void FStaticLightingSystem::SetupPrecomputedVisibility()
 
 							const float EdgePullback = .1f;
 
-							for (int32 SampleIndex = 0; SampleIndex < ARRAY_COUNT(SubsamplePositions); SampleIndex++)
+							for (int32 SampleIndex = 0; SampleIndex < UE_ARRAY_COUNT(SubsamplePositions); SampleIndex++)
 							{
 								const FVector2D SamplePosition = SubsamplePositions[SampleIndex] * (1 - 2 * EdgePullback) + FVector2D(EdgePullback, EdgePullback);
 
@@ -361,7 +361,7 @@ void FStaticLightingSystem::SetupPrecomputedVisibility()
 						const float CurrentMaxHeight = ExtremaIndex == 0 ? Cell.HitTriangles[HeightIndex].HeightRange.X : Cell.HitTriangles[HeightIndex].HeightRange.Y;
 						const float CompareHeight = CurrentMaxHeight + .5f * PrecomputedVisibilitySettings.PlayAreaHeight;
 
-						for (int32 TestIndex = 0; TestIndex < ARRAY_COUNT(TestHeights); TestIndex++)
+						for (int32 TestIndex = 0; TestIndex < UE_ARRAY_COUNT(TestHeights); TestIndex++)
 						{
 							const float TestHeight = CurrentMaxHeight + TestHeights[TestIndex] * PrecomputedVisibilitySettings.PlayAreaHeight;
 
@@ -527,7 +527,7 @@ void FStaticLightingSystem::SetupPrecomputedVisibility()
 			const FVector BoxExtent = SortMeshes[MeshIndex].Bounds.GetExtent();
 			float LocalDistance = 0;
 
-			for (int32 CornerIndex = 0; CornerIndex < ARRAY_COUNT(CubeCorners); CornerIndex++)
+			for (int32 CornerIndex = 0; CornerIndex < UE_ARRAY_COUNT(CubeCorners); CornerIndex++)
 			{
 				// Calculate max distance to a corner of the bounds as a measure of how much this mesh will expand the grid bounds
 				LocalDistance = FMath::Max(LocalDistance, (BoxCenter + BoxExtent * CubeCorners[CornerIndex]).SizeSquared());
diff --git a/Engine/Source/Programs/UnrealLightmass/Private/Lighting/Radiosity.cpp b/Engine/Source/Programs/UnrealLightmass/Private/Lighting/Radiosity.cpp
--- a/Engine/Source/Programs/UnrealLightmass/Private/Lighting/Radiosity.cpp
+++ b/Engine/Source/Programs/UnrealLightmass/Private/Lighting/Radiosity.cpp
@@ -586,7 +586,7 @@ void FStaticLightingSystem::RadiosityIterationTextureMapping(FStaticLightingText
 			NumAdaptiveRefinementLevels++;
 		}
 
-		const int32 RadiositySampleSet = FMath::Min<int32>(PassIndex, ARRAY_COUNT(CachedHemisphereSamplesForRadiosity) - 1);
+		const int32 RadiositySampleSet = FMath::Min<int32>(PassIndex, UE_ARRAY_COUNT(CachedHemisphereSamplesForRadiosity) - 1);
 
 		TLightingCache<FFinalGatherSample> RadiosityCache(TextureMapping->Mesh->BoundingBox, *this, 1);
 	
diff --git a/Engine/Source/Programs/UnrealLightmass/Private/Lighting/StaticMesh.cpp b/Engine/Source/Programs/UnrealLightmass/Private/Lighting/StaticMesh.cpp
--- a/Engine/Source/Programs/UnrealLightmass/Private/Lighting/StaticMesh.cpp
+++ b/Engine/Source/Programs/UnrealLightmass/Private/Lighting/StaticMesh.cpp
@@ -196,7 +196,7 @@ static void GetStaticLightingVertex(
 		OutVertex.WorldTangentZ = (OutVertex.WorldTangentX ^ OutVertex.WorldTangentY).GetSafeNormal();
 	}
 
-	for(uint32 LightmapTextureCoordinateIndex = 0; LightmapTextureCoordinateIndex < ARRAY_COUNT(InVertex.UVs); LightmapTextureCoordinateIndex++)
+	for(uint32 LightmapTextureCoordinateIndex = 0; LightmapTextureCoordinateIndex < UE_ARRAY_COUNT(InVertex.UVs); LightmapTextureCoordinateIndex++)
 	{
 		OutVertex.TextureCoordinates[LightmapTextureCoordinateIndex] = InVertex.UVs[LightmapTextureCoordinateIndex];
 	}
diff --git a/Engine/Source/Programs/UnrealLightmass/Private/Lighting/TextureMapping.cpp b/Engine/Source/Programs/UnrealLightmass/Private/Lighting/TextureMapping.cpp
--- a/Engine/Source/Programs/UnrealLightmass/Private/Lighting/TextureMapping.cpp
+++ b/Engine/Source/Programs/UnrealLightmass/Private/Lighting/TextureMapping.cpp
@@ -504,7 +504,7 @@ void FStaticLightingSystem::AdjustRepresentativeSurfelForTexelsTextureMapping(
 
 					float BackfaceDetectionDistance = TexelToVertex.TexelRadius * 2.5f;
 
-					for (int32 CornerIndex = 0; CornerIndex < ARRAY_COUNT(CornerSigns); CornerIndex++)
+					for (int32 CornerIndex = 0; CornerIndex < UE_ARRAY_COUNT(CornerSigns); CornerIndex++)
 					{
 						TraceToTexelCorner(
 							TexelCenterOffset,
@@ -525,7 +525,7 @@ void FStaticLightingSystem::AdjustRepresentativeSurfelForTexelsTextureMapping(
 					// Limit the distance that we will search for an intersecting backface in order to move the shading position to the texel radius
 					float ClosestBackfacingIntersectionDistanceSq = BackfaceDetectionDistance * BackfaceDetectionDistance;
 
-					for (int32 CornerIndex = 0; CornerIndex < ARRAY_COUNT(CornerSigns); CornerIndex++)
+					for (int32 CornerIndex = 0; CornerIndex < UE_ARRAY_COUNT(CornerSigns); CornerIndex++)
 					{
 						if (Intersections[CornerIndex].bIntersects)
 						{
@@ -2095,7 +2095,7 @@ void FStaticLightingSystem::CalculateDirectSignedDistanceFieldLightingTextureMap
 
 					// Search for a neighbor with different visibility
 					bool bNeighborsDifferent = false;
-					for (int32 i = 0 ; i < ARRAY_COUNT(Neighbors); i++)
+					for (int32 i = 0 ; i < UE_ARRAY_COUNT(Neighbors); i++)
 					{
 						if (X + Neighbors[i].X > 0
 							&& X + Neighbors[i].X < TextureMapping->CachedSizeX
@@ -2178,7 +2178,7 @@ void FStaticLightingSystem::CalculateDirectSignedDistanceFieldLightingTextureMap
 					if (!bAnyHighResSamplesMapped)
 					{
 						const FTexelToCorners& TexelToCorners = TexelToCornersMap(X, Y);
-						for (int32 CornerIndex = 0; CornerIndex < ARRAY_COUNT(Corners); CornerIndex++)
+						for (int32 CornerIndex = 0; CornerIndex < UE_ARRAY_COUNT(Corners); CornerIndex++)
 						{						
 							if (TexelToCorners.bValid[CornerIndex])
 							{
@@ -2297,7 +2297,7 @@ void FStaticLightingSystem::CalculateDirectSignedDistanceFieldLightingTextureMap
 							{
 								// Detect texels next to the shadow transition
 								bool bNeighborsDifferent = false;
-								for (int32 i = 0 ; i < ARRAY_COUNT(Neighbors); i++)
+								for (int32 i = 0 ; i < UE_ARRAY_COUNT(Neighbors); i++)
 								{
 									// Calculate the high resolution indices, which may go into neighboring low resolution samples
 									const int32 HighResNeighborX = LowResX * UpsampleFactor + HighResX + Neighbors[i].X;
@@ -2337,7 +2337,7 @@ void FStaticLightingSystem::CalculateDirectSignedDistanceFieldLightingTextureMap
 									float WorldSpacePerHighResTexelX = FLT_MAX;
 									float WorldSpacePerHighResTexelY = FLT_MAX;
 									// Determine how far to scatter transition distance by measuring the world space distance between this texel and its neighbors
-									for (int32 i = 0 ; i < ARRAY_COUNT(Neighbors); i++)
+									for (int32 i = 0 ; i < UE_ARRAY_COUNT(Neighbors); i++)
 									{
 										if (HighResX + Neighbors[i].X > 0
 											&& HighResX + Neighbors[i].X < UpsampleFactor
diff --git a/Engine/Source/Programs/UnrealLightmass/Private/LightmassCore/Misc/LMDebug.cpp b/Engine/Source/Programs/UnrealLightmass/Private/LightmassCore/Misc/LMDebug.cpp
--- a/Engine/Source/Programs/UnrealLightmass/Private/LightmassCore/Misc/LMDebug.cpp
+++ b/Engine/Source/Programs/UnrealLightmass/Private/LightmassCore/Misc/LMDebug.cpp
@@ -94,10 +94,10 @@ void appHandleCriticalError()
 				FCString::Strcpy(ChangelistVersionStr, *FString::FromInt(0));
 
 				TCHAR CmdLine[2048];
-				FCString::Strncpy(CmdLine, FCommandLine::Get(), ARRAY_COUNT(CmdLine));
-				FCString::Strncat(CmdLine, TEXT(" -unattended"), ARRAY_COUNT(CmdLine));
+				FCString::Strncpy(CmdLine, FCommandLine::Get(), UE_ARRAY_COUNT(CmdLine));
+				FCString::Strncat(CmdLine, TEXT(" -unattended"), UE_ARRAY_COUNT(CmdLine));
 				TCHAR BaseDir[260];
-				FCString::Strncpy(BaseDir, FPlatformProcess::BaseDir(), ARRAY_COUNT(BaseDir));
+				FCString::Strncpy(BaseDir, FPlatformProcess::BaseDir(), UE_ARRAY_COUNT(BaseDir));
 				TCHAR separator = 0;
 
 				TCHAR EngineMode[64];
diff --git a/Engine/Source/Runtime/AIModule/Private/BehaviorTree/BehaviorTreeTypes.cpp b/Engine/Source/Runtime/AIModule/Private/BehaviorTree/BehaviorTreeTypes.cpp
--- a/Engine/Source/Runtime/AIModule/Private/BehaviorTree/BehaviorTreeTypes.cpp
+++ b/Engine/Source/Runtime/AIModule/Private/BehaviorTree/BehaviorTreeTypes.cpp
@@ -427,31 +427,31 @@ FString UBehaviorTreeTypes::BTLoggingContext;
 FString UBehaviorTreeTypes::DescribeNodeResult(EBTNodeResult::Type NodeResult)
 {
 	static FString ResultDesc[] = { TEXT("Succeeded"), TEXT("Failed"), TEXT("Aborted"), TEXT("InProgress") };
-	return (NodeResult < ARRAY_COUNT(ResultDesc)) ? ResultDesc[NodeResult] : FString();
+	return (NodeResult < UE_ARRAY_COUNT(ResultDesc)) ? ResultDesc[NodeResult] : FString();
 }
 
 FString UBehaviorTreeTypes::DescribeFlowAbortMode(EBTFlowAbortMode::Type AbortMode)
 {
 	static FString AbortModeDesc[] = { TEXT("None"), TEXT("Lower Priority"), TEXT("Self"), TEXT("Both") };
-	return (AbortMode < ARRAY_COUNT(AbortModeDesc)) ? AbortModeDesc[AbortMode] : FString();
+	return (AbortMode < UE_ARRAY_COUNT(AbortModeDesc)) ? AbortModeDesc[AbortMode] : FString();
 }
 
 FString UBehaviorTreeTypes::DescribeActiveNode(EBTActiveNode::Type ActiveNodeType)
 {
 	static FString ActiveDesc[] = { TEXT("Composite"), TEXT("ActiveTask"), TEXT("AbortingTask"), TEXT("InactiveTask") };
-	return (ActiveNodeType < ARRAY_COUNT(ActiveDesc)) ? ActiveDesc[ActiveNodeType] : FString();
+	return (ActiveNodeType < UE_ARRAY_COUNT(ActiveDesc)) ? ActiveDesc[ActiveNodeType] : FString();
 }
 
 FString UBehaviorTreeTypes::DescribeTaskStatus(EBTTaskStatus::Type TaskStatus)
 {
 	static FString TaskStatusDesc[] = { TEXT("Active"), TEXT("Aborting"), TEXT("Inactive") };
-	return (TaskStatus < ARRAY_COUNT(TaskStatusDesc)) ? TaskStatusDesc[TaskStatus] : FString();
+	return (TaskStatus < UE_ARRAY_COUNT(TaskStatusDesc)) ? TaskStatusDesc[TaskStatus] : FString();
 }
 
 FString UBehaviorTreeTypes::DescribeNodeUpdateMode(EBTNodeUpdateMode::Type UpdateMode)
 {
 	static FString UpdateModeDesc[] = { TEXT("Unknown"), TEXT("Add"), TEXT("Remove") };
-	return (UpdateMode < ARRAY_COUNT(UpdateModeDesc)) ? UpdateModeDesc[UpdateMode] : FString();
+	return (UpdateMode < UE_ARRAY_COUNT(UpdateModeDesc)) ? UpdateModeDesc[UpdateMode] : FString();
 }
 
 FString UBehaviorTreeTypes::DescribeNodeHelper(const UBTNode* Node)
diff --git a/Engine/Source/Runtime/AIModule/Private/GameplayDebugger/GameplayDebuggerCategory_Navmesh.cpp b/Engine/Source/Runtime/AIModule/Private/GameplayDebugger/GameplayDebuggerCategory_Navmesh.cpp
--- a/Engine/Source/Runtime/AIModule/Private/GameplayDebugger/GameplayDebuggerCategory_Navmesh.cpp
+++ b/Engine/Source/Runtime/AIModule/Private/GameplayDebugger/GameplayDebuggerCategory_Navmesh.cpp
@@ -174,7 +174,7 @@ void FGameplayDebuggerCategory_Navmesh::CollectData(APlayerController* OwnerPC,
 		int32 TargetTileX = 0;
 		int32 TargetTileY = 0;
 		RecastNavMesh->GetNavMeshTileXY(TargetLocation, TargetTileX, TargetTileY);
-		for (int32 Idx = 0; Idx < ARRAY_COUNT(DeltaX); Idx++)
+		for (int32 Idx = 0; Idx < UE_ARRAY_COUNT(DeltaX); Idx++)
 		{
 			const int32 NeiX = TargetTileX + DeltaX[Idx];
 			const int32 NeiY = TargetTileY + DeltaY[Idx];
diff --git a/Engine/Source/Runtime/AIModule/Private/Navigation/CrowdManager.cpp b/Engine/Source/Runtime/AIModule/Private/Navigation/CrowdManager.cpp
--- a/Engine/Source/Runtime/AIModule/Private/Navigation/CrowdManager.cpp
+++ b/Engine/Source/Runtime/AIModule/Private/Navigation/CrowdManager.cpp
@@ -1088,7 +1088,7 @@ void UCrowdManager::DrawDebugSharedBoundary() const
 	const dtSharedBoundary* sharedBounds = DetourCrowd->getSharedBoundary();
 	for (int32 Idx = 0; Idx < sharedBounds->Data.Num(); Idx++)
 	{
-		FColor Color = Colors[Idx % ARRAY_COUNT(Colors)];
+		FColor Color = Colors[Idx % UE_ARRAY_COUNT(Colors)];
 		const FVector Center = Recast2UnrealPoint(sharedBounds->Data[Idx].Center);
 		DrawDebugCylinder(World, Center - FCrowdDebug::Offset, Center, sharedBounds->Data[Idx].Radius, 32, Color);
 
diff --git a/Engine/Source/Runtime/Apple/MetalRHI/Private/MetalBuffer.cpp b/Engine/Source/Runtime/Apple/MetalRHI/Private/MetalBuffer.cpp
--- a/Engine/Source/Runtime/Apple/MetalRHI/Private/MetalBuffer.cpp
+++ b/Engine/Source/Runtime/Apple/MetalRHI/Private/MetalBuffer.cpp
@@ -1051,7 +1051,7 @@ void FMetalSubBufferRing::Shrink()
 	if(Buffer.IsValid())
 	{
 		NSUInteger FrameMax = 0;
-		for (uint32 i = 0; i < ARRAY_COUNT(FrameSize); i++)
+		for (uint32 i = 0; i < UE_ARRAY_COUNT(FrameSize); i++)
 		{
 			FrameMax = FMath::Max(FrameMax, FrameSize[i]);
 		}
@@ -1073,7 +1073,7 @@ void FMetalSubBufferRing::Shrink()
 			LastFrameChange = GFrameNumberRenderThread;
 		}
 		
-		FrameSize[GFrameNumberRenderThread % ARRAY_COUNT(FrameSize)] = 0;
+		FrameSize[GFrameNumberRenderThread % UE_ARRAY_COUNT(FrameSize)] = 0;
 	}
 }
 
@@ -1114,7 +1114,7 @@ void FMetalSubBufferRing::Commit(mtlpp::CommandBuffer& CmdBuf)
 			BytesWritten = TrailLen + WriteHead;
 		}
 		
-		FrameSize[GFrameNumberRenderThread % ARRAY_COUNT(FrameSize)] += Align(BytesWritten, MinAlign);
+		FrameSize[GFrameNumberRenderThread % UE_ARRAY_COUNT(FrameSize)] += Align(BytesWritten, MinAlign);
 		
 		TSharedPtr<FMetalRingBufferRef, ESPMode::ThreadSafe> CmdBufferRingBuffer = Buffer;
 		FPlatformMisc::MemoryBarrier();
diff --git a/Engine/Source/Runtime/ApplicationCore/Private/HoloLens/HoloLensErrorOutputDevice.cpp b/Engine/Source/Runtime/ApplicationCore/Private/HoloLens/HoloLensErrorOutputDevice.cpp
--- a/Engine/Source/Runtime/ApplicationCore/Private/HoloLens/HoloLensErrorOutputDevice.cpp
+++ b/Engine/Source/Runtime/ApplicationCore/Private/HoloLens/HoloLensErrorOutputDevice.cpp
@@ -71,7 +71,7 @@ void FHoloLensErrorOutputDevice::HandleError()
 	GIsRunning = 0;
 	GIsCriticalError = 1;
 	GLogConsole = NULL;
-	GErrorHist[ARRAY_COUNT(GErrorHist) - 1] = 0;
+	GErrorHist[UE_ARRAY_COUNT(GErrorHist) - 1] = 0;
 
 	// Trigger the OnSystemFailure hook if it exists
 	// make sure it happens after GIsGuarded is set to 0 in case this hook crashes
diff --git a/Engine/Source/Runtime/ApplicationCore/Private/IOS/IOSInputInterface.cpp b/Engine/Source/Runtime/ApplicationCore/Private/IOS/IOSInputInterface.cpp
--- a/Engine/Source/Runtime/ApplicationCore/Private/IOS/IOSInputInterface.cpp
+++ b/Engine/Source/Runtime/ApplicationCore/Private/IOS/IOSInputInterface.cpp
@@ -138,7 +138,7 @@ void FIOSInputInterface::HandleConnection(GCController* Controller)
 
 	// find a good controller index to use
 	bool bFoundSlot = false;
-	for (int32 ControllerIndex = 0; ControllerIndex < ARRAY_COUNT(Controllers); ControllerIndex++)
+	for (int32 ControllerIndex = 0; ControllerIndex < UE_ARRAY_COUNT(Controllers); ControllerIndex++)
 	{
 		// is this one already connected for this type of controller?
 		if ((!bIsTreatedAsGamepad && Controllers[ControllerIndex].bIsRemoteConnected == false) ||
@@ -686,7 +686,7 @@ void FIOSInputInterface::CalibrateMotion(uint32 PlayerIndex)
 	}
 #endif
 
-	if (PlayerIndex >= 0 && PlayerIndex < ARRAY_COUNT(Controllers))
+	if (PlayerIndex >= 0 && PlayerIndex < UE_ARRAY_COUNT(Controllers))
 	{
 		Controllers[PlayerIndex].bNeedsReferenceAttitude = true;
 	}
@@ -708,7 +708,7 @@ bool FIOSInputInterface::Exec(UWorld* InWorld, const TCHAR* Cmd, FOutputDevice&
 }
 bool FIOSInputInterface::IsControllerAssignedToGamepad(int32 ControllerId) const
 {
-	return ControllerId < ARRAY_COUNT(Controllers) &&
+	return ControllerId < UE_ARRAY_COUNT(Controllers) &&
 		(Controllers[ControllerId].bIsGamepadConnected ||
 		 Controllers[ControllerId].bIsRemoteConnected);
 }
@@ -716,7 +716,7 @@ bool FIOSInputInterface::IsControllerAssignedToGamepad(int32 ControllerId) const
 bool FIOSInputInterface::IsGamepadAttached() const
 {
 	bool bIsAttached = false;
-	for(int32 i = 0; i < ARRAY_COUNT(Controllers); ++i)
+	for(int32 i = 0; i < UE_ARRAY_COUNT(Controllers); ++i)
 	{
 		bIsAttached |= IsControllerAssignedToGamepad(i);
 	}
diff --git a/Engine/Source/Runtime/ApplicationCore/Private/IOS/IOSView.cpp b/Engine/Source/Runtime/ApplicationCore/Private/IOS/IOSView.cpp
--- a/Engine/Source/Runtime/ApplicationCore/Private/IOS/IOSView.cpp
+++ b/Engine/Source/Runtime/ApplicationCore/Private/IOS/IOSView.cpp
@@ -422,7 +422,7 @@ id<MTLDevice> GMetalDevice = nil;
 -(int32) GetTouchIndex:(UITouch*)Touch
 {
 	// look for existing touch
-	for (int Index = 0; Index < ARRAY_COUNT(AllTouches); Index++)
+	for (int Index = 0; Index < UE_ARRAY_COUNT(AllTouches); Index++)
 	{
 		if (AllTouches[Index] == Touch)
 		{
@@ -431,7 +431,7 @@ id<MTLDevice> GMetalDevice = nil;
 	}
 	
 	// if we get here, it's a new touch, find a slot
-	for (int Index = 0; Index < ARRAY_COUNT(AllTouches); Index++)
+	for (int Index = 0; Index < UE_ARRAY_COUNT(AllTouches); Index++)
 	{
 		if (AllTouches[Index] == nil)
 		{
diff --git a/Engine/Source/Runtime/ApplicationCore/Private/Mac/HIDInputInterface.cpp b/Engine/Source/Runtime/ApplicationCore/Private/Mac/HIDInputInterface.cpp
--- a/Engine/Source/Runtime/ApplicationCore/Private/Mac/HIDInputInterface.cpp
+++ b/Engine/Source/Runtime/ApplicationCore/Private/Mac/HIDInputInterface.cpp
@@ -311,7 +311,7 @@ void HIDInputInterface::FHIDDeviceInfo::SetupMappings()
 
 		bool bIsXbox360Controller = false;
 
-		for (int i = 0; i < ARRAY_COUNT(XBox360Controllers); ++i)
+		for (int i = 0; i < UE_ARRAY_COUNT(XBox360Controllers); ++i)
 		{
 			const FXBox360ControllerID ControllerID = XBox360Controllers[i];
 			if (ControllerID.VendorID == VendorID && ControllerID.ProductID == ProductID)
diff --git a/Engine/Source/Runtime/ApplicationCore/Private/Mac/MacConsoleOutputDevice.cpp b/Engine/Source/Runtime/ApplicationCore/Private/Mac/MacConsoleOutputDevice.cpp
--- a/Engine/Source/Runtime/ApplicationCore/Private/Mac/MacConsoleOutputDevice.cpp
+++ b/Engine/Source/Runtime/ApplicationCore/Private/Mac/MacConsoleOutputDevice.cpp
@@ -180,8 +180,8 @@ void FMacConsoleOutputDevice::Serialize( const TCHAR* Data, ELogVerbosity::Type
 					
 					// turn the string into a bunch of 0's and 1's
 					TCHAR String[9];
-					FMemory::Memset(String, 0, sizeof(TCHAR) * ARRAY_COUNT(String));
-					FCString::Strncpy(String, Data, ARRAY_COUNT(String));
+					FMemory::Memset(String, 0, sizeof(TCHAR) * UE_ARRAY_COUNT(String));
+					FCString::Strncpy(String, Data, UE_ARRAY_COUNT(String));
 					for (TCHAR* S = String; *S; S++)
 					{
 						*S -= '0';
diff --git a/Engine/Source/Runtime/ApplicationCore/Private/Windows/WindowsConsoleOutputDevice.cpp b/Engine/Source/Runtime/ApplicationCore/Private/Windows/WindowsConsoleOutputDevice.cpp
--- a/Engine/Source/Runtime/ApplicationCore/Private/Windows/WindowsConsoleOutputDevice.cpp
+++ b/Engine/Source/Runtime/ApplicationCore/Private/Windows/WindowsConsoleOutputDevice.cpp
@@ -245,8 +245,8 @@ void FWindowsConsoleOutputDevice::SetColor( const TCHAR* Color )
 	{
 		// turn the string into a bunch of 0's and 1's
 		TCHAR String[9];
-		FMemory::Memset(String, 0, sizeof(TCHAR) * ARRAY_COUNT(String));
-		FCString::Strncpy(String, Color, ARRAY_COUNT(String));
+		FMemory::Memset(String, 0, sizeof(TCHAR) * UE_ARRAY_COUNT(String));
+		FCString::Strncpy(String, Color, UE_ARRAY_COUNT(String));
 		for (TCHAR* S = String; *S; S++)
 		{
 			*S -= '0';
diff --git a/Engine/Source/Runtime/ApplicationCore/Private/Windows/WindowsPlatformSurvey.cpp b/Engine/Source/Runtime/ApplicationCore/Private/Windows/WindowsPlatformSurvey.cpp
--- a/Engine/Source/Runtime/ApplicationCore/Private/Windows/WindowsPlatformSurvey.cpp
+++ b/Engine/Source/Runtime/ApplicationCore/Private/Windows/WindowsPlatformSurvey.cpp
@@ -179,9 +179,9 @@ bool FWindowsPlatformSurvey::GetSurveyResults( FHardwareSurveyResults& OutResult
 	LCID DefaultLocale = GetSystemDefaultLCID();
 	const int32 MaxLocaleStringLength = 9;
 	TCHAR LangBuffer[MaxLocaleStringLength];
-	int LangReturn = GetLocaleInfo(DefaultLocale, LOCALE_SISO639LANGNAME, LangBuffer, ARRAY_COUNT(LangBuffer));
+	int LangReturn = GetLocaleInfo(DefaultLocale, LOCALE_SISO639LANGNAME, LangBuffer, UE_ARRAY_COUNT(LangBuffer));
 	TCHAR CountryBuffer[MaxLocaleStringLength];
-	int CountryReturn = GetLocaleInfo(DefaultLocale, LOCALE_SISO3166CTRYNAME, CountryBuffer, ARRAY_COUNT(CountryBuffer));
+	int CountryReturn = GetLocaleInfo(DefaultLocale, LOCALE_SISO3166CTRYNAME, CountryBuffer, UE_ARRAY_COUNT(CountryBuffer));
 
 	if (LangReturn == 0 || CountryReturn == 0)
 	{
diff --git a/Engine/Source/Runtime/AudioMixer/Private/AudioMixerChannelMaps.cpp b/Engine/Source/Runtime/AudioMixer/Private/AudioMixerChannelMaps.cpp
--- a/Engine/Source/Runtime/AudioMixer/Private/AudioMixerChannelMaps.cpp
+++ b/Engine/Source/Runtime/AudioMixer/Private/AudioMixerChannelMaps.cpp
@@ -233,7 +233,7 @@ namespace Audio
 	void FMixerDevice::Get2DChannelMapInternal(const int32 NumSourceChannels, const int32 NumOutputChannels, const bool bIsCenterChannelOnly, TArray<float>& OutChannelMap) const
 	{
 		const int32 OutputChannelMapIndex = NumOutputChannels - 1;
-		check(OutputChannelMapIndex < ARRAY_COUNT(OutputChannelMaps));
+		check(OutputChannelMapIndex < UE_ARRAY_COUNT(OutputChannelMaps));
 
 		float* RESTRICT Matrix = OutputChannelMaps[OutputChannelMapIndex];
 		check(Matrix != nullptr);
diff --git a/Engine/Source/Runtime/AudioMixerCore/Private/AudioMixer.cpp b/Engine/Source/Runtime/AudioMixerCore/Private/AudioMixer.cpp
--- a/Engine/Source/Runtime/AudioMixerCore/Private/AudioMixer.cpp
+++ b/Engine/Source/Runtime/AudioMixerCore/Private/AudioMixer.cpp
@@ -612,7 +612,7 @@ namespace Audio
 		bInitialized = true;
 
 		// Create a hard-coded default channel order
-		check(ARRAY_COUNT(DefaultChannelOrder) == AUDIO_MIXER_MAX_OUTPUT_CHANNELS);
+		check(UE_ARRAY_COUNT(DefaultChannelOrder) == AUDIO_MIXER_MAX_OUTPUT_CHANNELS);
 		DefaultChannelOrder[0] = EAudioMixerChannel::FrontLeft;
 		DefaultChannelOrder[1] = EAudioMixerChannel::FrontRight;
 		DefaultChannelOrder[2] = EAudioMixerChannel::FrontCenter;
diff --git a/Engine/Source/Runtime/Core/Private/Android/AndroidErrorOutputDevice.cpp b/Engine/Source/Runtime/Core/Private/Android/AndroidErrorOutputDevice.cpp
--- a/Engine/Source/Runtime/Core/Private/Android/AndroidErrorOutputDevice.cpp
+++ b/Engine/Source/Runtime/Core/Private/Android/AndroidErrorOutputDevice.cpp
@@ -41,7 +41,7 @@ void FAndroidErrorOutputDevice::HandleError()
 	GIsRunning = 0;
 	GIsCriticalError = 1;
 	GLogConsole = NULL;
-	GErrorHist[ARRAY_COUNT(GErrorHist) - 1] = 0;
+	GErrorHist[UE_ARRAY_COUNT(GErrorHist) - 1] = 0;
 
 	// Dump the error and flush the log.
 #if !NO_LOGGING
diff --git a/Engine/Source/Runtime/Core/Private/Android/AndroidPlatformMemory.cpp b/Engine/Source/Runtime/Core/Private/Android/AndroidPlatformMemory.cpp
--- a/Engine/Source/Runtime/Core/Private/Android/AndroidPlatformMemory.cpp
+++ b/Engine/Source/Runtime/Core/Private/Android/AndroidPlatformMemory.cpp
@@ -112,7 +112,7 @@ FPlatformMemoryStats FAndroidPlatformMemory::GetStats()
 		do
 		{
 			char LineBuffer[256] = { 0 };
-			char *Line = fgets(LineBuffer, ARRAY_COUNT(LineBuffer), FileGlobalMemStats);
+			char *Line = fgets(LineBuffer, UE_ARRAY_COUNT(LineBuffer), FileGlobalMemStats);
 			if (Line == nullptr)
 			{
 				break;	// eof or an error
@@ -157,7 +157,7 @@ FPlatformMemoryStats FAndroidPlatformMemory::GetStats()
 		do
 		{
 			char LineBuffer[256] = { 0 };
-			char *Line = fgets(LineBuffer, ARRAY_COUNT(LineBuffer), ProcMemStats);
+			char *Line = fgets(LineBuffer, UE_ARRAY_COUNT(LineBuffer), ProcMemStats);
 			if (Line == nullptr)
 			{
 				break;	// eof or an error
@@ -220,7 +220,7 @@ uint64 FAndroidPlatformMemory::GetMemoryUsedFast()
 		while (1)
 		{
 			char LineBuffer[256] = { 0 };
-			char *Line = fgets(LineBuffer, ARRAY_COUNT(LineBuffer), ProcMemStats);
+			char *Line = fgets(LineBuffer, UE_ARRAY_COUNT(LineBuffer), ProcMemStats);
 			if (Line == nullptr)
 			{
 				break;	// eof or an error
diff --git a/Engine/Source/Runtime/Core/Private/Android/AndroidPlatformMisc.cpp b/Engine/Source/Runtime/Core/Private/Android/AndroidPlatformMisc.cpp
--- a/Engine/Source/Runtime/Core/Private/Android/AndroidPlatformMisc.cpp
+++ b/Engine/Source/Runtime/Core/Private/Android/AndroidPlatformMisc.cpp
@@ -664,7 +664,7 @@ int32 FAndroidMisc::NumberOfCores()
 					char LineBuffer[256] = { 0 };
 					do
 					{
-						char *Line = fgets(LineBuffer, ARRAY_COUNT(LineBuffer), FileGlobalCpuStats);
+						char *Line = fgets(LineBuffer, UE_ARRAY_COUNT(LineBuffer), FileGlobalCpuStats);
 						if (Line == nullptr)
 						{
 							break;	// eof or an error
@@ -897,7 +897,7 @@ const int32 TargetSignals[] =
 	SIGSYS
 };
 
-const int32 NumTargetSignals = ARRAY_COUNT(TargetSignals);
+const int32 NumTargetSignals = UE_ARRAY_COUNT(TargetSignals);
 
 struct sigaction PrevActions[NumTargetSignals];
 static bool PreviousSignalHandlersValid = false;
@@ -2287,7 +2287,7 @@ uint32 FAndroidMisc::GetCoreFrequency(int32 CoreIndex, ECoreFrequencyProperty Co
 	if (FILE* CoreFreqStateFile = fopen(QueryFile, "r"))
 	{
 		char curr_corefreq[32] = { 0 };
-		if( fgets(curr_corefreq, ARRAY_COUNT(curr_corefreq), CoreFreqStateFile) != nullptr)
+		if( fgets(curr_corefreq, UE_ARRAY_COUNT(curr_corefreq), CoreFreqStateFile) != nullptr)
 		{
 			ReturnFrequency = atol(curr_corefreq);
 		}
diff --git a/Engine/Source/Runtime/Core/Private/Apple/ApplePlatformCrashContext.cpp b/Engine/Source/Runtime/Core/Private/Apple/ApplePlatformCrashContext.cpp
--- a/Engine/Source/Runtime/Core/Private/Apple/ApplePlatformCrashContext.cpp
+++ b/Engine/Source/Runtime/Core/Private/Apple/ApplePlatformCrashContext.cpp
@@ -31,15 +31,15 @@ void FApplePlatformCrashContext::InitFromSignal(int32 InSignal, siginfo_t* InInf
 	Info = InInfo;
 	Context = reinterpret_cast< ucontext_t* >( InContext );
 	
-#define HANDLE_CASE(a,b) case a: FCStringAnsi::Strncpy(SignalDescription, #a ": " b, ARRAY_COUNT( SignalDescription ) - 1); break;
+#define HANDLE_CASE(a,b) case a: FCStringAnsi::Strncpy(SignalDescription, #a ": " b, UE_ARRAY_COUNT( SignalDescription ) - 1); break;
 	switch (Signal)
 	{
 		case SIGSEGV:
-			FCStringAnsi::Strncpy(SignalDescription, "SIGSEGV: invalid attempt to access memory at address 0x", ARRAY_COUNT( SignalDescription ) - 1);
+			FCStringAnsi::Strncpy(SignalDescription, "SIGSEGV: invalid attempt to access memory at address 0x", UE_ARRAY_COUNT( SignalDescription ) - 1);
 			FCStringAnsi::Strcat(SignalDescription, ItoANSI((uintptr_t)Info->si_addr, 16));
 			break;
 		case SIGBUS:
-			FCStringAnsi::Strncpy(SignalDescription, "SIGBUS: invalid attempt to access memory at address 0x", ARRAY_COUNT( SignalDescription ) - 1);
+			FCStringAnsi::Strncpy(SignalDescription, "SIGBUS: invalid attempt to access memory at address 0x", UE_ARRAY_COUNT( SignalDescription ) - 1);
 			FCStringAnsi::Strcat(SignalDescription, ItoANSI((uintptr_t)Info->si_addr, 16));
 			break;
 			HANDLE_CASE(SIGINT, "program interrupted")
@@ -55,7 +55,7 @@ void FApplePlatformCrashContext::InitFromSignal(int32 InSignal, siginfo_t* InInf
 			HANDLE_CASE(SIGSTOP, "stop")
 			
 		default:
-			FCStringAnsi::Strncpy(SignalDescription, "Signal ", ARRAY_COUNT( SignalDescription ) - 1);
+			FCStringAnsi::Strncpy(SignalDescription, "Signal ", UE_ARRAY_COUNT( SignalDescription ) - 1);
 			FCStringAnsi::Strcat(SignalDescription, ItoANSI(Signal, 10));
 			FCStringAnsi::Strcat(SignalDescription, " (unknown)");
 			break;
@@ -78,9 +78,9 @@ int32 FApplePlatformCrashContext::ReportCrash() const
 		// Generate the portable callstack for the crash xml
 		const_cast<FApplePlatformCrashContext*>(this)->CapturePortableCallStack(IgnoreDepth, Context);
 		// Walk the stack and dump it to the allocated memory (ignore first 2 callstack lines as those are in stack walking code)
-		FPlatformStackWalk::StackWalkAndDump( StackBuffer, ARRAY_COUNT(MinidumpCallstackInfo) - 1, IgnoreDepth, Context );
+		FPlatformStackWalk::StackWalkAndDump( StackBuffer, UE_ARRAY_COUNT(MinidumpCallstackInfo) - 1, IgnoreDepth, Context );
 		
-		FUTF8ToTCHAR_Convert::Convert(GErrorHist, ARRAY_COUNT(GErrorHist) - 1, MinidumpCallstackInfo, FCStringAnsi::Strlen(MinidumpCallstackInfo));
+		FUTF8ToTCHAR_Convert::Convert(GErrorHist, UE_ARRAY_COUNT(GErrorHist) - 1, MinidumpCallstackInfo, FCStringAnsi::Strlen(MinidumpCallstackInfo));
 		CreateExceptionInfoString(Signal, Info);
 	}
 	
@@ -89,15 +89,15 @@ int32 FApplePlatformCrashContext::ReportCrash() const
 
 void FApplePlatformCrashContext::CreateExceptionInfoString(int32 Signal, siginfo_t* Info)
 {
-#define HANDLE_CASE(a,b) case a: FCString::Strncpy(GErrorExceptionDescription, TEXT(#a ": " b), ARRAY_COUNT( SignalDescription ) - 1); break;
+#define HANDLE_CASE(a,b) case a: FCString::Strncpy(GErrorExceptionDescription, TEXT(#a ": " b), UE_ARRAY_COUNT( SignalDescription ) - 1); break;
 	switch (Signal)
 	{
 		case SIGSEGV:
-			FCString::Strncpy(GErrorExceptionDescription, TEXT("SIGSEGV: invalid attempt to access memory at address 0x"), ARRAY_COUNT( SignalDescription ) - 1);
+			FCString::Strncpy(GErrorExceptionDescription, TEXT("SIGSEGV: invalid attempt to access memory at address 0x"), UE_ARRAY_COUNT( SignalDescription ) - 1);
 			FCString::Strcat(GErrorExceptionDescription, ItoTCHAR((uintptr_t)Info->si_addr, 16));
 			break;
 		case SIGBUS:
-			FCString::Strncpy(GErrorExceptionDescription, TEXT("SIGBUS: invalid attempt to access memory at address 0x"), ARRAY_COUNT( SignalDescription ) - 1);
+			FCString::Strncpy(GErrorExceptionDescription, TEXT("SIGBUS: invalid attempt to access memory at address 0x"), UE_ARRAY_COUNT( SignalDescription ) - 1);
 			FCString::Strcat(GErrorExceptionDescription, ItoTCHAR((uintptr_t)Info->si_addr, 16));
 			break;
 			HANDLE_CASE(SIGINT, "program interrupted")
@@ -113,7 +113,7 @@ void FApplePlatformCrashContext::CreateExceptionInfoString(int32 Signal, siginfo
 			HANDLE_CASE(SIGSTOP, "stop")
 			
 		default:
-			FCString::Strncpy(GErrorExceptionDescription, TEXT("Signal "), ARRAY_COUNT( SignalDescription ) - 1);
+			FCString::Strncpy(GErrorExceptionDescription, TEXT("Signal "), UE_ARRAY_COUNT( SignalDescription ) - 1);
 			FCString::Strcat(GErrorExceptionDescription, ItoTCHAR(Signal, 10));
 			FCString::Strcat(GErrorExceptionDescription, TEXT(" (unknown)"));
 			break;
diff --git a/Engine/Source/Runtime/Core/Private/Apple/ApplePlatformStackWalk.cpp b/Engine/Source/Runtime/Core/Private/Apple/ApplePlatformStackWalk.cpp
--- a/Engine/Source/Runtime/Core/Private/Apple/ApplePlatformStackWalk.cpp
+++ b/Engine/Source/Runtime/Core/Private/Apple/ApplePlatformStackWalk.cpp
@@ -425,7 +425,7 @@ void CreateExceptionInfoString(int32 Signal, struct __siginfo* Info)
 	}
 	
 #if WITH_EDITORONLY_DATA
-	FCString::Strncpy(GErrorExceptionDescription, *ErrorString, FMath::Min(ErrorString.Len() + 1, (int32)ARRAY_COUNT(GErrorExceptionDescription)));
+	FCString::Strncpy(GErrorExceptionDescription, *ErrorString, FMath::Min(ErrorString.Len() + 1, (int32)UE_ARRAY_COUNT(GErrorExceptionDescription)));
 #endif
 #undef HANDLE_CASE
 }
@@ -448,7 +448,7 @@ int32 ReportCrash(ucontext_t *Context, int32 Signal, struct __siginfo* Info)
 		// Walk the stack and dump it to the allocated memory.
 		FPlatformStackWalk::StackWalkAndDump( StackTrace, StackTraceSize, 0, Context );
 #if WITH_EDITORONLY_DATA
-        FCString::Strncat( GErrorHist, ANSI_TO_TCHAR(StackTrace), ARRAY_COUNT(GErrorHist) - 1 );
+        FCString::Strncat( GErrorHist, ANSI_TO_TCHAR(StackTrace), UE_ARRAY_COUNT(GErrorHist) - 1 );
 		CreateExceptionInfoString(Signal, Info);
 #endif
 	}
diff --git a/Engine/Source/Runtime/Core/Private/Apple/ApplePlatformSymbolication.cpp b/Engine/Source/Runtime/Core/Private/Apple/ApplePlatformSymbolication.cpp
--- a/Engine/Source/Runtime/Core/Private/Apple/ApplePlatformSymbolication.cpp
+++ b/Engine/Source/Runtime/Core/Private/Apple/ApplePlatformSymbolication.cpp
@@ -413,7 +413,7 @@ bool FApplePlatformSymbolication::SymbolInfoForStrippedSymbol(FApplePlatformSymb
 				ANSICHAR const* FunctionName = CSSymbolGetName(Symbol);
 				if(FunctionName)
 				{
-					FCStringAnsi::Strcpy(Info.FunctionName, ARRAY_COUNT(Info.FunctionName), FunctionName);
+					FCStringAnsi::Strcpy(Info.FunctionName, UE_ARRAY_COUNT(Info.FunctionName), FunctionName);
 				}
 				
 				CSRange Range = CSSymbolGetRange(Symbol);
@@ -429,7 +429,7 @@ bool FApplePlatformSymbolication::SymbolInfoForStrippedSymbol(FApplePlatformSymb
 					ANSICHAR const* FileName = CSSourceInfoGetPath(SymbolInfo);
 					if(FileName)
 					{
-						FCStringAnsi::Strcpy(Info.Filename, ARRAY_COUNT(Info.FunctionName), FileName);
+						FCStringAnsi::Strcpy(Info.Filename, UE_ARRAY_COUNT(Info.FunctionName), FileName);
 					}
 				}
 				
@@ -458,8 +458,8 @@ bool FApplePlatformSymbolication::SymbolInfoForAddress(uint64 ProgramCounter, FP
 			if(!CSIsNull(Symbol))
 			{
 				out_SymbolInfo.LineNumber = CSSourceInfoGetLineNumber(Symbol);
-				FCStringAnsi::Strcpy(out_SymbolInfo.Filename, ARRAY_COUNT(out_SymbolInfo.Filename), CSSourceInfoGetPath(Symbol));
-				FCStringAnsi::Strcpy(out_SymbolInfo.FunctionName, ARRAY_COUNT(out_SymbolInfo.FunctionName), CSSymbolGetName(CSSourceInfoGetSymbol(Symbol)));
+				FCStringAnsi::Strcpy(out_SymbolInfo.Filename, UE_ARRAY_COUNT(out_SymbolInfo.Filename), CSSourceInfoGetPath(Symbol));
+				FCStringAnsi::Strcpy(out_SymbolInfo.FunctionName, UE_ARRAY_COUNT(out_SymbolInfo.FunctionName), CSSymbolGetName(CSSourceInfoGetSymbol(Symbol)));
 				CSRange CodeRange = CSSourceInfoGetRange(Symbol);
 				out_SymbolInfo.SymbolDisplacement = (ProgramCounter - CodeRange.Location);
 				
@@ -467,7 +467,7 @@ bool FApplePlatformSymbolication::SymbolInfoForAddress(uint64 ProgramCounter, FP
 				if(!CSIsNull(Owner))
 				{
 					ANSICHAR const* DylibName = CSSymbolOwnerGetName(Owner);
-					FCStringAnsi::Strcpy(out_SymbolInfo.ModuleName, ARRAY_COUNT(out_SymbolInfo.ModuleName), DylibName);
+					FCStringAnsi::Strcpy(out_SymbolInfo.ModuleName, UE_ARRAY_COUNT(out_SymbolInfo.ModuleName), DylibName);
 					
 					bOK = out_SymbolInfo.LineNumber != 0;
 				}
diff --git a/Engine/Source/Runtime/Core/Private/Containers/StackTracker.cpp b/Engine/Source/Runtime/Core/Private/Containers/StackTracker.cpp
--- a/Engine/Source/Runtime/Core/Private/Containers/StackTracker.cpp
+++ b/Engine/Source/Runtime/Core/Private/Containers/StackTracker.cpp
@@ -52,7 +52,7 @@ void FStackTracker::CaptureStackTrace(int32 EntriesToIgnore, void* UserData, int
 					{
 						ANSICHAR AddressInformation[512];
 						AddressInformation[0] = 0;
-						FPlatformStackWalk::ProgramCounterToHumanReadableString( 1, BackTrace[Index], AddressInformation, ARRAY_COUNT(AddressInformation)-1 );
+						FPlatformStackWalk::ProgramCounterToHumanReadableString( 1, BackTrace[Index], AddressInformation, UE_ARRAY_COUNT(AddressInformation)-1 );
 						FString Symbol(AddressInformation);
 						int32 Spot = Symbol.Find(TEXT(" - "));
 						if (Spot != INDEX_NONE)
@@ -193,11 +193,11 @@ void FStackTracker::DumpStackTraces(int32 StackThreshold, FOutputDevice& Ar, flo
 			
 
 			// Iterate over all addresses in the callstack to look up symbol name.
-			for( int32 AddressIndex=0; AddressIndex<ARRAY_COUNT(CallStack.Addresses) && CallStack.Addresses[AddressIndex]; AddressIndex++ )
+			for( int32 AddressIndex=0; AddressIndex<UE_ARRAY_COUNT(CallStack.Addresses) && CallStack.Addresses[AddressIndex]; AddressIndex++ )
 			{
 				ANSICHAR AddressInformation[512];
 				AddressInformation[0] = 0;
-				FPlatformStackWalk::ProgramCounterToHumanReadableString( AddressIndex, CallStack.Addresses[AddressIndex], AddressInformation, ARRAY_COUNT(AddressInformation)-1 );
+				FPlatformStackWalk::ProgramCounterToHumanReadableString( AddressIndex, CallStack.Addresses[AddressIndex], AddressInformation, UE_ARRAY_COUNT(AddressInformation)-1 );
 				CallStackString = CallStackString + LINE_TERMINATOR TEXT(",,,") + FString(AddressInformation);
 			}
 
diff --git a/Engine/Source/Runtime/Core/Private/Containers/String.cpp b/Engine/Source/Runtime/Core/Private/Containers/String.cpp
--- a/Engine/Source/Runtime/Core/Private/Containers/String.cpp
+++ b/Engine/Source/Runtime/Core/Private/Containers/String.cpp
@@ -1005,7 +1005,7 @@ int32 FString::ParseIntoArrayWS( TArray<FString>& OutArray, const TCHAR* pchExtr
 	};
 
 	// start with just the standard whitespaces
-	int32 NumWhiteSpaces = ARRAY_COUNT(WhiteSpace) - 1;
+	int32 NumWhiteSpaces = UE_ARRAY_COUNT(WhiteSpace) - 1;
 	// if we got one passed in, use that in addition
 	if (pchExtraDelim && *pchExtraDelim)
 	{
@@ -1026,7 +1026,7 @@ int32 FString::ParseIntoArrayLines(TArray<FString>& OutArray, bool InCullEmpty)
 	};
 
 	// start with just the standard line endings
-	int32 NumLineEndings = ARRAY_COUNT(LineEndings);	
+	int32 NumLineEndings = UE_ARRAY_COUNT(LineEndings);	
 	return ParseIntoArray(OutArray, LineEndings, NumLineEndings, InCullEmpty);
 }
 
@@ -1252,7 +1252,7 @@ static const TCHAR* CharToEscapeSeqMap[][2] =
 	{ TEXT("\""), TEXT("\\\"") }
 };
 
-static const uint32 MaxSupportedEscapeChars = ARRAY_COUNT(CharToEscapeSeqMap);
+static const uint32 MaxSupportedEscapeChars = UE_ARRAY_COUNT(CharToEscapeSeqMap);
 
 /**
  * Replaces certain characters with the "escaped" version of that character (i.e. replaces "\n" with "\\n").
@@ -1574,7 +1574,7 @@ FString SlugStringForValidName(const FString& DisplayString, const TCHAR* Replac
 	// Convert the display label, which may consist of just about any possible character, into a
 	// suitable name for a UObject (remove whitespace, certain symbols, etc.)
 	{
-		for ( int32 BadCharacterIndex = 0; BadCharacterIndex < ARRAY_COUNT(INVALID_OBJECTNAME_CHARACTERS) - 1; ++BadCharacterIndex )
+		for ( int32 BadCharacterIndex = 0; BadCharacterIndex < UE_ARRAY_COUNT(INVALID_OBJECTNAME_CHARACTERS) - 1; ++BadCharacterIndex )
 		{
 			const TCHAR TestChar[2] = { INVALID_OBJECTNAME_CHARACTERS[BadCharacterIndex], 0 };
 			const int32 NumReplacedChars = GeneratedName.ReplaceInline(TestChar, ReplaceWith);
diff --git a/Engine/Source/Runtime/Core/Private/GenericPlatform/GenericPlatformOutputDevices.cpp b/Engine/Source/Runtime/Core/Private/GenericPlatform/GenericPlatformOutputDevices.cpp
--- a/Engine/Source/Runtime/Core/Private/GenericPlatform/GenericPlatformOutputDevices.cpp
+++ b/Engine/Source/Runtime/Core/Private/GenericPlatform/GenericPlatformOutputDevices.cpp
@@ -57,7 +57,7 @@ FString FGenericPlatformOutputDevices::GetAbsoluteLogFilename()
 {
 	if (!CachedAbsoluteFilename[0])
 	{
-		FCString::Strcpy(CachedAbsoluteFilename, ARRAY_COUNT(CachedAbsoluteFilename), *FPaths::ProjectLogDir());
+		FCString::Strcpy(CachedAbsoluteFilename, UE_ARRAY_COUNT(CachedAbsoluteFilename), *FPaths::ProjectLogDir());
 		FString LogFilename;
 		const bool bShouldStopOnSeparator = false;
 		if (!FParse::Value(FCommandLine::Get(), TEXT("LOG="), LogFilename, bShouldStopOnSeparator))
@@ -89,7 +89,7 @@ FString FGenericPlatformOutputDevices::GetAbsoluteLogFilename()
 			LogFilename += TEXT(".log");
 		}
 
-		FCString::Strcat(CachedAbsoluteFilename, ARRAY_COUNT(CachedAbsoluteFilename) - FCString::Strlen(CachedAbsoluteFilename), *LogFilename);
+		FCString::Strcat(CachedAbsoluteFilename, UE_ARRAY_COUNT(CachedAbsoluteFilename) - FCString::Strlen(CachedAbsoluteFilename), *LogFilename);
 	}
 
 	return CachedAbsoluteFilename;
diff --git a/Engine/Source/Runtime/Core/Private/GenericPlatform/GenericPlatformTime.cpp b/Engine/Source/Runtime/Core/Private/GenericPlatform/GenericPlatformTime.cpp
--- a/Engine/Source/Runtime/Core/Private/GenericPlatform/GenericPlatformTime.cpp
+++ b/Engine/Source/Runtime/Core/Private/GenericPlatform/GenericPlatformTime.cpp
@@ -112,9 +112,9 @@ const TCHAR* FGenericPlatformTime::StrTimestamp()
 {
 	static TCHAR Result[1024];
 	*Result = 0;
-	StrDate( Result, ARRAY_COUNT(Result) );
+	StrDate( Result, UE_ARRAY_COUNT(Result) );
 	FCString::Strcat( Result, TEXT(" ") );
-	StrTime( Result + FCString::Strlen(Result), ARRAY_COUNT(Result) - FCString::Strlen(Result) );
+	StrTime( Result + FCString::Strlen(Result), UE_ARRAY_COUNT(Result) - FCString::Strlen(Result) );
 	return Result;
 }
 
diff --git a/Engine/Source/Runtime/Core/Private/GenericPlatform/GenericWidePlatformString.cpp b/Engine/Source/Runtime/Core/Private/GenericPlatform/GenericWidePlatformString.cpp
--- a/Engine/Source/Runtime/Core/Private/GenericPlatform/GenericWidePlatformString.cpp
+++ b/Engine/Source/Runtime/Core/Private/GenericPlatform/GenericWidePlatformString.cpp
@@ -478,7 +478,7 @@ int32 FGenericWidePlatformString::GetVarArgs( WIDECHAR* Dest, SIZE_T DestSize, c
 
 				// Yes, this is lame.
 				int CpyIdx = 0;
-				while (Percent < Src && CpyIdx < ARRAY_COUNT(FmtBuf))
+				while (Percent < Src && CpyIdx < UE_ARRAY_COUNT(FmtBuf))
 				{
 					if (LIKELY(*Percent != '*'))
 					{
@@ -488,7 +488,7 @@ int32 FGenericWidePlatformString::GetVarArgs( WIDECHAR* Dest, SIZE_T DestSize, c
 					else
 					{
 						snprintf(&FmtBuf[CpyIdx], sizeof(FmtBuf) - CpyIdx, "%d", FieldLen);
-						while (CpyIdx < ARRAY_COUNT(FmtBuf) && FmtBuf[CpyIdx] != 0)
+						while (CpyIdx < UE_ARRAY_COUNT(FmtBuf) && FmtBuf[CpyIdx] != 0)
 						{
 							++CpyIdx;
 						}
@@ -527,7 +527,7 @@ int32 FGenericWidePlatformString::GetVarArgs( WIDECHAR* Dest, SIZE_T DestSize, c
 
 				// Yes, this is lame.
 				int CpyIdx = 0;
-				while (Percent < Src && CpyIdx < ARRAY_COUNT(FmtBuf))
+				while (Percent < Src && CpyIdx < UE_ARRAY_COUNT(FmtBuf))
 				{
 					if (LIKELY(*Percent != '*'))
 					{
@@ -537,7 +537,7 @@ int32 FGenericWidePlatformString::GetVarArgs( WIDECHAR* Dest, SIZE_T DestSize, c
 					else
 					{
 						snprintf(&FmtBuf[CpyIdx], sizeof(FmtBuf) - CpyIdx, "%d", FieldLen);
-						while (CpyIdx < ARRAY_COUNT(FmtBuf) && FmtBuf[CpyIdx] != 0)
+						while (CpyIdx < UE_ARRAY_COUNT(FmtBuf) && FmtBuf[CpyIdx] != 0)
 						{
 							++CpyIdx;
 						}
@@ -563,7 +563,7 @@ int32 FGenericWidePlatformString::GetVarArgs( WIDECHAR* Dest, SIZE_T DestSize, c
 
 				// Yes, this is lame.
 				int CpyIdx = 0;
-				while (Percent < Src && CpyIdx < ARRAY_COUNT(FmtBuf))
+				while (Percent < Src && CpyIdx < UE_ARRAY_COUNT(FmtBuf))
 				{
 					FmtBuf[CpyIdx] = (ANSICHAR)*Percent;
 					Percent++;
@@ -602,7 +602,7 @@ int32 FGenericWidePlatformString::GetVarArgs( WIDECHAR* Dest, SIZE_T DestSize, c
 
 					// Yes, this is lame.
 					int CpyIdx = 0;
-					while (Percent < Src && CpyIdx < ARRAY_COUNT(FmtBuf))
+					while (Percent < Src && CpyIdx < UE_ARRAY_COUNT(FmtBuf))
 					{
 						FmtBuf[CpyIdx] = (ANSICHAR)*Percent;
 						Percent++;
@@ -627,7 +627,7 @@ int32 FGenericWidePlatformString::GetVarArgs( WIDECHAR* Dest, SIZE_T DestSize, c
 
 					// Yes, this is lame.
 					int CpyIdx = 0;
-					while (Percent < Src && CpyIdx < ARRAY_COUNT(FmtBuf))
+					while (Percent < Src && CpyIdx < UE_ARRAY_COUNT(FmtBuf))
 					{
 						FmtBuf[CpyIdx] = (ANSICHAR) *Percent;
 						Percent++;
@@ -704,7 +704,7 @@ int32 FGenericWidePlatformString::GetVarArgs( WIDECHAR* Dest, SIZE_T DestSize, c
 					CpyIdx = 0;
 				}
 
-				while (Percent < Src && CpyIdx < ARRAY_COUNT(FmtBuf))
+				while (Percent < Src && CpyIdx < UE_ARRAY_COUNT(FmtBuf))
 				{
 					FmtBuf[CpyIdx] = (ANSICHAR) *Percent;
 					Percent++;
@@ -731,7 +731,7 @@ int32 FGenericWidePlatformString::GetVarArgs( WIDECHAR* Dest, SIZE_T DestSize, c
 
 				// Yes, this is lame.
 				int CpyIdx = 0;
-				while (Percent < Src && CpyIdx < ARRAY_COUNT(FmtBuf))
+				while (Percent < Src && CpyIdx < UE_ARRAY_COUNT(FmtBuf))
 				{
 					FmtBuf[CpyIdx] = (ANSICHAR) *Percent;
 					Percent++;
@@ -740,11 +740,11 @@ int32 FGenericWidePlatformString::GetVarArgs( WIDECHAR* Dest, SIZE_T DestSize, c
 				FmtBuf[CpyIdx] = 0;
 
 				int RetCnt = snprintf(AnsiNum, sizeof (AnsiNum), FmtBuf, Val);
-				if (RetCnt >= ARRAY_COUNT(AnsiNum))
+				if (RetCnt >= UE_ARRAY_COUNT(AnsiNum))
 				{
 					// We should print what we have written into AnsiNum but ensure we null terminate before printing
-					AnsiNum[ARRAY_COUNT(AnsiNum) - 1] = '\0';
-					checkf(0, TEXT("Attempting to read past the size our buffer. Buffer Size: %d Size to read: %d. Current contents: '%s'\n"), ARRAY_COUNT(AnsiNum), RetCnt, UTF8_TO_TCHAR(AnsiNum));
+					AnsiNum[UE_ARRAY_COUNT(AnsiNum) - 1] = '\0';
+					checkf(0, TEXT("Attempting to read past the size our buffer. Buffer Size: %d Size to read: %d. Current contents: '%s'\n"), UE_ARRAY_COUNT(AnsiNum), RetCnt, UTF8_TO_TCHAR(AnsiNum));
 				}
 				if (!DestIter.Write(AnsiNum, RetCnt))
 				{
diff --git a/Engine/Source/Runtime/Core/Private/GenericPlatform/StandardPlatformString.cpp b/Engine/Source/Runtime/Core/Private/GenericPlatform/StandardPlatformString.cpp
--- a/Engine/Source/Runtime/Core/Private/GenericPlatform/StandardPlatformString.cpp
+++ b/Engine/Source/Runtime/Core/Private/GenericPlatform/StandardPlatformString.cpp
@@ -264,7 +264,7 @@ int32 FStandardPlatformString::GetVarArgs( WIDECHAR* Dest, SIZE_T DestSize, cons
 			Format += GetFormattingInfo(Format, Info);
 
 			WIDECHAR Formatted[1024];
-			int32 Length = ARRAY_COUNT(Formatted);
+			int32 Length = UE_ARRAY_COUNT(Formatted);
 			const WIDECHAR* FormattedArg = GetFormattedArgument(Info, ArgPtr, Formatted, Length);
 			if (FormattedArg && Length > 0)
 			{
diff --git a/Engine/Source/Runtime/Core/Private/HAL/LowLevelMemTracker.cpp b/Engine/Source/Runtime/Core/Private/HAL/LowLevelMemTracker.cpp
--- a/Engine/Source/Runtime/Core/Private/HAL/LowLevelMemTracker.cpp
+++ b/Engine/Source/Runtime/Core/Private/HAL/LowLevelMemTracker.cpp
@@ -130,7 +130,7 @@ extern const TCHAR* LLMGetTagName(ELLMTag Tag)
 #undef LLM_TAG_NAME_ARRAY
 
 	int32 Index = (int32)Tag;
-	if (Index >= 0 && Index < ARRAY_COUNT(Names))
+	if (Index >= 0 && Index < UE_ARRAY_COUNT(Names))
 	{
 		return Names[Index];
 	}
@@ -147,7 +147,7 @@ extern const ANSICHAR* LLMGetTagNameANSI(ELLMTag Tag)
 #undef LLM_TAG_NAME_ARRAY
 
 	int32 Index = (int32)Tag;
-	if (Index >= 0 && Index < ARRAY_COUNT(Names))
+	if (Index >= 0 && Index < UE_ARRAY_COUNT(Names))
 	{
 		return Names[Index];
 	}
@@ -164,7 +164,7 @@ extern FName LLMGetTagStat(ELLMTag Tag)
 #undef LLM_TAG_STAT_ARRAY
 
 	int32 Index = (int32)Tag;
-	if (Index >= 0 && Index < ARRAY_COUNT(Names))
+	if (Index >= 0 && Index < UE_ARRAY_COUNT(Names))
 	{
 		return Names[Index];
 	}
@@ -181,7 +181,7 @@ extern FName LLMGetTagStatGroup(ELLMTag Tag)
 #undef LLM_TAG_STAT_ARRAY
 
 	int32 Index = (int32)Tag;
-	if (Index >= 0 && Index < ARRAY_COUNT(Names))
+	if (Index >= 0 && Index < UE_ARRAY_COUNT(Names))
 	{
 		return Names[Index];
 	}
@@ -198,7 +198,7 @@ extern int32 LLMGetTagParent(ELLMTag Tag)
 #undef LLM_TAG_NAME_ARRAY
 
 	int32 Index = (int32)Tag;
-	if( Index >= 0 && Index < ARRAY_COUNT(ParentTags))
+	if( Index >= 0 && Index < UE_ARRAY_COUNT(ParentTags))
 	{
 		return ParentTags[Index];
 	}
@@ -215,7 +215,7 @@ bool LLMPrivate::HandleAssert(bool bLog, const TCHAR* Format, ...)
 	if (bLog)
 	{
 		TCHAR DescriptionString[4096];
-		GET_VARARGS(DescriptionString, ARRAY_COUNT(DescriptionString), ARRAY_COUNT(DescriptionString) - 1, Format, Format);
+		GET_VARARGS(DescriptionString, UE_ARRAY_COUNT(DescriptionString), UE_ARRAY_COUNT(DescriptionString) - 1, Format, Format);
 
 		FPlatformMisc::LowLevelOutputDebugString(DescriptionString);
 
diff --git a/Engine/Source/Runtime/Core/Private/HAL/MallocBinned2.cpp b/Engine/Source/Runtime/Core/Private/HAL/MallocBinned2.cpp
--- a/Engine/Source/Runtime/Core/Private/HAL/MallocBinned2.cpp
+++ b/Engine/Source/Runtime/Core/Private/HAL/MallocBinned2.cpp
@@ -655,8 +655,8 @@ FMallocBinned2::FMallocBinned2()
 	checkf(SmallBlockSizes[BINNED2_SMALL_POOL_COUNT - 1] == BINNED2_MAX_SMALL_POOL_SIZE, TEXT("BINNED2_MAX_SMALL_POOL_SIZE must equal the smallest block size"));
 	checkf(PageSize % BINNED2_LARGE_ALLOC == 0, TEXT("OS page size must be a multiple of BINNED2_LARGE_ALLOC"));
 	checkf(sizeof(FMallocBinned2::FFreeBlock) <= SmallBlockSizes[0], TEXT("Pool header must be able to fit into the smallest block"));
-	static_assert(ARRAY_COUNT(SmallBlockSizes) == BINNED2_SMALL_POOL_COUNT, "Small block size array size must match BINNED2_SMALL_POOL_COUNT");
-	static_assert(ARRAY_COUNT(SmallBlockSizes) <= 256, "Small block size array size must fit in a byte");
+	static_assert(UE_ARRAY_COUNT(SmallBlockSizes) == BINNED2_SMALL_POOL_COUNT, "Small block size array size must match BINNED2_SMALL_POOL_COUNT");
+	static_assert(UE_ARRAY_COUNT(SmallBlockSizes) <= 256, "Small block size array size must fit in a byte");
 	static_assert(sizeof(FFreeBlock) <= BINNED2_MINIMUM_ALIGNMENT, "Free block struct must be small enough to fit into a block.");
 
 	// Init pool tables.
diff --git a/Engine/Source/Runtime/Core/Private/HAL/MallocBinnedCommon.cpp b/Engine/Source/Runtime/Core/Private/HAL/MallocBinnedCommon.cpp
--- a/Engine/Source/Runtime/Core/Private/HAL/MallocBinnedCommon.cpp
+++ b/Engine/Source/Runtime/Core/Private/HAL/MallocBinnedCommon.cpp
@@ -119,31 +119,31 @@ uint8 FSizeTableEntry::FillSizeTable(uint64 PlatformPageSize, FSizeTableEntry* S
 {
 	int32 Index = 0;
 
-	for (int32 Sub = 0; Sub < ARRAY_COUNT(BinnedCommonSmallBlockSizes4k); Sub++)
+	for (int32 Sub = 0; Sub < UE_ARRAY_COUNT(BinnedCommonSmallBlockSizes4k); Sub++)
 	{
 		SizeTable[Index++] = FSizeTableEntry(BinnedCommonSmallBlockSizes4k[Sub], PlatformPageSize, 1, BasePageSize, MinimumAlignment);
 	}
-	for (int32 Sub = 0; Sub < ARRAY_COUNT(BinnedCommonSmallBlockSizes8k); Sub++)
+	for (int32 Sub = 0; Sub < UE_ARRAY_COUNT(BinnedCommonSmallBlockSizes8k); Sub++)
 	{
 		SizeTable[Index++] = FSizeTableEntry(BinnedCommonSmallBlockSizes8k[Sub], PlatformPageSize, 2, BasePageSize, MinimumAlignment);
 	}
-	for (int32 Sub = 0; Sub < ARRAY_COUNT(BinnedCommonSmallBlockSizes12k); Sub++)
+	for (int32 Sub = 0; Sub < UE_ARRAY_COUNT(BinnedCommonSmallBlockSizes12k); Sub++)
 	{
 		SizeTable[Index++] = FSizeTableEntry(BinnedCommonSmallBlockSizes12k[Sub], PlatformPageSize, 3, BasePageSize, MinimumAlignment);
 	}
-	for (int32 Sub = 0; Sub < ARRAY_COUNT(BinnedCommonSmallBlockSizes16k); Sub++)
+	for (int32 Sub = 0; Sub < UE_ARRAY_COUNT(BinnedCommonSmallBlockSizes16k); Sub++)
 	{
 		SizeTable[Index++] = FSizeTableEntry(BinnedCommonSmallBlockSizes16k[Sub], PlatformPageSize, 4, BasePageSize, MinimumAlignment);
 	}
-	for (int32 Sub = 0; Sub < ARRAY_COUNT(BinnedCommonSmallBlockSizes20k); Sub++)
+	for (int32 Sub = 0; Sub < UE_ARRAY_COUNT(BinnedCommonSmallBlockSizes20k); Sub++)
 	{
 		SizeTable[Index++] = FSizeTableEntry(BinnedCommonSmallBlockSizes20k[Sub], PlatformPageSize, 5, BasePageSize, MinimumAlignment);
 	}
-	for (int32 Sub = 0; Sub < ARRAY_COUNT(BinnedCommonSmallBlockSizes24k); Sub++)
+	for (int32 Sub = 0; Sub < UE_ARRAY_COUNT(BinnedCommonSmallBlockSizes24k); Sub++)
 	{
 		SizeTable[Index++] = FSizeTableEntry(BinnedCommonSmallBlockSizes24k[Sub], PlatformPageSize, 6, BasePageSize, MinimumAlignment);
 	}
-	for (int32 Sub = 0; Sub < ARRAY_COUNT(BinnedCommonSmallBlockSizes28k); Sub++)
+	for (int32 Sub = 0; Sub < UE_ARRAY_COUNT(BinnedCommonSmallBlockSizes28k); Sub++)
 	{
 		SizeTable[Index++] = FSizeTableEntry(BinnedCommonSmallBlockSizes28k[Sub], PlatformPageSize, 7, BasePageSize, MinimumAlignment);
 	}
diff --git a/Engine/Source/Runtime/Core/Private/HAL/MallocLeakDetection.cpp b/Engine/Source/Runtime/Core/Private/HAL/MallocLeakDetection.cpp
--- a/Engine/Source/Runtime/Core/Private/HAL/MallocLeakDetection.cpp
+++ b/Engine/Source/Runtime/Core/Private/HAL/MallocLeakDetection.cpp
@@ -96,7 +96,7 @@ void FMallocLeakDetection::PushContext(const TCHAR* Context)
 	bRecursive = true;
 
 	FContextString Str;
-	FCString::Strncpy(Str.Buffer, Context, ARRAY_COUNT(Str.Buffer));
+	FCString::Strncpy(Str.Buffer, Context, UE_ARRAY_COUNT(Str.Buffer));
 	TLContexts->Push(Str);
 	bRecursive = false;
 }
diff --git a/Engine/Source/Runtime/Core/Private/HoloLens/HoloLensPlatformAtomics.cpp b/Engine/Source/Runtime/Core/Private/HoloLens/HoloLensPlatformAtomics.cpp
--- a/Engine/Source/Runtime/Core/Private/HoloLens/HoloLensPlatformAtomics.cpp
+++ b/Engine/Source/Runtime/Core/Private/HoloLens/HoloLensPlatformAtomics.cpp
@@ -11,7 +11,7 @@ void FHoloLensAtomics::HandleAtomicsFailure( const TCHAR* InFormat, ... )
 	va_list Ptr;
 
 	va_start( Ptr, InFormat );	
-	FPlatformString::GetVarArgs( TempStr, ARRAY_COUNT(TempStr), InFormat, Ptr );
+	FPlatformString::GetVarArgs( TempStr, UE_ARRAY_COUNT(TempStr), InFormat, Ptr );
 	va_end( Ptr );
 
 	UE_LOG(LogTemp, Log, TempStr);
diff --git a/Engine/Source/Runtime/Core/Private/HoloLens/HoloLensPlatformCrashContext.cpp b/Engine/Source/Runtime/Core/Private/HoloLens/HoloLensPlatformCrashContext.cpp
--- a/Engine/Source/Runtime/Core/Private/HoloLens/HoloLensPlatformCrashContext.cpp
+++ b/Engine/Source/Runtime/Core/Private/HoloLens/HoloLensPlatformCrashContext.cpp
@@ -375,7 +375,7 @@ void CreateExceptionInfoString(EXCEPTION_RECORD* ExceptionRecord)
 	}
 
 #if WITH_EDITORONLY_DATA
-	FCString::Strncpy(GErrorExceptionDescription, *ErrorString, ARRAY_COUNT(GErrorExceptionDescription));
+	FCString::Strncpy(GErrorExceptionDescription, *ErrorString, UE_ARRAY_COUNT(GErrorExceptionDescription));
 #endif
 #undef HANDLE_CASE
 }
@@ -406,11 +406,11 @@ int32 ReportCrash( LPEXCEPTION_POINTERS ExceptionInfo )
 	if( ExceptionInfo->ExceptionRecord->ExceptionCode != 1 )
 	{
 		CreateExceptionInfoString( ExceptionInfo->ExceptionRecord );
-		FCString::Strncat( GErrorHist, GErrorExceptionDescription, ARRAY_COUNT( GErrorHist ) );
-		FCString::Strncat( GErrorHist, TEXT( "\r\n\r\n" ), ARRAY_COUNT( GErrorHist ) );
+		FCString::Strncat( GErrorHist, GErrorExceptionDescription, UE_ARRAY_COUNT( GErrorHist ) );
+		FCString::Strncat( GErrorHist, TEXT( "\r\n\r\n" ), UE_ARRAY_COUNT( GErrorHist ) );
 	}
 
-	FCString::Strncat( GErrorHist, ANSI_TO_TCHAR(StackTrace), ARRAY_COUNT(GErrorHist) );
+	FCString::Strncat( GErrorHist, ANSI_TO_TCHAR(StackTrace), UE_ARRAY_COUNT(GErrorHist) );
 
 	GMalloc->Free( StackTrace );
 
diff --git a/Engine/Source/Runtime/Core/Private/HoloLens/HoloLensProcess.cpp b/Engine/Source/Runtime/Core/Private/HoloLens/HoloLensProcess.cpp
--- a/Engine/Source/Runtime/Core/Private/HoloLens/HoloLensProcess.cpp
+++ b/Engine/Source/Runtime/Core/Private/HoloLens/HoloLensProcess.cpp
@@ -133,13 +133,13 @@ const TCHAR* FHoloLensProcess::ExecutableName(bool bRemoveExtension)
 	if (!Result[0])
 	{
 		// Get complete path for the executable
-		if (GetModuleFileName(NULL, Result, ARRAY_COUNT(Result)) != 0)
+		if (GetModuleFileName(NULL, Result, UE_ARRAY_COUNT(Result)) != 0)
 		{
 			// Remove all of the path information by finding the base filename
 			FString FileName = Result;
 			FString FileNameWithExt = Result;
-			FCString::Strncpy(Result, *(FPaths::GetBaseFilename(FileName)), ARRAY_COUNT(Result));
-			FCString::Strncpy(ResultWithExt, *(FPaths::GetCleanFilename(FileNameWithExt)), ARRAY_COUNT(ResultWithExt));
+			FCString::Strncpy(Result, *(FPaths::GetBaseFilename(FileName)), UE_ARRAY_COUNT(Result));
+			FCString::Strncpy(ResultWithExt, *(FPaths::GetCleanFilename(FileNameWithExt)), UE_ARRAY_COUNT(ResultWithExt));
 		}
 		// If the call failed, zero out the memory to be safe
 		else
diff --git a/Engine/Source/Runtime/Core/Private/HoloLens/HoloLensRunnableThread.cpp b/Engine/Source/Runtime/Core/Private/HoloLens/HoloLensRunnableThread.cpp
--- a/Engine/Source/Runtime/Core/Private/HoloLens/HoloLensRunnableThread.cpp
+++ b/Engine/Source/Runtime/Core/Private/HoloLens/HoloLensRunnableThread.cpp
@@ -33,8 +33,8 @@ uint32 FRunnableThreadHoloLens::GuardedRun()
 			GWarn->Flush();
 
 			// Append the thread name at the end of the error report.
-			FCString::Strncat(GErrorHist, LINE_TERMINATOR TEXT("Crash in runnable thread "), ARRAY_COUNT(GErrorHist));
-			FCString::Strncat(GErrorHist, *ThreadName, ARRAY_COUNT(GErrorHist));
+			FCString::Strncat(GErrorHist, LINE_TERMINATOR TEXT("Crash in runnable thread "), UE_ARRAY_COUNT(GErrorHist));
+			FCString::Strncat(GErrorHist, *ThreadName, UE_ARRAY_COUNT(GErrorHist));
 
 			ExitCode = 1;
 			// Generate status report.				
diff --git a/Engine/Source/Runtime/Core/Private/IOS/IOSPlatformMisc.cpp b/Engine/Source/Runtime/Core/Private/IOS/IOSPlatformMisc.cpp
--- a/Engine/Source/Runtime/Core/Private/IOS/IOSPlatformMisc.cpp
+++ b/Engine/Source/Runtime/Core/Private/IOS/IOSPlatformMisc.cpp
@@ -1413,7 +1413,7 @@ struct FIOSApplicationInfo
         TempSysCtlBufferSize = PATH_MAX+1;
         sysctlbyname("machdep.cpu.brand_string", MachineCPUString, &TempSysCtlBufferSize, NULL, 0);
         
-        gethostname(MachineName, ARRAY_COUNT(MachineName));
+        gethostname(MachineName, UE_ARRAY_COUNT(MachineName));
         
        BranchBaseDir = FString::Printf( TEXT( "%s!%s!%s!%d" ), *FApp::GetBranchName(), FPlatformProcess::BaseDir(), FPlatformMisc::GetEngineMode(), FEngineVersion::Current().GetChangelist() );
         
diff --git a/Engine/Source/Runtime/Core/Private/IOS/IOSPlatformProcess.cpp b/Engine/Source/Runtime/Core/Private/IOS/IOSPlatformProcess.cpp
--- a/Engine/Source/Runtime/Core/Private/IOS/IOSPlatformProcess.cpp
+++ b/Engine/Source/Runtime/Core/Private/IOS/IOSPlatformProcess.cpp
@@ -23,8 +23,8 @@ const TCHAR* FIOSPlatformProcess::ComputerName()
 
 	if( !Result[0] )
 	{
-		ANSICHAR AnsiResult[ARRAY_COUNT(Result)];
-		gethostname(AnsiResult, ARRAY_COUNT(Result));
+		ANSICHAR AnsiResult[UE_ARRAY_COUNT(Result)];
+		gethostname(AnsiResult, UE_ARRAY_COUNT(Result));
 		FCString::Strcpy(Result, ANSI_TO_TCHAR(AnsiResult));
 	}
 	return Result;
