Fix orbit selection to work with selected spline points.

diff --git a/Engine/Source/Editor/UnrealEd/Private/LevelEditorViewport.cpp b/Engine/Source/Editor/UnrealEd/Private/LevelEditorViewport.cpp
--- a/Engine/Source/Editor/UnrealEd/Private/LevelEditorViewport.cpp
+++ b/Engine/Source/Editor/UnrealEd/Private/LevelEditorViewport.cpp
@@ -5111,15 +5111,24 @@ bool FLevelEditorViewportClient::GetPivotForOrbit(FVector& Pivot) const
 		USceneComponent* Component = Cast<USceneComponent>(*It);
 		if (Component && Component->IsRegistered())
 		{
-			// It's possible that it doesn't have a bounding box, so just take its position in that case
-			FBox ComponentBBox = Component->Bounds.GetBox();
-			if (ComponentBBox.GetVolume() != 0)
+			TSharedPtr<FComponentVisualizer> Visualizer = GUnrealEd->FindComponentVisualizer(Component->GetClass());
+			FBox FocusOnSelectionBBox;
+			if (Visualizer && Visualizer->HasFocusOnSelectionBoundingBox(FocusOnSelectionBBox))
 			{
-				BoundingBox += ComponentBBox;
+				BoundingBox += FocusOnSelectionBBox;
 			}
 			else
 			{
-				BoundingBox += Component->GetComponentLocation();
+				// It's possible that it doesn't have a bounding box, so just take its position in that case
+				FBox ComponentBBox = Component->Bounds.GetBox();
+				if (ComponentBBox.GetVolume() != 0)
+				{
+					BoundingBox += ComponentBBox;
+				}
+				else
+				{
+					BoundingBox += Component->GetComponentLocation();
+				}
 			}
 			++NumValidComponents;
 		}
@@ -5145,7 +5154,16 @@ bool FLevelEditorViewportClient::GetPivotForOrbit(FVector& Pivot) const
 
 				if (PrimitiveComponent->IsRegistered() && !PrimitiveComponent->IgnoreBoundsForEditorFocus())
 				{
-					BoundingBox += PrimitiveComponent->Bounds.GetBox();
+					TSharedPtr<FComponentVisualizer> Visualizer = GUnrealEd->FindComponentVisualizer(PrimitiveComponent->GetClass());
+					FBox FocusOnSelectionBBox;
+					if (Visualizer && Visualizer->HasFocusOnSelectionBoundingBox(FocusOnSelectionBBox))
+					{
+						BoundingBox += FocusOnSelectionBBox;
+					}
+					else
+					{
+						BoundingBox += PrimitiveComponent->Bounds.GetBox();
+					}
 					++NumSelectedActors;
 				}
 			}
