Fix missing include
#rnx

diff --git a/Engine/Source/Editor/UnrealEd/Private/Commandlets/CompileAllBlueprintsCommandlet.cpp b/Engine/Source/Editor/UnrealEd/Private/Commandlets/CompileAllBlueprintsCommandlet.cpp
--- a/Engine/Source/Editor/UnrealEd/Private/Commandlets/CompileAllBlueprintsCommandlet.cpp
+++ b/Engine/Source/Editor/UnrealEd/Private/Commandlets/CompileAllBlueprintsCommandlet.cpp
@@ -9,6 +9,7 @@
 #include "Misc/FileHelper.h"
 #include "Misc/Paths.h"
 #include "Engine/Engine.h"
+#include "Kismet2/KismetEditorUtilities.h"
 
 DEFINE_LOG_CATEGORY_STATIC(LogCompileAllBlueprintsCommandlet, Log, All);
 
