Fix regression in logic to convert a function that was implementing an interface declared function into a standalone function

diff --git a/Engine/Source/Editor/UnrealEd/Private/Kismet2/BlueprintEditorUtils.cpp b/Engine/Source/Editor/UnrealEd/Private/Kismet2/BlueprintEditorUtils.cpp
--- a/Engine/Source/Editor/UnrealEd/Private/Kismet2/BlueprintEditorUtils.cpp
+++ b/Engine/Source/Editor/UnrealEd/Private/Kismet2/BlueprintEditorUtils.cpp
@@ -6162,47 +6162,20 @@ bool FBlueprintEditorUtils::RemoveInterfaceFunction(UBlueprint* Blueprint, FBPIn
 		{
 			CurrentGraph->Modify();
 			Blueprint->Modify();
-			UEdGraph* NewGraph = nullptr;
-			FString OriginalName;
-			CurrentGraph->GetName(OriginalName);
 
-			// Create a new one with it's signature if the user wants to preserve it
 			if (bPreserveFunction)
 			{
-				// Create a new function graph with this signature
-				static const FString TempName = "TEMP_INTERFACE_GRAPH_NAME";
-				const FName TempFuncGraphName = FBlueprintEditorUtils::GenerateUniqueGraphName(Blueprint, TempName);
-				UEdGraph* EventGraph = FBlueprintEditorUtils::FindEventGraph(Blueprint);
-				const UEdGraphSchema_K2* Schema = GetDefault<UEdGraphSchema_K2>();
-
-				if (EventGraph)
-				{
-					EventGraph->Modify();
-
-					NewGraph = FBlueprintEditorUtils::CreateNewGraph(
-						Blueprint,
-						TempFuncGraphName,
-						EventGraph->GetClass(),
-						EventGraph->GetSchema() ? EventGraph->GetSchema()->GetClass() : Schema->GetClass()
-					);
-
-					// Add the function graph like this to avoid the MarkBlueprintAsStructurallyModified call
-					CreateFunctionGraph(Blueprint, NewGraph, /*bIsUserCreated=*/true, Function);
-					Blueprint->FunctionGraphs.Add(NewGraph);
-					// Potentially adjust variable names for any child blueprints
-					ValidateBlueprintChildVariables(Blueprint, NewGraph->GetFName());
-				}
+				// Promote the graph if the user wants to preserve it:
+				PromoteGraphFromInterfaceOverride(Blueprint, CurrentGraph);
+				Blueprint->FunctionGraphs.Add(CurrentGraph);
 			}
 
 			FBlueprintEditorUtils::UpdateTransactionalFlags(Blueprint);
 
-			// Remove the interface graph
-			FBlueprintEditorUtils::RemoveGraph(Blueprint, CurrentGraph, EGraphRemoveFlags::MarkTransient);	// Do not recompile, yet
-			
-			// Rename the new graph to the original interface name
-			if (NewGraph)
+			if(!bPreserveFunction)
 			{
-				FBlueprintEditorUtils::RenameGraph(NewGraph, OriginalName);
+				// Remove the interface graph
+				FBlueprintEditorUtils::RemoveGraph(Blueprint, CurrentGraph, EGraphRemoveFlags::MarkTransient);	// Do not recompile, yet
 			}
 
 			return true;
