Fix a regression that causes local variables to be invalidated after renaming the entry node in a function graph.

#rnx

diff --git a/Engine/Source/Editor/Kismet/Private/SBlueprintPalette.cpp b/Engine/Source/Editor/Kismet/Private/SBlueprintPalette.cpp
--- a/Engine/Source/Editor/Kismet/Private/SBlueprintPalette.cpp
+++ b/Engine/Source/Editor/Kismet/Private/SBlueprintPalette.cpp
@@ -1320,25 +1320,6 @@ void SBlueprintPaletteItem::OnNameTextCommitted(const FText& NewText, ETextCommi
 			if (ExistingGraph == nullptr || ExistingGraph == Graph)
 			{
 				const FScopedTransaction Transaction( LOCTEXT( "Rename Function", "Rename Function" ) );
-
-				// Search through all function entry nodes for local variables to update their scope name
-				TArray<UK2Node_Variable*> VariableNodes;
-				Graph->GetNodesOfClass<UK2Node_Variable>(VariableNodes);
-				for (const UEdGraph* SubGraph : Graph->SubGraphs)
-				{
-					check(SubGraph != nullptr);
-					SubGraph->GetNodesOfClass<UK2Node_Variable>(VariableNodes);
-				}
-
-				for (UK2Node_Variable* const VariableNode : VariableNodes)
-				{
-					if(VariableNode->VariableReference.IsLocalScope())
-					{
-						// Update the variable's scope to be the graph's name (which mirrors the UFunction)
-						VariableNode->VariableReference.SetLocalMember(VariableNode->VariableReference.GetMemberName(), NewNameString, VariableNode->VariableReference.GetMemberGuid());
-					}
-				}
-
 				FBlueprintEditorUtils::RenameGraph(Graph, NewNameString );
 			}
 		}
diff --git a/Engine/Source/Editor/UnrealEd/Private/Kismet2/BlueprintEditorUtils.cpp b/Engine/Source/Editor/UnrealEd/Private/Kismet2/BlueprintEditorUtils.cpp
--- a/Engine/Source/Editor/UnrealEd/Private/Kismet2/BlueprintEditorUtils.cpp
+++ b/Engine/Source/Editor/UnrealEd/Private/Kismet2/BlueprintEditorUtils.cpp
@@ -2767,6 +2767,24 @@ void FBlueprintEditorUtils::RenameGraph(UEdGraph* Graph, const FString& NewNameS
 		// Note: This will find ALL children (including nested children) so there's no need to do this recursively.
 		ValidateBlueprintChildVariables(Blueprint, Graph->GetFName(), PostValidateChildBlueprintLambda);
 
+		// Find all variable nodes in this graph.
+		TArray<UK2Node_Variable*> VariableNodes;
+		Graph->GetNodesOfClass<UK2Node_Variable>(VariableNodes);
+		for (const UEdGraph* SubGraph : Graph->SubGraphs)
+		{
+			check(SubGraph != nullptr);
+			SubGraph->GetNodesOfClass<UK2Node_Variable>(VariableNodes);
+		}
+
+		// For any nodes that reference a local variable, update the variable's scope to be the graph's new name (which will mirror the UFunction).
+		for (UK2Node_Variable* const VariableNode : VariableNodes)
+		{
+			if (VariableNode->VariableReference.IsLocalScope())
+			{
+				VariableNode->VariableReference.SetLocalMember(VariableNode->VariableReference.GetMemberName(), NewNameStr, VariableNode->VariableReference.GetMemberGuid());
+			}
+		}
+
 		// Rename any function call points (including any calls w/ a child Blueprint as the target).
 		for (TObjectIterator<UK2Node_CallFunction> It(RF_ClassDefaultObject | RF_Transient); It; ++It)
 		{
