Remove GBlueprintUseCompilationManager related switches - no reason to support this as its been deprecated for several engine versions now and the paths that were using it (nativization, compileallblueprints) were trivial to port or unstable

diff --git a/Engine/Config/BaseEditor.ini b/Engine/Config/BaseEditor.ini
--- a/Engine/Config/BaseEditor.ini
+++ b/Engine/Config/BaseEditor.ini
@@ -504,7 +504,6 @@ bNativizeAnimBPOnlyWhenNonReducibleFuncitons=false
 +NoExportTypesWithDirectNativeFieldAccess=/Script/Engine.FormatArgumentData
 
 [/Script/UnrealEd.BlueprintEditorProjectSettings]
-bUseCompilationManager=true
 bValidateUnloadedSoftActorReferences=true
 
 [/Script/UnrealEd.AssetViewerSettings]
diff --git a/Engine/Source/Editor/KismetCompiler/Private/KismetCompiler.cpp b/Engine/Source/Editor/KismetCompiler/Private/KismetCompiler.cpp
--- a/Engine/Source/Editor/KismetCompiler/Private/KismetCompiler.cpp
+++ b/Engine/Source/Editor/KismetCompiler/Private/KismetCompiler.cpp
@@ -64,8 +64,6 @@ FSimpleMulticastDelegate FKismetCompilerContext::OnPostCompile;
 
 #define LOCTEXT_NAMESPACE "KismetCompiler"
 
-extern COREUOBJECT_API bool GBlueprintUseCompilationManager;
-
 //////////////////////////////////////////////////////////////////////////
 // Stats for this module
 DECLARE_CYCLE_STAT(TEXT("Create Schema"), EKismetCompilerStats_CreateSchema, STATGROUP_KismetCompiler );
@@ -3974,17 +3972,6 @@ void FKismetCompilerContext::CompileClassLayout(EInternalCompilerFlags InternalF
 		}
 	}
 
-	if (CompileOptions.DoesRequireBytecodeGeneration() && !GBlueprintUseCompilationManager)
-	{
-		TArray<UEdGraph*> AllGraphs;
-		Blueprint->GetAllGraphs(AllGraphs);
-		for (int32 i = 0; i < AllGraphs.Num(); i++)
-		{
-			//Reset error flags associated with nodes in each graph
-			ResetErrorFlags(AllGraphs[i]);
-		}
-	}
-
 	// Early validation
 	if (CompileOptions.CompileType == EKismetCompileType::Full)
 	{
diff --git a/Engine/Source/Editor/KismetCompiler/Private/KismetCompilerModule.cpp b/Engine/Source/Editor/KismetCompiler/Private/KismetCompilerModule.cpp
--- a/Engine/Source/Editor/KismetCompiler/Private/KismetCompilerModule.cpp
+++ b/Engine/Source/Editor/KismetCompiler/Private/KismetCompilerModule.cpp
@@ -44,7 +44,6 @@ class FKismet2CompilerModule : public IKismetCompilerInterface
 {
 public:
 	// Implementation of the IKismetCompilerInterface
-	virtual void CompileBlueprint(class UBlueprint* Blueprint, const FKismetCompilerOptions& CompileOptions, FCompilerResultsLog& Results, TSharedPtr<class FBlueprintCompileReinstancer> ParentReinstancer = NULL) override;
 	virtual void RefreshVariables(UBlueprint* Blueprint) override final;
 	virtual void CompileStructure(class UUserDefinedStruct* Struct, FCompilerResultsLog& Results) override;
 	virtual void RecoverCorruptedBlueprint(class UBlueprint* Blueprint) override;
@@ -56,7 +55,6 @@ public:
 	virtual FString GenerateCppWrapper(UBlueprintGeneratedClass* BPGC, const FCompilerNativizationOptions& NativizationOptions) override;
 	// End implementation
 private:
-	void CompileBlueprintInner(class UBlueprint* Blueprint, const FKismetCompilerOptions& CompileOptions, FCompilerResultsLog& Results, TSharedPtr<FBlueprintCompileReinstancer> Reinstancer);
 
 	TArray<IBlueprintCompiler*> Compilers;
 };
@@ -82,83 +80,6 @@ public:
 };
 
 // Compiles a blueprint.
-void FKismet2CompilerModule::CompileBlueprintInner(class UBlueprint* Blueprint, const FKismetCompilerOptions& CompileOptions, FCompilerResultsLog& Results, TSharedPtr<FBlueprintCompileReinstancer> Reinstancer)
-{
-	FBlueprintIsBeingCompiledHelper BeingCompiled(Blueprint);
-
-	Blueprint->CurrentMessageLog = &Results;
-
-	// Early out if blueprint parent is missing
-	if (Blueprint->ParentClass == NULL)
-	{
-		Results.Error(*LOCTEXT("KismetCompileError_MalformedParentClasss", "Blueprint @@ has missing or NULL parent class.").ToString(), Blueprint);
-	}
-	else
-	{
-		const uint32 PreviousSignatureCrc = Blueprint->CrcLastCompiledSignature;
-		const bool bIsFullCompile = CompileOptions.DoesRequireBytecodeGeneration() && (Blueprint->BlueprintType != BPTYPE_Interface);
-		const bool bRecompileDependencies = bIsFullCompile && !Blueprint->bIsRegeneratingOnLoad && Reinstancer.IsValid();
-
-		TArray<UBlueprint*> StoredDependentBlueprints;
-		if (bRecompileDependencies)
-		{
-			FBlueprintEditorUtils::GetDependentBlueprints(Blueprint, StoredDependentBlueprints);
-		}
-
-		// Loop through all external compiler delegates attempting to compile the blueprint.
-		bool Compiled = false;
-		for ( IBlueprintCompiler* Compiler : Compilers )
-		{
-			if ( Compiler->CanCompile(Blueprint) )
-			{
-				Compiled = true;
-				Compiler->Compile(Blueprint, CompileOptions, Results);
-				break;
-			}
-		}
-
-		// if no one handles it, then use the default blueprint compiler.
-		if ( !Compiled )
-		{
-			if ( UAnimBlueprint* AnimBlueprint = Cast<UAnimBlueprint>(Blueprint) )
-			{
-				FAnimBlueprintCompilerContext Compiler(AnimBlueprint, Results, CompileOptions);
-				Compiler.Compile();
-				check(Compiler.NewClass);
-			}
-			else
-			{
-				FKismetCompilerContext Compiler(Blueprint, Results, CompileOptions);
-				Compiler.Compile();
-				check(Compiler.NewClass);
-			}
-		}
-
-		if (bRecompileDependencies)
-		{
-			Reinstancer->BlueprintWasRecompiled(Blueprint, CompileOptions.CompileType == EKismetCompileType::BytecodeOnly);
-
-			const bool bSignatureWasChanged = PreviousSignatureCrc != Blueprint->CrcLastCompiledSignature;
-			UE_LOG(LogK2Compiler, Verbose, TEXT("Signature of Blueprint '%s' %s changed"), *GetNameSafe(Blueprint), bSignatureWasChanged ? TEXT("was") : TEXT("was not"));
-
-			if (bSignatureWasChanged)
-			{
-				for (UBlueprint* CurrentBP : StoredDependentBlueprints)
-				{
-					Reinstancer->EnlistDependentBlueprintToRecompile(CurrentBP, !(CurrentBP->IsPossiblyDirty() || CurrentBP->Status == BS_Error) && CurrentBP->IsValidForBytecodeOnlyRecompile());
-				}
-
-				if(!Blueprint->ParentClass->HasAnyClassFlags(CLASS_Native))
-				{
-					Reinstancer->EnlistDependentBlueprintToRecompile(Blueprint, true);
-				}
-			}
-		}
-	}
-
-	Blueprint->CurrentMessageLog = NULL;
-}
-
 
 void FKismet2CompilerModule::CompileStructure(UUserDefinedStruct* Struct, FCompilerResultsLog& Results)
 {
@@ -187,126 +108,6 @@ FString FKismet2CompilerModule::GenerateCppWrapper(UBlueprintGeneratedClass* BPG
 
 extern UNREALED_API FSecondsCounterData BlueprintCompileAndLoadTimerData;
 
-// Compiles a blueprint.
-void FKismet2CompilerModule::CompileBlueprint(class UBlueprint* Blueprint, const FKismetCompilerOptions& CompileOptions, FCompilerResultsLog& Results, TSharedPtr<FBlueprintCompileReinstancer> ParentReinstancer)
-{
-	FSecondsCounterScope Timer(BlueprintCompileAndLoadTimerData);
-	BP_SCOPED_COMPILER_EVENT_STAT(EKismetCompilerStats_CompileTime);
-
-	Results.SetSourcePath(Blueprint->GetPathName());
-
-	const bool bIsBrandNewBP = (Blueprint->SkeletonGeneratedClass == NULL) && (Blueprint->GeneratedClass == NULL) && (Blueprint->ParentClass != NULL) && !CompileOptions.bIsDuplicationInstigated;
-
-	for ( IBlueprintCompiler* Compiler : Compilers )
-	{
-		Compiler->PreCompile(Blueprint, CompileOptions);
-	}
-
-	if (CompileOptions.CompileType != EKismetCompileType::Cpp
-		&& CompileOptions.CompileType != EKismetCompileType::BytecodeOnly
-		&& CompileOptions.bRegenerateSkelton )
-	{
-		BP_SCOPED_COMPILER_EVENT_STAT(EKismetCompilerStats_CompileSkeletonClass);
-		auto SkeletonReinstancer = FBlueprintCompileReinstancer::Create(Blueprint->SkeletonGeneratedClass);
-
-		FCompilerResultsLog SkeletonResults;
-		SkeletonResults.bSilentMode = true;
-		FKismetCompilerOptions SkeletonCompileOptions;
-		SkeletonCompileOptions.CompileType = EKismetCompileType::SkeletonOnly;
-		CompileBlueprintInner(Blueprint, SkeletonCompileOptions, SkeletonResults, ParentReinstancer);
-
-		// Only when doing full compiles do we want to compile all skeletons before continuing to compile 
-		if (CompileOptions.CompileType == EKismetCompileType::Full)
-		{
-			SkeletonReinstancer->ReinstanceObjects();
-		}
-	}
-
-	// If this was a full compile, take appropriate actions depending on the success of failure of the compile
-	if( CompileOptions.IsGeneratedClassCompileType() )
-	{
-		BP_SCOPED_COMPILER_EVENT_STAT(EKismetCompilerStats_CompileGeneratedClass);
-
-		FBlueprintCompileReinstancer::OptionallyRefreshNodes(Blueprint);
-
-		// Perform the full compile
-		CompileBlueprintInner(Blueprint, CompileOptions, Results, ParentReinstancer);
-
-		if (Results.NumErrors == 0)
-		{
-			// Blueprint is error free.  Go ahead and fix up debug info
-			Blueprint->Status = (0 == Results.NumWarnings) ? BS_UpToDate : BS_UpToDateWithWarnings;
-			
-			Blueprint->BlueprintSystemVersion = UBlueprint::GetCurrentBlueprintSystemVersion();
-
-			// Reapply breakpoints to the bytecode of the new class
-			for (int32 Index = 0; Index < Blueprint->Breakpoints.Num(); ++Index)
-			{
-				UBreakpoint* Breakpoint = Blueprint->Breakpoints[Index];
-				FKismetDebugUtilities::ReapplyBreakpoint(Breakpoint);
-			}
-		}
-		else
-		{
-			// Should never get errors from a brand new blueprint!
-			ensure(!bIsBrandNewBP || (Results.NumErrors == 0));
-
-			// There were errors.  Compile the generated class to have function stubs
-			Blueprint->Status = BS_Error;
-
-			if(CompileOptions.bReinstanceAndStubOnFailure)
-			{
-				// Reinstance objects here, so we can preserve their memory layouts to reinstance them again
-				if (ParentReinstancer.IsValid())
-				{
-					ParentReinstancer->UpdateBytecodeReferences();
-
-					if(!Blueprint->bIsRegeneratingOnLoad)
-					{
-						ParentReinstancer->ReinstanceObjects();
-					}
-				}
-				const EBlueprintCompileReinstancerFlags Flags = (EKismetCompileType::BytecodeOnly == CompileOptions.CompileType) ? EBlueprintCompileReinstancerFlags::BytecodeOnly : EBlueprintCompileReinstancerFlags::None;
-				TSharedPtr<FBlueprintCompileReinstancer> StubReinstancer = FBlueprintCompileReinstancer::Create(Blueprint->GeneratedClass, Flags);
-
-				// Toss the half-baked class and generate a stubbed out skeleton class that can be used
-				FCompilerResultsLog StubResults;
-				StubResults.bSilentMode = true;
-				FKismetCompilerOptions StubCompileOptions(CompileOptions);
-				StubCompileOptions.CompileType = EKismetCompileType::StubAfterFailure;
-				{
-					CompileBlueprintInner(Blueprint, StubCompileOptions, StubResults, StubReinstancer);
-				}
-
-				StubReinstancer->UpdateBytecodeReferences();
-				if( !Blueprint->bIsRegeneratingOnLoad )
-				{
-					StubReinstancer->ReinstanceObjects();
-				}
-			}
-		}
-	}
-
-	for ( IBlueprintCompiler* Compiler : Compilers )
-	{
-		Compiler->PostCompile(Blueprint, CompileOptions);
-	}
-
-	UPackage* Package = Blueprint->GetOutermost();
-	if( Package )
-	{
-		UMetaData* MetaData =  Package->GetMetaData();
-		MetaData->RemoveMetaDataOutsidePackage();
-	}
-
-	if (!Results.bSilentMode)
-	{
-		FScopedBlueprintMessageLog MessageLog(Blueprint);
-		MessageLog.Log->ClearMessages();
-		MessageLog.Log->AddMessages(Results.Messages, false);
-	}
-}
-
 void FKismet2CompilerModule::RefreshVariables(UBlueprint* Blueprint)
 {
 #if 0 // disabled because this is a costly ensure:
diff --git a/Engine/Source/Editor/KismetCompiler/Public/KismetCompilerModule.h b/Engine/Source/Editor/KismetCompiler/Public/KismetCompilerModule.h
--- a/Engine/Source/Editor/KismetCompiler/Public/KismetCompilerModule.h
+++ b/Engine/Source/Editor/KismetCompiler/Public/KismetCompilerModule.h
@@ -50,15 +50,6 @@ protected:
 class IKismetCompilerInterface : public IModuleInterface
 {
 public:
-	/**
-	 * Compiles a blueprint.
-	 *
-	 * @param	Blueprint	The blueprint to compile.
-	 * @param	Results  	The results log for warnings and errors.
-	 * @param	Options		Compiler options.
-	 */
-	virtual void CompileBlueprint(class UBlueprint* Blueprint, const FKismetCompilerOptions& CompileOptions, FCompilerResultsLog& Results, TSharedPtr<class FBlueprintCompileReinstancer> ParentReinstancer = NULL)=0;
-
 	/**
 	 * Synchronizes Blueprint's GeneratedClass's properties with the NewVariable declarations in the blueprint
 	 * Used on load to ensure that all properties are present for instances.
diff --git a/Engine/Source/Editor/UnrealEd/Private/Kismet2/Kismet2.cpp b/Engine/Source/Editor/UnrealEd/Private/Kismet2/Kismet2.cpp
--- a/Engine/Source/Editor/UnrealEd/Private/Kismet2/Kismet2.cpp
+++ b/Engine/Source/Editor/UnrealEd/Private/Kismet2/Kismet2.cpp
@@ -94,8 +94,6 @@ DECLARE_CYCLE_STAT(TEXT("Validate Generated Class"), EKismetCompilerStats_Valida
 
 #define LOCTEXT_NAMESPACE "UnrealEd.Editor"
 
-extern COREUOBJECT_API bool GBlueprintUseCompilationManager;
-
 //////////////////////////////////////////////////////////////////////////
 // FArchiveInvalidateTransientRefs
 
@@ -512,29 +510,18 @@ UBlueprint* FKismetEditorUtilities::CreateBlueprint(UClass* ParentClass, UObject
 	// Create initial UClass
 	IKismetCompilerInterface& Compiler = FModuleManager::LoadModuleChecked<IKismetCompilerInterface>(KISMET_COMPILER_MODULENAME);
 
-	if (GBlueprintUseCompilationManager)
-	{
-		// Skip validation of the class default object here, because (a) the new CDO may fail validation since this
-		// is a new Blueprint that the user has not had a chance to modify any defaults for yet, and (b) in some cases,
-		// default value propagation to the new Blueprint's CDO may be deferred until after compilation (e.g. reparenting).
-		// Also skip the Blueprint search data update, as that will be handled by an OnAssetAdded() delegate in the FiB manager.
-		const EBlueprintCompileOptions CompileOptions =
-			EBlueprintCompileOptions::SkipGarbageCollection |
-			EBlueprintCompileOptions::SkipDefaultObjectValidation |
-			EBlueprintCompileOptions::SkipFiBSearchMetaUpdate;
+	// Skip validation of the class default object here, because (a) the new CDO may fail validation since this
+	// is a new Blueprint that the user has not had a chance to modify any defaults for yet, and (b) in some cases,
+	// default value propagation to the new Blueprint's CDO may be deferred until after compilation (e.g. reparenting).
+	// Also skip the Blueprint search data update, as that will be handled by an OnAssetAdded() delegate in the FiB manager.
+	const EBlueprintCompileOptions CompileOptions =
+		EBlueprintCompileOptions::SkipGarbageCollection |
+		EBlueprintCompileOptions::SkipDefaultObjectValidation |
+		EBlueprintCompileOptions::SkipFiBSearchMetaUpdate;
 
-		FBlueprintCompilationManager::CompileSynchronously(
-			FBPCompileRequest(NewBP, CompileOptions, nullptr)
-		);
-	}
-	else
-	{
-		FCompilerResultsLog Results;
-		const bool bReplaceExistingInstances = false;
-		NewBP->Status = BS_Dirty;
-		FKismetCompilerOptions CompileOptions;
-		Compiler.CompileBlueprint(NewBP, CompileOptions, Results);
-	}
+	FBlueprintCompilationManager::CompileSynchronously(
+		FBPCompileRequest(NewBP, CompileOptions, nullptr)
+	);
 
 	// Mark the BP as being regenerated, so it will not be confused as needing to be loaded and regenerated when a referenced BP loads.
 	NewBP->bHasBeenRegenerated = true;
@@ -760,219 +747,7 @@ void FKismetEditorUtilities::CompileBlueprint(UBlueprint* BlueprintObj, EBluepri
 {
 	DECLARE_SCOPE_HIERARCHICAL_COUNTER_FUNC()
 
-	if(GBlueprintUseCompilationManager)
-	{
-		FBlueprintCompilationManager::CompileSynchronously(FBPCompileRequest(BlueprintObj, CompileFlags, pResults));
-		return;
-	}
-
-	const bool bIsRegeneratingOnLoad		= (CompileFlags & EBlueprintCompileOptions::IsRegeneratingOnLoad		) != EBlueprintCompileOptions::None;
-	const bool bSkipGarbageCollection		= (CompileFlags & EBlueprintCompileOptions::SkipGarbageCollection		) != EBlueprintCompileOptions::None;
-	const bool bSaveIntermediateProducts	= (CompileFlags & EBlueprintCompileOptions::SaveIntermediateProducts	) != EBlueprintCompileOptions::None;
-	const bool bSkeletonUpToDate			= (CompileFlags & EBlueprintCompileOptions::SkeletonUpToDate			) != EBlueprintCompileOptions::None;
-	const bool bBatchCompile				= (CompileFlags & EBlueprintCompileOptions::BatchCompile				) != EBlueprintCompileOptions::None;
-	const bool bSkipReinstancing			= (CompileFlags & EBlueprintCompileOptions::SkipReinstancing			) != EBlueprintCompileOptions::None;
-
-	FSecondsCounterScope Timer(BlueprintCompileAndLoadTimerData); 
-	BP_SCOPED_COMPILER_EVENT_STAT(EKismetCompilerStats_CompileBlueprint);
-
-	// Wipe the PreCompile log, any generated messages are now irrelevant
-	BlueprintObj->PreCompileLog.Reset();
-
-	// Broadcast pre-compile
-#if WITH_EDITOR
-	{
-		BP_SCOPED_COMPILER_EVENT_STAT(EKismetCompilerStats_BroadcastPrecompile);
-		if(GEditor && GIsEditor)
-		{
-			GEditor->BroadcastBlueprintPreCompile(BlueprintObj);
-		}
-	}
-#endif
-
-	// Reset the flag, so if the user tries to use PIE it will warn them if the BP did not compile
-	BlueprintObj->bDisplayCompilePIEWarning = true;
-
-	UPackage* const BlueprintPackage = BlueprintObj->GetOutermost();
-	// compiling the blueprint will inherently dirty the package, but if there 
-	// weren't any changes to save before, there shouldn't be after
-	bool const bStartedWithUnsavedChanges = (BlueprintPackage != NULL) ? BlueprintPackage->IsDirty() : true;
-#if WITH_EDITOR
-	// Do not want to run this code without the editor present nor when running commandlets.
-	if (GEditor && GIsEditor)
-	{
-		BP_SCOPED_COMPILER_EVENT_STAT(EKismetCompilerStats_UpdateSearchMetaData);
-		// We do not want to regenerate a search Guid during loads, nothing has changed in the Blueprint and it is cached elsewhere
-		if (!bIsRegeneratingOnLoad)
-		{
-			FFindInBlueprintSearchManager::Get().AddOrUpdateBlueprintSearchMetadata(BlueprintObj);
-		}
-	}
-#endif
-
-	// The old class is either the GeneratedClass if we had an old successful compile, or the SkeletonGeneratedClass stub if there were previously fatal errors
-	UClass* OldClass = (BlueprintObj->GeneratedClass != NULL && (BlueprintObj->GeneratedClass != BlueprintObj->SkeletonGeneratedClass)) ? BlueprintObj->GeneratedClass : NULL;
-
-	// Load the compiler
-	IKismetCompilerInterface& Compiler = FModuleManager::LoadModuleChecked<IKismetCompilerInterface>(KISMET_COMPILER_MODULENAME);
-
-	// Prepare old objects for reinstancing
-	TGuardValue<bool> GuardTemplateNameFlag(GCompilingBlueprint, true);
-
-	// Compile
-	FCompilerResultsLog LocalResults;
-	FCompilerResultsLog& Results = (pResults != NULL) ? *pResults : LocalResults;
-
-	// Monitoring UE-20486, the OldClass->ClassGeneratedBy is NULL or otherwise not a UBlueprint.
-	if (OldClass && (OldClass->ClassGeneratedBy != BlueprintObj))
-	{
-		ensureMsgf(OldClass->ClassGeneratedBy == BlueprintObj, TEXT("Generated Class '%s' has an invalid ClassGeneratedBy '%s' while the expected is Blueprint '%s'"), *OldClass->GetPathName(), *GetPathNameSafe(OldClass->ClassGeneratedBy), *BlueprintObj->GetPathName());
-		OldClass->ClassGeneratedBy = BlueprintObj;
-	}
-	TSharedPtr<class FBlueprintCompileReinstancer> ReinstanceHelper;
-	if(!bSkipReinstancing)
-	{
-		ReinstanceHelper = FBlueprintCompileReinstancer::Create(OldClass);
-	}
-
-	// If enabled, suppress errors/warnings in the log if we're recompiling on load on a build machine
-	static const FBoolConfigValueHelper IgnoreCompileOnLoadErrorsOnBuildMachine(TEXT("Kismet"), TEXT("bIgnoreCompileOnLoadErrorsOnBuildMachine"), GEngineIni);
-	Results.bLogInfoOnly = BlueprintObj->bIsRegeneratingOnLoad && GIsBuildMachine && IgnoreCompileOnLoadErrorsOnBuildMachine;
-
-	FKismetCompilerOptions CompileOptions;
-	CompileOptions.bSaveIntermediateProducts = bSaveIntermediateProducts;
-	CompileOptions.bRegenerateSkelton = !bSkeletonUpToDate;
-	CompileOptions.bReinstanceAndStubOnFailure = !bSkipReinstancing;
-	Compiler.CompileBlueprint(BlueprintObj, CompileOptions, Results, ReinstanceHelper);
-
-	FBlueprintEditorUtils::UpdateDelegatesInBlueprint(BlueprintObj);
-
-	if (FBlueprintEditorUtils::IsLevelScriptBlueprint(BlueprintObj))
-	{
-		// When the Blueprint is recompiled, then update the bound events for level scripting
-		ULevelScriptBlueprint* LevelScriptBP = CastChecked<ULevelScriptBlueprint>(BlueprintObj);
-
-		if (ULevel* BPLevel = LevelScriptBP->GetLevel())
-		{
-			BPLevel->OnLevelScriptBlueprintChanged(LevelScriptBP);
-		}
-	}
-
-	if(!bSkipReinstancing)
-	{
-		ReinstanceHelper->UpdateBytecodeReferences();
-	}
-	
-	// in case any errors/warnings have been added since the call to Compiler.CompileBlueprint()
-	if (Results.NumErrors > 0)
-	{
-		BlueprintObj->Status = BS_Error;
-	}
-	else if (Results.NumWarnings > 0)
-	{
-		BlueprintObj->Status = BS_UpToDateWithWarnings;
-	}
-
-	const bool bIsInterface = FBlueprintEditorUtils::IsInterfaceBlueprint(BlueprintObj);
-	const bool bLetReinstancerRefreshDependBP = !bIsRegeneratingOnLoad && (OldClass != NULL) && !bIsInterface;
-	if (bLetReinstancerRefreshDependBP)
-	{
-		BP_SCOPED_COMPILER_EVENT_STAT(EKismetCompilerStats_RefreshDependentBlueprints);
-
-		TArray<UBlueprint*> DependentBPs;
-		FBlueprintEditorUtils::GetDependentBlueprints(BlueprintObj, DependentBPs);
-		if(!bSkipReinstancing)
-		{
-			ReinstanceHelper->ListDependentBlueprintsToRefresh(DependentBPs);
-		}
-	}
-
-	if ( (!bIsRegeneratingOnLoad) && (OldClass != NULL))
-	{
-		// Strip off any external components from the CDO, if needed because of reparenting, etc
-		FKismetEditorUtilities::StripExternalComponents(BlueprintObj);
-
-		// Ensure that external SCS node references match up with the generated class
-		if(BlueprintObj->SimpleConstructionScript)
-		{
-			BlueprintObj->SimpleConstructionScript->FixupRootNodeParentReferences();
-		}
-
-		// Replace instances of this class
-		if(!bSkipReinstancing)
-		{
-			ReinstanceHelper->ReinstanceObjects();
-		}
-		
-		// Notify everyone a blueprint has been compiled and reinstanced, but before GC so they can perform any final cleanup.
-		if ( GEditor )
-		{
-			GEditor->BroadcastBlueprintReinstanced();
-		}
-
-		if (!bSkipGarbageCollection)
-		{
-			BP_SCOPED_COMPILER_EVENT_STAT(EKismetCompilerStats_GarbageCollection);
-
-			// Garbage collect to make sure the old class and actors are disposed of
-			CollectGarbage(GARBAGE_COLLECTION_KEEPFLAGS);
-		}
-
-		// If you need to verify that all old instances are taken care of, uncomment this!
-		// ReinstanceHelper.VerifyReplacement();
-	}
-
-	if (!bBatchCompile)
-	{
-		BP_SCOPED_COMPILER_EVENT_STAT(EKismetCompilerStats_NotifyBlueprintChanged);
-
-		BlueprintObj->BroadcastCompiled();
-
-		if(GEditor)
-		{
-			GEditor->BroadcastBlueprintCompiled();	
-		}
-	}
-
-	if (!bLetReinstancerRefreshDependBP && (bIsInterface || !BlueprintObj->bIsRegeneratingOnLoad) && !bSkipReinstancing)
-	{
-		BP_SCOPED_COMPILER_EVENT_STAT(EKismetCompilerStats_RefreshDependentBlueprints);
-
-		TArray<UBlueprint*> DependentBPs;
-		FBlueprintEditorUtils::GetDependentBlueprints(BlueprintObj, DependentBPs);
-
-		// refresh each dependent blueprint
-		for (UBlueprint* Dependent : DependentBPs)
-		{
-			// for interface changes, auto-refresh nodes on any dependent blueprints
-			// note: RefreshAllNodes() will internally send a change notification event to the dependent blueprint
-			if (bIsInterface)
-			{
-				bool bPreviousRegenValue = Dependent->bIsRegeneratingOnLoad;
-				Dependent->bIsRegeneratingOnLoad = Dependent->bIsRegeneratingOnLoad || BlueprintObj->bIsRegeneratingOnLoad;
-				FBlueprintEditorUtils::RefreshAllNodes(Dependent);
-				Dependent->bIsRegeneratingOnLoad = bPreviousRegenValue;
-			}
-			else if(!BlueprintObj->bIsRegeneratingOnLoad)
-			{
-				// for non-interface changes, nodes with an external dependency have already been refreshed, and it is now safe to send a change notification event
-				Dependent->BroadcastChanged();
-			}
-		}
-	}
-
-	if(!bIsRegeneratingOnLoad && BlueprintObj->GeneratedClass)
-	{
-		BP_SCOPED_COMPILER_EVENT_STAT(EKismetCompilerStats_ValidateGeneratedClass);
-		UBlueprint::ValidateGeneratedClass(BlueprintObj->GeneratedClass);
-	}
-
-	if (BlueprintPackage != NULL)
-	{
-		BlueprintPackage->SetDirtyFlag(bStartedWithUnsavedChanges);
-	}	
-
-	UEdGraphPin::Purge();
+	FBlueprintCompilationManager::CompileSynchronously(FBPCompileRequest(BlueprintObj, CompileFlags, pResults));
 }
 
 /** Generates a blueprint skeleton only.  Minimal compile, no notifications will be sent, no GC, etc.  Only successful if there isn't already a skeleton generated */
@@ -984,97 +759,13 @@ bool FKismetEditorUtilities::GenerateBlueprintSkeleton(UBlueprint* BlueprintObj,
 
 	if( BlueprintObj->SkeletonGeneratedClass == NULL || bForceRegeneration )
 	{
-		if (GBlueprintUseCompilationManager)
-		{
-			FBlueprintCompilationManager::CompileSynchronously(
-				FBPCompileRequest(BlueprintObj, EBlueprintCompileOptions::RegenerateSkeletonOnly, nullptr)
-			);
-		}
-		else
-		{
-			UPackage* Package = BlueprintObj->GetOutermost();
-			bool bIsPackageDirty = Package ? Package->IsDirty() : false;
-
-			IKismetCompilerInterface& Compiler = FModuleManager::LoadModuleChecked<IKismetCompilerInterface>(KISMET_COMPILER_MODULENAME);
-
-			TGuardValue<bool> GuardTemplateNameFlag(GCompilingBlueprint, true);
-			FCompilerResultsLog Results;
-
-			FKismetCompilerOptions CompileOptions;
-			CompileOptions.CompileType = EKismetCompileType::SkeletonOnly;
-			Compiler.CompileBlueprint(BlueprintObj, CompileOptions, Results);
-			bRegeneratedSkeleton = true;
-
-			// Restore the package dirty flag here
-			if (Package != NULL)
-			{
-				Package->SetDirtyFlag(bIsPackageDirty);
-			}
-		}
+		FBlueprintCompilationManager::CompileSynchronously(
+			FBPCompileRequest(BlueprintObj, EBlueprintCompileOptions::RegenerateSkeletonOnly, nullptr)
+		);
 	}
 	return bRegeneratedSkeleton;
 }
 
-/** Recompiles the bytecode of a blueprint only.  Should only be run for recompiling dependencies during compile on load */
-void FKismetEditorUtilities::RecompileBlueprintBytecode(UBlueprint* BlueprintObj,  EBlueprintBytecodeRecompileOptions Flags)
-{
-	FSecondsCounterScope Timer(BlueprintCompileAndLoadTimerData); 
-
-	if(FBlueprintEditorUtils::IsCompileOnLoadDisabled(BlueprintObj))
-	{
-		return;
-	}
-
-	bool bBatchCompile = (Flags & EBlueprintBytecodeRecompileOptions::BatchCompile) != EBlueprintBytecodeRecompileOptions::None;
-	bool bSkipReinstancing = (Flags & EBlueprintBytecodeRecompileOptions::SkipReinstancing) != EBlueprintBytecodeRecompileOptions::None;
-
-	check(BlueprintObj);
-	checkf(BlueprintObj->GeneratedClass, TEXT("Invalid generated class for %s"), *BlueprintObj->GetName());
-
-	UPackage* const BlueprintPackage = BlueprintObj->GetOutermost();
-	bool const bStartedWithUnsavedChanges = (BlueprintPackage != NULL) ? BlueprintPackage->IsDirty() : true;
-
-	IKismetCompilerInterface& Compiler = FModuleManager::LoadModuleChecked<IKismetCompilerInterface>(KISMET_COMPILER_MODULENAME);
-
-	TGuardValue<bool> GuardTemplateNameFlag(GCompilingBlueprint, true);
-
-	TSharedPtr<FBlueprintCompileReinstancer> ReinstanceHelper;
-	if(!bSkipReinstancing)
-	{
-		ReinstanceHelper = FBlueprintCompileReinstancer::Create(BlueprintObj->GeneratedClass, EBlueprintCompileReinstancerFlags::BytecodeOnly | EBlueprintCompileReinstancerFlags::AutoInferSaveOnCompile);
-	}
-
-	FKismetCompilerOptions CompileOptions;
-	CompileOptions.CompileType = EKismetCompileType::BytecodeOnly;
-	{
-		FRecreateUberGraphFrameScope RecreateUberGraphFrameScope(BlueprintObj->GeneratedClass, true);
-		FCompilerResultsLog Results;
-		Compiler.CompileBlueprint(BlueprintObj, CompileOptions, Results, NULL);
-	}
-	
-	if(!bSkipReinstancing)
-	{
-		ReinstanceHelper->UpdateBytecodeReferences();
-	}
-
-	if (BlueprintPackage != NULL)
-	{
-		BlueprintPackage->SetDirtyFlag(bStartedWithUnsavedChanges);
-	}
-
-	if (!BlueprintObj->bIsRegeneratingOnLoad && !bBatchCompile)
-	{
-		BP_SCOPED_COMPILER_EVENT_STAT(EKismetCompilerStats_NotifyBlueprintChanged);
-
-		BlueprintObj->BroadcastCompiled();
-
-		if (GEditor)
-		{
-			GEditor->BroadcastBlueprintCompiled();
-		}
-	}
-}
-
 namespace ConformComponentsUtils
 {
 	static void ConformRemovedNativeComponents(UObject* BpCdo);
diff --git a/Engine/Source/Editor/UnrealEd/Private/Kismet2/KismetReinstanceUtilities.cpp b/Engine/Source/Editor/UnrealEd/Private/Kismet2/KismetReinstanceUtilities.cpp
--- a/Engine/Source/Editor/UnrealEd/Private/Kismet2/KismetReinstanceUtilities.cpp
+++ b/Engine/Source/Editor/UnrealEd/Private/Kismet2/KismetReinstanceUtilities.cpp
@@ -34,8 +34,6 @@
 #include "BlueprintEditorSettings.h"
 #include "Subsystems/AssetEditorSubsystem.h"
 
-extern COREUOBJECT_API bool GBlueprintUseCompilationManager;
-
 DECLARE_CYCLE_STAT(TEXT("Replace Instances"), EKismetReinstancerStats_ReplaceInstancesOfClass, STATGROUP_KismetReinstancer );
 DECLARE_CYCLE_STAT(TEXT("Find Referencers"), EKismetReinstancerStats_FindReferencers, STATGROUP_KismetReinstancer );
 DECLARE_CYCLE_STAT(TEXT("Replace References"), EKismetReinstancerStats_ReplaceReferences, STATGROUP_KismetReinstancer );
@@ -222,8 +220,6 @@ FReplaceInstancesOfClassParameters::FReplaceInstancesOfClassParameters(UClass* I
 // FBlueprintCompileReinstancer
 
 TSet<TWeakObjectPtr<UBlueprint>> FBlueprintCompileReinstancer::DependentBlueprintsToRefresh = TSet<TWeakObjectPtr<UBlueprint>>();
-TSet<TWeakObjectPtr<UBlueprint>> FBlueprintCompileReinstancer::DependentBlueprintsToRecompile = TSet<TWeakObjectPtr<UBlueprint>>();
-TSet<TWeakObjectPtr<UBlueprint>> FBlueprintCompileReinstancer::DependentBlueprintsToByteRecompile = TSet<TWeakObjectPtr<UBlueprint>>();
 TSet<TWeakObjectPtr<UBlueprint>> FBlueprintCompileReinstancer::CompiledBlueprintsToSave = TSet<TWeakObjectPtr<UBlueprint>>();
 
 UClass* FBlueprintCompileReinstancer::HotReloadedOldClass = nullptr;
@@ -683,17 +679,6 @@ void FBlueprintCompileReinstancer::ListDependentBlueprintsToRefresh(const TArray
 
 void FBlueprintCompileReinstancer::EnlistDependentBlueprintToRecompile(UBlueprint* BP, bool bBytecodeOnly)
 {
-	if (IsValid(BP))
-	{
-		if (bBytecodeOnly)
-		{
-			DependentBlueprintsToByteRecompile.Add(BP);
-		}
-		else
-		{
-			DependentBlueprintsToRecompile.Add(BP);
-		}
-	}
 }
 
 void FBlueprintCompileReinstancer::BlueprintWasRecompiled(UBlueprint* BP, bool bBytecodeOnly)
@@ -701,11 +686,6 @@ void FBlueprintCompileReinstancer::BlueprintWasRecompiled(UBlueprint* BP, bool b
 	if (IsValid(BP))
 	{
 		DependentBlueprintsToRefresh.Remove(BP);
-		DependentBlueprintsToByteRecompile.Remove(BP);
-		if (!bBytecodeOnly)
-		{
-			DependentBlueprintsToRecompile.Remove(BP);
-		}
 	}
 }
 
@@ -752,159 +732,6 @@ void FBlueprintCompileReinstancer::ReinstanceObjects(bool bForceAlwaysReinstance
 				// Mark it as the source reinstancer, no other reinstancer can get here until this Blueprint finishes compiling
 				bIsRootReinstancer = true;
 
-				TSet<TWeakObjectPtr<UBlueprint>> CompiledBlueprints;
-				// Blueprints will enqueue dirty and erroring dependents, in case those states would be 
-				// fixed up by having this dependency compiled first. However, this can result in an 
-				// infinite loop where two Blueprints with errors (unrelated to each other) keep 
-				// perpetually queuing the other. 
-				//
-				// To guard against this, we track the recompiled dependents (in order) and break the 
-				// cycle when we see that we've already compiled a dependent after its dependency
-				TArray<UBlueprint*> OrderedRecompiledDependents;
-
-				TSet<TWeakObjectPtr<UBlueprint>> RecompilationQueue = DependentBlueprintsToRecompile;
-				// empty the public facing queue so we can discern between old and new elements (added 
-				// as the result of subsequent recompiles) 
-				DependentBlueprintsToRecompile.Empty();
-
-				while (RecompilationQueue.Num())
-				{
-					auto Iter = RecompilationQueue.CreateIterator();
-					TWeakObjectPtr<UBlueprint> BPPtr = *Iter;
-					Iter.RemoveCurrent();
-					if (UBlueprint* BP = BPPtr.Get())
-					{
-						if (IsReinstancingSkeleton())
-						{
-							const bool bForceRegeneration = true;
-							FKismetEditorUtilities::GenerateBlueprintSkeleton(BP, bForceRegeneration);
-						}
-						else
-						{
-							// it's unsafe to GC in the middle of reinstancing because there may be other reinstancers still alive with references to 
-							// otherwise unreferenced classes:
-							// Full compiles first recompile all skeleton classes, so they are up-to-date
-							FKismetEditorUtilities::CompileBlueprint(BP, EBlueprintCompileOptions::SkipGarbageCollection | EBlueprintCompileOptions::SkeletonUpToDate);
-							CompiledBlueprints.Add(BP);
-						}
-
-						OrderedRecompiledDependents.Add(BP);
-
-						// if this BP compiled with an error, then I don't see any reason why we 
-						// should attempt to recompile its dependencies; if a subsequent recompile 
-						// would fix this up, then it'll get re-injected into the queue when that happens
-						if (BP->Status != EBlueprintStatus::BS_Error)
-						{
-							for (TWeakObjectPtr<UBlueprint>& DependentPtr : DependentBlueprintsToRecompile)
-							{
-								if (!DependentPtr.IsValid())
-								{
-									continue;
-								}
-								UBlueprint* NewDependent = DependentPtr.Get();
-
-								int32 DependentIndex = OrderedRecompiledDependents.FindLast(NewDependent);
-								if (DependentIndex != INDEX_NONE)
-								{
-									// even though we just pushed BP into the list and know that it
-									// exists as the last entry, we want to see if it was compiled 
-									// earlier (once before 'NewDependent'); so we use Find() to search 
-									// out the first entry
-									int32 RecompilingBpIndex = OrderedRecompiledDependents.Find(BP);
-									if (RecompilingBpIndex != INDEX_NONE && RecompilingBpIndex < DependentIndex)
-									{
-										// we've already recompiled this Blueprint once before (here in 
-										// this loop), already after its dependency has been compiled too;
-										// so, to avoid a potential infinite loop we cannot keep trying 
-										// to compile this
-										//
-										// NOTE: this may result in some a compiler error that would have 
-										//       been resolved in another subsequent compile (for example: 
-										//       B depends on A, A is compiled, A has an error, B compiles 
-										//       with an error as a result, C compiles and enqueues A as a 
-										//       dependent, A is recompiled without error now, B is not 
-										//       enqueued again because its already recompiled after A)
-										// 
-										// the true fix is to restructure the compiler so that these sort 
-										// of scenarios don't happen - until then, this is a fair trade 
-										// off... fallback to a byte code compile instead
-										DependentBlueprintsToByteRecompile.Add(DependentPtr);
-										continue;
-									}
-								}
-								RecompilationQueue.Add(DependentPtr);
-							}
-						}
-						DependentBlueprintsToRecompile.Empty();
-					}
-				}
-
-				TArray<UBlueprint*> OrderedBytecodeRecompile;
-
-				while (DependentBlueprintsToByteRecompile.Num())
-				{
-					auto Iter = DependentBlueprintsToByteRecompile.CreateIterator();
-					if (UBlueprint* BP = Iter->Get())
-					{
-						OrderedBytecodeRecompile.Add(BP);
-					}
-					Iter.RemoveCurrent();
-				}
-
-				// Make sure we compile classes that are deeper in the class hierarchy later
-				// than ones that are higher:
-				OrderedBytecodeRecompile.Sort(
-					[](const UBlueprint& LHS, const UBlueprint& RHS)
-				{
-					int32 LHS_Depth = 0;
-					int32 RHS_Depth = 0;
-
-					UStruct* Iter = LHS.ParentClass;
-					while (Iter)
-					{
-						LHS_Depth += 1;
-						Iter = Iter->GetSuperStruct();
-					}
-
-					Iter = RHS.ParentClass;
-					while (Iter)
-					{
-						RHS_Depth += 1;
-						Iter = Iter->GetSuperStruct();
-					}
-
-					// use name as tie breaker, just so we're stable
-					// across editor sessions:
-					return LHS_Depth != RHS_Depth ? (LHS_Depth < RHS_Depth) : LHS.GetName() < RHS.GetName();
-				}
-				);
-
-				DependentBlueprintsToByteRecompile.Empty();
-
-				for (int I = 0; I != OrderedBytecodeRecompile.Num(); ++I)
-				{
-					UBlueprint* BP = OrderedBytecodeRecompile[I];
-					FKismetEditorUtilities::RecompileBlueprintBytecode(BP, EBlueprintBytecodeRecompileOptions::BatchCompile);
-					ensure(0 == DependentBlueprintsToRecompile.Num());
-					CompiledBlueprints.Add(BP);
-
-					// We won't re-instance (and thus won't finalize) objects of this type, but we still need to know if the current
-					// Actor selection in the level editor includes instances of this type, and force a refresh on the selection if so.
-					// Note: If this is a Component BP instead of an Actor BP, the refresh is handled by IWCE's OnCompiled() BP delegate.
-					if (GEditor && BP->ParentClass && BP->ParentClass->IsChildOf<AActor>())
-					{
-						for (FSelectionIterator It(GEditor->GetSelectedActorIterator()); It; ++It)
-						{
-							AActor* ActorInstance = CastChecked<AActor>(*It);
-							if (ActorInstance->IsA(BP->GeneratedClass))
-							{
-								GEditor->SelectActor(ActorInstance, /*bInSelected =*/true, /*bNotify =*/true, false, true);
-							}
-						}
-					}
-				}
-
-
 				if (!IsReinstancingSkeleton())
 				{
 					TGuardValue<bool> ReinstancingGuard(GIsReinstancing, true);
@@ -931,11 +758,6 @@ void FBlueprintCompileReinstancer::ReinstanceObjects(bool bForceAlwaysReinstance
 						}
 					}
 
-					for (TWeakObjectPtr<UBlueprint>& CompiledBP : CompiledBlueprints)
-					{
-						CompiledBP->BroadcastCompiled();
-					}
-
 					{
 						BP_SCOPED_COMPILER_EVENT_STAT(EKismetCompilerStats_RefreshDependentBlueprintsInReinstancer);
 						for (TWeakObjectPtr<UBlueprint>& BPPtr : DependentBlueprintsToRefresh)
@@ -1653,12 +1475,12 @@ static void ReplaceObjectHelper(UObject*& OldObject, UClass* OldClass, UObject*&
 			for (UObject* OldArchetypeObject : OldArchetypeObjects)
 			{
 				OldToNewNameMap.Add(OldArchetypeObject, OldName);
-				OldArchetypeObject->Rename(*OldArchetypeName, OldArchetypeObject->GetOuter(), REN_DoNotDirty | REN_DontCreateRedirectors | (GBlueprintUseCompilationManager ? REN_ForceNoResetLoaders : 0));
+				OldArchetypeObject->Rename(*OldArchetypeName, OldArchetypeObject->GetOuter(), REN_DoNotDirty | REN_DontCreateRedirectors | REN_ForceNoResetLoaders);
 			}
 		}
 		else
 		{
-			OldObject->Rename(nullptr, OldObject->GetOuter(), REN_DoNotDirty | REN_DontCreateRedirectors | (GBlueprintUseCompilationManager ? REN_ForceNoResetLoaders : 0));
+			OldObject->Rename(nullptr, OldObject->GetOuter(), REN_DoNotDirty | REN_DontCreateRedirectors | REN_ForceNoResetLoaders);
 		}
 	}
 						
@@ -1798,7 +1620,7 @@ static void ReplaceActorHelper(UObject* OldObject, UClass* OldClass, UObject*& N
 	}
 #endif
 
-	OldActor->UObject::Rename(nullptr, OldObject->GetOuter(), REN_DoNotDirty | REN_DontCreateRedirectors | (GBlueprintUseCompilationManager ? REN_ForceNoResetLoaders : 0));
+	OldActor->UObject::Rename(nullptr, OldObject->GetOuter(), REN_DoNotDirty | REN_DontCreateRedirectors | REN_ForceNoResetLoaders);
 
 	AActor* NewActor = nullptr;
 	{
diff --git a/Engine/Source/Editor/UnrealEd/Private/Settings/EditorProjectSettings.cpp b/Engine/Source/Editor/UnrealEd/Private/Settings/EditorProjectSettings.cpp
--- a/Engine/Source/Editor/UnrealEd/Private/Settings/EditorProjectSettings.cpp
+++ b/Engine/Source/Editor/UnrealEd/Private/Settings/EditorProjectSettings.cpp
@@ -3,8 +3,6 @@
 #include "Settings/EditorProjectSettings.h"
 #include "UObject/UnrealType.h"
 
-extern COREUOBJECT_API bool GBlueprintUseCompilationManager;
-
 EUnit ConvertDefaultInputUnits(EDefaultLocationUnit In)
 {
 	typedef EDefaultLocationUnit EDefaultLocationUnit;
diff --git a/Engine/Source/Editor/UnrealEd/Public/Kismet2/KismetEditorUtilities.h b/Engine/Source/Editor/UnrealEd/Public/Kismet2/KismetEditorUtilities.h
--- a/Engine/Source/Editor/UnrealEd/Public/Kismet2/KismetEditorUtilities.h
+++ b/Engine/Source/Editor/UnrealEd/Public/Kismet2/KismetEditorUtilities.h
@@ -140,9 +140,6 @@ public:
 	/** Generates a blueprint skeleton only.  Minimal compile, no notifications will be sent, no GC, etc.  Only successful if there isn't already a skeleton generated */
 	static bool GenerateBlueprintSkeleton(UBlueprint* BlueprintObj, bool bForceRegeneration = false);
 
-	/** Recompiles the bytecode of a blueprint only.  Should only be run for recompiling dependencies during compile on load */
-	static void RecompileBlueprintBytecode(UBlueprint* BlueprintObj, EBlueprintBytecodeRecompileOptions Flags = EBlueprintBytecodeRecompileOptions::None);
-
 	/** Tries to make sure that a data-only blueprint is conformed to its native parent, in case any native class flags have changed */
 	static void ConformBlueprintFlagsAndComponents(UBlueprint* BlueprintObj);
 
diff --git a/Engine/Source/Editor/UnrealEd/Public/Kismet2/KismetReinstanceUtilities.h b/Engine/Source/Editor/UnrealEd/Public/Kismet2/KismetReinstanceUtilities.h
--- a/Engine/Source/Editor/UnrealEd/Public/Kismet2/KismetReinstanceUtilities.h
+++ b/Engine/Source/Editor/UnrealEd/Public/Kismet2/KismetReinstanceUtilities.h
@@ -78,8 +78,6 @@ protected:
 	friend struct FReinstancingJob;
 
 	static TSet<TWeakObjectPtr<UBlueprint>> DependentBlueprintsToRefresh;
-	static TSet<TWeakObjectPtr<UBlueprint>> DependentBlueprintsToRecompile;
-	static TSet<TWeakObjectPtr<UBlueprint>> DependentBlueprintsToByteRecompile;
 	static TSet<TWeakObjectPtr<UBlueprint>> CompiledBlueprintsToSave;
 
 	static UClass* HotReloadedOldClass;
diff --git a/Engine/Source/Editor/UnrealEd/Public/Settings/EditorProjectSettings.h b/Engine/Source/Editor/UnrealEd/Public/Settings/EditorProjectSettings.h
--- a/Engine/Source/Editor/UnrealEd/Public/Settings/EditorProjectSettings.h
+++ b/Engine/Source/Editor/UnrealEd/Public/Settings/EditorProjectSettings.h
@@ -164,14 +164,6 @@ class UNREALED_API UBlueprintEditorProjectSettings : public UDeveloperSettings
 	GENERATED_UCLASS_BODY()
 
 public:
-	/**
-	 * Flag to disable the compilation manager for blueprints - requires editor restart.
-	 * This flag is deprecated! In 4.21 the compilation manager will be the only way to compile
-	 * blueprints. Report any compilation manager issues immediately.
-	 */
-	UPROPERTY(EditAnywhere, config, Category = Blueprints, DisplayName = "Disable Compilation Manager (DEPRECATED)")
-	uint32 bDisableCompilationManager:1;
-	
 	/**
 	 * Flag to disable faster compiles for individual blueprints if they have no function signature
 	 * changes. This flag is deprecated! In 4.21 there will be no way to force all dependencies to 
diff --git a/Engine/Source/Runtime/CoreUObject/Private/Blueprint/BlueprintSupport.cpp b/Engine/Source/Runtime/CoreUObject/Private/Blueprint/BlueprintSupport.cpp
--- a/Engine/Source/Runtime/CoreUObject/Private/Blueprint/BlueprintSupport.cpp
+++ b/Engine/Source/Runtime/CoreUObject/Private/Blueprint/BlueprintSupport.cpp
@@ -28,9 +28,6 @@
 
 DEFINE_LOG_CATEGORY_STATIC(LogBlueprintSupport, Log, All);
 
-// Flag to enable the BlueprintCompilationManager:
-COREUOBJECT_API bool GBlueprintUseCompilationManager = true;
-
 const FName FBlueprintTags::GeneratedClassPath(TEXT("GeneratedClass"));
 const FName FBlueprintTags::ParentClassPath(TEXT("ParentClass"));
 const FName FBlueprintTags::NativeParentClassPath(TEXT("NativeParentClass"));
@@ -118,28 +115,6 @@ bool FBlueprintSupport::IsDeferredCDOInitializationDisabled()
 #endif
 }
 
-void FBlueprintSupport::InitializeCompilationManager()
-{
-	// 'real' initialization is done lazily because we're in a pretty tough
-	// spot in terms of dependencies:
-	bool bCompilationManagerDisabled = false;
-	bool bCompilationManagerEnabled = true;
-	GConfig->GetBool(TEXT("/Script/UnrealEd.BlueprintEditorProjectSettings"), TEXT("bUseCompilationManager"), bCompilationManagerEnabled, GEditorIni);
-	GConfig->GetBool(TEXT("/Script/UnrealEd.BlueprintEditorProjectSettings"), TEXT("bDisableCompilationManager"), bCompilationManagerDisabled, GEditorIni);
-	if (bCompilationManagerDisabled)
-	{
-		GBlueprintUseCompilationManager = false;
-	}
-	else if(!bCompilationManagerEnabled)
-	{
-		UE_LOG(LogScript, Warning,
-			TEXT("Warning: Compilation manager enabled, compilation manager will be mandatory in 4.21. \
-				Use bDisableCompilationManager to disable the compilation manager as a last resort or remove \
-				bUseCompilationManager from project .ini to fix this warning.")
-		);
-	}
-}
-
 static FFlushReinstancingQueueFPtr FlushReinstancingQueueFPtr = nullptr;
 
 void FBlueprintSupport::FlushReinstancingQueue()
@@ -461,24 +436,7 @@ FScopedClassDependencyGather::~FScopedClassDependencyGather()
 			}
 		};
 
-		if(!GBlueprintUseCompilationManager)
-		{
-			for( ; DependencyIter; ++DependencyIter )
-			{
-				UClass* Dependency = *DependencyIter;
-				if( Dependency->ClassGeneratedBy != BatchMasterClass->ClassGeneratedBy )
-				{
-					RecompileClassLambda(Dependency, LoadContext);
-				}
-			}
-
-			// Finally, recompile the master class to make sure it gets updated too
-			RecompileClassLambda(BatchMasterClass, LoadContext);
-		}
-		else
-		{
-			BatchMasterClass->ConditionalRecompileClass(LoadContext);
-		}
+		BatchMasterClass->ConditionalRecompileClass(LoadContext);
 
 		BatchMasterClass = NULL;
 	}
diff --git a/Engine/Source/Runtime/CoreUObject/Public/Blueprint/BlueprintSupport.h b/Engine/Source/Runtime/CoreUObject/Public/Blueprint/BlueprintSupport.h
--- a/Engine/Source/Runtime/CoreUObject/Public/Blueprint/BlueprintSupport.h
+++ b/Engine/Source/Runtime/CoreUObject/Public/Blueprint/BlueprintSupport.h
@@ -84,9 +84,6 @@ struct FBlueprintSupport
 	static bool IsDeferredExportCreationDisabled();
 	static bool IsDeferredCDOInitializationDisabled();
 
-	/** Initializes the BP Compilation Manager if it's enabled: */
-	COREUOBJECT_API static void InitializeCompilationManager();
-
 	/** Checks for any old instances and reinstances them: */
 	static void FlushReinstancingQueue();
 	COREUOBJECT_API static void SetFlushReinstancingQueueFPtr(FFlushReinstancingQueueFPtr Ptr);
diff --git a/Engine/Source/Runtime/Engine/Private/Blueprint.cpp b/Engine/Source/Runtime/Engine/Private/Blueprint.cpp
--- a/Engine/Source/Runtime/Engine/Private/Blueprint.cpp
+++ b/Engine/Source/Runtime/Engine/Private/Blueprint.cpp
@@ -580,69 +580,60 @@ void UBlueprint::PostDuplicate(bool bDuplicateForPIE)
 	}
 }
 
-extern COREUOBJECT_API bool GBlueprintUseCompilationManager;
-
 UClass* UBlueprint::RegenerateClass(UClass* ClassToRegenerate, UObject* PreviousCDO)
 {
 	LoadModulesRequiredForCompilation();
 
-	if(GBlueprintUseCompilationManager)
+	// ensure that we have UProperties for any properties declared in the blueprint:
+	if(!GeneratedClass || !HasAnyFlags(RF_BeingRegenerated) || bIsRegeneratingOnLoad || bHasBeenRegenerated)
 	{
-		// ensure that we have UProperties for any properties declared in the blueprint:
-		if(!GeneratedClass || !HasAnyFlags(RF_BeingRegenerated) || bIsRegeneratingOnLoad || bHasBeenRegenerated)
-		{
-			return GeneratedClass;
-		}
+		return GeneratedClass;
+	}
 		
-		// tag ourself as bIsRegeneratingOnLoad so that any reentrance via ForceLoad calls doesn't recurse:
-		bIsRegeneratingOnLoad = true;
+	// tag ourself as bIsRegeneratingOnLoad so that any reentrance via ForceLoad calls doesn't recurse:
+	bIsRegeneratingOnLoad = true;
 		
-		UPackage* Package = GetOutermost();
-		bool bIsPackageDirty = Package ? Package->IsDirty() : false;
+	UPackage* Package = GetOutermost();
+	bool bIsPackageDirty = Package ? Package->IsDirty() : false;
 
-		UClass* GeneratedClassResolved = GeneratedClass;
+	UClass* GeneratedClassResolved = GeneratedClass;
 
-		UBlueprint::ForceLoadMetaData(this);
-		if (ensure(GeneratedClassResolved->ClassDefaultObject ))
-		{
-			UBlueprint::ForceLoadMembers(GeneratedClassResolved);
-			UBlueprint::ForceLoadMembers(GeneratedClassResolved->ClassDefaultObject);
-		}
-		UBlueprint::ForceLoadMembers(this);
+	UBlueprint::ForceLoadMetaData(this);
+	if (ensure(GeneratedClassResolved->ClassDefaultObject ))
+	{
+		UBlueprint::ForceLoadMembers(GeneratedClassResolved);
+		UBlueprint::ForceLoadMembers(GeneratedClassResolved->ClassDefaultObject);
+	}
+	UBlueprint::ForceLoadMembers(this);
 		
-		FBlueprintEditorUtils::PreloadConstructionScript( this );
+	FBlueprintEditorUtils::PreloadConstructionScript( this );
 
-		FBlueprintEditorUtils::LinkExternalDependencies( this );
+	FBlueprintEditorUtils::LinkExternalDependencies( this );
 
-		FBlueprintEditorUtils::RefreshVariables(this);
+	FBlueprintEditorUtils::RefreshVariables(this);
 		
-		// Preload Overridden Components
-		if (InheritableComponentHandler)
-		{
-			InheritableComponentHandler->PreloadAll();
-		}
+	// Preload Overridden Components
+	if (InheritableComponentHandler)
+	{
+		InheritableComponentHandler->PreloadAll();
+	}
 
-		FBlueprintCompilationManager::NotifyBlueprintLoaded( this ); 
+	FBlueprintCompilationManager::NotifyBlueprintLoaded( this ); 
 		
-		FBlueprintEditorUtils::PreloadBlueprintSpecificData( this );
+	FBlueprintEditorUtils::PreloadBlueprintSpecificData( this );
 
-		FBlueprintEditorUtils::UpdateOutOfDateAnimBlueprints(this);
+	FBlueprintEditorUtils::UpdateOutOfDateAnimBlueprints(this);
 
-		// clear this now that we're not in a re-entrrant context - bHasBeenRegenerated will guard against 'real' 
-		// double regeneration calls:
-		bIsRegeneratingOnLoad = false;
+	// clear this now that we're not in a re-entrrant context - bHasBeenRegenerated will guard against 'real' 
+	// double regeneration calls:
+	bIsRegeneratingOnLoad = false;
 
-		if( Package )
-		{
-			Package->SetDirtyFlag(bIsPackageDirty);
-		}
-
-		return GeneratedClassResolved;
-	}
-	else
+	if( Package )
 	{
-		return FBlueprintEditorUtils::RegenerateBlueprintClass(this, ClassToRegenerate, PreviousCDO);
+		Package->SetDirtyFlag(bIsPackageDirty);
 	}
+
+	return GeneratedClassResolved;
 }
 
 void UBlueprint::RemoveChildRedirectors()
diff --git a/Engine/Source/Runtime/Engine/Private/BlueprintGeneratedClass.cpp b/Engine/Source/Runtime/Engine/Private/BlueprintGeneratedClass.cpp
--- a/Engine/Source/Runtime/Engine/Private/BlueprintGeneratedClass.cpp
+++ b/Engine/Source/Runtime/Engine/Private/BlueprintGeneratedClass.cpp
@@ -245,71 +245,17 @@ struct FConditionalRecompileClassHepler
 };
 
 extern UNREALED_API FSecondsCounterData BlueprintCompileAndLoadTimerData;
-extern COREUOBJECT_API bool GBlueprintUseCompilationManager;
 
 void UBlueprintGeneratedClass::ConditionalRecompileClass(FUObjectSerializeContext* InLoadContext)
 {
-	if(GBlueprintUseCompilationManager)
-	{
-		FBlueprintCompilationManager::FlushCompilationQueue(InLoadContext);
-		return;
-	}
-	
-	FSecondsCounterScope Timer(BlueprintCompileAndLoadTimerData);
-
-	UBlueprint* GeneratingBP = Cast<UBlueprint>(ClassGeneratedBy);
-	if (GeneratingBP && (GeneratingBP->SkeletonGeneratedClass != this))
-	{
-		const FConditionalRecompileClassHepler::ENeededAction NecessaryAction = FConditionalRecompileClassHepler::IsConditionalRecompilationNecessary(GeneratingBP);
-		if (FConditionalRecompileClassHepler::Recompile == NecessaryAction)
-		{
-			const bool bWasRegenerating = GeneratingBP->bIsRegeneratingOnLoad;
-			GeneratingBP->bIsRegeneratingOnLoad = true;
-
-			{
-				UPackage* const Package = GeneratingBP->GetOutermost();
-				const bool bStartedWithUnsavedChanges = Package != nullptr ? Package->IsDirty() : true;
-
-				// Make sure that nodes are up to date, so that we get any updated blueprint signatures
-				FBlueprintEditorUtils::RefreshExternalBlueprintDependencyNodes(GeneratingBP);
-
-				// Normal blueprints get their status reset by RecompileBlueprintBytecode, but macros will not:
-				if ((GeneratingBP->Status != BS_Error) && (GeneratingBP->BlueprintType == EBlueprintType::BPTYPE_MacroLibrary))
-				{
-					GeneratingBP->Status = BS_UpToDate;
-				}
-
-				if (Package != nullptr && Package->IsDirty() && !bStartedWithUnsavedChanges)
-				{
-					Package->SetDirtyFlag(false);
-				}
-			}
-			if ((GeneratingBP->Status != BS_Error) && (GeneratingBP->BlueprintType != EBlueprintType::BPTYPE_MacroLibrary))
-			{
-				FKismetEditorUtilities::RecompileBlueprintBytecode(GeneratingBP);
-			}
-
-			GeneratingBP->bIsRegeneratingOnLoad = bWasRegenerating;
-		}
-		else if (FConditionalRecompileClassHepler::StaticLink == NecessaryAction)
-		{
-			StaticLink(true);
-			if (*GeneratingBP->SkeletonGeneratedClass)
-			{
-				GeneratingBP->SkeletonGeneratedClass->StaticLink(true);
-			}
-		}
-	}
+	FBlueprintCompilationManager::FlushCompilationQueue(InLoadContext);
 }
 
 void UBlueprintGeneratedClass::FlushCompilationQueueForLevel()
 {
-	if(GBlueprintUseCompilationManager)
+	if(Cast<ULevelScriptBlueprint>(ClassGeneratedBy))
 	{
-		if(Cast<ULevelScriptBlueprint>(ClassGeneratedBy))
-		{
-			FBlueprintCompilationManager::FlushCompilationQueue(nullptr);
-		}
+		FBlueprintCompilationManager::FlushCompilationQueue(nullptr);
 	}
 }
 
diff --git a/Engine/Source/Runtime/Engine/Private/MemberReference.cpp b/Engine/Source/Runtime/Engine/Private/MemberReference.cpp
--- a/Engine/Source/Runtime/Engine/Private/MemberReference.cpp
+++ b/Engine/Source/Runtime/Engine/Private/MemberReference.cpp
@@ -9,8 +9,6 @@
 #include "BlueprintCompilationManager.h"
 #endif
 
-extern COREUOBJECT_API bool GBlueprintUseCompilationManager;
-
 //////////////////////////////////////////////////////////////////////////
 // FMemberReference
 
@@ -266,7 +264,7 @@ void FMemberReference::InitFieldRedirectMap()
 
 UClass* FMemberReference::GetClassToUse(UClass* InClass, bool bUseUpToDateClass)
 {
-	if(GBlueprintUseCompilationManager && bUseUpToDateClass)
+	if(bUseUpToDateClass)
 	{
 		return FBlueprintEditorUtils::GetMostUpToDateClass(InClass);
 	}
diff --git a/Engine/Source/Runtime/Launch/Private/LaunchEngineLoop.cpp b/Engine/Source/Runtime/Launch/Private/LaunchEngineLoop.cpp
--- a/Engine/Source/Runtime/Launch/Private/LaunchEngineLoop.cpp
+++ b/Engine/Source/Runtime/Launch/Private/LaunchEngineLoop.cpp
@@ -2380,7 +2380,7 @@ int32 FEngineLoop::PreInit(const TCHAR* CmdLine)
 					GUseRHIThread_InternalUseOnly = false;
 				}
 			}
-
+				
 			SCOPED_BOOT_TIMING("StartRenderingThread");
 			StartRenderingThread();
 		}
@@ -4926,10 +4926,6 @@ bool FEngineLoop::AppInit( )
 		FPlatformStackWalk::Init();
 	}
 
-#if WITH_EDITOR
-	FBlueprintSupport::InitializeCompilationManager();
-#endif
-
 	CheckForPrintTimesOverride();
 
 	// Check whether the project or any of its plugins are missing or are out of date
