Forgotten file
#rnx

diff --git a/Engine/Source/Editor/UnrealEd/Private/Kismet2/BlueprintEditorUtils.cpp b/Engine/Source/Editor/UnrealEd/Private/Kismet2/BlueprintEditorUtils.cpp
--- a/Engine/Source/Editor/UnrealEd/Private/Kismet2/BlueprintEditorUtils.cpp
+++ b/Engine/Source/Editor/UnrealEd/Private/Kismet2/BlueprintEditorUtils.cpp
@@ -120,8 +120,6 @@
 #include "AnimGraphNode_Root.h"
 #include "Subsystems/AssetEditorSubsystem.h"
 
-extern COREUOBJECT_API bool GBlueprintUseCompilationManager;
-
 #define LOCTEXT_NAMESPACE "Blueprint"
 
 DEFINE_LOG_CATEGORY(LogBlueprintDebug);
@@ -1857,16 +1855,9 @@ void FBlueprintEditorUtils::PostDuplicateBlueprint(UBlueprint* Blueprint, bool b
 			Blueprint->Timelines = NewBPGC->Timelines;
 			Blueprint->InheritableComponentHandler = NewBPGC->InheritableComponentHandler;
 
-			if (GBlueprintUseCompilationManager)
-			{
-				FBlueprintCompilationManager::CompileSynchronously(
-					FBPCompileRequest(Blueprint, EBlueprintCompileOptions::RegenerateSkeletonOnly, nullptr)
-				);
-			}
-			else
-			{
-				Compiler.CompileBlueprint(Blueprint, CompileOptions, Results);
-			}
+			FBlueprintCompilationManager::CompileSynchronously(
+				FBPCompileRequest(Blueprint, EBlueprintCompileOptions::RegenerateSkeletonOnly, nullptr)
+			);
 
 			// Create a new blueprint guid
 			Blueprint->GenerateNewGuid();
@@ -1908,24 +1899,15 @@ void FBlueprintEditorUtils::PostDuplicateBlueprint(UBlueprint* Blueprint, bool b
 				}
 			}
 
-			if (GBlueprintUseCompilationManager)
-			{
-				// Skip CDO validation in this case as we will not have yet propagated values to the new CDO. Also skip
-				// Blueprint search data updates, as that will be handled by an OnAssetAdded() delegate in the FiB manager.
-				const EBlueprintCompileOptions BPCompileOptions =
-					EBlueprintCompileOptions::SkipDefaultObjectValidation |
-					EBlueprintCompileOptions::SkipFiBSearchMetaUpdate;
+			// Skip CDO validation in this case as we will not have yet propagated values to the new CDO. Also skip
+			// Blueprint search data updates, as that will be handled by an OnAssetAdded() delegate in the FiB manager.
+			const EBlueprintCompileOptions BPCompileOptions =
+				EBlueprintCompileOptions::SkipDefaultObjectValidation |
+				EBlueprintCompileOptions::SkipFiBSearchMetaUpdate;
 
-				FBlueprintCompilationManager::CompileSynchronously(
-					FBPCompileRequest(Blueprint, BPCompileOptions, nullptr)
-				);
-			}
-			else
-			{
-				// Needs a full compile to handle the ArchiveReplaceObjectRef
-				CompileOptions.CompileType = EKismetCompileType::Full;
-				Compiler.CompileBlueprint(Blueprint, CompileOptions, Results);
-			}
+			FBlueprintCompilationManager::CompileSynchronously(
+				FBPCompileRequest(Blueprint, BPCompileOptions, nullptr)
+			);
 
 			FArchiveReplaceObjectRef<UObject> ReplaceTemplateRefs(NewBPGC, OldToNewMap, /*bNullPrivateRefs=*/ false, /*bIgnoreOuterRef=*/ false, /*bIgnoreArchetypeRef=*/ false);
 
@@ -1997,55 +1979,6 @@ void FBlueprintEditorUtils::MarkBlueprintAsStructurallyModified(UBlueprint* Blue
 {
 	FSecondsCounterScope Timer(BlueprintCompileAndLoadTimerData);
 
-	struct FRefreshHelper
-	{
-		static void SkeletalRecompileChildren(TArray<UClass*> SkelClassesToRecompile, bool bIsCompilingOnLoad)
-		{
-			FSecondsCounterScope SkeletalRecompileTimer(BlueprintCompileAndLoadTimerData);
-
-			for (UClass* SkelClass : SkelClassesToRecompile)
-			{
-				if (SkelClass->HasAnyClassFlags(CLASS_NewerVersionExists))
-				{
-					continue;
-				}
-
-				UBlueprint* SkelBlueprint = Cast<UBlueprint>(SkelClass->ClassGeneratedBy);
-				if (SkelBlueprint
-					&& SkelBlueprint->Status != BS_BeingCreated
-					&& !SkelBlueprint->bBeingCompiled
-					&& !SkelBlueprint->bIsRegeneratingOnLoad)
-				{
-					TArray<UClass*> ChildrenOfClass;
-					GetDerivedClasses(SkelClass, ChildrenOfClass, false);
-
-					IKismetCompilerInterface& Compiler = FModuleManager::LoadModuleChecked<IKismetCompilerInterface>(KISMET_COMPILER_MODULENAME);
-
-					FCompilerResultsLog Results;
-					Results.bSilentMode = true;
-					Results.bLogInfoOnly = true;
-
-					{
-						bool const bWasRegenerating = SkelBlueprint->bIsRegeneratingOnLoad;
-						SkelBlueprint->bIsRegeneratingOnLoad |= bIsCompilingOnLoad;
-
-						FKismetCompilerOptions CompileOptions;
-						CompileOptions.CompileType = EKismetCompileType::SkeletonOnly;
-						Compiler.CompileBlueprint(SkelBlueprint, CompileOptions, Results);
-						SkelBlueprint->Status = BS_Dirty;
-
-						SkelBlueprint->BroadcastCompiled();
-
-						SkelBlueprint->MarkPackageDirty();
-
-						SkeletalRecompileChildren(ChildrenOfClass, bIsCompilingOnLoad);
-						SkelBlueprint->bIsRegeneratingOnLoad = bWasRegenerating;
-					}
-				}
-			}
-		}
-	};
-
 	// The Blueprint has been structurally modified and this means that some node titles will need to be refreshed
 	GetDefault<UEdGraphSchema_K2>()->ForceVisualizationCacheClear();
 
@@ -2054,61 +1987,9 @@ void FBlueprintEditorUtils::MarkBlueprintAsStructurallyModified(UBlueprint* Blue
 	{
 		BP_SCOPED_COMPILER_EVENT_STAT(EKismetCompilerStats_MarkBlueprintasStructurallyModified);
 
-		if (GBlueprintUseCompilationManager)
-		{
-			FBlueprintCompilationManager::CompileSynchronously(
-				FBPCompileRequest(Blueprint, EBlueprintCompileOptions::RegenerateSkeletonOnly, nullptr)
-			);
-		}
-		else
-		{
-			FCompilerResultsLog Results;
-			Results.bLogInfoOnly = Blueprint->bIsRegeneratingOnLoad;
-
-			TArray<UClass*> ChildrenOfClass;
-			if (UClass* SkelClass = Blueprint->SkeletonGeneratedClass)
-			{
-				if (!Blueprint->bIsRegeneratingOnLoad)
-				{
-					if (IsInterfaceBlueprint(Blueprint))
-					{
-						// Find all dependent Blueprints that implement the interface. Note: Using
-						// GetDependentBlueprints() here as the result is cached and thus it should
-						// generally be a faster path than iterating through all loaded Blueprints.
-						TArray<UBlueprint*> DependentBlueprints;
-						GetDependentBlueprints(Blueprint, DependentBlueprints, true);
-						for (UBlueprint* DependentBlueprint : DependentBlueprints)
-						{
-							const bool bDependentBPImplementsInterface = DependentBlueprint->ImplementedInterfaces.ContainsByPredicate([&Blueprint](const FBPInterfaceDescription& InterfaceDesc)
-							{
-								return InterfaceDesc.Interface == Blueprint->GeneratedClass;
-							});
-
-							if (bDependentBPImplementsInterface)
-							{
-								ChildrenOfClass.Add(DependentBlueprint->SkeletonGeneratedClass);
-							}
-						}
-					}
-					else
-					{
-						GetDerivedClasses(SkelClass, ChildrenOfClass, false);
-					}
-				}
-			}
-
-			{
-				// Invoke the compiler to update the skeleton class definition
-				IKismetCompilerInterface& Compiler = FModuleManager::LoadModuleChecked<IKismetCompilerInterface>(KISMET_COMPILER_MODULENAME);
-
-				FKismetCompilerOptions CompileOptions;
-				CompileOptions.CompileType = EKismetCompileType::SkeletonOnly;
-				Compiler.CompileBlueprint(Blueprint, CompileOptions, Results);
-			}
-			UpdateDelegatesInBlueprint(Blueprint);
-
-			FRefreshHelper::SkeletalRecompileChildren(ChildrenOfClass, Blueprint->bIsRegeneratingOnLoad);
-		}
+		FBlueprintCompilationManager::CompileSynchronously(
+			FBPCompileRequest(Blueprint, EBlueprintCompileOptions::RegenerateSkeletonOnly, nullptr)
+		);
 
 		// Call general modification callback as well
 		MarkBlueprintAsModified(Blueprint);
@@ -2125,7 +2006,7 @@ void FBlueprintEditorUtils::MarkBlueprintAsStructurallyModified(UBlueprint* Blue
 // Blueprint has changed in some manner that invalidates the compiled data (link made/broken, default value changed, etc...)
 void FBlueprintEditorUtils::MarkBlueprintAsModified(UBlueprint* Blueprint, FPropertyChangedEvent PropertyChangedEvent)
 {
-	if(Blueprint->bBeingCompiled && GBlueprintUseCompilationManager)
+	if(Blueprint->bBeingCompiled)
 	{
 		return;
 	}
@@ -2245,16 +2126,9 @@ UClass* FBlueprintEditorUtils::GetMostUpToDateClass(UClass* FromClass)
 		return FromClass;
 	}
 
-	if(GBlueprintUseCompilationManager)
-	{
-		// It's really not safe/coherent to try and dig out the 'right' class. Things that need the 'most up to date'
-		// version of a class should always be looking at the skeleton:
-		return GetSkeletonClass(FromClass);
-	}
-	else
-	{
-		return FromClass;
-	}
+	// It's really not safe/coherent to try and dig out the 'right' class. Things that need the 'most up to date'
+	// version of a class should always be looking at the skeleton:
+	return GetSkeletonClass(FromClass);
 }
 
 const UClass* FBlueprintEditorUtils::GetMostUpToDateClass(const UClass* FromClass)
@@ -2264,13 +2138,7 @@ const UClass* FBlueprintEditorUtils::GetMostUpToDateClass(const UClass* FromClas
 
 bool FBlueprintEditorUtils::PropertyStillExists(UProperty* Property)
 {
-	if(GBlueprintUseCompilationManager)
-	{
-		return GetMostUpToDateProperty(Property) != nullptr;
-	}
-
-	// We can't reliably know if the property still exists, but assume that it does:
-	return true;
+	return GetMostUpToDateProperty(Property) != nullptr;
 }
 
 UProperty* FBlueprintEditorUtils::GetMostUpToDateProperty(UProperty* Property)
@@ -4444,7 +4312,6 @@ void FBlueprintEditorUtils::GetClassVariableList(const UBlueprint* Blueprint, TS
 	// Existing variables in the parent class and above, when using the compilation manager the previous SkeletonGeneratedClass will have been cleared when
 	// we're regenerating the SkeletonGeneratedClass. Using this function in the skeleton pass at all is highly dubious, but I am leaving it until the 
 	// compilation manager is on full time:
-	check(!Blueprint->bHasBeenRegenerated || Blueprint->bIsRegeneratingOnLoad || (Blueprint->SkeletonGeneratedClass != nullptr) || GBlueprintUseCompilationManager);
 	if (Blueprint->SkeletonGeneratedClass != nullptr)
 	{
 		for (TFieldIterator<UProperty> PropertyIt(Blueprint->SkeletonGeneratedClass, EFieldIteratorFlags::IncludeSuper); PropertyIt; ++PropertyIt)
