Remove CompileSkeletonOnly option from CompileAllBlueprints
#rnx

diff --git a/Engine/Source/Editor/UnrealEd/Private/Commandlets/CompileAllBlueprintsCommandlet.cpp b/Engine/Source/Editor/UnrealEd/Private/Commandlets/CompileAllBlueprintsCommandlet.cpp
--- a/Engine/Source/Editor/UnrealEd/Private/Commandlets/CompileAllBlueprintsCommandlet.cpp
+++ b/Engine/Source/Editor/UnrealEd/Private/Commandlets/CompileAllBlueprintsCommandlet.cpp
@@ -19,7 +19,6 @@ UCompileAllBlueprintsCommandlet::UCompileAllBlueprintsCommandlet(const FObjectIn
 : Super(ObjectInitializer)
 {
 	bResultsOnly = false;
-	bCompileSkeletonOnly = false;
 	bCookedOnly = false;
 	bDirtyOnly = false;
 	bSimpleAssetList = false;
@@ -52,7 +51,6 @@ void UCompileAllBlueprintsCommandlet::InitCommandLine(const FString& Params)
 	bResultsOnly = Switches.Contains(TEXT("ShowResultsOnly"));
 	bDirtyOnly = Switches.Contains(TEXT("DirtyOnly"));
 	bCookedOnly = Switches.Contains(TEXT("CookedOnly"));
-	bCompileSkeletonOnly = Switches.Contains(TEXT("CompileSkeletonOnly"));
 	bSimpleAssetList = Switches.Contains(TEXT("SimpleAssetList"));
 	
 	RequireAssetTags.Empty();
@@ -305,19 +303,6 @@ void UCompileAllBlueprintsCommandlet::CompileBlueprint(UBlueprint* Blueprint)
 		}
 		MessageLog.SetSourcePath(Blueprint->GetPathName());
 		MessageLog.BeginEvent(TEXT("Compile"));
-	
-		FKismetCompilerOptions CompileOptions;
-		//Defaults to use for CompileOptions
-		CompileOptions.CompileType = EKismetCompileType::Full;
-		CompileOptions.bSaveIntermediateProducts = false;
-		CompileOptions.bRegenerateSkelton = false;
-		CompileOptions.bIsDuplicationInstigated = false;
-		CompileOptions.bReinstanceAndStubOnFailure = false;
-
-		if (bCompileSkeletonOnly)
-		{
-			CompileOptions.CompileType = EKismetCompileType::SkeletonOnly;
-		}
 
 		FKismetEditorUtilities::CompileBlueprint(Blueprint, EBlueprintCompileOptions::SkipGarbageCollection, &MessageLog);
 
