Avoid duplication of child actors when a reinstancing batch includes both the child actor class as well as the parent class that contains the child actor component spawning it

diff --git a/Engine/Source/Editor/UnrealEd/Private/Kismet2/KismetReinstanceUtilities.cpp b/Engine/Source/Editor/UnrealEd/Private/Kismet2/KismetReinstanceUtilities.cpp
--- a/Engine/Source/Editor/UnrealEd/Private/Kismet2/KismetReinstanceUtilities.cpp
+++ b/Engine/Source/Editor/UnrealEd/Private/Kismet2/KismetReinstanceUtilities.cpp
@@ -1013,7 +1013,11 @@ private:
 
 void FActorReplacementHelper::Finalize(const TMap<UObject*, UObject*>& OldToNewInstanceMap, TSet<UObject*>* ObjectsThatShouldUseOldStuff, const TArray<UObject*>& ObjectsToReplace, const TMap<FSoftObjectPath, UObject*>& ReinstancedObjectsWeakReferenceMap)
 {
-	
+	if (NewActor->IsPendingKill())
+	{
+		return;
+	}
+
 	// because this is an editor context it's important to use this execution guard
 	FEditorScriptExecutionGuard ScriptGuard;
 
@@ -1572,7 +1576,7 @@ static void ReplaceObjectHelper(UObject*& OldObject, UClass* OldClass, UObject*&
 	}
 }
 
-static void ReplaceActorHelper(UObject* OldObject, UClass* OldClass, UObject*& NewUObject, UClass* NewClass, TMap<UObject*, UObject*>& OldToNewInstanceMap, TMap<UClass*, UClass*>& InOldToNewClassMap, AActor* OldActor, TMap<FSoftObjectPath, UObject*>& ReinstancedObjectsWeakReferenceMap, TMap<UObject*, FActorAttachmentData>& ActorAttachmentData, TArray<FActorReplacementHelper>& ReplacementActors, bool bPreserveRootComponent, bool& bSelectionChanged)
+static void ReplaceActorHelper(AActor* OldActor, UClass* OldClass, UObject*& NewUObject, UClass* NewClass, TMap<UObject*, UObject*>& OldToNewInstanceMap, TMap<UClass*, UClass*>& InOldToNewClassMap, TMap<FSoftObjectPath, UObject*>& ReinstancedObjectsWeakReferenceMap, TMap<UObject*, FActorAttachmentData>& ActorAttachmentData, TArray<FActorReplacementHelper>& ReplacementActors, bool bPreserveRootComponent, bool& bSelectionChanged)
 {
 	FVector  Location = FVector::ZeroVector;
 	FRotator Rotation = FRotator::ZeroRotator;
@@ -1620,7 +1624,7 @@ static void ReplaceActorHelper(UObject* OldObject, UClass* OldClass, UObject*& N
 	}
 #endif
 
-	OldActor->UObject::Rename(nullptr, OldObject->GetOuter(), REN_DoNotDirty | REN_DontCreateRedirectors | REN_ForceNoResetLoaders);
+	OldActor->UObject::Rename(nullptr, OldActor->GetOuter(), REN_DoNotDirty | REN_DontCreateRedirectors | REN_ForceNoResetLoaders);
 
 	AActor* NewActor = nullptr;
 	{
@@ -1641,11 +1645,11 @@ static void ReplaceActorHelper(UObject* OldObject, UClass* OldClass, UObject*& N
 	// running the NewActor's construction-script is saved for that 
 	// second pass (because the construction-script may reference 
 	// another instance that hasn't been replaced yet).
-	FActorAttachmentData& CurrentAttachmentData = ActorAttachmentData.FindChecked(OldObject);
+	FActorAttachmentData& CurrentAttachmentData = ActorAttachmentData.FindChecked(OldActor);
 	ReplacementActors.Add(FActorReplacementHelper(NewActor, OldActor, MoveTemp(CurrentAttachmentData)));
-	ActorAttachmentData.Remove(OldObject);
+	ActorAttachmentData.Remove(OldActor);
 
-	ReinstancedObjectsWeakReferenceMap.Add(OldObject, NewUObject);
+	ReinstancedObjectsWeakReferenceMap.Add(OldActor, NewUObject);
 
 	OldActor->DestroyConstructedComponents(); // don't want to serialize components from the old actor
 												// Unregister native components so we don't copy any sub-components they generate for themselves (like UCameraComponent does)
@@ -1684,7 +1688,6 @@ static void ReplaceActorHelper(UObject* OldObject, UClass* OldClass, UObject*& N
 		}
 	}
 
-	World->EditorDestroyActor(OldActor, /*bShouldModifyLevel =*/true);
 	OldToNewInstanceMap.Add(OldActor, NewActor);
 }
 
@@ -1828,6 +1831,16 @@ void FBlueprintCompileReinstancer::ReplaceInstancesOfClass_Inner(TMap<UClass*, U
 			}
 		}
 
+
+		FDelegateHandle OnLevelActorDeletedHandle = GEngine->OnLevelActorDeleted().AddLambda([&OldToNewInstanceMap](AActor* DestroyedActor)
+			{
+				if (UObject** ReplacementObject = OldToNewInstanceMap.Find(DestroyedActor))
+				{
+					AActor* ReplacementActor = CastChecked<AActor>(*ReplacementObject);
+					ReplacementActor->GetWorld()->EditorDestroyActor(ReplacementActor, /*bShouldModifyLevel =*/true);
+				}
+			});
+
 		// WARNING: for (TPair<UClass*, UClass*> OldToNewClass : InOldToNewClassMap) duplicated above 
 		// this loop only handles actors - which need to be reconstructed *after* their owned components 
 		// have been reinstanced:
@@ -1879,9 +1892,9 @@ void FBlueprintCompileReinstancer::ReplaceInstancesOfClass_Inner(TMap<UClass*, U
 					if (OldActor != nullptr)
 					{
 						UObject* NewUObject = nullptr;
-						if(OldActor->GetLevel())
+						if (OldActor->GetLevel())
 						{
-							ReplaceActorHelper(OldObject, OldClass, NewUObject, NewClass, OldToNewInstanceMap, InOldToNewClassMap, OldActor, ReinstancedObjectsWeakReferenceMap, ActorAttachmentData, ReplacementActors, bPreserveRootComponent, bSelectionChanged);
+							ReplaceActorHelper(OldActor, OldClass, NewUObject, NewClass, OldToNewInstanceMap, InOldToNewClassMap, ReinstancedObjectsWeakReferenceMap, ActorAttachmentData, ReplacementActors, bPreserveRootComponent, bSelectionChanged);
 						}
 						else
 						{
@@ -1899,6 +1912,15 @@ void FBlueprintCompileReinstancer::ReplaceInstancesOfClass_Inner(TMap<UClass*, U
 				}
 			}
 		}
+		GEngine->OnLevelActorDeleted().Remove(OnLevelActorDeletedHandle);
+
+		for (TPair<UObject*, UObject*> ReinstancedPair : OldToNewInstanceMap)
+		{
+			if (AActor* OldActor = Cast<AActor>(ReinstancedPair.Key))
+			{
+				OldActor->GetWorld()->EditorDestroyActor(OldActor, /*bShouldModifyLevel =*/true);
+			}
+		}
 	}
 
 	if(GEditor)
