Mark blueprints dirty when externally referenced blueprint variables are renamed

diff --git a/Engine/Source/Editor/UnrealEd/Private/Kismet2/BlueprintEditorUtils.cpp b/Engine/Source/Editor/UnrealEd/Private/Kismet2/BlueprintEditorUtils.cpp
--- a/Engine/Source/Editor/UnrealEd/Private/Kismet2/BlueprintEditorUtils.cpp
+++ b/Engine/Source/Editor/UnrealEd/Private/Kismet2/BlueprintEditorUtils.cpp
@@ -268,16 +268,21 @@ static void PromoteInterfaceImplementationToOverride(FBPInterfaceDescription con
  * @param  InOldVarName		The current name of the variable we want to replace
  * @param  InNewVarName		The name that we wish to change all references to
  */
-static void RenameVariableReferencesInGraph(UBlueprint* InBlueprint, UClass* InVariableClass, UEdGraph* InGraph, const FName& InOldVarName, const FName& InNewVarName)
+static bool RenameVariableReferencesInGraph(UBlueprint* InBlueprint, UClass* InVariableClass, UEdGraph* InGraph, const FName& InOldVarName, const FName& InNewVarName)
 {
+	bool bFoundReference = false;
+
 	for(UEdGraphNode* GraphNode : InGraph->Nodes)
 	{
 		// Allow node to handle variable renaming
 		if (UK2Node* const K2Node = Cast<UK2Node>(GraphNode))
 		{
+			bFoundReference |= K2Node->ReferencesVariable(InOldVarName, nullptr);
 			K2Node->HandleVariableRenamed(InBlueprint, InVariableClass, InGraph, InOldVarName, InNewVarName);
 		}
 	}
+
+	return bFoundReference;
 }
 
 FBlueprintEditorUtils::FOnRenameVariableReferences FBlueprintEditorUtils::OnRenameVariableReferencesEvent;
@@ -290,7 +295,10 @@ void FBlueprintEditorUtils::RenameVariableReferences(UBlueprint* Blueprint, UCla
 	// Update any graph nodes that reference the old variable name to instead reference the new name
 	for(UEdGraph* CurrentGraph : AllGraphs)
 	{
-		RenameVariableReferencesInGraph(Blueprint, VariableClass, CurrentGraph, OldVarName, NewVarName);
+		if (RenameVariableReferencesInGraph(Blueprint, VariableClass, CurrentGraph, OldVarName, NewVarName))
+		{
+			MarkBlueprintAsModified(Blueprint);
+		}
 	}
 
 	OnRenameVariableReferencesEvent.Broadcast(Blueprint, VariableClass, OldVarName, NewVarName);
@@ -5192,7 +5200,10 @@ void FBlueprintEditorUtils::RenameLocalVariable(UBlueprint* InBlueprint, const U
 			LocalVariable->FriendlyName = FName::NameToDisplayString( InNewName.ToString(), LocalVariable->VarType.PinCategory == UEdGraphSchema_K2::PC_Boolean );
 
 			// Update any existing references to the old name
-			RenameVariableReferencesInGraph(InBlueprint, InBlueprint->GeneratedClass, FindScopeGraph(InBlueprint, InScope), InOldName, InNewName);
+			if (RenameVariableReferencesInGraph(InBlueprint, InBlueprint->GeneratedClass, FindScopeGraph(InBlueprint, InScope), InOldName, InNewName))
+			{
+				MarkBlueprintAsModified(InBlueprint);
+			}
 
 			// Validate child blueprints and adjust variable names to avoid a potential name collision
 			FBlueprintEditorUtils::ValidateBlueprintChildVariables(InBlueprint, InNewName);
