Fix crash reinstancing of blueprints before GEngine has been created
#rnx

diff --git a/Engine/Source/Editor/UnrealEd/Private/Kismet2/KismetReinstanceUtilities.cpp b/Engine/Source/Editor/UnrealEd/Private/Kismet2/KismetReinstanceUtilities.cpp
--- a/Engine/Source/Editor/UnrealEd/Private/Kismet2/KismetReinstanceUtilities.cpp
+++ b/Engine/Source/Editor/UnrealEd/Private/Kismet2/KismetReinstanceUtilities.cpp
@@ -1832,14 +1832,14 @@ void FBlueprintCompileReinstancer::ReplaceInstancesOfClass_Inner(TMap<UClass*, U
 		}
 
 
-		FDelegateHandle OnLevelActorDeletedHandle = GEngine->OnLevelActorDeleted().AddLambda([&OldToNewInstanceMap](AActor* DestroyedActor)
+		FDelegateHandle OnLevelActorDeletedHandle = GEngine ? GEngine->OnLevelActorDeleted().AddLambda([&OldToNewInstanceMap](AActor* DestroyedActor)
 			{
 				if (UObject** ReplacementObject = OldToNewInstanceMap.Find(DestroyedActor))
 				{
 					AActor* ReplacementActor = CastChecked<AActor>(*ReplacementObject);
 					ReplacementActor->GetWorld()->EditorDestroyActor(ReplacementActor, /*bShouldModifyLevel =*/true);
 				}
-			});
+			}) : FDelegateHandle();
 
 		// WARNING: for (TPair<UClass*, UClass*> OldToNewClass : InOldToNewClassMap) duplicated above 
 		// this loop only handles actors - which need to be reconstructed *after* their owned components 
@@ -1912,7 +1912,10 @@ void FBlueprintCompileReinstancer::ReplaceInstancesOfClass_Inner(TMap<UClass*, U
 				}
 			}
 		}
-		GEngine->OnLevelActorDeleted().Remove(OnLevelActorDeletedHandle);
+		if (GEngine)
+		{
+			GEngine->OnLevelActorDeleted().Remove(OnLevelActorDeletedHandle);
+		}
 
 		for (TPair<UObject*, UObject*> ReinstancedPair : OldToNewInstanceMap)
 		{
