Copying //UE4/Dev-Editor to Dev-Main (//UE4/Dev-Main) for 4.24

diff --git a/Engine/Source/Editor/UnrealEd/Private/Fbx/FbxSkeletalMeshImport.cpp b/Engine/Source/Editor/UnrealEd/Private/Fbx/FbxSkeletalMeshImport.cpp
--- a/Engine/Source/Editor/UnrealEd/Private/Fbx/FbxSkeletalMeshImport.cpp
+++ b/Engine/Source/Editor/UnrealEd/Private/Fbx/FbxSkeletalMeshImport.cpp
@@ -27,6 +27,9 @@
 #include "Animation/MorphTarget.h"
 #include "PhysicsAssetUtils.h"
 #include "Rendering/SkeletalMeshModel.h"
+#include "Rendering/SkeletalMeshRenderData.h"
+#include "Rendering/SkeletalMeshLODRenderData.h"
+#include "Misc/ScopedSlowTask.h"
 
 #include "SkelImport.h"
 #include "Logging/TokenizedMessage.h"
@@ -54,6 +57,7 @@
 #include "Engine/SkeletalMeshSocket.h"
 #include "Assets/ClothingAsset.h"
 #include "LODUtilities.h"
+#include "IMeshBuilderModule.h"
 
 #define LOCTEXT_NAMESPACE "FBXImpoter"
 
@@ -1433,7 +1437,7 @@ bool UnFbx::FFbxImporter::FillSkeletalMeshImportData(TArray<FbxNode*>& NodeArray
 
 	if (ImportOptions->bPreserveSmoothingGroups)
 	{
-		bool bDuplicateUnSmoothWedges = (ImportOptions->NormalGenerationMethod != EFBXNormalGenerationMethod::MikkTSpace);
+		bool bDuplicateUnSmoothWedges = false; //We deprecate legacy build
 		DoUnSmoothVerts(*SkelMeshImportDataPtr, bDuplicateUnSmoothWedges);
 	}
 	else
@@ -1754,28 +1758,40 @@ USkeletalMesh* UnFbx::FFbxImporter::ImportSkeletalMesh(FImportSkeletalMeshArgs &
 		return nullptr;
 	}
 
+	//Stack the PostEditChange call, it will call post edit change when it will go out of scope
+	FScopedSkeletalMeshPostEditChange ScopedPostEditChange(ExistingSkelMesh);
+
+	ESkeletalMeshGeoImportVersions GeoImportVersion = ESkeletalMeshGeoImportVersions::LatestVersion;
+	ESkeletalMeshSkinningImportVersions SkinningImportVersion = ESkeletalMeshSkinningImportVersions::LatestVersion;
 	//Adjust the import data from the import options
 	if (ExistingSkelMesh != nullptr)
 	{
+		int32 SafeLODIndex = ImportSkeletalMeshArgs.LodIndex < 0 ? 0 : ImportSkeletalMeshArgs.LodIndex;
+		if (ExistingSkelMesh->GetImportedModel() && ExistingSkelMesh->GetImportedModel()->LODModels.IsValidIndex(SafeLODIndex))
+		{
+			const FRawSkeletalMeshBulkData& RawSkeletalMeshBulkData = ExistingSkelMesh->GetImportedModel()->LODModels[SafeLODIndex].RawSkeletalMeshBulkData;
+			GeoImportVersion = RawSkeletalMeshBulkData.GeoImportVersion;
+			SkinningImportVersion = RawSkeletalMeshBulkData.SkinningImportVersion;
+		}
+
 		if ((ImportOptions->bImportAsSkeletalSkinning || ImportOptions->bImportAsSkeletalGeometry))
 		{
 			if (ImportOptions->bImportAsSkeletalSkinning)
 			{
 				//Replace geometry import data by original existing skel mesh geometry data
-				ReplaceSkeletalMeshGeometryImportData(ExistingSkelMesh, SkelMeshImportDataPtr, ImportSkeletalMeshArgs.LodIndex);
+				ReplaceSkeletalMeshGeometryImportData(ExistingSkelMesh, SkelMeshImportDataPtr, SafeLODIndex);
 			}
 			else if (ImportOptions->bImportAsSkeletalGeometry)
 			{
 				//Replace skinning import data by original existing skel mesh skinning data
-				ReplaceSkeletalMeshSkinningImportData(ExistingSkelMesh, SkelMeshImportDataPtr, ImportSkeletalMeshArgs.LodIndex);
+				ReplaceSkeletalMeshSkinningImportData(ExistingSkelMesh, SkelMeshImportDataPtr, SafeLODIndex);
 			}
 		}
 
 		//Reuse the vertex color in case we are not re-importing weights only and there is some vertexcolor in the existing skeleMesh
 		if (ExistingSkelMesh->bHasVertexColors && ImportOptions->VertexColorImportOption == EVertexColorImportOption::Ignore)
 		{
-			int32 LODIndex = ImportSkeletalMeshArgs.LodIndex < 0 ? 0 : ImportSkeletalMeshArgs.LodIndex;
-			RemapSkeletalMeshVertexColorToImportData(ExistingSkelMesh, LODIndex, SkelMeshImportDataPtr);
+			RemapSkeletalMeshVertexColorToImportData(ExistingSkelMesh, SafeLODIndex, SkelMeshImportDataPtr);
 		}
 	}
 
@@ -1796,18 +1812,8 @@ USkeletalMesh* UnFbx::FFbxImporter::ImportSkeletalMesh(FImportSkeletalMeshArgs &
 		return nullptr;
 	}
 
-	// Array of re-import contexts for components using this mesh
-	// Will unregister before import, then re-register afterwards
-	TIndirectArray<FComponentReregisterContext> ComponentContexts;
-	TUniquePtr< FSkinnedMeshComponentRecreateRenderStateContext > SkinnedMeshComponentRecreateRenderStateContext;
-
-	FSkinnedMeshComponentRecreateRenderStateContext* RecreateExistingRenderStateContext = ExistingSkelMesh ? new FSkinnedMeshComponentRecreateRenderStateContext(ExistingSkelMesh, false) : nullptr;
-	//Backup the data before importing the new one
 	if (ExistingSkelMesh)
 	{
-		//Release the Renderdata resources before reimporting the skeletal mesh
-		SkinnedMeshComponentRecreateRenderStateContext = MakeUnique<FSkinnedMeshComponentRecreateRenderStateContext>(ExistingSkelMesh);
-
 		ExistingSkelMesh->PreEditChange(NULL);
 		//The backup of the skeletal mesh data empty the LOD array in the ImportedResource of the skeletal mesh
 		//If the import fail after this step the editor can crash when updating the bone later since the LODModel will not exist anymore
@@ -1816,12 +1822,18 @@ USkeletalMesh* UnFbx::FFbxImporter::ImportSkeletalMesh(FImportSkeletalMeshArgs &
 
 	if (SkeletalMesh == nullptr)
 	{
+		//Backup the PostEditchangeStackCounter before overriding the skeletalmesh
+		int32 PostEditChangeStackCounter = ExistingSkelMesh != nullptr ? ExistingSkelMesh->GetPostEditChangeStackCounter() : 0;
 		// Create the new mesh after saving the old data, since it will replace the old skeletalmesh
 		// This should happen only when doing a re-import. Otherwise the skeletal mesh is create before.
 		SkeletalMesh = NewObject<USkeletalMesh>(ImportSkeletalMeshArgs.InParent, ImportSkeletalMeshArgs.Name, ImportSkeletalMeshArgs.Flags);
+		SkeletalMesh->SetPostEditChangeStackCounter(PostEditChangeStackCounter);
 	}
 
+
 	SkeletalMesh->PreEditChange(NULL);
+	//Dirty the DDC Key for any imported Skeletal Mesh
+	SkeletalMesh->InvalidateDeriveDataCacheGUID();
 
 	FSkeletalMeshModel *ImportedResource = SkeletalMesh->GetImportedModel();
 	check(ImportedResource->LODModels.Num() == 0);
@@ -1849,9 +1861,24 @@ USkeletalMesh* UnFbx::FFbxImporter::ImportSkeletalMesh(FImportSkeletalMeshArgs &
 	// process bone influences from import data
 	ProcessImportMeshInfluences(*SkelMeshImportDataPtr);
 
-	//Store the original fbx import data
+	//Store the original fbx import data the SkelMeshImportDataPtr should not be modified after this
 	LODModel.RawSkeletalMeshBulkData.SaveRawMesh(*SkelMeshImportDataPtr);
-
+	if (ImportOptions->bImportAsSkeletalSkinning)
+	{
+		LODModel.RawSkeletalMeshBulkData.SkinningImportVersion = ESkeletalMeshSkinningImportVersions::LatestVersion;
+		LODModel.RawSkeletalMeshBulkData.GeoImportVersion = GeoImportVersion;
+	}
+	else if(ImportOptions->bImportAsSkeletalGeometry)
+	{
+		LODModel.RawSkeletalMeshBulkData.GeoImportVersion = ESkeletalMeshGeoImportVersions::LatestVersion;
+		LODModel.RawSkeletalMeshBulkData.SkinningImportVersion = SkinningImportVersion;
+	}
+	else
+	{
+		//We reimport both
+		LODModel.RawSkeletalMeshBulkData.SkinningImportVersion = ESkeletalMeshSkinningImportVersions::LatestVersion;
+		LODModel.RawSkeletalMeshBulkData.GeoImportVersion = ESkeletalMeshGeoImportVersions::LatestVersion;
+	}
 	SkeletalMesh->ResetLODInfo();
 	FSkeletalMeshLODInfo& NewLODInfo = SkeletalMesh->AddLODInfo();
 	NewLODInfo.ReductionSettings.NumOfTrianglesPercentage = 1.0f;
@@ -1877,38 +1904,61 @@ USkeletalMesh* UnFbx::FFbxImporter::ImportSkeletalMesh(FImportSkeletalMeshArgs &
 		TArray<int32> LODPointToRawMap;
 		SkelMeshImportDataPtr->CopyLODImportData(LODPoints,LODWedges,LODFaces,LODInfluences,LODPointToRawMap);
 
-		IMeshUtilities::MeshBuildOptions BuildOptions;
-		BuildOptions.OverlappingThresholds = ImportOptions->OverlappingThresholds;
-		BuildOptions.bComputeNormals = !ImportOptions->ShouldImportNormals() || !SkelMeshImportDataPtr->bHasNormals;
-		BuildOptions.bComputeTangents = !ImportOptions->ShouldImportTangents() || !SkelMeshImportDataPtr->bHasTangents;
+		FSkeletalMeshBuildSettings BuildOptions;
+		BuildOptions.bBuildAdjacencyBuffer = true;
+		BuildOptions.bRecomputeNormals = !ImportOptions->ShouldImportNormals() || !SkelMeshImportDataPtr->bHasNormals;
+		BuildOptions.bRecomputeTangents = !ImportOptions->ShouldImportTangents() || !SkelMeshImportDataPtr->bHasTangents;
 		BuildOptions.bUseMikkTSpace = (ImportOptions->NormalGenerationMethod == EFBXNormalGenerationMethod::MikkTSpace) && (!ImportOptions->ShouldImportNormals() || !ImportOptions->ShouldImportTangents());
 		BuildOptions.bComputeWeightedNormals = ImportOptions->bComputeWeightedNormals;
-		BuildOptions.bRemoveDegenerateTriangles = false;
-
-		IMeshUtilities& MeshUtilities = FModuleManager::Get().LoadModuleChecked<IMeshUtilities>("MeshUtilities");
-		
-		TArray<FText> WarningMessages;
-		TArray<FName> WarningNames;
-		// Create actual rendering data.
-		bool bBuildSuccess = MeshUtilities.BuildSkeletalMesh(ImportedResource->LODModels[0],SkeletalMesh->RefSkeleton,LODInfluences, LODWedges,LODFaces,LODPoints,LODPointToRawMap,BuildOptions, &WarningMessages, &WarningNames);
+		BuildOptions.bRemoveDegenerates = ImportOptions->bRemoveDegenerates;
+		BuildOptions.ThresholdPosition = ImportOptions->OverlappingThresholds.ThresholdPosition;
+		BuildOptions.ThresholdTangentNormal = ImportOptions->OverlappingThresholds.ThresholdTangentNormal;
+		BuildOptions.ThresholdUV = ImportOptions->OverlappingThresholds.ThresholdUV;
 
-		// temporary hack of message/names, should be one token or a struct
-		if(WarningMessages.Num() > 0 && WarningNames.Num() == WarningMessages.Num())
+		bool bUseLegacyMeshUtilities = false;
+		bool bBuildSuccess = false;
+		if (bUseLegacyMeshUtilities)
 		{
-			EMessageSeverity::Type MessageSeverity = bBuildSuccess ? EMessageSeverity::Warning : EMessageSeverity::Error;
+			IMeshUtilities::MeshBuildOptions LegacyBuildOptions;
+			LegacyBuildOptions.FillOptions(BuildOptions);
 
-			for(int32 MessageIdx = 0; MessageIdx<WarningMessages.Num(); ++MessageIdx)	
+			IMeshUtilities& MeshUtilities = FModuleManager::Get().LoadModuleChecked<IMeshUtilities>("MeshUtilities");
+
+			TArray<FText> WarningMessages;
+			TArray<FName> WarningNames;
+			// Create actual rendering data.
+			bBuildSuccess = MeshUtilities.BuildSkeletalMesh(ImportedResource->LODModels[0], SkeletalMesh->RefSkeleton, LODInfluences, LODWedges, LODFaces, LODPoints, LODPointToRawMap, LegacyBuildOptions, &WarningMessages, &WarningNames);
+
+			// temporary hack of message/names, should be one token or a struct
+			if (WarningMessages.Num() > 0 && WarningNames.Num() == WarningMessages.Num())
 			{
-				AddTokenizedErrorMessage(FTokenizedMessage::Create(MessageSeverity, WarningMessages[MessageIdx]), WarningNames[MessageIdx]);
+				EMessageSeverity::Type MessageSeverity = bBuildSuccess ? EMessageSeverity::Warning : EMessageSeverity::Error;
+
+				for (int32 MessageIdx = 0; MessageIdx < WarningMessages.Num(); ++MessageIdx)
+				{
+					AddTokenizedErrorMessage(FTokenizedMessage::Create(MessageSeverity, WarningMessages[MessageIdx]), WarningNames[MessageIdx]);
+				}
 			}
 		}
+		else
+		{
+			//The imported LOD is always 0 here, the LOD custom import will import the LOD alone(in a temporary skeletalmesh) and add it to the base skeletal mesh later
+			const int32 LODIndex = 0;
+			check(SkeletalMesh->GetLODInfo(LODIndex) != nullptr);
+			//Set the build options
+			SkeletalMesh->GetLODInfo(LODIndex)->BuildSettings = BuildOptions;
+			//New MeshDescription build process
+			IMeshBuilderModule& MeshBuilderModule = FModuleManager::Get().LoadModuleChecked<IMeshBuilderModule>("MeshBuilder");
+			//We must build the LODModel so we can restore properly the mesh, but we do not have to regenerate LODs
+			bBuildSuccess = MeshBuilderModule.BuildSkeletalMesh(SkeletalMesh, LODIndex, false);
+		}
 
 		if( !bBuildSuccess )
 		{
 			SkeletalMesh->MarkPendingKill();
 			return NULL;
 		}
-
+		
 		// Presize the per-section shadow casting array with the number of sections in the imported LOD.
 		const int32 NumSections = LODModel.Sections.Num();
 
@@ -1966,18 +2016,12 @@ USkeletalMesh* UnFbx::FFbxImporter::ImportSkeletalMesh(FImportSkeletalMeshArgs &
 		}
 
 		SkeletalMesh->CalculateInvRefMatrices();
-		SkeletalMesh->PostEditChange();
-		SkeletalMesh->MarkPackageDirty();
-
-		// Now iterate over all skeletal mesh components re-initialising them.
-		for (TObjectIterator<USkinnedMeshComponent> It; It; ++It)
+		if ((!SkeletalMesh->GetResourceForRendering() || !SkeletalMesh->GetResourceForRendering()->LODRenderData.IsValidIndex(0)) && ImportOptions->bCreatePhysicsAsset)
 		{
-			USkinnedMeshComponent* SkinComp = *It;
-			if (SkinComp->SkeletalMesh == SkeletalMesh)
-			{
-				ComponentContexts.Add(new FComponentReregisterContext(SkinComp));
-			}
+			//We need to have a valid render data to create physic asset
+			SkeletalMesh->Build();
 		}
+		SkeletalMesh->MarkPackageDirty();
 	}
 
 	if(ImportSkeletalMeshArgs.LodIndex == 0)
@@ -2133,13 +2177,6 @@ USkeletalMesh* UnFbx::FFbxImporter::ImportSkeletalMesh(FImportSkeletalMeshArgs &
 		ImportNodeCustomProperties(SkeletalMesh->Skeleton, SkeletonNode, true);
 	}
 
-	// ComponentContexts will now go out of scope, causing components to be re-registered
-
-	if (RecreateExistingRenderStateContext)
-	{
-		delete RecreateExistingRenderStateContext;
-	}
-
 	return SkeletalMesh;
 }
 
@@ -2343,8 +2380,6 @@ USkeletalMesh* UnFbx::FFbxImporter::ReimportSkeletalMesh(USkeletalMesh* Mesh, UF
 		return NULL;
 	}
 
-	char MeshName[1024];
-	FCStringAnsi::Strcpy(MeshName, TCHAR_TO_UTF8(*Mesh->GetName()));
 	TArray<FbxNode*>* FbxNodes = NULL;
 	USkeletalMesh* NewMesh = NULL;
 
@@ -2353,16 +2388,35 @@ USkeletalMesh* UnFbx::FFbxImporter::ReimportSkeletalMesh(USkeletalMesh* Mesh, UF
 	bool Old_ImportTextures = ImportOptions->bImportTextures;
 	bool Old_ImportAnimations = ImportOptions->bImportAnimations;
 
+	auto CleanUpImportOptionsData = [&]()
+	{
+		ImportOptions->bImportRigidMesh = Old_ImportRigidMesh;
+		ImportOptions->bImportMaterials = Old_ImportMaterials;
+		ImportOptions->bImportTextures = Old_ImportTextures;
+		ImportOptions->bImportAnimations = Old_ImportAnimations;
+	};
+
 	// support to update rigid animation mesh
 	ImportOptions->bImportRigidMesh = true;
 
-	// Unbind any existing clothing assets before we reimport the geometry
-	TArray<ClothingAssetUtils::FClothingAssetMeshBinding> ClothingBindings;
-	ClothingAssetUtils::GetMeshClothingAssetBindings(Mesh, ClothingBindings);
-
-	for(ClothingAssetUtils::FClothingAssetMeshBinding& Binding : ClothingBindings)
+	UFbxSkeletalMeshImportData* SKImportData = Cast<UFbxSkeletalMeshImportData>(Mesh->AssetImportData);
+	if (SKImportData)
 	{
-		Binding.Asset->UnbindFromSkeletalMesh(Mesh, Binding.LODIndex);
+		FSkeletalMeshLODInfo* LODInfo = Mesh->GetLODInfo(0);
+		if(LODInfo && Mesh->GetImportedModel() && Mesh->GetImportedModel()->LODModels.IsValidIndex(0))
+		{
+			const FSkeletalMeshLODModel& LODModel = Mesh->GetImportedModel()->LODModels[0];
+			if (!LODModel.RawSkeletalMeshBulkData.IsBuildDataAvailable())
+			{
+				//Set the build settings
+				LODInfo->BuildSettings.bBuildAdjacencyBuffer = true;
+				LODInfo->BuildSettings.bComputeWeightedNormals = SKImportData->bComputeWeightedNormals;
+				LODInfo->BuildSettings.bRecomputeNormals = SKImportData->NormalImportMethod == EFBXNormalImportMethod::FBXNIM_ComputeNormals;
+				LODInfo->BuildSettings.bRecomputeTangents = SKImportData->NormalImportMethod != EFBXNormalImportMethod::FBXNIM_ImportNormalsAndTangents;
+				LODInfo->BuildSettings.bRemoveDegenerates = true;
+				LODInfo->BuildSettings.bUseMikkTSpace = SKImportData->NormalGenerationMethod == EFBXNormalGenerationMethod::MikkTSpace;
+			}
+		}
 	}
 
 	// get meshes in Fbx file
@@ -2418,26 +2472,53 @@ USkeletalMesh* UnFbx::FFbxImporter::ReimportSkeletalMesh(USkeletalMesh* Mesh, UF
 			if (FbxNodes != nullptr)
 				break;
 		}
-		if (FbxNodes == nullptr)
-			return nullptr;
 	}
 	else
 	{
 		// if there is only one mesh, use it without name checking 
 		// (because the "Used As Full Name" option enables users name the Unreal mesh by themselves
-		if (FbxSkelMeshArray.Num() > 0)
+		if (FbxSkelMeshArray.Num() == 1)
 		{
 			FbxNodes = FbxSkelMeshArray[0];
 		}
-		else
+		else if (FbxSkelMeshArray.Num() > 1)
 		{
-			// @todo - FBX Importing - We need proper support for reimport if the file contains more than one skeletal mesh. 
-
+			//Search a set corresponding
+			for (TArray<FbxNode*>*SkeletalMeshNodes : FbxSkelMeshArray)
+			{
+				FbxNode* Node = GetMeshNodesFromName(Mesh->GetName(), *SkeletalMeshNodes);
+				if (Node != nullptr)
+				{
+					FbxNodes = SkeletalMeshNodes;
+					break;
+				}
+			}
+			if (FbxNodes == nullptr)
+			{
+				//Re import the first skeletalmesh found in the fbx file
+				FbxNodes = FbxSkelMeshArray[0];
+			}
 		}
 	}
-
+	
 	if (FbxNodes)
 	{
+		//Posteditchange will be call at the end of the reimport which will reallocate the render ressource
+		FScopedSkeletalMeshPostEditChange ScopedPostEditChange(Mesh);
+
+		FScopedSlowTask SlowTask(100.0f, NSLOCTEXT("SkeltalMeshReimportFactory", "SkeletalMeshReimportTask", "Reimporting skeletal mesh..."));
+		SlowTask.MakeDialog();
+
+		//The fbx file is valid we can reimport the skeletal mesh
+		Mesh->PreEditChange(nullptr);
+
+		// Unbind any existing clothing assets before we reimport the geometry
+		TArray<ClothingAssetUtils::FClothingAssetMeshBinding> ClothingBindings;
+		FLODUtilities::UnbindClothingAndBackup(Mesh, ClothingBindings);
+
+		//Reimport must force a new guid
+		Mesh->InvalidateDeriveDataCacheGUID();
+
 		for (int32 LODCounter = 1; LODCounter < Mesh->GetLODNum(); ++LODCounter)
 		{
 			//Reset all import with base mesh flags
@@ -2482,8 +2563,10 @@ USkeletalMesh* UnFbx::FFbxImporter::ReimportSkeletalMesh(USkeletalMesh* Mesh, UF
 		int32 LODIndex;
 		int32 SuccessfulLodIndex = 0;
 		bool bOperationCanceled = false;
+		float ProgressStep = 90.0 / MaxLODLevel;
 		for (LODIndex = 0; LODIndex < MaxLODLevel; LODIndex++)
 		{
+			SlowTask.EnterProgressFrame(ProgressStep);
 			int32 ImportedSuccessfulLodIndex = INDEX_NONE;
 			TArray<FbxNode*> SkelMeshNodeArray;
 			for (int32 j = 0; j < (*FbxNodes).Num(); j++)
@@ -2533,8 +2616,9 @@ USkeletalMesh* UnFbx::FFbxImporter::ReimportSkeletalMesh(USkeletalMesh* Mesh, UF
 				NewMesh = ImportSkeletalMesh( ImportSkeletalMeshArgs );
 				if (bOperationCanceled)
 				{
-					// User cancelled, clean up and return
-					return nullptr;
+					// User cancel, clean up and return
+					NewMesh = nullptr;
+					break;
 				}
 				if (NewMesh)
 				{
@@ -2581,7 +2665,7 @@ USkeletalMesh* UnFbx::FFbxImporter::ReimportSkeletalMesh(USkeletalMesh* Mesh, UF
 				ImportSkeletalMeshArgs.OutData = &OutData;
 
 				UObject *LODObject = ImportSkeletalMesh( ImportSkeletalMeshArgs );
-				bool bImportSucceeded = !bOperationCanceled && ImportSkeletalMeshLOD( Cast<USkeletalMesh>(LODObject), BaseSkeletalMesh, SuccessfulLodIndex, false, nullptr, TemplateImportData);
+				bool bImportSucceeded = !bOperationCanceled && ImportSkeletalMeshLOD( Cast<USkeletalMesh>(LODObject), BaseSkeletalMesh, SuccessfulLodIndex, TemplateImportData);
 
 				if (bImportSucceeded)
 				{
@@ -2617,47 +2701,41 @@ USkeletalMesh* UnFbx::FFbxImporter::ReimportSkeletalMesh(USkeletalMesh* Mesh, UF
 
 			if (NewMesh)
 			{
-				// import morph target
 				if ((ImportOptions->bImportSkeletalMeshLODs || LODIndex == 0) &&
 					ImportOptions->bImportMorph &&
 					ImportedSuccessfulLodIndex != INDEX_NONE &&
 					NewMesh->GetImportedModel() &&
-					NewMesh->GetImportedModel()->LODModels.IsValidIndex(LODIndex))
+					NewMesh->GetImportedModel()->LODModels.IsValidIndex(ImportedSuccessfulLodIndex))
 				{
-					// @fixme: @question : where do they import this morph? where to? What morph target sets?
 					ImportFbxMorphTarget(SkelMeshNodeArray, NewMesh, ImportedSuccessfulLodIndex, OutData);
 				}
 			}
 		}
-		if (NewMesh)
+
+		auto RebindClothing = [&ClothingBindings](USkeletalMesh* SkelMesh)
 		{
-			FSkeletalMeshModel* NewMeshResource = NewMesh->GetImportedModel();
-			if(NewMeshResource && ClothingBindings.Num() > 0)
+			FSkeletalMeshModel* MeshResource = SkelMesh->GetImportedModel();
+			if (MeshResource && ClothingBindings.Num() > 0)
 			{
-				NewMesh->PreEditChange(nullptr);
+				FLODUtilities::RestoreClothingFromBackup(SkelMesh, ClothingBindings);
+			}
+		};
 
-				for(ClothingAssetUtils::FClothingAssetMeshBinding& Binding : ClothingBindings)
-				{
-					if(NewMeshResource->LODModels.IsValidIndex(Binding.LODIndex) &&
-						NewMeshResource->LODModels[Binding.LODIndex].Sections.IsValidIndex(Binding.SectionIndex))
-					{
-						Binding.Asset->BindToSkeletalMesh(NewMesh, Binding.LODIndex, Binding.SectionIndex, Binding.AssetInternalLodIndex);
-					}
-				}
 
-				NewMesh->PostEditChange();
-			}
+		if (!bOperationCanceled && NewMesh)
+		{
+			RebindClothing(NewMesh);
 
 			//Update the import data so we can re-import correctly
 			UpdateSkeletalMeshImportData(NewMesh, TemplateImportData, INDEX_NONE, &ImportMaterialOriginalNameData, &ImportMeshLodData);
-			//If we have import some morph target we have to rebuild the render resources since morph target are now using GPU
-			if (NewMesh->MorphTargets.Num() > 0)
-			{
-				NewMesh->ReleaseResources();
-				//Rebuild the resources with a post edit change since we have added some morph targets
-				NewMesh->PostEditChange();
-			}
 		}
+		else if(bOperationCanceled && Mesh != nullptr)
+		{
+			//Operation was cancel
+			//rebind the clothing on the original mesh
+			RebindClothing(Mesh);
+		}
+		SlowTask.EnterProgressFrame(10.0f);
 	}
 	else
 	{
@@ -2665,10 +2743,7 @@ USkeletalMesh* UnFbx::FFbxImporter::ReimportSkeletalMesh(USkeletalMesh* Mesh, UF
 		AddTokenizedErrorMessage(FTokenizedMessage::Create(EMessageSeverity::Warning, FText::Format(LOCTEXT("FbxSkeletaLMeshimport_NoFBXMeshMatch", "No FBX mesh matches the Unreal mesh '{0}'."), FText::FromString(Mesh->GetName()))), FFbxErrors::Generic_Mesh_MeshNotFound);
 	}
 
-	ImportOptions->bImportRigidMesh = Old_ImportRigidMesh;
-	ImportOptions->bImportMaterials = Old_ImportMaterials;
-	ImportOptions->bImportTextures = Old_ImportTextures;
-	ImportOptions->bImportAnimations = Old_ImportAnimations;
+	CleanUpImportOptionsData();
 
 	return NewMesh;
 }
@@ -3402,7 +3477,7 @@ bool UnFbx::FFbxImporter::FillSkelMeshImporterFromFbx( FSkeletalMeshImportData&
 			//
 			// normals, tangents and binormals
 			//
-			if( ImportOptions->ShouldImportNormals() && bHasNormalInformation )
+			if( bHasNormalInformation )
 			{
 				int32 TmpIndex = LocalIndex*3 + VertexIndex;
 				//normals may have different reference and mapping mode than tangents and binormals
@@ -3416,7 +3491,7 @@ bool UnFbx::FFbxImporter::FillSkelMeshImporterFromFbx( FSkeletalMeshImportData&
 
 				FbxVector4 TempValue;
 
-				if( ImportOptions->ShouldImportTangents() && bHasTangentInformation )
+				if( bHasTangentInformation )
 				{
 					TempValue = LayerElementTangent->GetDirectArray().GetAt(TangentMapIndex);
 					TempValue = TotalMatrixForNormal.MultT(TempValue);
@@ -3713,6 +3788,9 @@ bool UnFbx::FFbxImporter::FillSkelMeshImporterFromFbx( FSkeletalMeshImportData&
 
 void UnFbx::FFbxImporter::InsertNewLODToBaseSkeletalMesh(USkeletalMesh* InSkeletalMesh, USkeletalMesh* BaseSkeletalMesh, int32 DesiredLOD, UFbxSkeletalMeshImportData* TemplateImportData)
 {
+	//Scope a post edit change
+	FScopedSkeletalMeshPostEditChange ScopedPostEditChange(BaseSkeletalMesh);
+
 	FSkeletalMeshModel* ImportedResource = InSkeletalMesh->GetImportedModel();
 	FSkeletalMeshModel* DestImportedResource = BaseSkeletalMesh->GetImportedModel();
 
@@ -3772,6 +3850,13 @@ void UnFbx::FFbxImporter::InsertNewLODToBaseSkeletalMesh(USkeletalMesh* InSkelet
 
 	// Set up LODMaterialMap to number of materials in new mesh.
 	FSkeletalMeshLODInfo& LODInfo = *(BaseSkeletalMesh->GetLODInfo(DesiredLOD));
+
+	//Copy the build settings
+	if(InSkeletalMesh->GetLODInfo(0))
+	{
+		const FSkeletalMeshLODInfo& ImportedLODInfo = *(InSkeletalMesh->GetLODInfo(0));
+		LODInfo.BuildSettings = ImportedLODInfo.BuildSettings;
+	}
 	
 	LODInfo.LODMaterialMap.Empty();
 	// Now set up the material mapping array.
@@ -3865,12 +3950,9 @@ void UnFbx::FFbxImporter::InsertNewLODToBaseSkeletalMesh(USkeletalMesh* InSkelet
 			}
 		}
 	}
-
-	// rebuild vertex buffers and reinit RHI resources
-	BaseSkeletalMesh->PostEditChange();
 }
 
-bool UnFbx::FFbxImporter::ImportSkeletalMeshLOD(USkeletalMesh* InSkeletalMesh, USkeletalMesh* BaseSkeletalMesh, int32 DesiredLOD, bool bNeedToReregister, TArray<UActorComponent*>* ReregisterAssociatedComponents, UFbxSkeletalMeshImportData* TemplateImportData)
+bool UnFbx::FFbxImporter::ImportSkeletalMeshLOD(USkeletalMesh* InSkeletalMesh, USkeletalMesh* BaseSkeletalMesh, int32 DesiredLOD, UFbxSkeletalMeshImportData* TemplateImportData)
 {
 	check(InSkeletalMesh);
 	check(BaseSkeletalMesh);
@@ -3934,6 +4016,8 @@ bool UnFbx::FFbxImporter::ImportSkeletalMeshLOD(USkeletalMesh* InSkeletalMesh, U
 		}
 	}
 
+	FScopedSkeletalMeshPostEditChange ScopedPostEditChange(BaseSkeletalMesh);
+
 	//Skeleton do not match and we are importing skeletal geometry only so apply the base LOD Skinning
 	if (bApplyBaseSkinning)
 	{
@@ -3941,8 +4025,10 @@ bool UnFbx::FFbxImporter::ImportSkeletalMeshLOD(USkeletalMesh* InSkeletalMesh, U
 		//Apply the base LOD skinning to the new LOD geometry, DestImportedResource here is the base lod and is the source for the apply skinning.
 		//The Dest here mean we merge the just import LOD to the base LodModel.
 		SkeletalMeshHelper::ApplySkinning(InSkeletalMesh, DestImportedResource->LODModels[0], ImportedResource->LODModels[0]);
-		IMeshUtilities& MeshUtilities = FModuleManager::Get().LoadModuleChecked<IMeshUtilities>("MeshUtilities");
-		MeshUtilities.RemoveBonesFromMesh(BaseSkeletalMesh, DesiredLOD, NULL);
+
+		//Excluded bones will be remove when we rebuild the asset
+		//IMeshUtilities& MeshUtilities = FModuleManager::Get().LoadModuleChecked<IMeshUtilities>("MeshUtilities");
+		//MeshUtilities.RemoveBonesFromMesh(BaseSkeletalMesh, DesiredLOD, NULL);
 	}
 
 	FSkeletalMeshLODModel& NewLODModel = ImportedResource->LODModels[0];
@@ -4105,32 +4191,7 @@ bool UnFbx::FFbxImporter::ImportSkeletalMeshLOD(USkeletalMesh* InSkeletalMesh, U
 		}
 	}
 
-	if (bNeedToReregister)
-	{
-		// Shut down the skeletal mesh component that is previewing this mesh.
-		if (ReregisterAssociatedComponents)
-		{
-			FMultiComponentReregisterContext ReregisterContext(*ReregisterAssociatedComponents);
-			// wait until resources are released
-			FlushRenderingCommands();
-
-			InsertNewLODToBaseSkeletalMesh(InSkeletalMesh, BaseSkeletalMesh, DesiredLOD, TemplateImportData);
-
-			// ReregisterContexts go out of scope here, reregistering associated components with the scene.
-		}
-		else
-		{
-			TComponentReregisterContext<USkinnedMeshComponent> ComponentReregisterContext;
-
-			InsertNewLODToBaseSkeletalMesh(InSkeletalMesh, BaseSkeletalMesh, DesiredLOD, TemplateImportData);
-
-			// ReregisterContexts go out of scope here, reregistering skel components with the scene.
-		}
-	}
-	else
-	{
-		InsertNewLODToBaseSkeletalMesh(InSkeletalMesh, BaseSkeletalMesh, DesiredLOD, TemplateImportData);
-	}
+	InsertNewLODToBaseSkeletalMesh(InSkeletalMesh, BaseSkeletalMesh, DesiredLOD, TemplateImportData);
 
 	return true;
 }
@@ -4271,9 +4332,21 @@ void UnFbx::FFbxImporter::ImportMorphTargetsInternal( TArray<FbxNode*>& SkelMesh
 		check(BaseImportData.MorphTargetNames.Num() == BaseImportData.MorphTargets.Num() && BaseImportData.MorphTargetNames.Num() == BaseImportData.MorphTargetModifiedPoints.Num());
 	}
 
-	//Build MorphTargets
-	// Calculate overlapping corners and tangents
-	FLODUtilities::BuildMorphTargets(BaseSkelMesh, BaseImportData, LODIndex, ImportOptions->ShouldImportNormals(), ImportOptions->ShouldImportTangents(), (ImportOptions->NormalGenerationMethod == EFBXNormalGenerationMethod::MikkTSpace));
+	if (BaseSkelMesh->GetImportedModel() && BaseSkelMesh->GetImportedModel()->LODModels.IsValidIndex(LODIndex))
+	{
+		//If we can build the skeletal mesh there is no need to build the morph target now, all the necessary build morph target data was copied before.
+		if (!BaseSkelMesh->GetImportedModel()->LODModels[LODIndex].RawSkeletalMeshBulkData.IsBuildDataAvailable())
+		{
+			//Build MorphTargets
+			FLODUtilities::BuildMorphTargets(
+				BaseSkelMesh,
+				BaseImportData,
+				LODIndex,
+				ImportOptions->ShouldImportNormals(),
+				ImportOptions->ShouldImportTangents(),
+				(ImportOptions->NormalGenerationMethod == EFBXNormalGenerationMethod::MikkTSpace));
+		}
+	}
 
 	GWarn->EndSlowTask();
 }
@@ -4281,6 +4354,9 @@ void UnFbx::FFbxImporter::ImportMorphTargetsInternal( TArray<FbxNode*>& SkelMesh
 // Import Morph target
 void UnFbx::FFbxImporter::ImportFbxMorphTarget(TArray<FbxNode*> &SkelMeshNodeArray, USkeletalMesh* BaseSkelMesh, int32 LODIndex, FSkeletalMeshImportData &BaseSkeletalMeshImportData)
 {
+	//Stack the PostEditChange call, it will call post edit change when it will go out of scope
+	FScopedSkeletalMeshPostEditChange ScopedPostEditChange(BaseSkelMesh);
+
 	bool bHasMorph = false;
 	int32 NodeIndex;
 	// check if there are morph in this geometry
@@ -4299,7 +4375,12 @@ void UnFbx::FFbxImporter::ImportFbxMorphTarget(TArray<FbxNode*> &SkelMeshNodeArr
 	
 	if (bHasMorph)
 	{
-		ImportMorphTargetsInternal( SkelMeshNodeArray, BaseSkelMesh, LODIndex, BaseSkeletalMeshImportData );
+		ImportMorphTargetsInternal( SkelMeshNodeArray, BaseSkelMesh, LODIndex, BaseSkeletalMeshImportData);
+		//Save the rawMesh
+		if (BaseSkelMesh->GetImportedModel() && BaseSkelMesh->GetImportedModel()->LODModels.IsValidIndex(LODIndex))
+		{
+			BaseSkelMesh->GetImportedModel()->LODModels[LODIndex].RawSkeletalMeshBulkData.SaveRawMesh(BaseSkeletalMeshImportData);
+		}
 	}
 }
 
diff --git a/Engine/Source/Editor/UnrealEd/Private/Fbx/FbxStaticMeshImport.cpp b/Engine/Source/Editor/UnrealEd/Private/Fbx/FbxStaticMeshImport.cpp
--- a/Engine/Source/Editor/UnrealEd/Private/Fbx/FbxStaticMeshImport.cpp
+++ b/Engine/Source/Editor/UnrealEd/Private/Fbx/FbxStaticMeshImport.cpp
@@ -32,8 +32,7 @@
 #include "Interfaces/ITargetPlatformManagerModule.h"
 #include "Misc/FbxErrors.h"
 #include "PhysicsEngine/BodySetup.h"
-#include "MeshDescription.h"
-#include "MeshAttributes.h"
+#include "StaticMeshAttributes.h"
 #include "MeshDescriptionOperations.h"
 #include "IMeshBuilderModule.h"
 #include "Settings/EditorExperimentalSettings.h"
@@ -284,6 +283,8 @@ bool UnFbx::FFbxImporter::BuildStaticMeshFromGeometry(FbxNode* Node, UStaticMesh
 	FMeshDescription* MeshDescription = StaticMesh->GetMeshDescription(LODIndex);
 	//The mesh description should have been created before calling BuildStaticMeshFromGeometry
 	check(MeshDescription);
+	FStaticMeshAttributes Attributes(*MeshDescription);
+
 	//remove the bad polygons before getting any data from mesh
 	Mesh->RemoveBadPolygons();
 
@@ -553,15 +554,15 @@ bool UnFbx::FFbxImporter::BuildStaticMeshFromGeometry(FbxNode* Node, UStaticMesh
 	bool OddNegativeScale = IsOddNegativeScale(TotalMatrix);
 	bool bHasNonDegeneratePolygons = false;
 
-	TVertexAttributesRef<FVector> VertexPositions = MeshDescription->VertexAttributes().GetAttributesRef<FVector>(MeshAttribute::Vertex::Position);
-	TVertexInstanceAttributesRef<FVector> VertexInstanceNormals = MeshDescription->VertexInstanceAttributes().GetAttributesRef<FVector>(MeshAttribute::VertexInstance::Normal);
-	TVertexInstanceAttributesRef<FVector> VertexInstanceTangents = MeshDescription->VertexInstanceAttributes().GetAttributesRef<FVector>(MeshAttribute::VertexInstance::Tangent);
-	TVertexInstanceAttributesRef<float> VertexInstanceBinormalSigns = MeshDescription->VertexInstanceAttributes().GetAttributesRef<float>(MeshAttribute::VertexInstance::BinormalSign);
-	TVertexInstanceAttributesRef<FVector4> VertexInstanceColors = MeshDescription->VertexInstanceAttributes().GetAttributesRef<FVector4>(MeshAttribute::VertexInstance::Color);
-	TVertexInstanceAttributesRef<FVector2D> VertexInstanceUVs = MeshDescription->VertexInstanceAttributes().GetAttributesRef<FVector2D>(MeshAttribute::VertexInstance::TextureCoordinate);
-	TEdgeAttributesRef<bool> EdgeHardnesses = MeshDescription->EdgeAttributes().GetAttributesRef<bool>(MeshAttribute::Edge::IsHard);
-	TEdgeAttributesRef<float> EdgeCreaseSharpnesses = MeshDescription->EdgeAttributes().GetAttributesRef<float>(MeshAttribute::Edge::CreaseSharpness);
-	TPolygonGroupAttributesRef<FName> PolygonGroupImportedMaterialSlotNames = MeshDescription->PolygonGroupAttributes().GetAttributesRef<FName>(MeshAttribute::PolygonGroup::ImportedMaterialSlotName);
+	TVertexAttributesRef<FVector> VertexPositions = Attributes.GetVertexPositions();
+	TVertexInstanceAttributesRef<FVector> VertexInstanceNormals = Attributes.GetVertexInstanceNormals();
+	TVertexInstanceAttributesRef<FVector> VertexInstanceTangents = Attributes.GetVertexInstanceTangents();
+	TVertexInstanceAttributesRef<float> VertexInstanceBinormalSigns = Attributes.GetVertexInstanceBinormalSigns();
+	TVertexInstanceAttributesRef<FVector4> VertexInstanceColors = Attributes.GetVertexInstanceColors();
+	TVertexInstanceAttributesRef<FVector2D> VertexInstanceUVs = Attributes.GetVertexInstanceUVs();
+	TEdgeAttributesRef<bool> EdgeHardnesses = Attributes.GetEdgeHardnesses();
+	TEdgeAttributesRef<float> EdgeCreaseSharpnesses = Attributes.GetEdgeCreaseSharpnesses();
+	TPolygonGroupAttributesRef<FName> PolygonGroupImportedMaterialSlotNames = Attributes.GetPolygonGroupMaterialSlotNames();
 
 	int32 VertexOffset = MeshDescription->Vertices().Num();
 	int32 VertexInstanceOffset = MeshDescription->VertexInstances().Num();
@@ -1071,51 +1072,6 @@ void UnFbx::FFbxImporter::AddStaticMeshSourceModelGeneratedLOD(UStaticMesh* Stat
 	}
 }
 
-FbxNode* UnFbx::FFbxImporter::GetMeshNodesFromName(UStaticMesh* StaticMesh, TArray<FbxNode*>& FbxMeshArray)
-{
-	char MeshName[1024];
-	FCStringAnsi::Strcpy(MeshName, 1024, TCHAR_TO_UTF8(*StaticMesh->GetName()));
-	// find the Fbx mesh node that the Unreal Mesh matches according to name
-	int32 MeshIndex;
-	for (MeshIndex = 0; MeshIndex < FbxMeshArray.Num(); MeshIndex++)
-	{
-		const char* FbxMeshName = FbxMeshArray[MeshIndex]->GetName();
-		// The name of Unreal mesh may have a prefix, so we match from end
-		int32 i = 0;
-		char* MeshPtr = MeshName + FCStringAnsi::Strlen(MeshName) - 1;
-		if (FCStringAnsi::Strlen(FbxMeshName) <= FCStringAnsi::Strlen(MeshName))
-		{
-			const char* FbxMeshPtr = FbxMeshName + FCStringAnsi::Strlen(FbxMeshName) - 1;
-			while (i < FCStringAnsi::Strlen(FbxMeshName))
-			{
-				bool bIsPointAndUnderscore = *FbxMeshPtr == '.' && *MeshPtr == '_';
-
-				if (*MeshPtr != *FbxMeshPtr && !bIsPointAndUnderscore)
-				{
-					break;
-				}
-				else
-				{
-					i++;
-					MeshPtr--;
-					FbxMeshPtr--;
-				}
-			}
-		}
-
-		if (i == FCStringAnsi::Strlen(FbxMeshName)) // matched
-		{
-			// check further
-			if (FCStringAnsi::Strlen(FbxMeshName) == FCStringAnsi::Strlen(MeshName) || // the name of Unreal mesh is full match
-				*MeshPtr == '_')														// or the name of Unreal mesh has a prefix
-			{
-				return FbxMeshArray[MeshIndex];
-			}
-		}
-	}
-	return nullptr;
-}
-
 UStaticMesh* UnFbx::FFbxImporter::ReimportStaticMesh(UStaticMesh* Mesh, UFbxStaticMeshImportData* TemplateImportData)
 {
 	TArray<FbxNode*> FbxMeshArray;
@@ -1198,7 +1154,7 @@ UStaticMesh* UnFbx::FFbxImporter::ReimportStaticMesh(UStaticMesh* Mesh, UFbxStat
 	}
 	else if(!bCombineMeshes && !bCombineMeshesLOD)
 	{
-		Node = GetMeshNodesFromName(Mesh, FbxMeshArray);
+		Node = GetMeshNodesFromName(Mesh->GetName(), FbxMeshArray);
 	}
 
 	// If there is no match it may be because an LOD group was imported where
@@ -1506,6 +1462,8 @@ UStaticMesh* UnFbx::FFbxImporter::ImportStaticMeshAsSingle(UObject* InParent, TA
 		MeshDescription->Empty();
 	}
 
+	FStaticMeshAttributes Attributes(*MeshDescription);
+
 	// make sure it has a new lighting guid
 	StaticMesh->LightingGuid = FGuid::NewGuid();
 
@@ -1553,9 +1511,9 @@ UStaticMesh* UnFbx::FFbxImporter::ImportStaticMeshAsSingle(UObject* InParent, TA
 
 	if (bBuildStatus)
 	{
-		TVertexInstanceAttributesRef<FVector2D> VertexInstanceUVs = MeshDescription->VertexInstanceAttributes().GetAttributesRef<FVector2D>(MeshAttribute::VertexInstance::TextureCoordinate);
+		TVertexInstanceAttributesRef<FVector2D> VertexInstanceUVs = Attributes.GetVertexInstanceUVs();
 		int32 FirstOpenUVChannel = VertexInstanceUVs.GetNumIndices() >= MAX_MESH_TEXTURE_COORDS_MD ? 1 : VertexInstanceUVs.GetNumIndices();
-		TPolygonGroupAttributesRef<FName> PolygonGroupImportedMaterialSlotNames = MeshDescription->PolygonGroupAttributes().GetAttributesRef<FName>(MeshAttribute::PolygonGroup::ImportedMaterialSlotName);
+		TPolygonGroupAttributesRef<FName> PolygonGroupImportedMaterialSlotNames = Attributes.GetPolygonGroupMaterialSlotNames();
 
 		TArray<FStaticMaterial> MaterialToAdd;
 		for (const FPolygonGroupID PolygonGroupID : MeshDescription->PolygonGroups().GetElementIDs())
@@ -1944,12 +1902,6 @@ void UnFbx::FFbxImporter::PostImportStaticMesh(UStaticMesh* StaticMesh, TArray<F
 	IConsoleVariable* CVarDistanceFieldInterface = IConsoleManager::Get().FindConsoleVariable(TEXT("r.GenerateMeshDistanceFields"));
 	bool bOriginalGenerateMeshDistanceField = StaticMesh->bGenerateMeshDistanceField;
 	
-	//Always triangulate the original mesh description after we import it
-	FMeshDescription* MeshDescription = StaticMesh->GetMeshDescription(LODIndex);
-	if (MeshDescription)
-	{
-		MeshDescription->TriangulateMesh();
-	}
 	bool bUserCancel = false;
 	if (StaticMesh->FixLODRequiresAdjacencyInformation(LODIndex))
 	{
diff --git a/Engine/Source/Editor/UnrealEd/Private/FbxMeshUtils.cpp b/Engine/Source/Editor/UnrealEd/Private/FbxMeshUtils.cpp
--- a/Engine/Source/Editor/UnrealEd/Private/FbxMeshUtils.cpp
+++ b/Engine/Source/Editor/UnrealEd/Private/FbxMeshUtils.cpp
@@ -35,6 +35,7 @@
 
 #include "Assets/ClothingAsset.h"
 #include "SkinWeightsUtilities.h"
+#include "LODUtilities.h"
 
 DEFINE_LOG_CATEGORY_STATIC(LogExportMeshUtils, Log, All);
 
@@ -280,10 +281,13 @@ namespace FbxMeshUtils
 			return false;
 		}
 
+		FScopedSkeletalMeshPostEditChange ScopePostEditChange(SelectedSkelMesh);
+
 		//If the imported LOD already exist, we will need to reimport all the skin weight profiles
 		bool bMustReimportAlternateSkinWeightProfile = false;
 
 		// Get a list of all the clothing assets affecting this LOD so we can re-apply later
+		TArray<ClothingAssetUtils::FClothingAssetMeshBinding> ClothingBindings;
 		TArray<UClothingAssetBase*> ClothingAssetsInUse;
 		TArray<int32> ClothingAssetSectionIndices;
 		TArray<int32> ClothingAssetInternalLodIndices;
@@ -292,47 +296,16 @@ namespace FbxMeshUtils
 		if(ImportedResource && ImportedResource->LODModels.IsValidIndex(LODLevel))
 		{
 			bMustReimportAlternateSkinWeightProfile = true;
-			FSkeletalMeshLODModel& LodModel = ImportedResource->LODModels[LODLevel];
-
-			const int32 NumSections = LodModel.Sections.Num();
-
-			for(int32 SectionIndex = 0; SectionIndex < NumSections; ++SectionIndex)
-			{
-				FSkelMeshSection& Section = LodModel.Sections[SectionIndex];
-
-				if(Section.HasClothingData())
-				{
-					UClothingAssetBase* AssetInUse = SelectedSkelMesh->GetSectionClothingAsset(LODLevel, SectionIndex);
-					ClothingAssetsInUse.Add(AssetInUse);
-					ClothingAssetSectionIndices.Add(SectionIndex);
-					ClothingAssetInternalLodIndices.Add(Section.ClothingData.AssetLodIndex);
-				}
-			}
-		}
-
-		// Remove our clothing assets while we import this LOD
-		for(UClothingAssetBase* ClothingAsset : ClothingAssetsInUse)
-		{
-			ClothingAsset->UnbindFromSkeletalMesh(SelectedSkelMesh, LODLevel);
+			FLODUtilities::UnbindClothingAndBackup(SelectedSkelMesh, ClothingBindings, LODLevel);
 		}
 
 		//Lambda to call to re-apply the clothing
-		auto ReapplyClothing = [&SelectedSkelMesh, &ClothingAssetsInUse, &ClothingAssetSectionIndices, &ClothingAssetInternalLodIndices, &ImportedResource, &LODLevel]()
+		auto ReapplyClothing = [&SelectedSkelMesh, &ClothingBindings, &ImportedResource, &LODLevel]()
 		{
-			// Re-apply our clothing assets
-			int32 NumClothingAssetsToApply = ClothingAssetsInUse.Num();
 			if (ImportedResource && ImportedResource->LODModels.IsValidIndex(LODLevel))
 			{
-				FSkeletalMeshLODModel& LodModel = ImportedResource->LODModels[LODLevel];
-				for (int32 AssetIndex = 0; AssetIndex < NumClothingAssetsToApply; ++AssetIndex)
-				{
-					// Only if the same section exists
-					if (LodModel.Sections.IsValidIndex(ClothingAssetSectionIndices[AssetIndex]))
-					{
-						UClothingAssetBase* AssetToApply = ClothingAssetsInUse[AssetIndex];
-						AssetToApply->BindToSkeletalMesh(SelectedSkelMesh, LODLevel, ClothingAssetSectionIndices[AssetIndex], ClothingAssetInternalLodIndices[AssetIndex]);
-					}
-				}
+				// Re-apply our clothing assets
+				FLODUtilities::RestoreClothingFromBackup(SelectedSkelMesh, ClothingBindings, LODLevel);
 			}
 		};
 
@@ -395,6 +368,24 @@ namespace FbxMeshUtils
 			TArray<FString> LODStrings;
 			TArray<FbxNode*>* MeshObject = NULL;;
 
+			//Set the build options if the BuildDat is not available so it is the same option we use to import the LOD
+			if (ImportedResource && ImportedResource->LODModels.IsValidIndex(LODLevel) && !ImportedResource->LODModels[LODLevel].RawSkeletalMeshBulkData.IsBuildDataAvailable())
+			{
+				FSkeletalMeshLODInfo* LODInfo = SelectedSkelMesh->GetLODInfo(LODLevel);
+				if (LODInfo)
+				{
+					LODInfo->BuildSettings.bBuildAdjacencyBuffer = true;
+					LODInfo->BuildSettings.bRecomputeNormals = !ImportOptions->ShouldImportNormals();
+					LODInfo->BuildSettings.bRecomputeTangents = !ImportOptions->ShouldImportTangents();
+					LODInfo->BuildSettings.bUseMikkTSpace = (ImportOptions->NormalGenerationMethod == EFBXNormalGenerationMethod::MikkTSpace) && (!ImportOptions->ShouldImportNormals() || !ImportOptions->ShouldImportTangents());
+					LODInfo->BuildSettings.bComputeWeightedNormals = ImportOptions->bComputeWeightedNormals;
+					LODInfo->BuildSettings.bRemoveDegenerates = ImportOptions->bRemoveDegenerates;
+					LODInfo->BuildSettings.ThresholdPosition = ImportOptions->OverlappingThresholds.ThresholdPosition;
+					LODInfo->BuildSettings.ThresholdTangentNormal = ImportOptions->OverlappingThresholds.ThresholdTangentNormal;
+					LODInfo->BuildSettings.ThresholdUV = ImportOptions->OverlappingThresholds.ThresholdUV;
+				}
+			}
+
 			// Populate the mesh array
 			FFbxImporter->FillFbxSkelMeshArrayInScene(FFbxImporter->Scene->GetRootNode(), MeshArray, false, ImportOptions->bImportAsSkeletalGeometry || ImportOptions->bImportAsSkeletalSkinning, ImportOptions->bImportScene);
 
@@ -524,12 +515,8 @@ namespace FbxMeshUtils
 
 				TempSkelMesh = (USkeletalMesh*)FFbxImporter->ImportSkeletalMesh( ImportSkeletalMeshArgs );
 				// Add the new imported LOD to the existing model (check skeleton compatibility)
-				if( TempSkelMesh  && FFbxImporter->ImportSkeletalMeshLOD(TempSkelMesh, SelectedSkelMesh, SelectedLOD, true, nullptr, TempAssetImportData))
+				if( TempSkelMesh  && FFbxImporter->ImportSkeletalMeshLOD(TempSkelMesh, SelectedSkelMesh, SelectedLOD, TempAssetImportData))
 				{
-					TComponentReregisterContext<USkinnedMeshComponent> ReregisterContext;
-					SelectedSkelMesh->ReleaseResources();
-					SelectedSkelMesh->ReleaseResourcesFence.Wait();
-					
 					//Update the import data for this lod
 					UnFbx::FFbxImporter::UpdateSkeletalMeshImportData(SelectedSkelMesh, nullptr, SelectedLOD, &ImportMaterialOriginalNameData, &ImportMeshLodData);
 
@@ -555,13 +542,9 @@ namespace FbxMeshUtils
 					{
 						//We cannot use anymore the FFbxImporter after the cleanup
 						CleanUpScene();
-						FSkinWeightsUtilities::ReimportAlternateSkinWeight(SelectedSkelMesh, SelectedLOD, false);
+						FSkinWeightsUtilities::ReimportAlternateSkinWeight(SelectedSkelMesh, SelectedLOD);
 					}
 
-					SelectedSkelMesh->PostEditChange();
-					SelectedSkelMesh->InitResources();
-					SelectedSkelMesh->MarkPackageDirty();
-
 					// Notification of success
 					FNotificationInfo NotificationInfo(FText::GetEmpty());
 					NotificationInfo.Text = FText::Format(NSLOCTEXT("UnrealEd", "LODImportSuccessful", "Mesh for LOD {0} imported successfully!"), FText::AsNumber(SelectedLOD));
@@ -633,7 +616,7 @@ namespace FbxMeshUtils
 		return ChosenFilname;
 	}
 
-	bool ImportMeshLODDialog(class UObject* SelectedMesh, int32 LODLevel)
+	bool ImportMeshLODDialog(class UObject* SelectedMesh, int32 LODLevel, bool bNotifyCB /*= true*/)
 	{
 		if(!SelectedMesh)
 		{
@@ -732,7 +715,7 @@ namespace FbxMeshUtils
 			FMessageDialog::Open(EAppMsgType::Ok, FText::Format(LOCTEXT("LODImport_Failure", "Failed to import LOD{0}"), FText::AsNumber(LODLevel)));
 		}
 
-		if (bImportSuccess)
+		if (bImportSuccess && bNotifyCB)
 		{
 			if (SkeletalMesh)
 			{
diff --git a/Engine/Source/Editor/UnrealEd/Private/GeomFitUtils.cpp b/Engine/Source/Editor/UnrealEd/Private/GeomFitUtils.cpp
--- a/Engine/Source/Editor/UnrealEd/Private/GeomFitUtils.cpp
+++ b/Engine/Source/Editor/UnrealEd/Private/GeomFitUtils.cpp
@@ -20,8 +20,7 @@
 #include "PhysicsEngine/SphylElem.h"
 #include "PhysicsEngine/BodySetup.h"
 #include "Engine/StaticMesh.h"
-#include "MeshDescription.h"
-#include "MeshAttributes.h"
+#include "StaticMeshAttributes.h"
 #include "Settings/EditorExperimentalSettings.h"
 
 #define LOCAL_EPS (0.01f)
@@ -190,12 +189,13 @@ int32 GenerateKDopAsSimpleCollision(UStaticMesh* StaticMesh, const TArray<FVecto
 
 /* ******************************** BOX ******************************** */
 
-static void CalcBoundingBox(const FMeshDescription& MeshDescription, FVector& Center, FVector& Extents, FVector& LimitVec)
+static void CalcBoundingBox(const FMeshDescription* MeshDescription, FVector& Center, FVector& Extents, FVector& LimitVec)
 {
-	TVertexAttributesConstRef<FVector> VertexPositions = MeshDescription.VertexAttributes().GetAttributesRef<FVector>(MeshAttribute::Vertex::Position);
+	FStaticMeshConstAttributes Attributes(*MeshDescription);
+	TVertexAttributesConstRef<FVector> VertexPositions = Attributes.GetVertexPositions();
 
 	FBox Box(ForceInit);
-	for (const FVertexID VertexID : MeshDescription.Vertices().GetElementIDs())
+	for (const FVertexID VertexID : MeshDescription->Vertices().GetElementIDs())
 	{
 		Box += VertexPositions[VertexID] * LimitVec;
 	}
@@ -207,8 +207,9 @@ void ComputeBoundingBox(UStaticMesh* StaticMesh, FVector& Center, FVector& Exten
 	// Calculate bounding Box.
 	
 	FMeshDescription* MeshDescription = StaticMesh->GetMeshDescription(0);
+	check(MeshDescription);
 	FVector unitVec = FVector(1.f);
-	CalcBoundingBox(*MeshDescription, Center, Extents, unitVec);
+	CalcBoundingBox(MeshDescription, Center, Extents, unitVec);
 }
 
 int32 GenerateBoxAsSimpleCollision(UStaticMesh* StaticMesh)
@@ -224,7 +225,8 @@ int32 GenerateBoxAsSimpleCollision(UStaticMesh* StaticMesh)
 	FVector unitVec = bs->BuildScale3D;
 	FVector Center, Extents;
 	FMeshDescription* MeshDescription = StaticMesh->GetMeshDescription(0);
-	CalcBoundingBox(*MeshDescription, Center, Extents, unitVec);
+	check(MeshDescription);
+	CalcBoundingBox(MeshDescription, Center, Extents, unitVec);
 	bs->Modify();
 
 	// Create new GUID
@@ -255,19 +257,21 @@ int32 GenerateBoxAsSimpleCollision(UStaticMesh* StaticMesh)
 
 // This algorithm taken from Ritter, 1990
 // This one seems to do well with asymmetric input.
-static void CalcBoundingSphere(const FMeshDescription& MeshDescription, FSphere& sphere, FVector& LimitVec)
+static void CalcBoundingSphere(const FMeshDescription* MeshDescription, FSphere& sphere, FVector& LimitVec)
 {
-	if (MeshDescription.Vertices().Num() == 0)
+	if (MeshDescription->Vertices().Num() == 0)
 		return;
 
 	FBox Box;
 	FVector MinIx[3];
 	FVector MaxIx[3];
 
-	TVertexAttributesConstRef<FVector> VertexPositions = MeshDescription.VertexAttributes().GetAttributesRef<FVector>(MeshAttribute::Vertex::Position);
+	FStaticMeshConstAttributes Attributes(*MeshDescription);
+
+	TVertexAttributesConstRef<FVector> VertexPositions = Attributes.GetVertexPositions();
 
 	bool bFirstVertex = true;
-	for (const FVertexID VertexID : MeshDescription.Vertices().GetElementIDs())
+	for (const FVertexID VertexID : MeshDescription->Vertices().GetElementIDs())
 	{
 		FVector p = VertexPositions[VertexID] * LimitVec;
 		if (bFirstVertex)
@@ -348,7 +352,7 @@ static void CalcBoundingSphere(const FMeshDescription& MeshDescription, FSphere&
 	float r2 = FMath::Square(r);
 
 	// Now check each point lies within this sphere. If not - expand it a bit.
-	for (const FVertexID VertexID : MeshDescription.Vertices().GetElementIDs())
+	for (const FVertexID VertexID : MeshDescription->Vertices().GetElementIDs())
 	{
 		const FVector cToP = (VertexPositions[VertexID] * LimitVec) - sphere.Center;
 
@@ -371,7 +375,7 @@ static void CalcBoundingSphere(const FMeshDescription& MeshDescription, FSphere&
 
 // This is the one thats already used by unreal.
 // Seems to do better with more symmetric input...
-static void CalcBoundingSphere2(const FMeshDescription& MeshDescription, FSphere& sphere, FVector& LimitVec)
+static void CalcBoundingSphere2(const FMeshDescription* MeshDescription, FSphere& sphere, FVector& LimitVec)
 {
 	FVector Center, Extents;
 	CalcBoundingBox(MeshDescription, Center, Extents, LimitVec);
@@ -379,9 +383,10 @@ static void CalcBoundingSphere2(const FMeshDescription& MeshDescription, FSphere
 	sphere.Center = Center;
 	sphere.W = 0.0f;
 
-	TVertexAttributesConstRef<FVector> VertexPositions = MeshDescription.VertexAttributes().GetAttributesRef<FVector>(MeshAttribute::Vertex::Position);
+	FStaticMeshConstAttributes Attributes(*MeshDescription);
+	TVertexAttributesConstRef<FVector> VertexPositions = Attributes.GetVertexPositions();
 
-	for (const FVertexID VertexID : MeshDescription.Vertices().GetElementIDs())
+	for (const FVertexID VertexID : MeshDescription->Vertices().GetElementIDs())
 	{
 		float Dist = FVector::DistSquared(VertexPositions[VertexID] * LimitVec, sphere.Center);
 		if (Dist > sphere.W)
@@ -405,8 +410,9 @@ int32 GenerateSphereAsSimpleCollision(UStaticMesh* StaticMesh)
 
 	// Calculate bounding sphere.
 	FMeshDescription* MeshDescription = StaticMesh->GetMeshDescription(0);
-	CalcBoundingSphere(*MeshDescription, bSphere, unitVec);
-	CalcBoundingSphere2(*MeshDescription, bSphere2, unitVec);
+	check(MeshDescription);
+	CalcBoundingSphere(MeshDescription, bSphere, unitVec);
+	CalcBoundingSphere2(MeshDescription, bSphere2, unitVec);
 
 	if(bSphere.W < bSphere2.W)
 		bestSphere = bSphere;
@@ -442,9 +448,9 @@ int32 GenerateSphereAsSimpleCollision(UStaticMesh* StaticMesh)
 
 /* ******************************** SPHYL ******************************** */
 
-static void CalcBoundingSphyl(const FMeshDescription& MeshDescription, FSphere& sphere, float& length, FRotator& rotation, FVector& LimitVec)
+static void CalcBoundingSphyl(const FMeshDescription* MeshDescription, FSphere& sphere, float& length, FRotator& rotation, FVector& LimitVec)
 {
-	if (MeshDescription.Vertices().Num() == 0)
+	if (MeshDescription->Vertices().Num() == 0)
 		return;
 
 	FVector Center, Extents;
@@ -475,10 +481,11 @@ static void CalcBoundingSphyl(const FMeshDescription& MeshDescription, FSphere&
 	float r = Extents.GetMax();
 	float r2 = FMath::Square(r);
 	
-	TVertexAttributesConstRef<FVector> VertexPositions = MeshDescription.VertexAttributes().GetAttributesRef<FVector>(MeshAttribute::Vertex::Position);
+	FStaticMeshConstAttributes Attributes(*MeshDescription);
+	TVertexAttributesConstRef<FVector> VertexPositions = Attributes.GetVertexPositions();
 
 	// Now check each point lies within this the radius. If not - expand it a bit.
-	for (const FVertexID VertexID : MeshDescription.Vertices().GetElementIDs())
+	for (const FVertexID VertexID : MeshDescription->Vertices().GetElementIDs())
 	{
 		FVector cToP = (VertexPositions[VertexID] * LimitVec) - sphere.Center;
 		cToP = rotation.UnrotateVector(cToP);
@@ -499,7 +506,7 @@ static void CalcBoundingSphyl(const FMeshDescription& MeshDescription, FSphere&
 	float hl = FMath::Max(0.0f, Extent - r);
 
 	// Now check each point lies within the length. If not - expand it a bit.
-	for (const FVertexID VertexID : MeshDescription.Vertices().GetElementIDs())
+	for (const FVertexID VertexID : MeshDescription->Vertices().GetElementIDs())
 	{
 		FVector cToP = (VertexPositions[VertexID] * LimitVec) - sphere.Center;
 		cToP = rotation.UnrotateVector(cToP);
@@ -546,7 +553,8 @@ int32 GenerateSphylAsSimpleCollision(UStaticMesh* StaticMesh)
 
 	// Calculate bounding box.
 	FMeshDescription* MeshDescription = StaticMesh->GetMeshDescription(0);
-	CalcBoundingSphyl(*MeshDescription, sphere, length, rotation, unitVec);
+	check(MeshDescription);
+	CalcBoundingSphyl(MeshDescription, sphere, length, rotation, unitVec);
 
 	// Dont use if radius is zero.
 	if (sphere.W <= 0.f)
diff --git a/Engine/Source/Editor/UnrealEd/Private/SkinWeightsUtilities.cpp b/Engine/Source/Editor/UnrealEd/Private/SkinWeightsUtilities.cpp
--- a/Engine/Source/Editor/UnrealEd/Private/SkinWeightsUtilities.cpp
+++ b/Engine/Source/Editor/UnrealEd/Private/SkinWeightsUtilities.cpp
@@ -34,7 +34,7 @@
 
 DEFINE_LOG_CATEGORY_STATIC(LogSkinWeightsUtilities, Log, All);
 
-bool FSkinWeightsUtilities::ImportAlternateSkinWeight(USkeletalMesh* SkeletalMesh, FString Path, int32 TargetLODIndex, const FName& ProfileName, bool bReregisterComponent)
+bool FSkinWeightsUtilities::ImportAlternateSkinWeight(USkeletalMesh* SkeletalMesh, FString Path, int32 TargetLODIndex, const FName& ProfileName)
 {
 	check(SkeletalMesh);
 	check(SkeletalMesh->GetLODInfo(TargetLODIndex));
@@ -53,6 +53,8 @@ bool FSkinWeightsUtilities::ImportAlternateSkinWeight(USkeletalMesh* SkeletalMes
 		UE_LOG(LogSkinWeightsUtilities, Error, TEXT("Path containing Skin Weight Profile data does not exist (%s)."), *Path);
 		return false;
 	}
+	FScopedSkeletalMeshPostEditChange ScopePostEditChange(SkeletalMesh);
+
 	UnFbx::FBXImportOptions ImportOptions;
 	//Import the alternate fbx into a temporary skeletal mesh using the same import options
 	UFbxFactory* FbxFactory = NewObject<UFbxFactory>(UFbxFactory::StaticClass());
@@ -200,21 +202,7 @@ bool FSkinWeightsUtilities::ImportAlternateSkinWeight(USkeletalMesh* SkeletalMes
 				bool bUseMikkTSpace = ImportOptions.NormalGenerationMethod == EFBXNormalGenerationMethod::MikkTSpace;
 
 				TArray<FRawSkinWeight>& SkinWeights = ProfileData.SkinWeights;
-				if (bReregisterComponent)
-				{
-					TComponentReregisterContext<USkinnedMeshComponent> ReregisterContext;
-					SkeletalMesh->ReleaseResources();
-					SkeletalMesh->ReleaseResourcesFence.Wait();
-
-					bResult = FLODUtilities::UpdateAlternateSkinWeights(SkeletalMesh, ProfileName, TmpSkeletalMesh, TargetLODIndex, SrcLodIndex, OverlappingThresholds, ShouldImportNormals, ShouldImportTangents, bUseMikkTSpace, ImportOptions.bComputeWeightedNormals);
-
-					SkeletalMesh->PostEditChange();
-					SkeletalMesh->InitResources();
-				}
-				else
-				{
-					bResult = FLODUtilities::UpdateAlternateSkinWeights(SkeletalMesh, ProfileName, TmpSkeletalMesh, TargetLODIndex, SrcLodIndex, OverlappingThresholds, ShouldImportNormals, ShouldImportTangents, bUseMikkTSpace, ImportOptions.bComputeWeightedNormals);
-				}
+				bResult = FLODUtilities::UpdateAlternateSkinWeights(SkeletalMesh, ProfileName, TmpSkeletalMesh, TargetLODIndex, SrcLodIndex, OverlappingThresholds, ShouldImportNormals, ShouldImportTangents, bUseMikkTSpace, ImportOptions.bComputeWeightedNormals);
 				
 				if (!bResult)
 				{
@@ -252,57 +240,45 @@ bool FSkinWeightsUtilities::ImportAlternateSkinWeight(USkeletalMesh* SkeletalMes
 	return bResult;
 }
 
-bool FSkinWeightsUtilities::ReimportAlternateSkinWeight(USkeletalMesh* SkeletalMesh, int32 TargetLODIndex, bool bReregisterComponent)
+bool FSkinWeightsUtilities::ReimportAlternateSkinWeight(USkeletalMesh* SkeletalMesh, int32 TargetLODIndex)
 {
 	bool bResult = false;
 
-	//Bulk work of the function, we use a lambda because of the re-register component option.
-	auto DoWork = [&SkeletalMesh, &TargetLODIndex, &bResult]()
+	const TArray<FSkinWeightProfileInfo>& SkinWeightProfiles = SkeletalMesh->GetSkinWeightProfiles();
+	if (SkinWeightProfiles.Num() <= 0)
 	{
-		const TArray<FSkinWeightProfileInfo>& SkinWeightProfiles = SkeletalMesh->GetSkinWeightProfiles();
-		for (int32 ProfileIndex = 0; ProfileIndex < SkinWeightProfiles.Num(); ++ProfileIndex)
-		{
-			const FSkinWeightProfileInfo& ProfileInfo = SkinWeightProfiles[ProfileIndex];
+		return bResult;
+	}
 
-			const FString* PathNamePtr = ProfileInfo.PerLODSourceFiles.Find(TargetLODIndex);
-			//Skip profile that do not have data for TargetLODIndex
-			if (!PathNamePtr)
-			{
-				continue;
-			}
+	FScopedSkeletalMeshPostEditChange ScopePostEditChange(SkeletalMesh);
 
-			const FString& PathName = *PathNamePtr;
+	for (int32 ProfileIndex = 0; ProfileIndex < SkinWeightProfiles.Num(); ++ProfileIndex)
+	{
+		const FSkinWeightProfileInfo& ProfileInfo = SkinWeightProfiles[ProfileIndex];
 
-			if (FPaths::FileExists(PathName))
-			{
-				bResult |= FSkinWeightsUtilities::ImportAlternateSkinWeight(SkeletalMesh, PathName, TargetLODIndex, ProfileInfo.Name, false);
-			}
-			else
-			{
-				const FString PickedFileName = FSkinWeightsUtilities::PickSkinWeightFBXPath(TargetLODIndex);
-				if (!PickedFileName.IsEmpty() && FPaths::FileExists(PickedFileName))
-				{
-					bResult |= FSkinWeightsUtilities::ImportAlternateSkinWeight(SkeletalMesh, PickedFileName, TargetLODIndex, ProfileInfo.Name, false);
-				}
-			}
+		const FString* PathNamePtr = ProfileInfo.PerLODSourceFiles.Find(TargetLODIndex);
+		//Skip profile that do not have data for TargetLODIndex
+		if (!PathNamePtr)
+		{
+			continue;
 		}
-	};
 
-	if (bReregisterComponent)
-	{
-		TComponentReregisterContext<USkinnedMeshComponent> ReregisterContext;
-		SkeletalMesh->ReleaseResources();
-		SkeletalMesh->ReleaseResourcesFence.Wait();
-
-		DoWork();
+		const FString& PathName = *PathNamePtr;
 
-		SkeletalMesh->PostEditChange();
-		SkeletalMesh->InitResources();
-	}
-	else
-	{
-		DoWork();
+		if (FPaths::FileExists(PathName))
+		{
+			bResult |= FSkinWeightsUtilities::ImportAlternateSkinWeight(SkeletalMesh, PathName, TargetLODIndex, ProfileInfo.Name);
+		}
+		else
+		{
+			const FString PickedFileName = FSkinWeightsUtilities::PickSkinWeightFBXPath(TargetLODIndex);
+			if (!PickedFileName.IsEmpty() && FPaths::FileExists(PickedFileName))
+			{
+				bResult |= FSkinWeightsUtilities::ImportAlternateSkinWeight(SkeletalMesh, PickedFileName, TargetLODIndex, ProfileInfo.Name);
+			}
+		}
 	}
+
 	
 	if (bResult)
 	{
diff --git a/Engine/Source/Editor/UnrealEd/Private/StaticMeshEdit.cpp b/Engine/Source/Editor/UnrealEd/Private/StaticMeshEdit.cpp
--- a/Engine/Source/Editor/UnrealEd/Private/StaticMeshEdit.cpp
+++ b/Engine/Source/Editor/UnrealEd/Private/StaticMeshEdit.cpp
@@ -13,6 +13,7 @@
 #include "EditorFramework/ThumbnailInfo.h"
 #include "Engine/MeshMerging.h"
 #include "Engine/StaticMesh.h"
+#include "StaticMeshAttributes.h"
 #include "Engine/StaticMeshSocket.h"
 #include "Engine/Polys.h"
 #include "Editor.h"
@@ -31,8 +32,6 @@
 #include "Interfaces/ITargetPlatformManagerModule.h"
 #include "Interfaces/ITargetPlatform.h"
 #include "Settings/EditorExperimentalSettings.h"
-#include "MeshDescription.h"
-#include "MeshAttributes.h"
 
 #include "Modules/ModuleManager.h"
 #include "IMeshReductionManagerModule.h"
@@ -595,9 +594,8 @@ UStaticMesh* CreateStaticMesh(FMeshDescription& RawMesh,TArray<FStaticMaterial>&
 	auto StaticMesh = NewObject<UStaticMesh>(InOuter, InName, RF_Public | RF_Standalone);
 
 	// Add one LOD for the base mesh
-	FStaticMeshSourceModel& SrcModel = StaticMesh->AddSourceModel();
-	FMeshDescription* MeshDescription = StaticMesh->CreateMeshDescription(0);
-	*MeshDescription = RawMesh;
+	StaticMesh->AddSourceModel();
+	FMeshDescription* MeshDescription = StaticMesh->CreateMeshDescription(0, RawMesh);
 	StaticMesh->CommitMeshDescription(0);
 	StaticMesh->StaticMaterials = Materials;
 
@@ -805,10 +803,9 @@ UStaticMesh* CreateStaticMeshFromBrush(UObject* Outer, FName Name, ABrush* Brush
 	UStaticMesh* StaticMesh = NewObject<UStaticMesh>(Outer, Name, RF_Public | RF_Standalone);
 
 	// Add one LOD for the base mesh
-	FStaticMeshSourceModel& SrcModel = StaticMesh->AddSourceModel();
+	StaticMesh->AddSourceModel();
 	const int32 LodIndex = StaticMesh->GetNumSourceModels() - 1;
 	FMeshDescription* MeshDescription = StaticMesh->CreateMeshDescription(LodIndex);
-	UStaticMesh::RegisterMeshAttributes(*MeshDescription);
 
 	// Fill out the mesh description and materials from the brush geometry
 	TArray<FStaticMaterial> Materials;
@@ -1490,9 +1487,7 @@ void RestoreExistingMeshData(ExistingStaticMeshData* ExistingMeshDataPtr, UStati
 		FStaticMeshSourceModel& SrcModel = NewMesh->AddSourceModel();
 		if (ExistingMeshDataPtr->ExistingLODData[i].ExistingMeshDescription.IsValid())
 		{
-			FMeshDescription* MeshDescription = NewMesh->CreateMeshDescription(i);
-			*MeshDescription = MoveTemp(*ExistingMeshDataPtr->ExistingLODData[i].ExistingMeshDescription);
-			ExistingMeshDataPtr->ExistingLODData[i].ExistingMeshDescription.Reset();
+			FMeshDescription* MeshDescription = NewMesh->CreateMeshDescription(i, MoveTemp(*ExistingMeshDataPtr->ExistingLODData[i].ExistingMeshDescription));
 			NewMesh->CommitMeshDescription(i);
 		}
 		SrcModel.BuildSettings = ExistingMeshDataPtr->ExistingLODData[i].ExistingBuildSettings;
diff --git a/Engine/Source/Editor/UnrealEd/Public/FbxImporter.h b/Engine/Source/Editor/UnrealEd/Public/FbxImporter.h
--- a/Engine/Source/Editor/UnrealEd/Public/FbxImporter.h
+++ b/Engine/Source/Editor/UnrealEd/Public/FbxImporter.h
@@ -756,9 +756,9 @@ public:
 	void AddStaticMeshSourceModelGeneratedLOD(UStaticMesh* StaticMesh, int32 LODIndex);
 
 	/**
-	* Return the node that match the staticmesh name. Return nullptr in case there is no match
+	* Return the node that match the mesh name. Return nullptr in case there is no match
 	*/
-	FbxNode* GetMeshNodesFromName(UStaticMesh* StaticMesh, TArray<FbxNode*>& FbxMeshArray);
+	FbxNode* GetMeshNodesFromName(const FString& ReimportMeshName, TArray<FbxNode*>& FbxMeshArray);
 
 	/**
 	 * re-import Unreal static mesh from updated Fbx file
@@ -917,11 +917,9 @@ public:
 	 * @param InSkeletalMesh - LOD mesh object
 	 * @param BaseSkeletalMesh - base mesh object
 	 * @param DesiredLOD - LOD level
-	 * @param bNeedToReregister - if true, re-register this skeletal mesh to shut down the skeletal mesh component that is previewing this mesh. 
-									But you can set this to false when in the first loading before rendering this mesh for a performance issue 
 	   @param ReregisterAssociatedComponents - if NULL, just re-registers all SkinnedMeshComponents but if you set the specific components, will only re-registers those components
 	 */
-	UNREALED_API bool ImportSkeletalMeshLOD(USkeletalMesh* InSkeletalMesh, USkeletalMesh* BaseSkeletalMesh, int32 DesiredLOD, bool bNeedToReregister = true, TArray<UActorComponent*>* ReregisterAssociatedComponents = NULL, UFbxSkeletalMeshImportData* TemplateImportData = nullptr);
+	UNREALED_API bool ImportSkeletalMeshLOD(USkeletalMesh* InSkeletalMesh, USkeletalMesh* BaseSkeletalMesh, int32 DesiredLOD, UFbxSkeletalMeshImportData* TemplateImportData = nullptr);
 
 	/**
 	 * Empties the FBX scene, releasing its memory.
diff --git a/Engine/Source/Editor/UnrealEd/Public/FbxMeshUtils.h b/Engine/Source/Editor/UnrealEd/Public/FbxMeshUtils.h
--- a/Engine/Source/Editor/UnrealEd/Public/FbxMeshUtils.h
+++ b/Engine/Source/Editor/UnrealEd/Public/FbxMeshUtils.h
@@ -39,7 +39,7 @@ namespace FbxMeshUtils
 	 * @param LODLevel			The level of the lod to import
 	 * @return					Whether or not the import succeeded
 	 */
-	UNREALED_API bool ImportMeshLODDialog( UObject* Mesh, int32 LODLevel);
+	UNREALED_API bool ImportMeshLODDialog( UObject* Mesh, int32 LODLevel, bool bNotifyCB = true);
 
 	/**
 	 * Sets import option before importing
diff --git a/Engine/Source/Editor/UnrealEd/Public/SkelImport.h b/Engine/Source/Editor/UnrealEd/Public/SkelImport.h
--- a/Engine/Source/Editor/UnrealEd/Public/SkelImport.h
+++ b/Engine/Source/Editor/UnrealEd/Public/SkelImport.h
@@ -22,18 +22,20 @@ class FSkeletalMeshLODModel;
 
 struct ExistingMeshLodSectionData
 {
-	ExistingMeshLodSectionData(FName InImportedMaterialSlotName, bool InbCastShadow, bool InbRecomputeTangents, int32 InGenerateUpTo, bool InbDisabled)
+	ExistingMeshLodSectionData(FName InImportedMaterialSlotName, bool InbCastShadow, bool InbRecomputeTangents, int32 InGenerateUpTo, bool InbDisabled, bool InbBoneChunkedSection)
 	: ImportedMaterialSlotName(InImportedMaterialSlotName)
 	, bCastShadow(InbCastShadow)
 	, bRecomputeTangents(InbRecomputeTangents)
 	, GenerateUpTo(InGenerateUpTo)
 	, bDisabled(InbDisabled)
+	, bBoneChunkedSection(InbBoneChunkedSection)
 	{}
 	FName ImportedMaterialSlotName;
 	bool bCastShadow;
 	bool bRecomputeTangents;
 	int32 GenerateUpTo;
 	bool bDisabled;
+	bool bBoneChunkedSection;
 };
 
 struct ExistingSkelMeshData
@@ -41,6 +43,7 @@ struct ExistingSkelMeshData
 	TArray<USkeletalMeshSocket*>			ExistingSockets;
 	TArray<FReductionBaseSkeletalMeshBulkData*> ExistingOriginalReductionSourceMeshData;
 	TIndirectArray<FSkeletalMeshLODModel>	ExistingLODModels;
+	FSkeletalMeshLODInfo					ExistingBaseLODInfo;
 	TArray<FSkeletalMeshLODInfo>			ExistingLODInfo;
 	FReferenceSkeleton						ExistingRefSkeleton;
 	TArray<FSkeletalMaterial>				ExistingMaterials;
diff --git a/Engine/Source/Editor/UnrealEd/Public/SkinWeightsUtilities.h b/Engine/Source/Editor/UnrealEd/Public/SkinWeightsUtilities.h
--- a/Engine/Source/Editor/UnrealEd/Public/SkinWeightsUtilities.h
+++ b/Engine/Source/Editor/UnrealEd/Public/SkinWeightsUtilities.h
@@ -27,7 +27,7 @@ public:
 	* @Param ProfileName - The name of the profile to associate the imported skin weight
 	* @param bReregisterComponent - if true the component using the skeletal mesh will all be re register.
 	*/
-	static bool ImportAlternateSkinWeight(USkeletalMesh* SkeletalMesh, FString Path, int32 TargetLODIndex, const FName& ProfileName, bool bReregisterComponent = true);
+	static bool ImportAlternateSkinWeight(USkeletalMesh* SkeletalMesh, FString Path, int32 TargetLODIndex, const FName& ProfileName);
 	
 	/**
 	* This function reimport all skin weights profile for a specified LOD. Return true if the weights are successfully updates.
@@ -37,7 +37,7 @@ public:
 	* @Param TargetLODIndex - The LODIndex to imp[ort the skin weight
 	* @param bReregisterComponent - if true the component using the skeletal mesh will all be re register.
 	*/
-	static bool ReimportAlternateSkinWeight(USkeletalMesh* SkeletalMesh, int32 TargetLODIndex, bool bReregisterComponent = true);
+	static bool ReimportAlternateSkinWeight(USkeletalMesh* SkeletalMesh, int32 TargetLODIndex);
 
 	static bool RemoveSkinnedWeightProfileData(USkeletalMesh* SkeletalMesh, const FName& ProfileName, int32 LODIndex);
 
diff --git a/Engine/Source/Editor/WorldBrowser/Private/Tiles/WorldTileCollectionModel.cpp b/Engine/Source/Editor/WorldBrowser/Private/Tiles/WorldTileCollectionModel.cpp
--- a/Engine/Source/Editor/WorldBrowser/Private/Tiles/WorldTileCollectionModel.cpp
+++ b/Engine/Source/Editor/WorldBrowser/Private/Tiles/WorldTileCollectionModel.cpp
@@ -23,10 +23,7 @@
 #include "AssetRegistryModule.h"
 #include "MeshUtilities.h"
 
-#include "MeshDescription.h"
-#include "MeshAttributes.h"
-#include "MeshAttributeArray.h"
-
+#include "StaticMeshAttributes.h"
 #include "Materials/Material.h"
 #include "MaterialUtilities.h"
 #include "LandscapeLayerInfoObject.h"
@@ -2158,12 +2155,13 @@ bool FWorldTileCollectionModel::GenerateLODLevels(FLevelModelList InLevelList, i
 				//Set the Imported version before calling the build
 				StaticMesh->ImportVersion = EImportStaticMeshVersion::LastVersion;
 			}
-			FMeshDescription& LandscapeRawMesh = *(StaticMesh->CreateMeshDescription(0));
+			FMeshDescription* LandscapeRawMesh = StaticMesh->CreateMeshDescription(0);
+			FStaticMeshAttributes Attributes(*LandscapeRawMesh);
 		
-			Landscape->ExportToRawMesh(LandscapeLOD, LandscapeRawMesh);
+			Landscape->ExportToRawMesh(LandscapeLOD, *LandscapeRawMesh);
 		
-			TVertexAttributesRef<FVector> VertexPositions = LandscapeRawMesh.VertexAttributes().GetAttributesRef<FVector>(MeshAttribute::Vertex::Position);
-			for (const FVertexID& VertexID : LandscapeRawMesh.Vertices().GetElementIDs())
+			TVertexAttributesRef<FVector> VertexPositions = Attributes.GetVertexPositions();
+			for (const FVertexID& VertexID : LandscapeRawMesh->Vertices().GetElementIDs())
 			{
 				VertexPositions[VertexID] -= LandscapeWorldLocation;
 			}
diff --git a/Engine/Source/Runtime/ClothingSystemRuntimeInterface/Public/ClothingAssetBase.h b/Engine/Source/Runtime/ClothingSystemRuntimeInterface/Public/ClothingAssetBase.h
--- a/Engine/Source/Runtime/ClothingSystemRuntimeInterface/Public/ClothingAssetBase.h
+++ b/Engine/Source/Runtime/ClothingSystemRuntimeInterface/Public/ClothingAssetBase.h
@@ -23,7 +23,7 @@ public:
 	* @param InSubmeshIdx Submesh in this asset to replace section with
 	* @param InAssetLodIndex Internal clothing LOD to use
 	*/
-	virtual bool BindToSkeletalMesh(USkeletalMesh* InSkelMesh, int32 InMeshLodIndex, int32 InSectionIndex, int32 InAssetLodIndex, bool bCallPostEditChange = true)
+	virtual bool BindToSkeletalMesh(USkeletalMesh* InSkelMesh, int32 InMeshLodIndex, int32 InSectionIndex, int32 InAssetLodIndex)
 	PURE_VIRTUAL(UClothingAssetBase::BindToSkeletalMesh, return false;);
 
 	/**
diff --git a/Engine/Source/Runtime/ClothingSystemRuntimeNv/Public/Assets/ClothingAsset.h b/Engine/Source/Runtime/ClothingSystemRuntimeNv/Public/Assets/ClothingAsset.h
--- a/Engine/Source/Runtime/ClothingSystemRuntimeNv/Public/Assets/ClothingAsset.h
+++ b/Engine/Source/Runtime/ClothingSystemRuntimeNv/Public/Assets/ClothingAsset.h
@@ -526,7 +526,7 @@ public:
 	// UClothingAssetBase Interface ////////////////////////////////////////////
 	virtual void RefreshBoneMapping(USkeletalMesh* InSkelMesh) override;
 #if WITH_EDITOR
-	virtual bool BindToSkeletalMesh(USkeletalMesh* InSkelMesh, int32 InMeshLodIndex, int32 InSectionIndex, int32 InAssetLodIndex, bool bCallPostEditChange=true) override;
+	virtual bool BindToSkeletalMesh(USkeletalMesh* InSkelMesh, int32 InMeshLodIndex, int32 InSectionIndex, int32 InAssetLodIndex) override;
 	virtual void UnbindFromSkeletalMesh(USkeletalMesh* InSkelMesh) override;
 	virtual void UnbindFromSkeletalMesh(USkeletalMesh* InSkelMesh, int32 InMeshLodIndex) override;
 	virtual void InvalidateCachedData() override;
diff --git a/Engine/Source/Runtime/Core/Public/UObject/EditorObjectVersion.h b/Engine/Source/Runtime/Core/Public/UObject/EditorObjectVersion.h
--- a/Engine/Source/Runtime/Core/Public/UObject/EditorObjectVersion.h
+++ b/Engine/Source/Runtime/Core/Public/UObject/EditorObjectVersion.h
@@ -83,6 +83,8 @@ struct CORE_API FEditorObjectVersion
 		MeshDescriptionTriangles,
 		//Add weighted area and angle when computing the normals
 		ComputeWeightedNormals,
+		// SkeletalMesh now can be rebuild in editor, no more need to re-import
+		SkeletalMeshBuildRefactor,
 		// -----<new versions can be added above this line>-------------------------------------------------
 		VersionPlusOne,
 		LatestVersion = VersionPlusOne - 1
diff --git a/Engine/Source/Runtime/Engine/Classes/Engine/EngineTypes.h b/Engine/Source/Runtime/Engine/Classes/Engine/EngineTypes.h
--- a/Engine/Source/Runtime/Engine/Classes/Engine/EngineTypes.h
+++ b/Engine/Source/Runtime/Engine/Classes/Engine/EngineTypes.h
@@ -2675,6 +2675,96 @@ struct FMeshBuildSettings
 	}
 };
 
+/**
+ * Settings applied when building a mesh.
+ */
+USTRUCT()
+struct FSkeletalMeshBuildSettings
+{
+	GENERATED_BODY()
+
+	/** If true, normals in the raw mesh are ignored and recomputed. */
+	UPROPERTY(EditAnywhere, Category=BuildSettings)
+	uint8 bRecomputeNormals:1;
+
+	/** If true, tangents in the raw mesh are ignored and recomputed. */
+	UPROPERTY(EditAnywhere, Category=BuildSettings)
+	uint8 bRecomputeTangents:1;
+	
+	/** If true, degenerate triangles will be removed. */
+	UPROPERTY(EditAnywhere, Category=BuildSettings)
+	uint8 bUseMikkTSpace:1;
+	
+	/** If true, we will use the surface area and the corner angle of the triangle as a ratio when computing the normals. */
+	UPROPERTY(EditAnywhere, Category = BuildSettings)
+	uint8 bComputeWeightedNormals : 1;
+
+	/** If true, degenerate triangles will be removed. */
+	UPROPERTY(EditAnywhere, Category=BuildSettings)
+	uint8 bRemoveDegenerates:1;
+	
+	/** If true, Tangents will be stored at 16 bit vs 8 bit precision. */
+	UPROPERTY(EditAnywhere, Category = BuildSettings)
+	uint8 bUseHighPrecisionTangentBasis:1;
+
+	/** If true, UVs will be stored at full floating point precision. */
+	UPROPERTY(EditAnywhere, Category=BuildSettings)
+	uint8 bUseFullPrecisionUVs:1;
+	
+	/** Required for PNT tessellation but can be slow. Recommend disabling for larger meshes. */
+	UPROPERTY(EditAnywhere, Category=BuildSettings)
+	uint8 bBuildAdjacencyBuffer:1;
+
+	/** Threshold use to decide if two vertex position are equal. */
+	UPROPERTY(EditAnywhere, Category = BuildSettings)
+	float ThresholdPosition;
+
+	/** Threshold use to decide if two normal, tangents or bi-normals are equal. */
+	UPROPERTY(EditAnywhere, Category = BuildSettings)
+	float ThresholdTangentNormal;
+
+	/** Threshold use to decide if two UVs are equal. */
+	UPROPERTY(EditAnywhere, Category = BuildSettings)
+	float ThresholdUV;
+
+	/** Default settings. */
+	FSkeletalMeshBuildSettings()
+		: bRecomputeNormals(true)
+		, bRecomputeTangents(true)
+		, bUseMikkTSpace(true)
+		, bComputeWeightedNormals(false)
+		, bRemoveDegenerates(true)
+		, bUseHighPrecisionTangentBasis(false)
+		, bUseFullPrecisionUVs(false)
+		, bBuildAdjacencyBuffer(true)
+		, ThresholdPosition(0.00002)
+		, ThresholdTangentNormal(0.00002)
+		, ThresholdUV(0.0009765625)
+	{}
+
+	/** Equality operator. */
+	bool operator==(const FSkeletalMeshBuildSettings& Other) const
+	{
+		return bRecomputeNormals == Other.bRecomputeNormals
+			&& bRecomputeTangents == Other.bRecomputeTangents
+			&& bUseMikkTSpace == Other.bUseMikkTSpace
+			&& bComputeWeightedNormals == Other.bComputeWeightedNormals
+			&& bRemoveDegenerates == Other.bRemoveDegenerates
+			&& bUseHighPrecisionTangentBasis == Other.bUseHighPrecisionTangentBasis
+			&& bUseFullPrecisionUVs == Other.bUseFullPrecisionUVs
+			&& bBuildAdjacencyBuffer == Other.bBuildAdjacencyBuffer
+			&& ThresholdPosition == Other.ThresholdPosition
+			&& ThresholdTangentNormal == Other.ThresholdTangentNormal
+			&& ThresholdUV == Other.ThresholdUV;
+	}
+
+	/** Inequality. */
+	bool operator!=(const FSkeletalMeshBuildSettings& Other) const
+	{
+		return !(*this == Other);
+	}
+};
+
 /** Event used by AActor::TakeDamage and related functions */
 USTRUCT(BlueprintType)
 struct ENGINE_API FDamageEvent
diff --git a/Engine/Source/Runtime/Engine/Classes/Engine/SkeletalMeshLODSettings.h b/Engine/Source/Runtime/Engine/Classes/Engine/SkeletalMeshLODSettings.h
--- a/Engine/Source/Runtime/Engine/Classes/Engine/SkeletalMeshLODSettings.h
+++ b/Engine/Source/Runtime/Engine/Classes/Engine/SkeletalMeshLODSettings.h
@@ -5,6 +5,7 @@
 #include "CoreMinimal.h"
 #include "UObject/ObjectMacros.h"
 #include "UObject/Object.h"
+#include "Engine/EngineTypes.h"
 #include "SkeletalMeshReductionSettings.h"
 #include "DataAsset.h"
 #include "PerPlatformProperties.h"
@@ -63,6 +64,9 @@ struct FSkeletalMeshLODGroupSettings
 	/** Get Skeletal mesh optimizations setting structure for the given LOD level */
 	ENGINE_API FSkeletalMeshOptimizationSettings GetReductionSettings() const;
 
+	/** Get Skeletal mesh build setting structure for the given LOD level */
+	ENGINE_API FSkeletalMeshBuildSettings GetBuildSettings() const;
+
 	/** Get Skeletal mesh optimizations setting structure for the given LOD level */
 	UE_DEPRECATED(4.20, "Please use GetReductionSettings instead")
 	FSkeletalMeshOptimizationSettings GetSettings() const
@@ -107,6 +111,10 @@ struct FSkeletalMeshLODGroupSettings
 	/** The optimization settings to use for the respective LOD level */
 	UPROPERTY(EditAnywhere, Category = Reduction)
 	FSkeletalMeshOptimizationSettings ReductionSettings;
+
+	/** The build settings to use for the respective LOD level */
+	UPROPERTY(EditAnywhere, Category = Build)
+	FSkeletalMeshBuildSettings BuildSettings;
 };
 
 UCLASS(config = Engine, defaultconfig, BlueprintType, MinimalAPI)
diff --git a/Engine/Source/Runtime/Engine/Private/Components/SceneCaptureComponent.cpp b/Engine/Source/Runtime/Engine/Private/Components/SceneCaptureComponent.cpp
--- a/Engine/Source/Runtime/Engine/Private/Components/SceneCaptureComponent.cpp
+++ b/Engine/Source/Runtime/Engine/Private/Components/SceneCaptureComponent.cpp
@@ -7,6 +7,7 @@
 #include "Components/SceneCaptureComponent.h"
 #include "Misc/ScopeLock.h"
 #include "UObject/RenderingObjectVersion.h"
+#include "UObject/EditorObjectVersion.h"
 #include "UObject/ConstructorHelpers.h"
 #include "GameFramework/Actor.h"
 #include "RenderingThread.h"
@@ -746,15 +747,15 @@ void APlanarReflection::PostLoad()
 	Super::PostLoad();
 
 	if (GetLinkerCustomVersion(FEditorObjectVersion::GUID) < FEditorObjectVersion::ChangeSceneCaptureRootComponent)
-{
-		if (PlanarReflectionComponent)
 	{
+		if (PlanarReflectionComponent)
+		{
 			PRAGMA_DISABLE_DEPRECATION_WARNINGS
 			PlanarReflectionComponent->bShowPreviewPlane = bShowPreviewPlane_DEPRECATED;
 			PRAGMA_ENABLE_DEPRECATION_WARNINGS
 		}
-			}
-		}
+	}
+}
 
 void APlanarReflection::OnInterpToggle(bool bEnable)
 {
diff --git a/Engine/Source/Runtime/Engine/Private/SkeletalMeshLODModel.cpp b/Engine/Source/Runtime/Engine/Private/SkeletalMeshLODModel.cpp
--- a/Engine/Source/Runtime/Engine/Private/SkeletalMeshLODModel.cpp
+++ b/Engine/Source/Runtime/Engine/Private/SkeletalMeshLODModel.cpp
@@ -442,12 +442,43 @@ FArchive& operator<<(FArchive& Ar, FSkelMeshSection& S)
 		{
 			S.GenerateUpToLodIndex = -1;
 		}
+
+		if (Ar.CustomVer(FEditorObjectVersion::GUID) >= FEditorObjectVersion::SkeletalMeshBuildRefactor)
+		{
+			Ar << S.OriginalDataSectionIndex;
+			Ar << S.ChunkedParentSectionIndex;
+		}
+		else if (Ar.IsLoading())
+		{
+			S.OriginalDataSectionIndex = INDEX_NONE;
+			S.ChunkedParentSectionIndex = INDEX_NONE;
+		}
 		return Ar;
 	}
 
 	return Ar;
 }
 
+// Serialization.
+FArchive& operator<<(FArchive& Ar, FSkelMeshSourceSectionUserData& S)
+{
+	FStripDataFlags StripFlags(Ar);
+	// When data is cooked we do not serialize anything
+	//This is for editor only editing
+	if (StripFlags.IsEditorDataStripped())
+	{
+		return Ar;
+	}
+
+	Ar << S.bRecomputeTangent;
+	Ar << S.bCastShadow;
+	Ar << S.bDisabled;
+	Ar << S.GenerateUpToLodIndex;
+	Ar << S.CorrespondClothAssetIndex;
+	Ar << S.ClothingData;
+
+	return Ar;
+}
 
 //////////////////////////////////////////////////////////////////////////
 
@@ -555,12 +586,19 @@ void FSkeletalMeshLODModel::Serialize(FArchive& Ar, UObject* Owner, int32 Idx)
 
 	Ar.UsingCustomVersion(FSkeletalMeshCustomVersion::GUID);
 	Ar.UsingCustomVersion(FFortniteMainBranchObjectVersion::GUID);
+	Ar.UsingCustomVersion(FEditorObjectVersion::GUID);
 
 	if (StripFlags.IsDataStrippedForServer())
 	{
 		TArray<FSkelMeshSection> TempSections;
 		Ar << TempSections;
 
+		if (Ar.CustomVer(FEditorObjectVersion::GUID) >= FEditorObjectVersion::SkeletalMeshBuildRefactor)
+		{
+			TMap<int32, FSkelMeshSourceSectionUserData> TempUserSectionsData;
+			Ar << TempUserSectionsData;
+		}
+
 		// For old content, load as a multi-size container
 		if (Ar.IsLoading() && Ar.CustomVer(FSkeletalMeshCustomVersion::GUID) < FSkeletalMeshCustomVersion::SplitModelAndRenderData)
 		{
@@ -579,6 +617,12 @@ void FSkeletalMeshLODModel::Serialize(FArchive& Ar, UObject* Owner, int32 Idx)
 	else
 	{
 		Ar << Sections;
+		
+		if (!StripFlags.IsEditorDataStripped() && Ar.CustomVer(FEditorObjectVersion::GUID) >= FEditorObjectVersion::SkeletalMeshBuildRefactor)
+		{
+			//Editor builds only
+			Ar << UserSectionsData;
+		}
 
 		// For old content, load as a multi-size container, but convert into regular array
 		if (Ar.IsLoading() && Ar.CustomVer(FSkeletalMeshCustomVersion::GUID) < FSkeletalMeshCustomVersion::SplitModelAndRenderData)
@@ -887,7 +931,7 @@ void FSkeletalMeshLODModel::GetNonClothVertices(TArray<FSoftSkinVertex>& OutVert
 	for (int32 SectionIndex = 0; SectionIndex < NumSections; SectionIndex++)
 	{
 		const FSkelMeshSection& Section = Sections[SectionIndex];
-		if (Section.HasClothingData())
+		if (Section.ClothingData.AssetGuid.IsValid())
 		{
 			continue;
 		}
@@ -909,4 +953,74 @@ bool FSkeletalMeshLODModel::DoSectionsNeedExtraBoneInfluences() const
 	return false;
 }
 
+void FSkeletalMeshLODModel::SyncronizeUserSectionsDataArray(bool bResetNonUsedSection /*= false*/)
+{
+	int32 SectionNum = Sections.Num();
+	for (int32 SectionIndex = 0; SectionIndex < SectionNum; ++SectionIndex)
+	{
+		FSkelMeshSection& Section = Sections[SectionIndex];
+		FSkelMeshSourceSectionUserData& SectionUserData = UserSectionsData.FindOrAdd(Section.OriginalDataSectionIndex);
+		Section.bCastShadow					= SectionUserData.bCastShadow;
+		Section.bRecomputeTangent			= SectionUserData.bRecomputeTangent;
+		Section.bDisabled					= SectionUserData.bDisabled;
+		Section.GenerateUpToLodIndex		= SectionUserData.GenerateUpToLodIndex;
+		Section.CorrespondClothAssetIndex	= SectionUserData.CorrespondClothAssetIndex;
+		Section.ClothingData.AssetGuid		= SectionUserData.ClothingData.AssetGuid;
+		Section.ClothingData.AssetLodIndex	= SectionUserData.ClothingData.AssetLodIndex;
+	}
+
+	//Reset normally happen when we re-import a skeletalmesh, we never want to reset this when we build the skeletalmesh (reduce can remove section, but we need to keep the original section data)
+	if (bResetNonUsedSection)
+	{
+		//Make sure we have the correct amount of UserSectionData we delete all the entries and recreate them with the previously sync Sections
+		UserSectionsData.Reset();
+		for (int32 SectionIndex = 0; SectionIndex < SectionNum; ++SectionIndex)
+		{
+			FSkelMeshSection& Section = Sections[SectionIndex];
+			//We only need parent section, no need to iterate bone chunked sections
+			if (Section.ChunkedParentSectionIndex != INDEX_NONE)
+			{
+				continue;
+			}
+			FSkelMeshSourceSectionUserData& SectionUserData = UserSectionsData.FindOrAdd(Section.OriginalDataSectionIndex);
+			SectionUserData.bCastShadow = Section.bCastShadow;
+			SectionUserData.bRecomputeTangent = Section.bRecomputeTangent;
+			SectionUserData.bDisabled = Section.bDisabled;
+			SectionUserData.GenerateUpToLodIndex = Section.GenerateUpToLodIndex;
+			SectionUserData.CorrespondClothAssetIndex = Section.CorrespondClothAssetIndex;
+			SectionUserData.ClothingData.AssetGuid = Section.ClothingData.AssetGuid;
+			SectionUserData.ClothingData.AssetLodIndex = Section.ClothingData.AssetLodIndex;
+		}
+	}
+}
+
+FString FSkeletalMeshLODModel::GetLODModelDeriveDataKey() const
+{
+	FString KeySuffix = TEXT("LODMODEL");
+
+	TArray<uint8> ByteData;
+	FMemoryWriter Ar(ByteData, true);
+
+	//Add the bulk data ID (if someone modify the original imported data, this ID will change)
+	FString BulkDatIDString = RawSkeletalMeshBulkData.GetIdString();
+	Ar << BulkDatIDString;
+	int32 UserSectionCount = UserSectionsData.Num();
+	Ar << UserSectionCount;
+	for (auto Kvp : UserSectionsData)
+	{
+		Ar << Kvp.Key;
+		Ar << Kvp.Value;
+	}
+
+	FSHA1 Sha;
+	Sha.Update(ByteData.GetData(), ByteData.Num() * ByteData.GetTypeSize());
+	Sha.Final();
+	// Retrieve the hash and use it to construct a pseudo-GUID.
+	uint32 Hash[5];
+	Sha.GetHash((uint8*)Hash);
+	KeySuffix += FGuid(Hash[0] ^ Hash[4], Hash[1], Hash[2], Hash[3]).ToString(EGuidFormats::Digits);
+
+	return KeySuffix;
+}
+
 #endif // WITH_EDITOR
diff --git a/Engine/Source/Runtime/Engine/Private/SkeletalMeshLODSettings.cpp b/Engine/Source/Runtime/Engine/Private/SkeletalMeshLODSettings.cpp
--- a/Engine/Source/Runtime/Engine/Private/SkeletalMeshLODSettings.cpp
+++ b/Engine/Source/Runtime/Engine/Private/SkeletalMeshLODSettings.cpp
@@ -6,6 +6,7 @@
 #include "Animation/Skeleton.h"
 #include "Animation/AnimSequence.h"
 #include "UObject/FortniteMainBranchObjectVersion.h"
+#include "UObject/EditorObjectVersion.h"
 
 DEFINE_LOG_CATEGORY_STATIC(LogSkeletalMeshLODSettings, Warning, All)
 
@@ -47,6 +48,7 @@ bool USkeletalMeshLODSettings::SetLODSettingsToMesh(USkeletalMesh* InMesh, int32
 		FSkeletalMeshLODInfo* LODInfo = InMesh->GetLODInfo(LODIndex);
 		const FSkeletalMeshLODGroupSettings& Setting = LODGroups[LODIndex];
 		LODInfo->ReductionSettings = Setting.ReductionSettings;
+		LODInfo->BuildSettings = Setting.BuildSettings;
 		LODInfo->ScreenSize = Setting.ScreenSize;
 		LODInfo->LODHysteresis = Setting.LODHysteresis;
 		LODInfo->WeightOfPrioritization = Setting.WeightOfPrioritization;
@@ -206,6 +208,7 @@ int32 USkeletalMeshLODSettings::SetLODSettingsFromMesh(USkeletalMesh* InMesh)
 			FSkeletalMeshLODInfo* LODInfo = InMesh->GetLODInfo(Index);
 			FSkeletalMeshLODGroupSettings& Setting = LODGroups[Index];
 			Setting.ReductionSettings = LODInfo->ReductionSettings;
+			Setting.BuildSettings = LODInfo->BuildSettings;
 			Setting.ScreenSize = LODInfo->ScreenSize;
 			Setting.LODHysteresis = LODInfo->LODHysteresis;
 			// copy mesh setting to shared setting
@@ -259,6 +262,7 @@ void USkeletalMeshLODSettings::Serialize(FArchive& Ar)
 	Super::Serialize(Ar);
 
 	Ar.UsingCustomVersion(FFortniteMainBranchObjectVersion::GUID);
+	Ar.UsingCustomVersion(FEditorObjectVersion::GUID);
 
 	if (Ar.CustomVer(FFortniteMainBranchObjectVersion::GUID) < FFortniteMainBranchObjectVersion::ConvertReductionSettingOptions)
 	{
@@ -275,6 +279,27 @@ void USkeletalMeshLODSettings::Serialize(FArchive& Ar)
 			}
 		}
 	}
+
+	if (Ar.CustomVer(FEditorObjectVersion::GUID) < FEditorObjectVersion::SkeletalMeshBuildRefactor)
+	{
+		//Since we use the bake source model when loading old asset we do not want to change anything
+		for (int32 Index = 0; Index < LODGroups.Num(); ++Index)
+		{
+			FSkeletalMeshBuildSettings& BuildSettings = LODGroups[Index].BuildSettings;
+			BuildSettings.bRecomputeNormals = false;
+			BuildSettings.bRecomputeTangents = false;
+			BuildSettings.bUseMikkTSpace = false;
+			BuildSettings.bRemoveDegenerates = false;
+			BuildSettings.ThresholdPosition = 0.0f;
+			BuildSettings.ThresholdTangentNormal = 0.0f;
+			BuildSettings.ThresholdUV = 0.0f;
+
+			//The user will have to make sure those settings are ok since before the modification precisions
+			//were one value for the whole skeletal mesh LODs. We set them to false by default
+			BuildSettings.bUseFullPrecisionUVs = GVertexElementTypeSupport.IsSupported(VET_Half2) ? false : true;
+			BuildSettings.bUseHighPrecisionTangentBasis = false;
+		}
+	}
 }
 
 /////////////////////////////////////////////////////////////
@@ -285,6 +310,14 @@ FSkeletalMeshOptimizationSettings FSkeletalMeshLODGroupSettings::GetReductionSet
 	return ReductionSettings;
 }
 
+/////////////////////////////////////////////////////////////
+// FSkeletalMeshBuildSettings 
+/////////////////////////////////////////////////////////////
+FSkeletalMeshBuildSettings FSkeletalMeshLODGroupSettings::GetBuildSettings() const
+{
+	return BuildSettings;
+}
+
 const float FSkeletalMeshLODGroupSettings::GetScreenSize() const
 {
 	return ScreenSize.Default;
diff --git a/Engine/Source/Runtime/Engine/Private/SkeletalMeshMerge.cpp b/Engine/Source/Runtime/Engine/Private/SkeletalMeshMerge.cpp
--- a/Engine/Source/Runtime/Engine/Private/SkeletalMeshMerge.cpp
+++ b/Engine/Source/Runtime/Engine/Private/SkeletalMeshMerge.cpp
@@ -163,8 +163,6 @@ bool FSkeletalMeshMerge::FinalizeMesh()
 	// If things are going ok so far...
 	if (Result)
 	{
-		// force 16 bit UVs if supported on hardware
-		MergeMesh->bUseFullPrecisionUVs = GVertexElementTypeSupport.IsSupported(VET_Half2) ? false : true;
 
 		// Array of per-lod number of UV sets
 		TArray<uint32> PerLODNumUVSets;
@@ -194,7 +192,9 @@ bool FSkeletalMeshMerge::FinalizeMesh()
 		MergeMesh->AllocateResourceForRendering();
 		for (int32 LODIdx = 0; LODIdx < MaxNumLODs; LODIdx++)
 		{
-			if (!MergeMesh->bUseFullPrecisionUVs)
+			const FSkeletalMeshLODInfo* LODInfoPtr = MergeMesh->GetLODInfo(LODIdx);
+			bool bUseFullPrecisionUVs = LODInfoPtr ? LODInfoPtr->BuildSettings.bUseFullPrecisionUVs : GVertexElementTypeSupport.IsSupported(VET_Half2) ? false : true;
+			if (!bUseFullPrecisionUVs)
 			{
 				if (PerLODExtraBoneInfluences[LODIdx])
 				{
@@ -303,8 +303,7 @@ void FSkeletalMeshMerge::GenerateNewSectionArray( TArray<FNewSectionInfo>& NewSe
 				// get the material for this section
 				int32 MaterialIndex = Section.MaterialIndex;
 				// use the remapping of material indices for all LODs besides the base LOD 
-				if( LODIdx > 0 && 
-					SrcLODInfo.LODMaterialMap.IsValidIndex(Section.MaterialIndex) )
+				if(SrcLODInfo.LODMaterialMap.IsValidIndex(Section.MaterialIndex) && SrcMesh->Materials.IsValidIndex(SrcLODInfo.LODMaterialMap[Section.MaterialIndex]))
 				{
 					MaterialIndex = FMath::Clamp<int32>( SrcLODInfo.LODMaterialMap[Section.MaterialIndex], 0, SrcMesh->Materials.Num() );
 				}
@@ -564,6 +563,9 @@ void FSkeletalMeshMerge::GenerateLODModel( int32 LODIdx )
 				}
 			}
 #endif
+			MergeLODInfo.BuildSettings.bUseFullPrecisionUVs |= SrcLODInfo.BuildSettings.bUseFullPrecisionUVs;
+			MergeLODInfo.BuildSettings.bUseHighPrecisionTangentBasis |= SrcLODInfo.BuildSettings.bUseHighPrecisionTangentBasis;
+
 			MergeLODInfo.LODHysteresis = FMath::Min<float>(MergeLODInfo.LODHysteresis,SrcLODInfo.LODHysteresis);
 
 			// get the source skel LOD model from this merge entry
@@ -730,9 +732,9 @@ void FSkeletalMeshMerge::GenerateLODModel( int32 LODIdx )
 	// sort required bone array in strictly increasing order
 	MergeLODData.RequiredBones.Sort();
 	MergeMesh->RefSkeleton.EnsureParentsExistAndSort(MergeLODData.ActiveBoneIndices);
-
+	
 	// copy the new vertices and indices to the vertex buffer for the new model
-	MergeLODData.StaticVertexBuffers.StaticMeshVertexBuffer.SetUseFullPrecisionUVs(MergeMesh->bUseFullPrecisionUVs);
+	MergeLODData.StaticVertexBuffers.StaticMeshVertexBuffer.SetUseFullPrecisionUVs(MergeLODInfo.BuildSettings.bUseFullPrecisionUVs);
 
 	MergeLODData.StaticVertexBuffers.PositionVertexBuffer.Init(MergedVertexBuffer.Num(), bNeedsCPUAccess);
 	MergeLODData.StaticVertexBuffers.StaticMeshVertexBuffer.Init(MergedVertexBuffer.Num(), TotalNumUVs, bNeedsCPUAccess);
diff --git a/Engine/Source/Runtime/Engine/Private/SkeletalMeshModel.cpp b/Engine/Source/Runtime/Engine/Private/SkeletalMeshModel.cpp
--- a/Engine/Source/Runtime/Engine/Private/SkeletalMeshModel.cpp
+++ b/Engine/Source/Runtime/Engine/Private/SkeletalMeshModel.cpp
@@ -89,6 +89,31 @@ FString FSkeletalMeshModel::GetIdString() const
 	return GuidString;
 }
 
+void FSkeletalMeshModel::SyncronizeLODUserSectionsData()
+{
+	for (int32 LODIndex = 0; LODIndex < LODModels.Num(); ++LODIndex)
+	{
+		FSkeletalMeshLODModel& Model = LODModels[LODIndex];
+		Model.SyncronizeUserSectionsDataArray();
+	}
+}
+
+FString FSkeletalMeshModel::GetLODModelIdString() const
+{
+	FSHA1 Sha;
+	for (int32 LODIndex = 0; LODIndex < LODModels.Num(); ++LODIndex)
+	{
+		const FSkeletalMeshLODModel& Model = LODModels[LODIndex];
+		Model.BuildStringID = Model.GetLODModelDeriveDataKey();
+		TArray<TCHAR> IDArray = Model.BuildStringID.GetCharArray();
+		Sha.Update((uint8*)IDArray.GetData(), IDArray.Num() * IDArray.GetTypeSize());
+	}
+	Sha.Final();
+	// Retrieve the hash and use it to construct a pseudo-GUID. 
+	uint32 Hash[5];
+	Sha.GetHash((uint8*)Hash);
+	return FGuid(Hash[0] ^ Hash[4], Hash[1], Hash[2], Hash[3]).ToString(EGuidFormats::Digits);
+}
 
 void FSkeletalMeshModel::GetResourceSizeEx(FResourceSizeEx& CumulativeResourceSize)
 {
