Copying //UE4/Dev-Editor to Dev-Main (//UE4/Dev-Main) for 4.24

diff --git a/Engine/Source/Editor/PropertyEditor/Private/PropertyCustomizationHelpers.cpp b/Engine/Source/Editor/PropertyEditor/Private/PropertyCustomizationHelpers.cpp
--- a/Engine/Source/Editor/PropertyEditor/Private/PropertyCustomizationHelpers.cpp
+++ b/Engine/Source/Editor/PropertyEditor/Private/PropertyCustomizationHelpers.cpp
@@ -815,9 +815,9 @@ public:
 	* @param Section		The Section being used
 	* @param bCanBeReplced	Whether or not the Section can be replaced by a user
 	*/
-	virtual void AddSection(int32 LodIndex, int32 SectionIndex, FName InMaterialSlotName, int32 InMaterialSlotIndex, FName InOriginalMaterialSlotName, const TMap<int32, FName> &InAvailableMaterialSlotName, const UMaterialInterface* Material, bool IsSectionUsingCloth) override
+	virtual void AddSection(int32 LodIndex, int32 SectionIndex, FName InMaterialSlotName, int32 InMaterialSlotIndex, FName InOriginalMaterialSlotName, const TMap<int32, FName> &InAvailableMaterialSlotName, const UMaterialInterface* Material, bool IsSectionUsingCloth, bool bIsChunkSection, int32 DefaultMaterialIndex) override
 	{
-		FSectionListItem SectionItem(LodIndex, SectionIndex, InMaterialSlotName, InMaterialSlotIndex, InOriginalMaterialSlotName, InAvailableMaterialSlotName, Material, IsSectionUsingCloth, ThumbnailSize);
+		FSectionListItem SectionItem(LodIndex, SectionIndex, InMaterialSlotName, InMaterialSlotIndex, InOriginalMaterialSlotName, InAvailableMaterialSlotName, Material, IsSectionUsingCloth, ThumbnailSize, bIsChunkSection, DefaultMaterialIndex);
 		if (!Sections.Contains(SectionItem))
 		{
 			Sections.Add(SectionItem);
@@ -934,7 +934,10 @@ public:
 
 	TSharedRef<SWidget> CreateValueContent(const TSharedPtr<FAssetThumbnailPool>& ThumbnailPool)
 	{
-		FText MaterialSlotNameTooltipText = SectionItem.IsSectionUsingCloth ? FText(LOCTEXT("SectionIndex_MaterialSlotNameTooltip", "Cannot change the material slot when the mesh section use the cloth system.")) : FText::GetEmpty();
+		FFormatNamedArguments Arguments;
+		Arguments.Add(TEXT("DefaultMaterialIndex"), SectionItem.DefaultMaterialIndex);
+		FText BaseMaterialSlotTooltip = SectionItem.DefaultMaterialIndex != SectionItem.MaterialSlotIndex ? FText::Format(LOCTEXT("SectionIndex_BaseMaterialSlotNameTooltip", "This section material slot was change from the default value [{DefaultMaterialIndex}]."), Arguments) : FText::GetEmpty();
+		FText MaterialSlotNameTooltipText = SectionItem.IsSectionUsingCloth ? FText(LOCTEXT("SectionIndex_MaterialSlotNameTooltip", "Cannot change the material slot when the mesh section use the cloth system.")) : BaseMaterialSlotTooltip;
 		return
 			SNew(SHorizontalBox)
 			+ SHorizontalBox::Slot()
@@ -948,6 +951,7 @@ public:
 				.HAlign(HAlign_Fill)
 				[
 					SNew(SHorizontalBox)
+					.Visibility(SectionItem.bIsChunkSection ? EVisibility::Collapsed : EVisibility::Visible)
 					+ SHorizontalBox::Slot()
 					.FillWidth(1.0f)
 					[
@@ -1023,6 +1027,22 @@ public:
 						]
 					]
 				]
+				+ SVerticalBox::Slot()
+				.AutoHeight()
+				.Padding(0.0f)
+				.VAlign(VAlign_Center)
+				.HAlign(HAlign_Fill)
+				[
+					SNew(SHorizontalBox)
+					.Visibility(SectionItem.bIsChunkSection ? EVisibility::Visible : EVisibility::Collapsed)
+					+ SHorizontalBox::Slot()
+					.FillWidth(1.0f)
+					[
+						SNew(STextBlock)
+						.Font(IDetailLayoutBuilder::GetDetailFont())
+						.Text(LOCTEXT("SectionListItemChunkSectionValueLabel", "Chunked"))
+					]
+				]
 			];
 	}
 
@@ -1076,7 +1096,12 @@ private:
 	{
 		FString MaterialSlotDisplayName;
 		SectionItem.MaterialSlotName.ToString(MaterialSlotDisplayName);
-		MaterialSlotDisplayName = TEXT("[") + FString::FromInt(SectionItem.MaterialSlotIndex) + TEXT("] ") + MaterialSlotDisplayName;
+		FString MaterialSlotRemapString = TEXT("");
+		if (SectionItem.DefaultMaterialIndex != SectionItem.MaterialSlotIndex)
+		{
+			MaterialSlotRemapString = TEXT(" (Modified)");
+		}
+		MaterialSlotDisplayName = TEXT("[") + FString::FromInt(SectionItem.MaterialSlotIndex) + TEXT("] ") + MaterialSlotDisplayName + MaterialSlotRemapString;
 		return FText::FromString(MaterialSlotDisplayName);
 	}
 
@@ -1210,7 +1235,7 @@ void FSectionList::GenerateChildContent(IDetailChildrenBuilder& ChildrenBuilder)
 			if (bDisplayAllSectionsInSlot)
 			{
 				FDetailWidgetRow& ChildRow = ChildrenBuilder.AddCustomRow(Section.Material.IsValid() ? FText::FromString(Section.Material->GetName()) : FText::GetEmpty());
-				AddSectionItem(ChildRow, CurrentLODIndex, FSectionListItem(CurrentLODIndex, Section.SectionIndex, Section.MaterialSlotName, Section.MaterialSlotIndex, Section.OriginalMaterialSlotName, Section.AvailableMaterialSlotName, Section.Material.Get(), Section.IsSectionUsingCloth, ThumbnailSize), !bDisplayAllSectionsInSlot);
+				AddSectionItem(ChildRow, CurrentLODIndex, FSectionListItem(CurrentLODIndex, Section.SectionIndex, Section.MaterialSlotName, Section.MaterialSlotIndex, Section.OriginalMaterialSlotName, Section.AvailableMaterialSlotName, Section.Material.Get(), Section.IsSectionUsingCloth, ThumbnailSize, Section.bIsChunkSection, Section.DefaultMaterialIndex), !bDisplayAllSectionsInSlot);
 			}
 		}
 	}
@@ -1292,6 +1317,7 @@ void FSectionList::AddSectionItem(FDetailWidgetRow& Row, int32 LodIndex, const s
 {
 	uint32 NumSections = SectionListBuilder->GetNumSections(LodIndex);
 
+	bool bIsChunkSection = Item.bIsChunkSection;
 	TSharedRef<FSectionItemView> NewView = FSectionItemView::Create(Item, SectionListDelegates.OnSectionChanged, SectionListDelegates.OnGenerateCustomNameWidgets, SectionListDelegates.OnGenerateCustomSectionWidgets, SectionListDelegates.OnResetSectionToDefaultClicked, NumSections, ThumbnailSize);
 
 	TSharedPtr<SWidget> RightSideContent;
@@ -1318,13 +1344,17 @@ void FSectionList::AddSectionItem(FDetailWidgetRow& Row, int32 LodIndex, const s
 		ViewedSections.Add(NewView);
 	}
 
-	Row.CopyAction(FUIAction(FExecuteAction::CreateSP(this, &FSectionList::OnCopySectionItem, LodIndex, Item.SectionIndex), FCanExecuteAction::CreateSP(this, &FSectionList::OnCanCopySectionItem, LodIndex, Item.SectionIndex)));
-	Row.PasteAction(FUIAction(FExecuteAction::CreateSP(this, &FSectionList::OnPasteSectionItem, LodIndex, Item.SectionIndex)));
-
-	if(SectionListDelegates.OnEnableSectionItem.IsBound())
+	//Chunk section cannot be copy enable or disable, do the operation on the parent section
+	if (!bIsChunkSection)
 	{
-		Row.AddCustomContextMenuAction(FUIAction(FExecuteAction::CreateSP(this, &FSectionList::OnEnableSectionItem, LodIndex, Item.SectionIndex, true)), LOCTEXT("SectionItemContexMenu_Enable", "Enable"));
-		Row.AddCustomContextMenuAction(FUIAction(FExecuteAction::CreateSP(this, &FSectionList::OnEnableSectionItem, LodIndex, Item.SectionIndex, false)), LOCTEXT("SectionItemContexMenu_Disable", "Disable"));
+		Row.CopyAction(FUIAction(FExecuteAction::CreateSP(this, &FSectionList::OnCopySectionItem, LodIndex, Item.SectionIndex), FCanExecuteAction::CreateSP(this, &FSectionList::OnCanCopySectionItem, LodIndex, Item.SectionIndex)));
+		Row.PasteAction(FUIAction(FExecuteAction::CreateSP(this, &FSectionList::OnPasteSectionItem, LodIndex, Item.SectionIndex)));
+
+		if (SectionListDelegates.OnEnableSectionItem.IsBound())
+		{
+			Row.AddCustomContextMenuAction(FUIAction(FExecuteAction::CreateSP(this, &FSectionList::OnEnableSectionItem, LodIndex, Item.SectionIndex, true)), LOCTEXT("SectionItemContexMenu_Enable", "Enable"));
+			Row.AddCustomContextMenuAction(FUIAction(FExecuteAction::CreateSP(this, &FSectionList::OnEnableSectionItem, LodIndex, Item.SectionIndex, false)), LOCTEXT("SectionItemContexMenu_Disable", "Disable"));
+		}
 	}
 
 	Row.NameContent()
diff --git a/Engine/Source/Editor/PropertyEditor/Public/PropertyCustomizationHelpers.h b/Engine/Source/Editor/PropertyEditor/Public/PropertyCustomizationHelpers.h
--- a/Engine/Source/Editor/PropertyEditor/Public/PropertyCustomizationHelpers.h
+++ b/Engine/Source/Editor/PropertyEditor/Public/PropertyCustomizationHelpers.h
@@ -576,7 +576,7 @@ public:
 	* @param Section		The Section being used
 	* @param bCanBeReplced	Whether or not the Section can be replaced by a user
 	*/
-	virtual void AddSection(int32 LodIndex, int32 SectionIndex, FName InMaterialSlotName, int32 InMaterialSlotIndex, FName InOriginalMaterialSlotName, const TMap<int32, FName> &InAvailableMaterialSlotName, const UMaterialInterface* Material, bool IsSectionUsingCloth) = 0;
+	virtual void AddSection(int32 LodIndex, int32 SectionIndex, FName InMaterialSlotName, int32 InMaterialSlotIndex, FName InOriginalMaterialSlotName, const TMap<int32, FName> &InAvailableMaterialSlotName, const UMaterialInterface* Material, bool IsSectionUsingCloth, bool bIsChunkSection, int32 DefaultMaterialIndex) = 0;
 };
 
 
@@ -607,8 +607,11 @@ struct FSectionListItem
 	/* Available material slot name*/
 	TMap<int32, FName> AvailableMaterialSlotName;
 
+	bool bIsChunkSection;
+	int32 DefaultMaterialIndex;
 
-	FSectionListItem(int32 InLodIndex, int32 InSectionIndex, FName InMaterialSlotName, int32 InMaterialSlotIndex, FName InOriginalMaterialSlotName, const TMap<int32, FName> &InAvailableMaterialSlotName, const UMaterialInterface* InMaterial, bool InIsSectionUsingCloth, int32 InThumbnailSize)
+
+	FSectionListItem(int32 InLodIndex, int32 InSectionIndex, FName InMaterialSlotName, int32 InMaterialSlotIndex, FName InOriginalMaterialSlotName, const TMap<int32, FName> &InAvailableMaterialSlotName, const UMaterialInterface* InMaterial, bool InIsSectionUsingCloth, int32 InThumbnailSize, bool InIsChunkSection, int32 InDefaultMaterialIndex)
 		: LodIndex(InLodIndex)
 		, SectionIndex(InSectionIndex)
 		, IsSectionUsingCloth(InIsSectionUsingCloth)
@@ -618,11 +621,21 @@ struct FSectionListItem
 		, MaterialSlotIndex(InMaterialSlotIndex)
 		, OriginalMaterialSlotName(InOriginalMaterialSlotName)
 		, AvailableMaterialSlotName(InAvailableMaterialSlotName)
+		, bIsChunkSection(InIsChunkSection)
+		, DefaultMaterialIndex(InDefaultMaterialIndex)
 	{}
 
 	bool operator==(const FSectionListItem& Other) const
 	{
-		bool IsSectionItemEqual = LodIndex == Other.LodIndex && SectionIndex == Other.SectionIndex && MaterialSlotIndex == Other.MaterialSlotIndex && MaterialSlotName == Other.MaterialSlotName && Material == Other.Material && AvailableMaterialSlotName.Num() == Other.AvailableMaterialSlotName.Num() && IsSectionUsingCloth == Other.IsSectionUsingCloth;
+		bool IsSectionItemEqual = LodIndex == Other.LodIndex
+			&& SectionIndex == Other.SectionIndex
+			&& MaterialSlotIndex == Other.MaterialSlotIndex
+			&& MaterialSlotName == Other.MaterialSlotName
+			&& Material == Other.Material
+			&& AvailableMaterialSlotName.Num() == Other.AvailableMaterialSlotName.Num()
+			&& IsSectionUsingCloth == Other.IsSectionUsingCloth
+			&& bIsChunkSection == Other.bIsChunkSection
+			&& DefaultMaterialIndex == Other.DefaultMaterialIndex;
 		if (IsSectionItemEqual)
 		{
 			for (auto Kvp : AvailableMaterialSlotName)
diff --git a/Engine/Source/Editor/SkeletalMeshEditor/Private/SkeletalMeshEditor.cpp b/Engine/Source/Editor/SkeletalMeshEditor/Private/SkeletalMeshEditor.cpp
--- a/Engine/Source/Editor/SkeletalMeshEditor/Private/SkeletalMeshEditor.cpp
+++ b/Engine/Source/Editor/SkeletalMeshEditor/Private/SkeletalMeshEditor.cpp
@@ -44,6 +44,8 @@
 #include "EditorFramework/AssetImportData.h"
 #include "Factories/FbxSkeletalMeshImportData.h"
 
+#include "Misc/MessageDialog.h"
+
 const FName SkeletalMeshEditorAppIdentifier = FName(TEXT("SkeletalMeshEditorApp"));
 
 const FName SkeletalMeshEditorModes::SkeletalMeshEditorMode(TEXT("SkeletalMeshEditorMode"));
@@ -78,6 +80,112 @@ FSkeletalMeshEditor::~FSkeletalMeshEditor()
 	}
 }
 
+bool IsReductionParentBaseLODUseSkeletalMeshBuildWorkflow(USkeletalMesh* SkeletalMesh, int32 TestLODIndex)
+{
+	FSkeletalMeshLODInfo* LODInfo = SkeletalMesh->GetLODInfo(TestLODIndex);
+	if (LODInfo == nullptr || !SkeletalMesh->GetImportedModel() || !SkeletalMesh->GetImportedModel()->LODModels.IsValidIndex(TestLODIndex))
+	{
+		return false;
+	}
+
+	if (SkeletalMesh->GetImportedModel()->LODModels[TestLODIndex].RawSkeletalMeshBulkData.IsBuildDataAvailable())
+	{
+		return true;
+	}
+	if (LODInfo->bHasBeenSimplified || SkeletalMesh->IsReductionActive(TestLODIndex))
+	{
+		int32 ReduceBaseLOD = LODInfo->ReductionSettings.BaseLOD;
+		if (ReduceBaseLOD < TestLODIndex)
+		{
+			return IsReductionParentBaseLODUseSkeletalMeshBuildWorkflow(SkeletalMesh, ReduceBaseLOD);
+		}
+	}
+	return false;
+}
+
+bool FSkeletalMeshEditor::OnRequestClose()
+{
+	bool bAllowClose = true;
+
+	if (PersonaToolkit.IsValid() && SkeletalMesh)
+	{
+		bool bHaveModifiedLOD = false;
+		for (int32 LODIndex = 0; LODIndex < SkeletalMesh->GetLODNum(); ++LODIndex)
+		{
+			FSkeletalMeshLODInfo* LODInfo = SkeletalMesh->GetLODInfo(LODIndex);
+			if (LODInfo == nullptr || !SkeletalMesh->GetImportedModel() || !SkeletalMesh->GetImportedModel()->LODModels.IsValidIndex(LODIndex))
+			{
+				continue;
+			}
+			
+			//Do not prevent exiting if we are not using the skeletal mesh build workflow
+			if (!SkeletalMesh->GetImportedModel()->LODModels[LODIndex].RawSkeletalMeshBulkData.IsBuildDataAvailable())
+			{
+				if (!LODInfo->bHasBeenSimplified && !SkeletalMesh->IsReductionActive(LODIndex))
+				{
+					continue;
+				}
+				//Do not prevent exit if the generated LOD is base on a LODModel not using the skeletalmesh build workflow
+				int32 ReduceBaseLOD = LODInfo->ReductionSettings.BaseLOD;
+				if(!IsReductionParentBaseLODUseSkeletalMeshBuildWorkflow(SkeletalMesh, ReduceBaseLOD))
+				{
+					continue;
+				}
+				
+			}
+
+			bool bValidLODSettings = false;
+			if (SkeletalMesh->LODSettings != nullptr)
+			{
+				const int32 NumSettings = FMath::Min(SkeletalMesh->LODSettings->GetNumberOfSettings(), SkeletalMesh->GetLODNum());
+				if (LODIndex < NumSettings)
+				{
+					bValidLODSettings = true;
+				}
+			}
+			const FSkeletalMeshLODGroupSettings* SkeletalMeshLODGroupSettings = bValidLODSettings ? &SkeletalMesh->LODSettings->GetSettingsForLODLevel(LODIndex) : nullptr;
+
+			FGuid BuildGUID = LODInfo->ComputeDeriveDataCacheKey(SkeletalMeshLODGroupSettings);
+			if (LODInfo->BuildGUID != BuildGUID)
+			{
+				bHaveModifiedLOD = true;
+				break;
+			}
+			FString BuildStringID = SkeletalMesh->GetImportedModel()->LODModels[LODIndex].GetLODModelDeriveDataKey();
+			if (SkeletalMesh->GetImportedModel()->LODModels[LODIndex].BuildStringID != BuildStringID)
+			{
+				bHaveModifiedLOD = true;
+				break;
+			}
+		}
+
+		if (bHaveModifiedLOD)
+		{
+			// find out the user wants to do with this dirty material
+			EAppReturnType::Type OkCancelReply = FMessageDialog::Open(
+				EAppMsgType::OkCancel,
+				FText::Format(LOCTEXT("SkeletalMeshEditorShouldApplyLODChanges", "We have to apply level of detail changes to {0} before exiting the skeletal mesh editor."), FText::FromString(PersonaToolkit->GetMesh()->GetName()))
+			);
+
+			switch (OkCancelReply)
+			{
+			case EAppReturnType::Ok:
+				{
+					SkeletalMesh->MarkPackageDirty();
+					SkeletalMesh->PostEditChange();
+					bAllowClose = true;
+				}
+				break;
+			case EAppReturnType::Cancel:
+				// Don't exit.
+				bAllowClose = false;
+				break;
+			}
+		}
+	}
+	return bAllowClose;
+}
+
 void FSkeletalMeshEditor::RegisterTabSpawners(const TSharedRef<class FTabManager>& InTabManager)
 {
 	WorkspaceMenuCategory = InTabManager->AddLocalWorkspaceMenuCategory(LOCTEXT("WorkspaceMenu_SkeletalMeshEditor", "Skeletal Mesh Editor"));
@@ -546,8 +654,7 @@ void FSkeletalMeshEditor::OnRemoveSectionFromLodAndBelowMenuItemClicked(int32 Lo
 		const FSkeletalMeshLODInfo* CurrentSkeletalMeshLODInfo = SkeletalMesh->GetLODInfo(GenerateLodIndex);
 		if (CurrentSkeletalMeshLODInfo != nullptr && CurrentSkeletalMeshLODInfo->bHasBeenSimplified && BaseLodIndexes.Contains(CurrentSkeletalMeshLODInfo->ReductionSettings.BaseLOD))
 		{
-			bool bRestoreClothing = true;
-			FLODUtilities::SimplifySkeletalMeshLOD(UpdateContext, GenerateLodIndex, true, bRestoreClothing);
+			FLODUtilities::SimplifySkeletalMeshLOD(UpdateContext, GenerateLodIndex, true);
 			BaseLodIndexes.Add(GenerateLodIndex);
 		}
 	}
@@ -668,7 +775,7 @@ void FSkeletalMeshEditor::OnCreateClothingAssetMenuItemClicked(FSkeletalMeshClot
 						{
 							SectionIndex = i;
 							AssetLodIndex = SkelLod.Sections[i].ClothingData.AssetLodIndex;
-							TargetAssetPtr->UnbindFromSkeletalMesh(Mesh, Params.TargetLod);
+							RemoveClothing(Params.TargetLod, SectionIndex);
 							break;
 						}
 					}
@@ -786,13 +893,54 @@ void FSkeletalMeshEditor::ApplyClothing(UClothingAssetBase* InAsset, int32 InLod
 {
 	USkeletalMesh* Mesh = GetPersonaToolkit()->GetPreviewMesh();
 
-	if(UClothingAsset* ClothingAsset = Cast<UClothingAsset>(InAsset))
+	if (Mesh == nullptr || Mesh->GetImportedModel() == nullptr || !Mesh->GetImportedModel()->LODModels.IsValidIndex(InLodIndex))
 	{
-		ClothingAsset->BindToSkeletalMesh(Mesh, InLodIndex, InSectionIndex, InClothingLod);
+		return;
 	}
-	else
+
+	FSkeletalMeshLODModel& LODModel = Mesh->GetImportedModel()->LODModels[InLodIndex];
+	const FSkelMeshSection& Section = LODModel.Sections[InSectionIndex];
+	FScopedSuspendAlternateSkinWeightPreview ScopedSuspendAlternateSkinnWeightPreview(Mesh);
 	{
-		RemoveClothing(InLodIndex, InSectionIndex);
+		FScopedSkeletalMeshPostEditChange ScopedPostEditChange(Mesh);
+		FScopedTransaction Transaction(LOCTEXT("SkeletalMeshEditorApplyClothingTransaction", "Persona editor: Apply Section Cloth"));
+		Mesh->Modify();
+
+		FSkelMeshSourceSectionUserData& OriginalSectionData = LODModel.UserSectionsData.FindOrAdd(Section.OriginalDataSectionIndex);
+		auto ClearOriginalSectionUserData = [&OriginalSectionData]()
+		{
+			OriginalSectionData.CorrespondClothAssetIndex = INDEX_NONE;
+			OriginalSectionData.ClothingData.AssetGuid = FGuid();
+			OriginalSectionData.ClothingData.AssetLodIndex = INDEX_NONE;
+		};
+		if (UClothingAsset* ClothingAsset = Cast<UClothingAsset>(InAsset))
+		{
+			// Look for a currently bound asset an unbind it if necessary first
+			if (UClothingAssetBase* CurrentAsset = Mesh->GetSectionClothingAsset(InLodIndex, InSectionIndex))
+			{
+				CurrentAsset->UnbindFromSkeletalMesh(Mesh, InLodIndex);
+				ClearOriginalSectionUserData();
+			}
+
+			if (ClothingAsset->BindToSkeletalMesh(Mesh, InLodIndex, InSectionIndex, InClothingLod))
+			{
+				//Successful bind so set the SectionUserData
+				int32 AssetIndex = INDEX_NONE;
+				check(Mesh->MeshClothingAssets.Find(ClothingAsset, AssetIndex));
+				OriginalSectionData.CorrespondClothAssetIndex = AssetIndex;
+				OriginalSectionData.ClothingData.AssetGuid = ClothingAsset->GetAssetGuid();
+				OriginalSectionData.ClothingData.AssetLodIndex = InClothingLod;
+			}
+		}
+		else if (Mesh)
+		{
+			//User set none, so unbind anything that is bind
+			if (UClothingAssetBase* CurrentAsset = Mesh->GetSectionClothingAsset(InLodIndex, InSectionIndex))
+			{
+				CurrentAsset->UnbindFromSkeletalMesh(Mesh, InLodIndex);
+				ClearOriginalSectionUserData();
+			}
+		}
 	}
 }
 
@@ -800,11 +948,31 @@ void FSkeletalMeshEditor::RemoveClothing(int32 InLodIndex, int32 InSectionIndex)
 {
 	USkeletalMesh* Mesh = GetPersonaToolkit()->GetPreviewMesh();
 
-	if(Mesh)
+	if (Mesh->GetImportedModel() == nullptr || !Mesh->GetImportedModel()->LODModels.IsValidIndex(InLodIndex))
 	{
-		if(UClothingAssetBase* CurrentAsset = Mesh->GetSectionClothingAsset(InLodIndex, InSectionIndex))
+		return;
+	}
+
+	FSkeletalMeshLODModel& LODModel = Mesh->GetImportedModel()->LODModels[InLodIndex];
+	const FSkelMeshSection& Section = LODModel.Sections[InSectionIndex];
+	FScopedSuspendAlternateSkinWeightPreview ScopedSuspendAlternateSkinnWeightPreview(Mesh);
+	{
+		FScopedSkeletalMeshPostEditChange ScopedPostEditChange(Mesh);
+		FScopedTransaction Transaction(LOCTEXT("SkeletalMeshEditorRemoveClothingTransaction", "Persona editor: Remove Section Cloth"));
+		Mesh->Modify();
+
+		FSkelMeshSourceSectionUserData& OriginalSectionData = LODModel.UserSectionsData.FindOrAdd(Section.OriginalDataSectionIndex);
+		auto ClearOriginalSectionUserData = [&OriginalSectionData]()
+		{
+			OriginalSectionData.CorrespondClothAssetIndex = INDEX_NONE;
+			OriginalSectionData.ClothingData.AssetGuid = FGuid();
+			OriginalSectionData.ClothingData.AssetLodIndex = INDEX_NONE;
+		};
+		// Look for a currently bound asset an unbind it if necessary first
+		if (UClothingAssetBase* CurrentAsset = Mesh->GetSectionClothingAsset(InLodIndex, InSectionIndex))
 		{
 			CurrentAsset->UnbindFromSkeletalMesh(Mesh, InLodIndex);
+			ClearOriginalSectionUserData();
 		}
 	}
 }
@@ -883,92 +1051,92 @@ UObject* FSkeletalMeshEditor::HandleGetAsset()
 
 bool FSkeletalMeshEditor::HandleReimportMeshInternal(int32 SourceFileIndex /*= INDEX_NONE*/, bool bWithNewFile /*= false*/)
 {
-	// Make sure we clear any skin weight profile currently being previewed
-	UDebugSkelMeshComponent* MeshComponent = GetPersonaToolkit()->GetPreviewScene()->GetPreviewMeshComponent();
-	const FName ProfileName = MeshComponent ? MeshComponent->GetCurrentSkinWeightProfileName() : NAME_None;
-	if (MeshComponent)
-	{
-		MeshComponent->ClearSkinWeightProfile();
-	}
-
-	// Reimport the asset
-	const bool bResult = FReimportManager::Instance()->Reimport(SkeletalMesh, true, true, TEXT(""), nullptr, SourceFileIndex, bWithNewFile);
-
-	// Refresh skeleton tree
-	SkeletonTree->Refresh();
-
-	// Re-set skin weigh profile
-	if (MeshComponent && (ProfileName != NAME_None))
+	FScopedSuspendAlternateSkinWeightPreview ScopedSuspendAlternateSkinnWeightPreview(SkeletalMesh);
+	bool bResult = false;
 	{
-		MeshComponent->SetSkinWeightProfile(ProfileName);
+		FScopedSkeletalMeshPostEditChange ScopedPostEditChange(SkeletalMesh);
+		// Reimport the asset
+		bResult = FReimportManager::Instance()->Reimport(SkeletalMesh, true, true, TEXT(""), nullptr, SourceFileIndex, bWithNewFile);
+		// Refresh skeleton tree
+		SkeletonTree->Refresh();
 	}
-
 	return bResult;
 }
 
 void FSkeletalMeshEditor::HandleReimportMesh(int32 SourceFileIndex /*= INDEX_NONE*/)
 {
+	FScopedSuspendAlternateSkinWeightPreview ScopedSuspendAlternateSkinnWeightPreview(SkeletalMesh);
+
 	HandleReimportMeshInternal(SourceFileIndex, false);
 }
 
 void FSkeletalMeshEditor::HandleReimportMeshWithNewFile(int32 SourceFileIndex /*= INDEX_NONE*/)
 {
+	FScopedSuspendAlternateSkinWeightPreview ScopedSuspendAlternateSkinnWeightPreview(SkeletalMesh);
+
 	HandleReimportMeshInternal(SourceFileIndex, true);
 }
 
 void ReimportAllCustomLODs(USkeletalMesh* SkeletalMesh, UDebugSkelMeshComponent* PreviewMeshComponent, bool bWithNewFile)
 {
-	//Find the dependencies of the generated LOD
-	TArray<bool> Dependencies;
-	Dependencies.AddZeroed(SkeletalMesh->GetLODNum());
-	//Avoid making LOD 0 to true in the dependencies since everything that should be regenerate base on LOD 0 is already regenerate at this point.
-	//But we need to regenerate every generated LOD base on any re-import custom LOD
-	//Reimport all custom LODs
-	for (int32 LodIndex = 1; LodIndex < SkeletalMesh->GetLODNum(); ++LodIndex)
+	FScopedSuspendAlternateSkinWeightPreview ScopedSuspendAlternateSkinnWeightPreview(SkeletalMesh);
 	{
-		//Do not reimport LOD that was re-import with the base mesh
-		if (SkeletalMesh->GetLODInfo(LodIndex)->bImportWithBaseMesh)
-		{
-			continue;
-		}
-		if (SkeletalMesh->GetLODInfo(LodIndex)->bHasBeenSimplified == false)
+		FScopedSkeletalMeshPostEditChange ScopedPostEditChange(SkeletalMesh);
+
+		//Find the dependencies of the generated LOD
+		TArray<bool> Dependencies;
+		Dependencies.AddZeroed(SkeletalMesh->GetLODNum());
+		//Avoid making LOD 0 to true in the dependencies since everything that should be regenerate base on LOD 0 is already regenerate at this point.
+		//But we need to regenerate every generated LOD base on any re-import custom LOD
+		//Reimport all custom LODs
+		for (int32 LodIndex = 1; LodIndex < SkeletalMesh->GetLODNum(); ++LodIndex)
 		{
-			FString SourceFilenameBackup = SkeletalMesh->GetLODInfo(LodIndex)->SourceImportFilename;
-			if (bWithNewFile)
+			//Do not reimport LOD that was re-import with the base mesh
+			if (SkeletalMesh->GetLODInfo(LodIndex)->bImportWithBaseMesh)
 			{
-				SkeletalMesh->GetLODInfo(LodIndex)->SourceImportFilename.Empty();
+				continue;
 			}
-
-			if (!FbxMeshUtils::ImportMeshLODDialog(SkeletalMesh, LodIndex))
+			if (SkeletalMesh->GetLODInfo(LodIndex)->bHasBeenSimplified == false)
 			{
+				FString SourceFilenameBackup = SkeletalMesh->GetLODInfo(LodIndex)->SourceImportFilename;
 				if (bWithNewFile)
 				{
-					SkeletalMesh->GetLODInfo(LodIndex)->SourceImportFilename = SourceFilenameBackup;
+					SkeletalMesh->GetLODInfo(LodIndex)->SourceImportFilename.Empty();
+				}
+
+				if (!FbxMeshUtils::ImportMeshLODDialog(SkeletalMesh, LodIndex, false))
+				{
+					if (bWithNewFile)
+					{
+						SkeletalMesh->GetLODInfo(LodIndex)->SourceImportFilename = SourceFilenameBackup;
+					}
+				}
+				else
+				{
+					Dependencies[LodIndex] = true;
 				}
 			}
-			else
+			else if (Dependencies[SkeletalMesh->GetLODInfo(LodIndex)->ReductionSettings.BaseLOD])
 			{
+				//Regenerate the LOD
+				FSkeletalMeshUpdateContext UpdateContext;
+				UpdateContext.SkeletalMesh = SkeletalMesh;
+				UpdateContext.AssociatedComponents.Push(PreviewMeshComponent);
+				FLODUtilities::SimplifySkeletalMeshLOD(UpdateContext, LodIndex, false);
 				Dependencies[LodIndex] = true;
 			}
 		}
-		else if(Dependencies[SkeletalMesh->GetLODInfo(LodIndex)->ReductionSettings.BaseLOD])
-		{
-			//Regenerate the LOD
-			FSkeletalMeshUpdateContext UpdateContext;
-			UpdateContext.SkeletalMesh = SkeletalMesh;
-			UpdateContext.AssociatedComponents.Push(PreviewMeshComponent);
-			bool bRestoreClothing = true;
-			FLODUtilities::SimplifySkeletalMeshLOD(UpdateContext, LodIndex, true, bRestoreClothing);
-			Dependencies[LodIndex] = true;
-		}
 	}
 }
 
 void FSkeletalMeshEditor::HandleReimportAllMesh(int32 SourceFileIndex /*= INDEX_NONE*/)
 {
+	FScopedSuspendAlternateSkinWeightPreview ScopedSuspendAlternateSkinnWeightPreview(SkeletalMesh);
 	// Reimport the asset
 	if (SkeletalMesh)
 	{
+		FScopedSkeletalMeshPostEditChange ScopedPostEditChange(SkeletalMesh);
+
 		//Reimport base LOD
 		if (HandleReimportMeshInternal(SourceFileIndex, false))
 		{
@@ -980,9 +1148,12 @@ void FSkeletalMeshEditor::HandleReimportAllMesh(int32 SourceFileIndex /*= INDEX_
 
 void FSkeletalMeshEditor::HandleReimportAllMeshWithNewFile(int32 SourceFileIndex /*= INDEX_NONE*/)
 {
+	FScopedSuspendAlternateSkinWeightPreview ScopedSuspendAlternateSkinnWeightPreview(SkeletalMesh);
+
 	// Reimport the asset
 	if (SkeletalMesh)
 	{
+		FScopedSkeletalMeshPostEditChange ScopedPostEditChange(SkeletalMesh);
 		TArray<UObject*> ImportObjs;
 		ImportObjs.Add(SkeletalMesh);
 		if (HandleReimportMeshInternal(SourceFileIndex, true))
diff --git a/Engine/Source/Editor/SkeletalMeshEditor/Private/SkeletalMeshEditor.h b/Engine/Source/Editor/SkeletalMeshEditor/Private/SkeletalMeshEditor.h
--- a/Engine/Source/Editor/SkeletalMeshEditor/Private/SkeletalMeshEditor.h
+++ b/Engine/Source/Editor/SkeletalMeshEditor/Private/SkeletalMeshEditor.h
@@ -62,6 +62,10 @@ public:
 	virtual FString GetWorldCentricTabPrefix() const override;
 	virtual FLinearColor GetWorldCentricTabColorScale() const override;
 
+	//~ Begin FAssetEditorToolkit Interface.
+	virtual bool OnRequestClose() override;
+	//~ End FAssetEditorToolkit Interface.
+
 	/** FEditorUndoClient interface */
 	virtual void PostUndo(bool bSuccess) override;
 	virtual void PostRedo(bool bSuccess) override;
diff --git a/Engine/Source/Editor/StaticMeshEditor/Private/StaticMeshEditor.cpp b/Engine/Source/Editor/StaticMeshEditor/Private/StaticMeshEditor.cpp
--- a/Engine/Source/Editor/StaticMeshEditor/Private/StaticMeshEditor.cpp
+++ b/Engine/Source/Editor/StaticMeshEditor/Private/StaticMeshEditor.cpp
@@ -1398,7 +1398,7 @@ void FStaticMeshEditor::HandleReimportAllMesh()
 					continue;
 				}
 
-				bool bHasBeenSimplified = StaticMesh->GetMeshDescription(LodIndex) == nullptr || StaticMesh->IsReductionActive(LodIndex);
+				bool bHasBeenSimplified = !StaticMesh->IsMeshDescriptionValid(LodIndex) || StaticMesh->IsReductionActive(LodIndex);
 				if (!bHasBeenSimplified)
 				{
 					FbxMeshUtils::ImportMeshLODDialog(StaticMesh, LodIndex);
diff --git a/Engine/Source/Editor/StaticMeshEditor/Private/StaticMeshEditorTools.cpp b/Engine/Source/Editor/StaticMeshEditor/Private/StaticMeshEditorTools.cpp
--- a/Engine/Source/Editor/StaticMeshEditor/Private/StaticMeshEditorTools.cpp
+++ b/Engine/Source/Editor/StaticMeshEditor/Private/StaticMeshEditorTools.cpp
@@ -1971,7 +1971,8 @@ void FMeshSectionSettingsLayout::OnGetSectionsForView(ISectionListBuilder& OutSe
 				{
 					SectionMaterial = UMaterial::GetDefaultMaterial(MD_Surface);
 				}
-				OutSections.AddSection(LODIndex, SectionIndex, CurrentSectionMaterialSlotName, MaterialIndex, CurrentSectionOriginalImportedMaterialName, AvailableSectionName, StaticMesh.StaticMaterials[MaterialIndex].MaterialInterface, false);
+				//TODO: Need to know if a section material slot assignment was change from the default value (implemented in skeletalmesh editor)
+				OutSections.AddSection(LODIndex, SectionIndex, CurrentSectionMaterialSlotName, MaterialIndex, CurrentSectionOriginalImportedMaterialName, AvailableSectionName, StaticMesh.StaticMaterials[MaterialIndex].MaterialInterface, false, false, MaterialIndex);
 			}
 		}
 	}
diff --git a/Engine/Source/Editor/UnrealEd/Private/EditorExporters.cpp b/Engine/Source/Editor/UnrealEd/Private/EditorExporters.cpp
--- a/Engine/Source/Editor/UnrealEd/Private/EditorExporters.cpp
+++ b/Engine/Source/Editor/UnrealEd/Private/EditorExporters.cpp
@@ -59,7 +59,7 @@
 #include "Editor.h"
 #include "MatineeExporter.h"
 #include "FbxExporter.h"
-#include "MeshDescription.h"
+#include "StaticMeshAttributes.h"
 #include "MeshDescriptionOperations.h"
 #include "MaterialUtilities.h"
 #include "InstancedFoliageActor.h"
diff --git a/Engine/Source/Editor/UnrealEd/Private/Factories/EditorFactories.cpp b/Engine/Source/Editor/UnrealEd/Private/Factories/EditorFactories.cpp
--- a/Engine/Source/Editor/UnrealEd/Private/Factories/EditorFactories.cpp
+++ b/Engine/Source/Editor/UnrealEd/Private/Factories/EditorFactories.cpp
@@ -6221,6 +6221,7 @@ EReimportResult::Type UReimportFbxSkeletalMeshFactory::Reimport( UObject* Obj, i
 		ImportOptions->SkeletonForAnimation = SkeletalMesh->Skeleton;
 		ImportOptions->bCreatePhysicsAsset = false;
 		ImportOptions->PhysicsAsset = SkeletalMesh->PhysicsAsset;
+		
 
 		ImportOptions = GetImportOptions( FFbxImporter, ReimportUI, bShowOptionDialog, bIsAutomated, Obj->GetPathName(), bOperationCanceled, bOutImportAll, bIsObjFormat, Filename, bForceImportType, FBXIT_SkeletalMesh);
 
@@ -6231,11 +6232,32 @@ EReimportResult::Type UReimportFbxSkeletalMeshFactory::Reimport( UObject* Obj, i
 		}
 	}
 
+	//Set the build option to reflect the user choice in the dialog
+	UFbxSkeletalMeshImportData* SKImportData = UFbxSkeletalMeshImportData::GetImportDataForSkeletalMesh(SkeletalMesh, ReimportUI->SkeletalMeshImportData);
+	if (SKImportData)
+	{
+		FSkeletalMeshLODInfo* LODInfo = SkeletalMesh->GetLODInfo(0);
+		if (LODInfo && SkeletalMesh->GetImportedModel() && SkeletalMesh->GetImportedModel()->LODModels.IsValidIndex(0))
+		{
+			const FSkeletalMeshLODModel& LODModel = SkeletalMesh->GetImportedModel()->LODModels[0];
+			if (LODModel.RawSkeletalMeshBulkData.IsBuildDataAvailable())
+			{
+				//Set the build settings
+				LODInfo->BuildSettings.bComputeWeightedNormals = SKImportData->bComputeWeightedNormals;
+				LODInfo->BuildSettings.bRecomputeNormals = SKImportData->NormalImportMethod == EFBXNormalImportMethod::FBXNIM_ComputeNormals;
+				LODInfo->BuildSettings.bRecomputeTangents = SKImportData->NormalImportMethod != EFBXNormalImportMethod::FBXNIM_ImportNormalsAndTangents;
+				LODInfo->BuildSettings.bUseMikkTSpace = SKImportData->NormalGenerationMethod == EFBXNormalGenerationMethod::MikkTSpace;
+			}
+		}
+	}
+
 	UE_LOG(LogEditorFactories, Log, TEXT("Performing atomic reimport of [%s]"), *Filename);
 	CurrentFilename = Filename;
 
 	if( !bOperationCanceled )
 	{
+		FScopedSkeletalMeshPostEditChange ScopedPostEditChange(SkeletalMesh);
+
 		ImportOptions->bCanShowDialog = !IsUnattended;
 
 		if (ImportOptions->bImportAsSkeletalSkinning)
@@ -6293,7 +6315,7 @@ EReimportResult::Type UReimportFbxSkeletalMeshFactory::Reimport( UObject* Obj, i
 			//Restore skin weight profile infos, then reimport affected LODs
 			TArray<FSkinWeightProfileInfo>&SkinWeightsProfile = SkeletalMesh->GetSkinWeightProfiles();
 			SkinWeightsProfile = ExistingSkinWeightProfileInfos;
-			FSkinWeightsUtilities::ReimportAlternateSkinWeight(SkeletalMesh, 0, true);
+			FSkinWeightsUtilities::ReimportAlternateSkinWeight(SkeletalMesh, 0);
 		}
 
 		// Reimporting can have dangerous effects if the mesh is still in the transaction buffer.  Reset the transaction buffer if this is the case
diff --git a/Engine/Source/Editor/UnrealEd/Private/Factories/SkeletalMeshImport.cpp b/Engine/Source/Editor/UnrealEd/Private/Factories/SkeletalMeshImport.cpp
--- a/Engine/Source/Editor/UnrealEd/Private/Factories/SkeletalMeshImport.cpp
+++ b/Engine/Source/Editor/UnrealEd/Private/Factories/SkeletalMeshImport.cpp
@@ -255,17 +255,20 @@ ExistingSkelMeshData* SaveExistingSkelMeshData(USkeletalMesh* ExistingSkelMesh,
 
 	for (int32 LodIndex = 0; LodIndex < ImportedResource->LODModels.Num(); ++LodIndex)
 	{
+		const FSkeletalMeshLODModel& LODModel = ImportedResource->LODModels[LodIndex];
 		ExistingMeshDataPtr->ExistingImportMeshLodSectionMaterialData.AddZeroed();
-		for (int32 SectionIndex = 0; SectionIndex < ImportedResource->LODModels[LodIndex].Sections.Num(); ++SectionIndex)
+		for (int32 SectionIndex = 0; SectionIndex < LODModel.Sections.Num(); ++SectionIndex)
 		{
-			int32 SectionMaterialIndex = ImportedResource->LODModels[LodIndex].Sections[SectionIndex].MaterialIndex;
-			bool SectionCastShadow = ImportedResource->LODModels[LodIndex].Sections[SectionIndex].bCastShadow;
-			bool SectionRecomputeTangents = ImportedResource->LODModels[LodIndex].Sections[SectionIndex].bRecomputeTangent;
-			int32 GenerateUpTo = ImportedResource->LODModels[LodIndex].Sections[SectionIndex].GenerateUpToLodIndex;
-			bool bDisabled = ImportedResource->LODModels[LodIndex].Sections[SectionIndex].bDisabled;
+			const FSkelMeshSourceSectionUserData& UserSectionData = LODModel.UserSectionsData.FindChecked(LODModel.Sections[SectionIndex].OriginalDataSectionIndex);
+			int32 SectionMaterialIndex = LODModel.Sections[SectionIndex].MaterialIndex;
+			bool SectionCastShadow = UserSectionData.bCastShadow;
+			bool SectionRecomputeTangents = UserSectionData.bRecomputeTangent;
+			int32 GenerateUpTo = UserSectionData.GenerateUpToLodIndex;
+			bool bDisabled = UserSectionData.bDisabled;
+			bool bBoneChunkedSection = LODModel.Sections[SectionIndex].ChunkedParentSectionIndex != INDEX_NONE;
 			if (ExistingMeshDataPtr->ExistingImportMaterialOriginalNameData.IsValidIndex(SectionMaterialIndex))
 			{
-				ExistingMeshDataPtr->ExistingImportMeshLodSectionMaterialData[LodIndex].Add(ExistingMeshLodSectionData(ExistingMeshDataPtr->ExistingImportMaterialOriginalNameData[SectionMaterialIndex], SectionCastShadow, SectionRecomputeTangents, GenerateUpTo, bDisabled));
+				ExistingMeshDataPtr->ExistingImportMeshLodSectionMaterialData[LodIndex].Add(ExistingMeshLodSectionData(ExistingMeshDataPtr->ExistingImportMaterialOriginalNameData[SectionMaterialIndex], SectionCastShadow, SectionRecomputeTangents, GenerateUpTo, bDisabled, bBoneChunkedSection));
 			}
 		}
 	}
@@ -289,6 +292,7 @@ ExistingSkelMeshData* SaveExistingSkelMeshData(USkeletalMesh* ExistingSkelMesh,
 			//Save the imported LOD reduction settings
 			ExistingMeshDataPtr->ExistingReimportLODReductionSettings = LODInfo->ReductionSettings;
 		}
+		ExistingMeshDataPtr->ExistingBaseLODInfo = *LODInfo;
 
 		// Remove the zero'th LOD (ie: the LOD being reimported).
 		if (!ReimportSpecificLOD)
@@ -310,9 +314,9 @@ ExistingSkelMeshData* SaveExistingSkelMeshData(USkeletalMesh* ExistingSkelMesh,
 			{
 				FSkeletalMeshLODModel BaseLODModel;
 				TMap<FString, TArray<FMorphTargetDelta>> BaseLODMorphTargetData;
-				ImportedResource->OriginalReductionSourceMeshData[ReductionLODIndex]->LoadReductionData(BaseLODModel, BaseLODMorphTargetData);
+				ImportedResource->OriginalReductionSourceMeshData[ReductionLODIndex]->LoadReductionData(BaseLODModel, BaseLODMorphTargetData, ExistingSkelMesh);
 				FReductionBaseSkeletalMeshBulkData* ReductionLODData = new FReductionBaseSkeletalMeshBulkData();
-				ReductionLODData->SaveReductionData(BaseLODModel, BaseLODMorphTargetData);
+				ReductionLODData->SaveReductionData(BaseLODModel, BaseLODMorphTargetData, ExistingSkelMesh);
 				//Add necessary empty slot
 				while (ExistingMeshDataPtr->ExistingOriginalReductionSourceMeshData.Num() < LODModelIndex)
 				{
@@ -361,9 +365,6 @@ ExistingSkelMeshData* SaveExistingSkelMeshData(USkeletalMesh* ExistingSkelMesh,
 	ExistingMeshDataPtr->ExistingMorphTargets.Empty(ExistingSkelMesh->MorphTargets.Num());
 	ExistingMeshDataPtr->ExistingMorphTargets.Append(ExistingSkelMesh->MorphTargets);
 	
-	ExistingMeshDataPtr->bExistingUseFullPrecisionUVs = ExistingSkelMesh->bUseFullPrecisionUVs;
-	ExistingMeshDataPtr->bExistingUseHighPrecisionTangentBasis = ExistingSkelMesh->bUseHighPrecisionTangentBasis;
-
 	ExistingMeshDataPtr->ExistingAssetImportData = ExistingSkelMesh->AssetImportData;
 	ExistingMeshDataPtr->ExistingThumbnailInfo = ExistingSkelMesh->ThumbnailInfo;
 
@@ -442,6 +443,31 @@ void RegenerateDependentLODs(USkeletalMesh* SkelMesh, int32 LODIndex)
 	}
 }
 
+void RestoreDependentLODs(ExistingSkelMeshData* MeshData, USkeletalMesh* SkeletalMesh)
+{
+	check(SkeletalMesh != nullptr);
+	int32 TotalLOD = MeshData->ExistingLODModels.Num();
+	FSkeletalMeshModel* SkeletalMeshImportedModel = SkeletalMesh->GetImportedModel();
+
+	for (int32 Index = 0; Index < TotalLOD; ++Index)
+	{
+		int32 LODIndex = Index + 1;
+		if (LODIndex >= SkeletalMesh->GetLODInfoArray().Num())
+		{
+			FSkeletalMeshLODInfo& ExistLODInfo = MeshData->ExistingLODInfo[Index];
+			FSkeletalMeshLODModel& ExistLODModel = MeshData->ExistingLODModels[Index];
+			// reset material maps, it won't work anyway. 
+			ExistLODInfo.LODMaterialMap.Empty();
+
+			FSkeletalMeshLODModel* NewLODModel = new FSkeletalMeshLODModel(ExistLODModel);
+			SkeletalMeshImportedModel->LODModels.Add(NewLODModel);
+			// add LOD info back
+			SkeletalMesh->AddLODInfo(ExistLODInfo);
+			check(LODIndex < SkeletalMesh->GetLODInfoArray().Num());
+		}
+	}
+}
+
 void TryRegenerateLODs(ExistingSkelMeshData* MeshData, USkeletalMesh* SkeletalMesh)
 {
 	check(SkeletalMesh != nullptr);
@@ -710,24 +736,15 @@ void RestoreExistingSkelMeshData(ExistingSkelMeshData* MeshData, USkeletalMesh*
 	}
 
 	//Copy back the reimport LOD specific data
-	if (SkeletalMesh->GetLODInfoArray().IsValidIndex(ReimportLODIndex) && MeshData->bIsReimportLODReduced)
+	if (SkeletalMesh->GetLODInfoArray().IsValidIndex(ReimportLODIndex))
 	{
 		FSkeletalMeshLODInfo& BaseLODInfo = SkeletalMesh->GetLODInfoArray()[ReimportLODIndex];
-		//Restore the reimport LOD reduction settings
-		BaseLODInfo.ReductionSettings = MeshData->ExistingReimportLODReductionSettings;
-		if (SkeletalMeshImportedModel && SkeletalMeshImportedModel->OriginalReductionSourceMeshData.IsValidIndex(ReimportLODIndex))
+		//Restore the build setting first
+		BaseLODInfo.BuildSettings = MeshData->ExistingBaseLODInfo.BuildSettings;
+		if (MeshData->bIsReimportLODReduced)
 		{
-			SkeletalMeshImportedModel->OriginalReductionSourceMeshData[ReimportLODIndex]->EmptyBulkData();
-		}
-
-		if (SkeletalMesh->IsReductionActive(ReimportLODIndex))
-		{
-			//Regenerate the reimport LOD
-			GWarn->BeginSlowTask(LOCTEXT("RegenReimportedLOD", "Generating reimported LOD"), true);
-			FSkeletalMeshUpdateContext UpdateContext;
-			UpdateContext.SkeletalMesh = SkeletalMesh;
-			FLODUtilities::SimplifySkeletalMeshLOD(UpdateContext, ReimportLODIndex, false);
-			GWarn->EndSlowTask();
+			//Restore the reimport LOD reduction settings
+			BaseLODInfo.ReductionSettings = MeshData->ExistingBaseLODInfo.ReductionSettings;
 		}
 	}
 
@@ -755,7 +772,7 @@ void RestoreExistingSkelMeshData(ExistingSkelMeshData* MeshData, USkeletalMesh*
 		// We copy back and fix-up the LODs that still work with this skeleton.
 		if (MeshData->ExistingLODModels.Num() > 0)
 		{
-			auto RestoreReductionSourceData = [&SkeletalMeshImportedModel, &MeshData](int32 ExistingIndex, int32 NewIndex)
+			auto RestoreReductionSourceData = [&SkeletalMesh, &SkeletalMeshImportedModel, &MeshData](int32 ExistingIndex, int32 NewIndex)
 			{
 				if (!MeshData->ExistingOriginalReductionSourceMeshData.IsValidIndex(ExistingIndex) || MeshData->ExistingOriginalReductionSourceMeshData[ExistingIndex]->IsEmpty())
 				{
@@ -764,9 +781,9 @@ void RestoreExistingSkelMeshData(ExistingSkelMeshData* MeshData, USkeletalMesh*
 				//Restore the original reduction source mesh data
 				FSkeletalMeshLODModel BaseLODModel;
 				TMap<FString, TArray<FMorphTargetDelta>> BaseLODMorphTargetData;
-				MeshData->ExistingOriginalReductionSourceMeshData[ExistingIndex]->LoadReductionData(BaseLODModel, BaseLODMorphTargetData);
+				MeshData->ExistingOriginalReductionSourceMeshData[ExistingIndex]->LoadReductionData(BaseLODModel, BaseLODMorphTargetData, SkeletalMesh);
 				FReductionBaseSkeletalMeshBulkData* ReductionLODData = new FReductionBaseSkeletalMeshBulkData();
-				ReductionLODData->SaveReductionData(BaseLODModel, BaseLODMorphTargetData);
+				ReductionLODData->SaveReductionData(BaseLODModel, BaseLODMorphTargetData, SkeletalMesh);
 				//Add necessary empty slot
 				while (SkeletalMeshImportedModel->OriginalReductionSourceMeshData.Num() < NewIndex)
 				{
@@ -776,7 +793,6 @@ void RestoreExistingSkelMeshData(ExistingSkelMeshData* MeshData, USkeletalMesh*
 				SkeletalMeshImportedModel->OriginalReductionSourceMeshData.Add(ReductionLODData);
 			};
 
-
 			if (SkeletonsAreCompatible(SkeletalMesh->RefSkeleton, MeshData->ExistingRefSkeleton, bImportSkinningOnly))
 			{
 				// First create mapping table from old skeleton to new skeleton.
@@ -889,8 +905,8 @@ void RestoreExistingSkelMeshData(ExistingSkelMeshData* MeshData, USkeletalMesh*
 					}
 				}
 			}
-			//Always regenerate LODs when re-importing LOD 0
-			TryRegenerateLODs(MeshData, SkeletalMesh);
+			//We just need to restore the LOD model and LOD info the build will regenerate the LODs
+			RestoreDependentLODs(MeshData, SkeletalMesh);
 		}
 
 		for (int32 AssetIndex = 0; AssetIndex < MeshData->ExistingPhysicsAssets.Num(); ++AssetIndex)
@@ -920,9 +936,6 @@ void RestoreExistingSkelMeshData(ExistingSkelMeshData* MeshData, USkeletalMesh*
 		SkeletalMesh->MorphTargets.Append(MeshData->ExistingMorphTargets);
 		SkeletalMesh->InitMorphTargets();
 
-		SkeletalMesh->bUseFullPrecisionUVs = MeshData->bExistingUseFullPrecisionUVs;
-		SkeletalMesh->bUseHighPrecisionTangentBasis= MeshData->bExistingUseHighPrecisionTangentBasis;
-
 		SkeletalMesh->AssetImportData = MeshData->ExistingAssetImportData.Get();
 		SkeletalMesh->ThumbnailInfo = MeshData->ExistingThumbnailInfo.Get();
 
@@ -936,15 +949,39 @@ void RestoreExistingSkelMeshData(ExistingSkelMeshData* MeshData, USkeletalMesh*
 
 		SkeletalMesh->SetSamplingInfo(MeshData->ExistingSamplingInfo);
 	}
-	else
-	{
-		RegenerateDependentLODs(SkeletalMesh, ReimportLODIndex);
-	}
 
 	//Restore the section change only for the reimport LOD, other LOD are not affected since the material array can only grow.
 	if (MeshData->UseMaterialNameSlotWorkflow)
 	{
+		//Restore the base LOD materialMap the LODs LODMaterialMap are restore differently
+		if (ReimportLODIndex < 1 && SkeletalMesh->GetLODInfoArray().IsValidIndex(ReimportLODIndex))
+		{
+			FSkeletalMeshLODInfo& BaseLODInfo = SkeletalMesh->GetLODInfoArray()[ReimportLODIndex];
+			//Restore the Base MaterialMap
+			if (SkeletalMesh->GetImportedModel() && SkeletalMesh->GetImportedModel()->LODModels.IsValidIndex(ReimportLODIndex))
+			{
+				FSkeletalMeshLODModel& LODModel = SkeletalMesh->GetImportedModel()->LODModels[ReimportLODIndex];
+				for (int32 SectionIndex = 0; SectionIndex < LODModel.Sections.Num(); ++SectionIndex)
+				{
+					int32 MaterialIndex = LODModel.Sections[SectionIndex].MaterialIndex;
+					if (MeshData->ExistingBaseLODInfo.LODMaterialMap.IsValidIndex(MaterialIndex))
+					{
+						int32 ExistingLODMaterialIndex = MeshData->ExistingBaseLODInfo.LODMaterialMap[MaterialIndex];
+						while (BaseLODInfo.LODMaterialMap.Num() <= MaterialIndex)
+						{
+							BaseLODInfo.LODMaterialMap.Add(INDEX_NONE);
+						}
+						if (RemapMaterial.IsValidIndex(ExistingLODMaterialIndex))
+						{
+							ExistingLODMaterialIndex = RemapMaterial[ExistingLODMaterialIndex];
+						}
+						BaseLODInfo.LODMaterialMap[MaterialIndex] = ExistingLODMaterialIndex;
+					}
+				}
+			}
+		}
 		FSkeletalMeshLODModel &NewSkelMeshLodModel = SkeletalMeshImportedModel->LODModels[ReimportLODIndex];
+		int32 ParentSectionIndex = 0;
 		//Restore the section changes from the old import data
 		for (int32 SectionIndex = 0; SectionIndex < NewSkelMeshLodModel.Sections.Num(); SectionIndex++)
 		{
@@ -963,6 +1000,20 @@ void RestoreExistingSkelMeshData(ExistingSkelMeshData* MeshData, USkeletalMesh*
 						NewSkelMeshLodModel.Sections[SectionIndex].bRecomputeTangent = MeshData->ExistingImportMeshLodSectionMaterialData[ReimportLODIndex][RemapMaterial[NewMeshSectionMaterialIndex]].bRecomputeTangents;
 						NewSkelMeshLodModel.Sections[SectionIndex].GenerateUpToLodIndex = MeshData->ExistingImportMeshLodSectionMaterialData[ReimportLODIndex][RemapMaterial[NewMeshSectionMaterialIndex]].GenerateUpTo;
 						NewSkelMeshLodModel.Sections[SectionIndex].bDisabled = MeshData->ExistingImportMeshLodSectionMaterialData[ReimportLODIndex][RemapMaterial[NewMeshSectionMaterialIndex]].bDisabled;
+						bool bBoneChunkedSection = MeshData->ExistingImportMeshLodSectionMaterialData[ReimportLODIndex][RemapMaterial[NewMeshSectionMaterialIndex]].bBoneChunkedSection;
+						if (!bBoneChunkedSection)
+						{
+							//Set the new Parent Index
+							ParentSectionIndex = SectionIndex;
+							FSkelMeshSourceSectionUserData& UserSectionData = NewSkelMeshLodModel.UserSectionsData.FindOrAdd(ParentSectionIndex);
+							UserSectionData.bDisabled = NewSkelMeshLodModel.Sections[SectionIndex].bDisabled;
+							UserSectionData.bCastShadow = NewSkelMeshLodModel.Sections[SectionIndex].bCastShadow;
+							UserSectionData.bRecomputeTangent = NewSkelMeshLodModel.Sections[SectionIndex].bRecomputeTangent;
+							UserSectionData.GenerateUpToLodIndex = NewSkelMeshLodModel.Sections[SectionIndex].GenerateUpToLodIndex;
+							//The cloth will be rebind later after the reimport is done
+						}
+						NewSkelMeshLodModel.Sections[SectionIndex].ChunkedParentSectionIndex = bBoneChunkedSection ? ParentSectionIndex : INDEX_NONE;
+						NewSkelMeshLodModel.Sections[SectionIndex].OriginalDataSectionIndex = ParentSectionIndex;
 					}
 				}
 			}
@@ -1007,11 +1058,12 @@ void RestoreExistingSkelMeshData(ExistingSkelMeshData* MeshData, USkeletalMesh*
 				break;
 			}
 		}
+		//Make sure we reset the User section array to only what we have in the fbx
+		NewSkelMeshLodModel.SyncronizeUserSectionsDataArray(true);
 	}
 
 	if (!bImportSkinningOnly && !MeshData->bIsReimportLODReduced)
 	{
-		//Fix the OriginalReductionSourceMeshData
 		if (ReimportLODIndex < 0)
 		{
 			if (SkeletalMeshImportedModel->OriginalReductionSourceMeshData.IsValidIndex(0))
diff --git a/Engine/Source/Editor/UnrealEd/Private/Fbx/FbxFactory.cpp b/Engine/Source/Editor/UnrealEd/Private/Fbx/FbxFactory.cpp
--- a/Engine/Source/Editor/UnrealEd/Private/Fbx/FbxFactory.cpp
+++ b/Engine/Source/Editor/UnrealEd/Private/Fbx/FbxFactory.cpp
@@ -549,6 +549,10 @@ UObject* UFbxFactory::FactoryCreateFile
 						int32 LODIndex;
 						int32 SuccessfulLodIndex = 0;
 						bool bImportSkeletalMeshLODs = ImportUI->SkeletalMeshImportData->bImportMeshLODs;
+						
+						//The skeletalmesh will be set after we import the LOD 0 since it is not created yet.
+						FScopedSkeletalMeshPostEditChange ScopedPostEditChange(nullptr);
+
 						for (LODIndex = 0; LODIndex < MaxLODLevel; LODIndex++)
 						{
 							//We need to know what is the imported lod index when importing the morph targets
@@ -605,7 +609,7 @@ UObject* UFbxFactory::FactoryCreateFile
 
 								USkeletalMesh* NewMesh = FbxImporter->ImportSkeletalMesh( ImportSkeletalMeshArgs );
 								CreatedObject = NewMesh;
-
+								
 								if(bOperationCanceled)
 								{
 									// User cancelled, clean up and return
@@ -617,6 +621,9 @@ UObject* UFbxFactory::FactoryCreateFile
 
 								if ( NewMesh )
 								{
+									//Set the base skeletalmesh to the scoped post edit change variable
+									ScopedPostEditChange.SetSkeletalMesh(NewMesh);
+
 									if (ImportOptions->bImportAnimations)
 									{
 										// We need to remove all scaling from the root node before we set up animation data.
@@ -643,7 +650,7 @@ UObject* UFbxFactory::FactoryCreateFile
 								LODInfo.ReductionSettings.BaseLOD = 0;
 								LODInfo.bImportWithBaseMesh = true;
 								LODInfo.SourceImportFilename = FString(TEXT(""));
-								FLODUtilities::SimplifySkeletalMeshLOD(UpdateContext, SuccessfulLodIndex, false);
+								FLODUtilities::SimplifySkeletalMeshLOD(UpdateContext, SuccessfulLodIndex);
 								ImportedSuccessfulLodIndex = SuccessfulLodIndex;
 								SuccessfulLodIndex++;
 							}
@@ -662,7 +669,7 @@ UObject* UFbxFactory::FactoryCreateFile
 								ImportSkeletalMeshArgs.OutData = &OutData;
 
 								USkeletalMesh *LODObject = FbxImporter->ImportSkeletalMesh( ImportSkeletalMeshArgs );
-								bool bImportSucceeded = !bOperationCanceled && FbxImporter->ImportSkeletalMeshLOD(LODObject, BaseSkeletalMesh, SuccessfulLodIndex, false);
+								bool bImportSucceeded = !bOperationCanceled && FbxImporter->ImportSkeletalMeshLOD(LODObject, BaseSkeletalMesh, SuccessfulLodIndex);
 
 								if (bImportSucceeded)
 								{
@@ -704,15 +711,6 @@ UObject* UFbxFactory::FactoryCreateFile
 							
 							USkeletalMesh* SkeletalMesh = Cast<USkeletalMesh>(CreatedObject);
 							UnFbx::FFbxImporter::UpdateSkeletalMeshImportData(SkeletalMesh, ImportUI->SkeletalMeshImportData, INDEX_NONE, nullptr, nullptr);
-							
-							//If we have import some morph target we have to rebuild the render resources since morph target are now using GPU
-							if (SkeletalMesh)// && SkeletalMesh->MorphTargets.Num() > 0)
-							{
-								SkeletalMesh->ReleaseResources();
-								//Rebuild the resources with a post edit change since we have added some morph targets
-								SkeletalMesh->PostEditChange();
-							}
-
 						}
 					}
 				
@@ -1391,7 +1389,7 @@ namespace ImportCompareHelper
 		}
 		else if (!bCombineMeshes && !bCombineMeshesLOD)
 		{
-			Node = FFbxImporter->GetMeshNodesFromName(StaticMesh, FbxMeshArray);
+			Node = FFbxImporter->GetMeshNodesFromName(StaticMesh->GetName(), FbxMeshArray);
 		}
 
 		// If there is no match it may be because an LOD group was imported where
@@ -1481,7 +1479,38 @@ namespace ImportCompareHelper
 			return;
 		}
 
-		const TArray<FbxNode*>& SkeletalMeshNodes = *(SkeletalMeshArray[0]);
+		TArray<FbxNode*>* FbxNodes = nullptr;
+		// if there is only one mesh, use it without name checking 
+		// (because the "Used As Full Name" option enables users name the Unreal mesh by themselves
+		if (SkeletalMeshArray.Num() == 1)
+		{
+			FbxNodes = SkeletalMeshArray[0];
+		}
+		else if(SkeletalMeshArray.Num() > 1)
+		{
+			FbxNodes = nullptr;
+			//Search a set corresponding
+			for (TArray<FbxNode*>*SkeletalMeshNodes : SkeletalMeshArray)
+			{
+				FbxNode* Node = FFbxImporter->GetMeshNodesFromName(SkeletalMesh->GetName(), *SkeletalMeshNodes);
+				if (Node != nullptr)
+				{
+					FbxNodes = SkeletalMeshNodes;
+					break;
+				}
+			}
+			if (FbxNodes == nullptr)
+			{
+				//Re import the first skeletalmesh found in the fbx file
+				FbxNodes = SkeletalMeshArray[0];
+			}
+		}
+		if (FbxNodes == nullptr)
+		{
+			return;
+		}
+
+		const TArray<FbxNode*>& SkeletalMeshNodes = *(FbxNodes);
 		if (SkeletalMeshNodes.Num() == 0)
 		{
 			return;
diff --git a/Engine/Source/Editor/UnrealEd/Private/Fbx/FbxMainImport.cpp b/Engine/Source/Editor/UnrealEd/Private/Fbx/FbxMainImport.cpp
--- a/Engine/Source/Editor/UnrealEd/Private/Fbx/FbxMainImport.cpp
+++ b/Engine/Source/Editor/UnrealEd/Private/Fbx/FbxMainImport.cpp
@@ -1741,6 +1741,52 @@ FName FFbxImporter::MakeNameForMesh(FString InName, FbxObject* FbxObject)
 	return OutputName;
 }
 
+FbxNode* FFbxImporter::GetMeshNodesFromName(const FString& ReimportMeshName, TArray<FbxNode*>& FbxMeshArray)
+{
+	//StaticMesh->GetName()
+	char MeshName[1024];
+	FCStringAnsi::Strcpy(MeshName, 1024, TCHAR_TO_UTF8(*ReimportMeshName));
+	// find the Fbx mesh node that the Unreal Mesh matches according to name
+	int32 MeshIndex;
+	for (MeshIndex = 0; MeshIndex < FbxMeshArray.Num(); MeshIndex++)
+	{
+		const char* FbxMeshName = FbxMeshArray[MeshIndex]->GetName();
+		// The name of Unreal mesh may have a prefix, so we match from end
+		int32 i = 0;
+		char* MeshPtr = MeshName + FCStringAnsi::Strlen(MeshName) - 1;
+		if (FCStringAnsi::Strlen(FbxMeshName) <= FCStringAnsi::Strlen(MeshName))
+		{
+			const char* FbxMeshPtr = FbxMeshName + FCStringAnsi::Strlen(FbxMeshName) - 1;
+			while (i < FCStringAnsi::Strlen(FbxMeshName))
+			{
+				bool bIsPointAndUnderscore = *FbxMeshPtr == '.' && *MeshPtr == '_';
+
+				if (*MeshPtr != *FbxMeshPtr && !bIsPointAndUnderscore)
+				{
+					break;
+				}
+				else
+				{
+					i++;
+					MeshPtr--;
+					FbxMeshPtr--;
+				}
+			}
+		}
+
+		if (i == FCStringAnsi::Strlen(FbxMeshName)) // matched
+		{
+			// check further
+			if (FCStringAnsi::Strlen(FbxMeshName) == FCStringAnsi::Strlen(MeshName) || // the name of Unreal mesh is full match
+				*MeshPtr == '_')														// or the name of Unreal mesh has a prefix
+			{
+				return FbxMeshArray[MeshIndex];
+			}
+		}
+	}
+	return nullptr;
+}
+
 FbxAMatrix FFbxImporter::ComputeSkeletalMeshTotalMatrix(FbxNode* Node, FbxNode *RootSkeletalNode)
 {
 	if (ImportOptions->bImportScene && !ImportOptions->bTransformVertexToAbsolute && RootSkeletalNode != nullptr && RootSkeletalNode != Node)
diff --git a/Engine/Source/Editor/UnrealEd/Private/Fbx/FbxSceneImportFactory.cpp b/Engine/Source/Editor/UnrealEd/Private/Fbx/FbxSceneImportFactory.cpp
--- a/Engine/Source/Editor/UnrealEd/Private/Fbx/FbxSceneImportFactory.cpp
+++ b/Engine/Source/Editor/UnrealEd/Private/Fbx/FbxSceneImportFactory.cpp
@@ -58,6 +58,7 @@
 
 #include "HAL/FileManager.h"
 #include "LODUtilities.h"
+#include "ComponentReregisterContext.h"
 
 #define LOCTEXT_NAMESPACE "FBXSceneImportFactory"
 
@@ -1886,6 +1887,9 @@ UObject* UFbxSceneImportFactory::ImportOneSkeletalMesh(void* VoidRootNodeToImpor
 		}
 	}
 
+	//The skeletalmesh will be set after we import the LOD 0 since it is not created yet.
+	FScopedSkeletalMeshPostEditChange ScopedPostEditChange(nullptr);
+
 	int32 LODIndex;
 	for (LODIndex = 0; LODIndex < MaxLODLevel; LODIndex++)
 	{
@@ -1969,6 +1973,7 @@ UObject* UFbxSceneImportFactory::ImportOneSkeletalMesh(void* VoidRootNodeToImpor
 			NewObject = NewMesh;
 			if (NewMesh)
 			{
+				ScopedPostEditChange.SetSkeletalMesh(NewMesh);
 				TSharedPtr<FFbxNodeInfo> SkelMeshNodeInfo;
 				if (FindSceneNodeInfo(SceneInfo, SkelMeshNodeArray[0]->GetUniqueID(), SkelMeshNodeInfo) && SkelMeshNodeInfo.IsValid() && SkelMeshNodeInfo->AttributeInfo.IsValid())
 				{
@@ -2046,17 +2051,10 @@ UObject* UFbxSceneImportFactory::ImportOneSkeletalMesh(void* VoidRootNodeToImpor
 		}
 	}
 	
-	USkeletalMesh* ImportedSkelMesh = Cast<USkeletalMesh>(NewObject);
-	//If we have import some morph target we have to rebuild the render resources since morph target are now using GPU
-	if (ImportedSkelMesh && ImportedSkelMesh->MorphTargets.Num() > 0)
-	{
-		ImportedSkelMesh->ReleaseResources();
-		//Rebuild the resources with a post edit change since we have added some morph targets
-		ImportedSkelMesh->PostEditChange();
-	}
-	
 	//Put back the options
 	GlobalImportSettings->bBakePivotInVertex = Old_bBakePivotInVertex;
+
+	//FScopedSkeletalMeshPostEditChange will call post edit change when going out of scope
 	return NewObject;
 }
 
diff --git a/Engine/Source/Editor/UnrealEd/Private/Fbx/FbxSkeletalMeshImport.cpp b/Engine/Source/Editor/UnrealEd/Private/Fbx/FbxSkeletalMeshImport.cpp
--- a/Engine/Source/Editor/UnrealEd/Private/Fbx/FbxSkeletalMeshImport.cpp
+++ b/Engine/Source/Editor/UnrealEd/Private/Fbx/FbxSkeletalMeshImport.cpp
@@ -27,6 +27,9 @@
 #include "Animation/MorphTarget.h"
 #include "PhysicsAssetUtils.h"
 #include "Rendering/SkeletalMeshModel.h"
+#include "Rendering/SkeletalMeshRenderData.h"
+#include "Rendering/SkeletalMeshLODRenderData.h"
+#include "Misc/ScopedSlowTask.h"
 
 #include "SkelImport.h"
 #include "Logging/TokenizedMessage.h"
@@ -54,6 +57,7 @@
 #include "Engine/SkeletalMeshSocket.h"
 #include "Assets/ClothingAsset.h"
 #include "LODUtilities.h"
+#include "IMeshBuilderModule.h"
 
 #define LOCTEXT_NAMESPACE "FBXImpoter"
 
@@ -1433,7 +1437,7 @@ bool UnFbx::FFbxImporter::FillSkeletalMeshImportData(TArray<FbxNode*>& NodeArray
 
 	if (ImportOptions->bPreserveSmoothingGroups)
 	{
-		bool bDuplicateUnSmoothWedges = (ImportOptions->NormalGenerationMethod != EFBXNormalGenerationMethod::MikkTSpace);
+		bool bDuplicateUnSmoothWedges = false; //We deprecate legacy build
 		DoUnSmoothVerts(*SkelMeshImportDataPtr, bDuplicateUnSmoothWedges);
 	}
 	else
@@ -1754,28 +1758,40 @@ USkeletalMesh* UnFbx::FFbxImporter::ImportSkeletalMesh(FImportSkeletalMeshArgs &
 		return nullptr;
 	}
 
+	//Stack the PostEditChange call, it will call post edit change when it will go out of scope
+	FScopedSkeletalMeshPostEditChange ScopedPostEditChange(ExistingSkelMesh);
+
+	ESkeletalMeshGeoImportVersions GeoImportVersion = ESkeletalMeshGeoImportVersions::LatestVersion;
+	ESkeletalMeshSkinningImportVersions SkinningImportVersion = ESkeletalMeshSkinningImportVersions::LatestVersion;
 	//Adjust the import data from the import options
 	if (ExistingSkelMesh != nullptr)
 	{
+		int32 SafeLODIndex = ImportSkeletalMeshArgs.LodIndex < 0 ? 0 : ImportSkeletalMeshArgs.LodIndex;
+		if (ExistingSkelMesh->GetImportedModel() && ExistingSkelMesh->GetImportedModel()->LODModels.IsValidIndex(SafeLODIndex))
+		{
+			const FRawSkeletalMeshBulkData& RawSkeletalMeshBulkData = ExistingSkelMesh->GetImportedModel()->LODModels[SafeLODIndex].RawSkeletalMeshBulkData;
+			GeoImportVersion = RawSkeletalMeshBulkData.GeoImportVersion;
+			SkinningImportVersion = RawSkeletalMeshBulkData.SkinningImportVersion;
+		}
+
 		if ((ImportOptions->bImportAsSkeletalSkinning || ImportOptions->bImportAsSkeletalGeometry))
 		{
 			if (ImportOptions->bImportAsSkeletalSkinning)
 			{
 				//Replace geometry import data by original existing skel mesh geometry data
-				ReplaceSkeletalMeshGeometryImportData(ExistingSkelMesh, SkelMeshImportDataPtr, ImportSkeletalMeshArgs.LodIndex);
+				ReplaceSkeletalMeshGeometryImportData(ExistingSkelMesh, SkelMeshImportDataPtr, SafeLODIndex);
 			}
 			else if (ImportOptions->bImportAsSkeletalGeometry)
 			{
 				//Replace skinning import data by original existing skel mesh skinning data
-				ReplaceSkeletalMeshSkinningImportData(ExistingSkelMesh, SkelMeshImportDataPtr, ImportSkeletalMeshArgs.LodIndex);
+				ReplaceSkeletalMeshSkinningImportData(ExistingSkelMesh, SkelMeshImportDataPtr, SafeLODIndex);
 			}
 		}
 
 		//Reuse the vertex color in case we are not re-importing weights only and there is some vertexcolor in the existing skeleMesh
 		if (ExistingSkelMesh->bHasVertexColors && ImportOptions->VertexColorImportOption == EVertexColorImportOption::Ignore)
 		{
-			int32 LODIndex = ImportSkeletalMeshArgs.LodIndex < 0 ? 0 : ImportSkeletalMeshArgs.LodIndex;
-			RemapSkeletalMeshVertexColorToImportData(ExistingSkelMesh, LODIndex, SkelMeshImportDataPtr);
+			RemapSkeletalMeshVertexColorToImportData(ExistingSkelMesh, SafeLODIndex, SkelMeshImportDataPtr);
 		}
 	}
 
@@ -1796,18 +1812,8 @@ USkeletalMesh* UnFbx::FFbxImporter::ImportSkeletalMesh(FImportSkeletalMeshArgs &
 		return nullptr;
 	}
 
-	// Array of re-import contexts for components using this mesh
-	// Will unregister before import, then re-register afterwards
-	TIndirectArray<FComponentReregisterContext> ComponentContexts;
-	TUniquePtr< FSkinnedMeshComponentRecreateRenderStateContext > SkinnedMeshComponentRecreateRenderStateContext;
-
-	FSkinnedMeshComponentRecreateRenderStateContext* RecreateExistingRenderStateContext = ExistingSkelMesh ? new FSkinnedMeshComponentRecreateRenderStateContext(ExistingSkelMesh, false) : nullptr;
-	//Backup the data before importing the new one
 	if (ExistingSkelMesh)
 	{
-		//Release the Renderdata resources before reimporting the skeletal mesh
-		SkinnedMeshComponentRecreateRenderStateContext = MakeUnique<FSkinnedMeshComponentRecreateRenderStateContext>(ExistingSkelMesh);
-
 		ExistingSkelMesh->PreEditChange(NULL);
 		//The backup of the skeletal mesh data empty the LOD array in the ImportedResource of the skeletal mesh
 		//If the import fail after this step the editor can crash when updating the bone later since the LODModel will not exist anymore
@@ -1816,12 +1822,18 @@ USkeletalMesh* UnFbx::FFbxImporter::ImportSkeletalMesh(FImportSkeletalMeshArgs &
 
 	if (SkeletalMesh == nullptr)
 	{
+		//Backup the PostEditchangeStackCounter before overriding the skeletalmesh
+		int32 PostEditChangeStackCounter = ExistingSkelMesh != nullptr ? ExistingSkelMesh->GetPostEditChangeStackCounter() : 0;
 		// Create the new mesh after saving the old data, since it will replace the old skeletalmesh
 		// This should happen only when doing a re-import. Otherwise the skeletal mesh is create before.
 		SkeletalMesh = NewObject<USkeletalMesh>(ImportSkeletalMeshArgs.InParent, ImportSkeletalMeshArgs.Name, ImportSkeletalMeshArgs.Flags);
+		SkeletalMesh->SetPostEditChangeStackCounter(PostEditChangeStackCounter);
 	}
 
+
 	SkeletalMesh->PreEditChange(NULL);
+	//Dirty the DDC Key for any imported Skeletal Mesh
+	SkeletalMesh->InvalidateDeriveDataCacheGUID();
 
 	FSkeletalMeshModel *ImportedResource = SkeletalMesh->GetImportedModel();
 	check(ImportedResource->LODModels.Num() == 0);
@@ -1849,9 +1861,24 @@ USkeletalMesh* UnFbx::FFbxImporter::ImportSkeletalMesh(FImportSkeletalMeshArgs &
 	// process bone influences from import data
 	ProcessImportMeshInfluences(*SkelMeshImportDataPtr);
 
-	//Store the original fbx import data
+	//Store the original fbx import data the SkelMeshImportDataPtr should not be modified after this
 	LODModel.RawSkeletalMeshBulkData.SaveRawMesh(*SkelMeshImportDataPtr);
-
+	if (ImportOptions->bImportAsSkeletalSkinning)
+	{
+		LODModel.RawSkeletalMeshBulkData.SkinningImportVersion = ESkeletalMeshSkinningImportVersions::LatestVersion;
+		LODModel.RawSkeletalMeshBulkData.GeoImportVersion = GeoImportVersion;
+	}
+	else if(ImportOptions->bImportAsSkeletalGeometry)
+	{
+		LODModel.RawSkeletalMeshBulkData.GeoImportVersion = ESkeletalMeshGeoImportVersions::LatestVersion;
+		LODModel.RawSkeletalMeshBulkData.SkinningImportVersion = SkinningImportVersion;
+	}
+	else
+	{
+		//We reimport both
+		LODModel.RawSkeletalMeshBulkData.SkinningImportVersion = ESkeletalMeshSkinningImportVersions::LatestVersion;
+		LODModel.RawSkeletalMeshBulkData.GeoImportVersion = ESkeletalMeshGeoImportVersions::LatestVersion;
+	}
 	SkeletalMesh->ResetLODInfo();
 	FSkeletalMeshLODInfo& NewLODInfo = SkeletalMesh->AddLODInfo();
 	NewLODInfo.ReductionSettings.NumOfTrianglesPercentage = 1.0f;
@@ -1877,38 +1904,61 @@ USkeletalMesh* UnFbx::FFbxImporter::ImportSkeletalMesh(FImportSkeletalMeshArgs &
 		TArray<int32> LODPointToRawMap;
 		SkelMeshImportDataPtr->CopyLODImportData(LODPoints,LODWedges,LODFaces,LODInfluences,LODPointToRawMap);
 
-		IMeshUtilities::MeshBuildOptions BuildOptions;
-		BuildOptions.OverlappingThresholds = ImportOptions->OverlappingThresholds;
-		BuildOptions.bComputeNormals = !ImportOptions->ShouldImportNormals() || !SkelMeshImportDataPtr->bHasNormals;
-		BuildOptions.bComputeTangents = !ImportOptions->ShouldImportTangents() || !SkelMeshImportDataPtr->bHasTangents;
+		FSkeletalMeshBuildSettings BuildOptions;
+		BuildOptions.bBuildAdjacencyBuffer = true;
+		BuildOptions.bRecomputeNormals = !ImportOptions->ShouldImportNormals() || !SkelMeshImportDataPtr->bHasNormals;
+		BuildOptions.bRecomputeTangents = !ImportOptions->ShouldImportTangents() || !SkelMeshImportDataPtr->bHasTangents;
 		BuildOptions.bUseMikkTSpace = (ImportOptions->NormalGenerationMethod == EFBXNormalGenerationMethod::MikkTSpace) && (!ImportOptions->ShouldImportNormals() || !ImportOptions->ShouldImportTangents());
 		BuildOptions.bComputeWeightedNormals = ImportOptions->bComputeWeightedNormals;
-		BuildOptions.bRemoveDegenerateTriangles = false;
-
-		IMeshUtilities& MeshUtilities = FModuleManager::Get().LoadModuleChecked<IMeshUtilities>("MeshUtilities");
-		
-		TArray<FText> WarningMessages;
-		TArray<FName> WarningNames;
-		// Create actual rendering data.
-		bool bBuildSuccess = MeshUtilities.BuildSkeletalMesh(ImportedResource->LODModels[0],SkeletalMesh->RefSkeleton,LODInfluences, LODWedges,LODFaces,LODPoints,LODPointToRawMap,BuildOptions, &WarningMessages, &WarningNames);
+		BuildOptions.bRemoveDegenerates = ImportOptions->bRemoveDegenerates;
+		BuildOptions.ThresholdPosition = ImportOptions->OverlappingThresholds.ThresholdPosition;
+		BuildOptions.ThresholdTangentNormal = ImportOptions->OverlappingThresholds.ThresholdTangentNormal;
+		BuildOptions.ThresholdUV = ImportOptions->OverlappingThresholds.ThresholdUV;
 
-		// temporary hack of message/names, should be one token or a struct
-		if(WarningMessages.Num() > 0 && WarningNames.Num() == WarningMessages.Num())
+		bool bUseLegacyMeshUtilities = false;
+		bool bBuildSuccess = false;
+		if (bUseLegacyMeshUtilities)
 		{
-			EMessageSeverity::Type MessageSeverity = bBuildSuccess ? EMessageSeverity::Warning : EMessageSeverity::Error;
+			IMeshUtilities::MeshBuildOptions LegacyBuildOptions;
+			LegacyBuildOptions.FillOptions(BuildOptions);
 
-			for(int32 MessageIdx = 0; MessageIdx<WarningMessages.Num(); ++MessageIdx)	
+			IMeshUtilities& MeshUtilities = FModuleManager::Get().LoadModuleChecked<IMeshUtilities>("MeshUtilities");
+
+			TArray<FText> WarningMessages;
+			TArray<FName> WarningNames;
+			// Create actual rendering data.
+			bBuildSuccess = MeshUtilities.BuildSkeletalMesh(ImportedResource->LODModels[0], SkeletalMesh->RefSkeleton, LODInfluences, LODWedges, LODFaces, LODPoints, LODPointToRawMap, LegacyBuildOptions, &WarningMessages, &WarningNames);
+
+			// temporary hack of message/names, should be one token or a struct
+			if (WarningMessages.Num() > 0 && WarningNames.Num() == WarningMessages.Num())
 			{
-				AddTokenizedErrorMessage(FTokenizedMessage::Create(MessageSeverity, WarningMessages[MessageIdx]), WarningNames[MessageIdx]);
+				EMessageSeverity::Type MessageSeverity = bBuildSuccess ? EMessageSeverity::Warning : EMessageSeverity::Error;
+
+				for (int32 MessageIdx = 0; MessageIdx < WarningMessages.Num(); ++MessageIdx)
+				{
+					AddTokenizedErrorMessage(FTokenizedMessage::Create(MessageSeverity, WarningMessages[MessageIdx]), WarningNames[MessageIdx]);
+				}
 			}
 		}
+		else
+		{
+			//The imported LOD is always 0 here, the LOD custom import will import the LOD alone(in a temporary skeletalmesh) and add it to the base skeletal mesh later
+			const int32 LODIndex = 0;
+			check(SkeletalMesh->GetLODInfo(LODIndex) != nullptr);
+			//Set the build options
+			SkeletalMesh->GetLODInfo(LODIndex)->BuildSettings = BuildOptions;
+			//New MeshDescription build process
+			IMeshBuilderModule& MeshBuilderModule = FModuleManager::Get().LoadModuleChecked<IMeshBuilderModule>("MeshBuilder");
+			//We must build the LODModel so we can restore properly the mesh, but we do not have to regenerate LODs
+			bBuildSuccess = MeshBuilderModule.BuildSkeletalMesh(SkeletalMesh, LODIndex, false);
+		}
 
 		if( !bBuildSuccess )
 		{
 			SkeletalMesh->MarkPendingKill();
 			return NULL;
 		}
-
+		
 		// Presize the per-section shadow casting array with the number of sections in the imported LOD.
 		const int32 NumSections = LODModel.Sections.Num();
 
@@ -1966,18 +2016,12 @@ USkeletalMesh* UnFbx::FFbxImporter::ImportSkeletalMesh(FImportSkeletalMeshArgs &
 		}
 
 		SkeletalMesh->CalculateInvRefMatrices();
-		SkeletalMesh->PostEditChange();
-		SkeletalMesh->MarkPackageDirty();
-
-		// Now iterate over all skeletal mesh components re-initialising them.
-		for (TObjectIterator<USkinnedMeshComponent> It; It; ++It)
+		if ((!SkeletalMesh->GetResourceForRendering() || !SkeletalMesh->GetResourceForRendering()->LODRenderData.IsValidIndex(0)) && ImportOptions->bCreatePhysicsAsset)
 		{
-			USkinnedMeshComponent* SkinComp = *It;
-			if (SkinComp->SkeletalMesh == SkeletalMesh)
-			{
-				ComponentContexts.Add(new FComponentReregisterContext(SkinComp));
-			}
+			//We need to have a valid render data to create physic asset
+			SkeletalMesh->Build();
 		}
+		SkeletalMesh->MarkPackageDirty();
 	}
 
 	if(ImportSkeletalMeshArgs.LodIndex == 0)
@@ -2133,13 +2177,6 @@ USkeletalMesh* UnFbx::FFbxImporter::ImportSkeletalMesh(FImportSkeletalMeshArgs &
 		ImportNodeCustomProperties(SkeletalMesh->Skeleton, SkeletonNode, true);
 	}
 
-	// ComponentContexts will now go out of scope, causing components to be re-registered
-
-	if (RecreateExistingRenderStateContext)
-	{
-		delete RecreateExistingRenderStateContext;
-	}
-
 	return SkeletalMesh;
 }
 
@@ -2343,8 +2380,6 @@ USkeletalMesh* UnFbx::FFbxImporter::ReimportSkeletalMesh(USkeletalMesh* Mesh, UF
 		return NULL;
 	}
 
-	char MeshName[1024];
-	FCStringAnsi::Strcpy(MeshName, TCHAR_TO_UTF8(*Mesh->GetName()));
 	TArray<FbxNode*>* FbxNodes = NULL;
 	USkeletalMesh* NewMesh = NULL;
 
@@ -2353,16 +2388,35 @@ USkeletalMesh* UnFbx::FFbxImporter::ReimportSkeletalMesh(USkeletalMesh* Mesh, UF
 	bool Old_ImportTextures = ImportOptions->bImportTextures;
 	bool Old_ImportAnimations = ImportOptions->bImportAnimations;
 
+	auto CleanUpImportOptionsData = [&]()
+	{
+		ImportOptions->bImportRigidMesh = Old_ImportRigidMesh;
+		ImportOptions->bImportMaterials = Old_ImportMaterials;
+		ImportOptions->bImportTextures = Old_ImportTextures;
+		ImportOptions->bImportAnimations = Old_ImportAnimations;
+	};
+
 	// support to update rigid animation mesh
 	ImportOptions->bImportRigidMesh = true;
 
-	// Unbind any existing clothing assets before we reimport the geometry
-	TArray<ClothingAssetUtils::FClothingAssetMeshBinding> ClothingBindings;
-	ClothingAssetUtils::GetMeshClothingAssetBindings(Mesh, ClothingBindings);
-
-	for(ClothingAssetUtils::FClothingAssetMeshBinding& Binding : ClothingBindings)
+	UFbxSkeletalMeshImportData* SKImportData = Cast<UFbxSkeletalMeshImportData>(Mesh->AssetImportData);
+	if (SKImportData)
 	{
-		Binding.Asset->UnbindFromSkeletalMesh(Mesh, Binding.LODIndex);
+		FSkeletalMeshLODInfo* LODInfo = Mesh->GetLODInfo(0);
+		if(LODInfo && Mesh->GetImportedModel() && Mesh->GetImportedModel()->LODModels.IsValidIndex(0))
+		{
+			const FSkeletalMeshLODModel& LODModel = Mesh->GetImportedModel()->LODModels[0];
+			if (!LODModel.RawSkeletalMeshBulkData.IsBuildDataAvailable())
+			{
+				//Set the build settings
+				LODInfo->BuildSettings.bBuildAdjacencyBuffer = true;
+				LODInfo->BuildSettings.bComputeWeightedNormals = SKImportData->bComputeWeightedNormals;
+				LODInfo->BuildSettings.bRecomputeNormals = SKImportData->NormalImportMethod == EFBXNormalImportMethod::FBXNIM_ComputeNormals;
+				LODInfo->BuildSettings.bRecomputeTangents = SKImportData->NormalImportMethod != EFBXNormalImportMethod::FBXNIM_ImportNormalsAndTangents;
+				LODInfo->BuildSettings.bRemoveDegenerates = true;
+				LODInfo->BuildSettings.bUseMikkTSpace = SKImportData->NormalGenerationMethod == EFBXNormalGenerationMethod::MikkTSpace;
+			}
+		}
 	}
 
 	// get meshes in Fbx file
@@ -2418,26 +2472,53 @@ USkeletalMesh* UnFbx::FFbxImporter::ReimportSkeletalMesh(USkeletalMesh* Mesh, UF
 			if (FbxNodes != nullptr)
 				break;
 		}
-		if (FbxNodes == nullptr)
-			return nullptr;
 	}
 	else
 	{
 		// if there is only one mesh, use it without name checking 
 		// (because the "Used As Full Name" option enables users name the Unreal mesh by themselves
-		if (FbxSkelMeshArray.Num() > 0)
+		if (FbxSkelMeshArray.Num() == 1)
 		{
 			FbxNodes = FbxSkelMeshArray[0];
 		}
-		else
+		else if (FbxSkelMeshArray.Num() > 1)
 		{
-			// @todo - FBX Importing - We need proper support for reimport if the file contains more than one skeletal mesh. 
-
+			//Search a set corresponding
+			for (TArray<FbxNode*>*SkeletalMeshNodes : FbxSkelMeshArray)
+			{
+				FbxNode* Node = GetMeshNodesFromName(Mesh->GetName(), *SkeletalMeshNodes);
+				if (Node != nullptr)
+				{
+					FbxNodes = SkeletalMeshNodes;
+					break;
+				}
+			}
+			if (FbxNodes == nullptr)
+			{
+				//Re import the first skeletalmesh found in the fbx file
+				FbxNodes = FbxSkelMeshArray[0];
+			}
 		}
 	}
-
+	
 	if (FbxNodes)
 	{
+		//Posteditchange will be call at the end of the reimport which will reallocate the render ressource
+		FScopedSkeletalMeshPostEditChange ScopedPostEditChange(Mesh);
+
+		FScopedSlowTask SlowTask(100.0f, NSLOCTEXT("SkeltalMeshReimportFactory", "SkeletalMeshReimportTask", "Reimporting skeletal mesh..."));
+		SlowTask.MakeDialog();
+
+		//The fbx file is valid we can reimport the skeletal mesh
+		Mesh->PreEditChange(nullptr);
+
+		// Unbind any existing clothing assets before we reimport the geometry
+		TArray<ClothingAssetUtils::FClothingAssetMeshBinding> ClothingBindings;
+		FLODUtilities::UnbindClothingAndBackup(Mesh, ClothingBindings);
+
+		//Reimport must force a new guid
+		Mesh->InvalidateDeriveDataCacheGUID();
+
 		for (int32 LODCounter = 1; LODCounter < Mesh->GetLODNum(); ++LODCounter)
 		{
 			//Reset all import with base mesh flags
@@ -2482,8 +2563,10 @@ USkeletalMesh* UnFbx::FFbxImporter::ReimportSkeletalMesh(USkeletalMesh* Mesh, UF
 		int32 LODIndex;
 		int32 SuccessfulLodIndex = 0;
 		bool bOperationCanceled = false;
+		float ProgressStep = 90.0 / MaxLODLevel;
 		for (LODIndex = 0; LODIndex < MaxLODLevel; LODIndex++)
 		{
+			SlowTask.EnterProgressFrame(ProgressStep);
 			int32 ImportedSuccessfulLodIndex = INDEX_NONE;
 			TArray<FbxNode*> SkelMeshNodeArray;
 			for (int32 j = 0; j < (*FbxNodes).Num(); j++)
@@ -2533,8 +2616,9 @@ USkeletalMesh* UnFbx::FFbxImporter::ReimportSkeletalMesh(USkeletalMesh* Mesh, UF
 				NewMesh = ImportSkeletalMesh( ImportSkeletalMeshArgs );
 				if (bOperationCanceled)
 				{
-					// User cancelled, clean up and return
-					return nullptr;
+					// User cancel, clean up and return
+					NewMesh = nullptr;
+					break;
 				}
 				if (NewMesh)
 				{
@@ -2581,7 +2665,7 @@ USkeletalMesh* UnFbx::FFbxImporter::ReimportSkeletalMesh(USkeletalMesh* Mesh, UF
 				ImportSkeletalMeshArgs.OutData = &OutData;
 
 				UObject *LODObject = ImportSkeletalMesh( ImportSkeletalMeshArgs );
-				bool bImportSucceeded = !bOperationCanceled && ImportSkeletalMeshLOD( Cast<USkeletalMesh>(LODObject), BaseSkeletalMesh, SuccessfulLodIndex, false, nullptr, TemplateImportData);
+				bool bImportSucceeded = !bOperationCanceled && ImportSkeletalMeshLOD( Cast<USkeletalMesh>(LODObject), BaseSkeletalMesh, SuccessfulLodIndex, TemplateImportData);
 
 				if (bImportSucceeded)
 				{
@@ -2617,47 +2701,41 @@ USkeletalMesh* UnFbx::FFbxImporter::ReimportSkeletalMesh(USkeletalMesh* Mesh, UF
 
 			if (NewMesh)
 			{
-				// import morph target
 				if ((ImportOptions->bImportSkeletalMeshLODs || LODIndex == 0) &&
 					ImportOptions->bImportMorph &&
 					ImportedSuccessfulLodIndex != INDEX_NONE &&
 					NewMesh->GetImportedModel() &&
-					NewMesh->GetImportedModel()->LODModels.IsValidIndex(LODIndex))
+					NewMesh->GetImportedModel()->LODModels.IsValidIndex(ImportedSuccessfulLodIndex))
 				{
-					// @fixme: @question : where do they import this morph? where to? What morph target sets?
 					ImportFbxMorphTarget(SkelMeshNodeArray, NewMesh, ImportedSuccessfulLodIndex, OutData);
 				}
 			}
 		}
-		if (NewMesh)
+
+		auto RebindClothing = [&ClothingBindings](USkeletalMesh* SkelMesh)
 		{
-			FSkeletalMeshModel* NewMeshResource = NewMesh->GetImportedModel();
-			if(NewMeshResource && ClothingBindings.Num() > 0)
+			FSkeletalMeshModel* MeshResource = SkelMesh->GetImportedModel();
+			if (MeshResource && ClothingBindings.Num() > 0)
 			{
-				NewMesh->PreEditChange(nullptr);
+				FLODUtilities::RestoreClothingFromBackup(SkelMesh, ClothingBindings);
+			}
+		};
 
-				for(ClothingAssetUtils::FClothingAssetMeshBinding& Binding : ClothingBindings)
-				{
-					if(NewMeshResource->LODModels.IsValidIndex(Binding.LODIndex) &&
-						NewMeshResource->LODModels[Binding.LODIndex].Sections.IsValidIndex(Binding.SectionIndex))
-					{
-						Binding.Asset->BindToSkeletalMesh(NewMesh, Binding.LODIndex, Binding.SectionIndex, Binding.AssetInternalLodIndex);
-					}
-				}
 
-				NewMesh->PostEditChange();
-			}
+		if (!bOperationCanceled && NewMesh)
+		{
+			RebindClothing(NewMesh);
 
 			//Update the import data so we can re-import correctly
 			UpdateSkeletalMeshImportData(NewMesh, TemplateImportData, INDEX_NONE, &ImportMaterialOriginalNameData, &ImportMeshLodData);
-			//If we have import some morph target we have to rebuild the render resources since morph target are now using GPU
-			if (NewMesh->MorphTargets.Num() > 0)
-			{
-				NewMesh->ReleaseResources();
-				//Rebuild the resources with a post edit change since we have added some morph targets
-				NewMesh->PostEditChange();
-			}
 		}
+		else if(bOperationCanceled && Mesh != nullptr)
+		{
+			//Operation was cancel
+			//rebind the clothing on the original mesh
+			RebindClothing(Mesh);
+		}
+		SlowTask.EnterProgressFrame(10.0f);
 	}
 	else
 	{
@@ -2665,10 +2743,7 @@ USkeletalMesh* UnFbx::FFbxImporter::ReimportSkeletalMesh(USkeletalMesh* Mesh, UF
 		AddTokenizedErrorMessage(FTokenizedMessage::Create(EMessageSeverity::Warning, FText::Format(LOCTEXT("FbxSkeletaLMeshimport_NoFBXMeshMatch", "No FBX mesh matches the Unreal mesh '{0}'."), FText::FromString(Mesh->GetName()))), FFbxErrors::Generic_Mesh_MeshNotFound);
 	}
 
-	ImportOptions->bImportRigidMesh = Old_ImportRigidMesh;
-	ImportOptions->bImportMaterials = Old_ImportMaterials;
-	ImportOptions->bImportTextures = Old_ImportTextures;
-	ImportOptions->bImportAnimations = Old_ImportAnimations;
+	CleanUpImportOptionsData();
 
 	return NewMesh;
 }
@@ -3402,7 +3477,7 @@ bool UnFbx::FFbxImporter::FillSkelMeshImporterFromFbx( FSkeletalMeshImportData&
 			//
 			// normals, tangents and binormals
 			//
-			if( ImportOptions->ShouldImportNormals() && bHasNormalInformation )
+			if( bHasNormalInformation )
 			{
 				int32 TmpIndex = LocalIndex*3 + VertexIndex;
 				//normals may have different reference and mapping mode than tangents and binormals
@@ -3416,7 +3491,7 @@ bool UnFbx::FFbxImporter::FillSkelMeshImporterFromFbx( FSkeletalMeshImportData&
 
 				FbxVector4 TempValue;
 
-				if( ImportOptions->ShouldImportTangents() && bHasTangentInformation )
+				if( bHasTangentInformation )
 				{
 					TempValue = LayerElementTangent->GetDirectArray().GetAt(TangentMapIndex);
 					TempValue = TotalMatrixForNormal.MultT(TempValue);
@@ -3713,6 +3788,9 @@ bool UnFbx::FFbxImporter::FillSkelMeshImporterFromFbx( FSkeletalMeshImportData&
 
 void UnFbx::FFbxImporter::InsertNewLODToBaseSkeletalMesh(USkeletalMesh* InSkeletalMesh, USkeletalMesh* BaseSkeletalMesh, int32 DesiredLOD, UFbxSkeletalMeshImportData* TemplateImportData)
 {
+	//Scope a post edit change
+	FScopedSkeletalMeshPostEditChange ScopedPostEditChange(BaseSkeletalMesh);
+
 	FSkeletalMeshModel* ImportedResource = InSkeletalMesh->GetImportedModel();
 	FSkeletalMeshModel* DestImportedResource = BaseSkeletalMesh->GetImportedModel();
 
@@ -3772,6 +3850,13 @@ void UnFbx::FFbxImporter::InsertNewLODToBaseSkeletalMesh(USkeletalMesh* InSkelet
 
 	// Set up LODMaterialMap to number of materials in new mesh.
 	FSkeletalMeshLODInfo& LODInfo = *(BaseSkeletalMesh->GetLODInfo(DesiredLOD));
+
+	//Copy the build settings
+	if(InSkeletalMesh->GetLODInfo(0))
+	{
+		const FSkeletalMeshLODInfo& ImportedLODInfo = *(InSkeletalMesh->GetLODInfo(0));
+		LODInfo.BuildSettings = ImportedLODInfo.BuildSettings;
+	}
 	
 	LODInfo.LODMaterialMap.Empty();
 	// Now set up the material mapping array.
@@ -3865,12 +3950,9 @@ void UnFbx::FFbxImporter::InsertNewLODToBaseSkeletalMesh(USkeletalMesh* InSkelet
 			}
 		}
 	}
-
-	// rebuild vertex buffers and reinit RHI resources
-	BaseSkeletalMesh->PostEditChange();
 }
 
-bool UnFbx::FFbxImporter::ImportSkeletalMeshLOD(USkeletalMesh* InSkeletalMesh, USkeletalMesh* BaseSkeletalMesh, int32 DesiredLOD, bool bNeedToReregister, TArray<UActorComponent*>* ReregisterAssociatedComponents, UFbxSkeletalMeshImportData* TemplateImportData)
+bool UnFbx::FFbxImporter::ImportSkeletalMeshLOD(USkeletalMesh* InSkeletalMesh, USkeletalMesh* BaseSkeletalMesh, int32 DesiredLOD, UFbxSkeletalMeshImportData* TemplateImportData)
 {
 	check(InSkeletalMesh);
 	check(BaseSkeletalMesh);
@@ -3934,6 +4016,8 @@ bool UnFbx::FFbxImporter::ImportSkeletalMeshLOD(USkeletalMesh* InSkeletalMesh, U
 		}
 	}
 
+	FScopedSkeletalMeshPostEditChange ScopedPostEditChange(BaseSkeletalMesh);
+
 	//Skeleton do not match and we are importing skeletal geometry only so apply the base LOD Skinning
 	if (bApplyBaseSkinning)
 	{
@@ -3941,8 +4025,10 @@ bool UnFbx::FFbxImporter::ImportSkeletalMeshLOD(USkeletalMesh* InSkeletalMesh, U
 		//Apply the base LOD skinning to the new LOD geometry, DestImportedResource here is the base lod and is the source for the apply skinning.
 		//The Dest here mean we merge the just import LOD to the base LodModel.
 		SkeletalMeshHelper::ApplySkinning(InSkeletalMesh, DestImportedResource->LODModels[0], ImportedResource->LODModels[0]);
-		IMeshUtilities& MeshUtilities = FModuleManager::Get().LoadModuleChecked<IMeshUtilities>("MeshUtilities");
-		MeshUtilities.RemoveBonesFromMesh(BaseSkeletalMesh, DesiredLOD, NULL);
+
+		//Excluded bones will be remove when we rebuild the asset
+		//IMeshUtilities& MeshUtilities = FModuleManager::Get().LoadModuleChecked<IMeshUtilities>("MeshUtilities");
+		//MeshUtilities.RemoveBonesFromMesh(BaseSkeletalMesh, DesiredLOD, NULL);
 	}
 
 	FSkeletalMeshLODModel& NewLODModel = ImportedResource->LODModels[0];
@@ -4105,32 +4191,7 @@ bool UnFbx::FFbxImporter::ImportSkeletalMeshLOD(USkeletalMesh* InSkeletalMesh, U
 		}
 	}
 
-	if (bNeedToReregister)
-	{
-		// Shut down the skeletal mesh component that is previewing this mesh.
-		if (ReregisterAssociatedComponents)
-		{
-			FMultiComponentReregisterContext ReregisterContext(*ReregisterAssociatedComponents);
-			// wait until resources are released
-			FlushRenderingCommands();
-
-			InsertNewLODToBaseSkeletalMesh(InSkeletalMesh, BaseSkeletalMesh, DesiredLOD, TemplateImportData);
-
-			// ReregisterContexts go out of scope here, reregistering associated components with the scene.
-		}
-		else
-		{
-			TComponentReregisterContext<USkinnedMeshComponent> ComponentReregisterContext;
-
-			InsertNewLODToBaseSkeletalMesh(InSkeletalMesh, BaseSkeletalMesh, DesiredLOD, TemplateImportData);
-
-			// ReregisterContexts go out of scope here, reregistering skel components with the scene.
-		}
-	}
-	else
-	{
-		InsertNewLODToBaseSkeletalMesh(InSkeletalMesh, BaseSkeletalMesh, DesiredLOD, TemplateImportData);
-	}
+	InsertNewLODToBaseSkeletalMesh(InSkeletalMesh, BaseSkeletalMesh, DesiredLOD, TemplateImportData);
 
 	return true;
 }
@@ -4271,9 +4332,21 @@ void UnFbx::FFbxImporter::ImportMorphTargetsInternal( TArray<FbxNode*>& SkelMesh
 		check(BaseImportData.MorphTargetNames.Num() == BaseImportData.MorphTargets.Num() && BaseImportData.MorphTargetNames.Num() == BaseImportData.MorphTargetModifiedPoints.Num());
 	}
 
-	//Build MorphTargets
-	// Calculate overlapping corners and tangents
-	FLODUtilities::BuildMorphTargets(BaseSkelMesh, BaseImportData, LODIndex, ImportOptions->ShouldImportNormals(), ImportOptions->ShouldImportTangents(), (ImportOptions->NormalGenerationMethod == EFBXNormalGenerationMethod::MikkTSpace));
+	if (BaseSkelMesh->GetImportedModel() && BaseSkelMesh->GetImportedModel()->LODModels.IsValidIndex(LODIndex))
+	{
+		//If we can build the skeletal mesh there is no need to build the morph target now, all the necessary build morph target data was copied before.
+		if (!BaseSkelMesh->GetImportedModel()->LODModels[LODIndex].RawSkeletalMeshBulkData.IsBuildDataAvailable())
+		{
+			//Build MorphTargets
+			FLODUtilities::BuildMorphTargets(
+				BaseSkelMesh,
+				BaseImportData,
+				LODIndex,
+				ImportOptions->ShouldImportNormals(),
+				ImportOptions->ShouldImportTangents(),
+				(ImportOptions->NormalGenerationMethod == EFBXNormalGenerationMethod::MikkTSpace));
+		}
+	}
 
 	GWarn->EndSlowTask();
 }
@@ -4281,6 +4354,9 @@ void UnFbx::FFbxImporter::ImportMorphTargetsInternal( TArray<FbxNode*>& SkelMesh
 // Import Morph target
 void UnFbx::FFbxImporter::ImportFbxMorphTarget(TArray<FbxNode*> &SkelMeshNodeArray, USkeletalMesh* BaseSkelMesh, int32 LODIndex, FSkeletalMeshImportData &BaseSkeletalMeshImportData)
 {
+	//Stack the PostEditChange call, it will call post edit change when it will go out of scope
+	FScopedSkeletalMeshPostEditChange ScopedPostEditChange(BaseSkelMesh);
+
 	bool bHasMorph = false;
 	int32 NodeIndex;
 	// check if there are morph in this geometry
@@ -4299,7 +4375,12 @@ void UnFbx::FFbxImporter::ImportFbxMorphTarget(TArray<FbxNode*> &SkelMeshNodeArr
 	
 	if (bHasMorph)
 	{
-		ImportMorphTargetsInternal( SkelMeshNodeArray, BaseSkelMesh, LODIndex, BaseSkeletalMeshImportData );
+		ImportMorphTargetsInternal( SkelMeshNodeArray, BaseSkelMesh, LODIndex, BaseSkeletalMeshImportData);
+		//Save the rawMesh
+		if (BaseSkelMesh->GetImportedModel() && BaseSkelMesh->GetImportedModel()->LODModels.IsValidIndex(LODIndex))
+		{
+			BaseSkelMesh->GetImportedModel()->LODModels[LODIndex].RawSkeletalMeshBulkData.SaveRawMesh(BaseSkeletalMeshImportData);
+		}
 	}
 }
 
diff --git a/Engine/Source/Editor/UnrealEd/Private/Fbx/FbxStaticMeshImport.cpp b/Engine/Source/Editor/UnrealEd/Private/Fbx/FbxStaticMeshImport.cpp
--- a/Engine/Source/Editor/UnrealEd/Private/Fbx/FbxStaticMeshImport.cpp
+++ b/Engine/Source/Editor/UnrealEd/Private/Fbx/FbxStaticMeshImport.cpp
@@ -32,8 +32,7 @@
 #include "Interfaces/ITargetPlatformManagerModule.h"
 #include "Misc/FbxErrors.h"
 #include "PhysicsEngine/BodySetup.h"
-#include "MeshDescription.h"
-#include "MeshAttributes.h"
+#include "StaticMeshAttributes.h"
 #include "MeshDescriptionOperations.h"
 #include "IMeshBuilderModule.h"
 #include "Settings/EditorExperimentalSettings.h"
@@ -284,6 +283,8 @@ bool UnFbx::FFbxImporter::BuildStaticMeshFromGeometry(FbxNode* Node, UStaticMesh
 	FMeshDescription* MeshDescription = StaticMesh->GetMeshDescription(LODIndex);
 	//The mesh description should have been created before calling BuildStaticMeshFromGeometry
 	check(MeshDescription);
+	FStaticMeshAttributes Attributes(*MeshDescription);
+
 	//remove the bad polygons before getting any data from mesh
 	Mesh->RemoveBadPolygons();
 
@@ -553,15 +554,15 @@ bool UnFbx::FFbxImporter::BuildStaticMeshFromGeometry(FbxNode* Node, UStaticMesh
 	bool OddNegativeScale = IsOddNegativeScale(TotalMatrix);
 	bool bHasNonDegeneratePolygons = false;
 
-	TVertexAttributesRef<FVector> VertexPositions = MeshDescription->VertexAttributes().GetAttributesRef<FVector>(MeshAttribute::Vertex::Position);
-	TVertexInstanceAttributesRef<FVector> VertexInstanceNormals = MeshDescription->VertexInstanceAttributes().GetAttributesRef<FVector>(MeshAttribute::VertexInstance::Normal);
-	TVertexInstanceAttributesRef<FVector> VertexInstanceTangents = MeshDescription->VertexInstanceAttributes().GetAttributesRef<FVector>(MeshAttribute::VertexInstance::Tangent);
-	TVertexInstanceAttributesRef<float> VertexInstanceBinormalSigns = MeshDescription->VertexInstanceAttributes().GetAttributesRef<float>(MeshAttribute::VertexInstance::BinormalSign);
-	TVertexInstanceAttributesRef<FVector4> VertexInstanceColors = MeshDescription->VertexInstanceAttributes().GetAttributesRef<FVector4>(MeshAttribute::VertexInstance::Color);
-	TVertexInstanceAttributesRef<FVector2D> VertexInstanceUVs = MeshDescription->VertexInstanceAttributes().GetAttributesRef<FVector2D>(MeshAttribute::VertexInstance::TextureCoordinate);
-	TEdgeAttributesRef<bool> EdgeHardnesses = MeshDescription->EdgeAttributes().GetAttributesRef<bool>(MeshAttribute::Edge::IsHard);
-	TEdgeAttributesRef<float> EdgeCreaseSharpnesses = MeshDescription->EdgeAttributes().GetAttributesRef<float>(MeshAttribute::Edge::CreaseSharpness);
-	TPolygonGroupAttributesRef<FName> PolygonGroupImportedMaterialSlotNames = MeshDescription->PolygonGroupAttributes().GetAttributesRef<FName>(MeshAttribute::PolygonGroup::ImportedMaterialSlotName);
+	TVertexAttributesRef<FVector> VertexPositions = Attributes.GetVertexPositions();
+	TVertexInstanceAttributesRef<FVector> VertexInstanceNormals = Attributes.GetVertexInstanceNormals();
+	TVertexInstanceAttributesRef<FVector> VertexInstanceTangents = Attributes.GetVertexInstanceTangents();
+	TVertexInstanceAttributesRef<float> VertexInstanceBinormalSigns = Attributes.GetVertexInstanceBinormalSigns();
+	TVertexInstanceAttributesRef<FVector4> VertexInstanceColors = Attributes.GetVertexInstanceColors();
+	TVertexInstanceAttributesRef<FVector2D> VertexInstanceUVs = Attributes.GetVertexInstanceUVs();
+	TEdgeAttributesRef<bool> EdgeHardnesses = Attributes.GetEdgeHardnesses();
+	TEdgeAttributesRef<float> EdgeCreaseSharpnesses = Attributes.GetEdgeCreaseSharpnesses();
+	TPolygonGroupAttributesRef<FName> PolygonGroupImportedMaterialSlotNames = Attributes.GetPolygonGroupMaterialSlotNames();
 
 	int32 VertexOffset = MeshDescription->Vertices().Num();
 	int32 VertexInstanceOffset = MeshDescription->VertexInstances().Num();
@@ -1071,51 +1072,6 @@ void UnFbx::FFbxImporter::AddStaticMeshSourceModelGeneratedLOD(UStaticMesh* Stat
 	}
 }
 
-FbxNode* UnFbx::FFbxImporter::GetMeshNodesFromName(UStaticMesh* StaticMesh, TArray<FbxNode*>& FbxMeshArray)
-{
-	char MeshName[1024];
-	FCStringAnsi::Strcpy(MeshName, 1024, TCHAR_TO_UTF8(*StaticMesh->GetName()));
-	// find the Fbx mesh node that the Unreal Mesh matches according to name
-	int32 MeshIndex;
-	for (MeshIndex = 0; MeshIndex < FbxMeshArray.Num(); MeshIndex++)
-	{
-		const char* FbxMeshName = FbxMeshArray[MeshIndex]->GetName();
-		// The name of Unreal mesh may have a prefix, so we match from end
-		int32 i = 0;
-		char* MeshPtr = MeshName + FCStringAnsi::Strlen(MeshName) - 1;
-		if (FCStringAnsi::Strlen(FbxMeshName) <= FCStringAnsi::Strlen(MeshName))
-		{
-			const char* FbxMeshPtr = FbxMeshName + FCStringAnsi::Strlen(FbxMeshName) - 1;
-			while (i < FCStringAnsi::Strlen(FbxMeshName))
-			{
-				bool bIsPointAndUnderscore = *FbxMeshPtr == '.' && *MeshPtr == '_';
-
-				if (*MeshPtr != *FbxMeshPtr && !bIsPointAndUnderscore)
-				{
-					break;
-				}
-				else
-				{
-					i++;
-					MeshPtr--;
-					FbxMeshPtr--;
-				}
-			}
-		}
-
-		if (i == FCStringAnsi::Strlen(FbxMeshName)) // matched
-		{
-			// check further
-			if (FCStringAnsi::Strlen(FbxMeshName) == FCStringAnsi::Strlen(MeshName) || // the name of Unreal mesh is full match
-				*MeshPtr == '_')														// or the name of Unreal mesh has a prefix
-			{
-				return FbxMeshArray[MeshIndex];
-			}
-		}
-	}
-	return nullptr;
-}
-
 UStaticMesh* UnFbx::FFbxImporter::ReimportStaticMesh(UStaticMesh* Mesh, UFbxStaticMeshImportData* TemplateImportData)
 {
 	TArray<FbxNode*> FbxMeshArray;
@@ -1198,7 +1154,7 @@ UStaticMesh* UnFbx::FFbxImporter::ReimportStaticMesh(UStaticMesh* Mesh, UFbxStat
 	}
 	else if(!bCombineMeshes && !bCombineMeshesLOD)
 	{
-		Node = GetMeshNodesFromName(Mesh, FbxMeshArray);
+		Node = GetMeshNodesFromName(Mesh->GetName(), FbxMeshArray);
 	}
 
 	// If there is no match it may be because an LOD group was imported where
@@ -1506,6 +1462,8 @@ UStaticMesh* UnFbx::FFbxImporter::ImportStaticMeshAsSingle(UObject* InParent, TA
 		MeshDescription->Empty();
 	}
 
+	FStaticMeshAttributes Attributes(*MeshDescription);
+
 	// make sure it has a new lighting guid
 	StaticMesh->LightingGuid = FGuid::NewGuid();
 
@@ -1553,9 +1511,9 @@ UStaticMesh* UnFbx::FFbxImporter::ImportStaticMeshAsSingle(UObject* InParent, TA
 
 	if (bBuildStatus)
 	{
-		TVertexInstanceAttributesRef<FVector2D> VertexInstanceUVs = MeshDescription->VertexInstanceAttributes().GetAttributesRef<FVector2D>(MeshAttribute::VertexInstance::TextureCoordinate);
+		TVertexInstanceAttributesRef<FVector2D> VertexInstanceUVs = Attributes.GetVertexInstanceUVs();
 		int32 FirstOpenUVChannel = VertexInstanceUVs.GetNumIndices() >= MAX_MESH_TEXTURE_COORDS_MD ? 1 : VertexInstanceUVs.GetNumIndices();
-		TPolygonGroupAttributesRef<FName> PolygonGroupImportedMaterialSlotNames = MeshDescription->PolygonGroupAttributes().GetAttributesRef<FName>(MeshAttribute::PolygonGroup::ImportedMaterialSlotName);
+		TPolygonGroupAttributesRef<FName> PolygonGroupImportedMaterialSlotNames = Attributes.GetPolygonGroupMaterialSlotNames();
 
 		TArray<FStaticMaterial> MaterialToAdd;
 		for (const FPolygonGroupID PolygonGroupID : MeshDescription->PolygonGroups().GetElementIDs())
@@ -1944,12 +1902,6 @@ void UnFbx::FFbxImporter::PostImportStaticMesh(UStaticMesh* StaticMesh, TArray<F
 	IConsoleVariable* CVarDistanceFieldInterface = IConsoleManager::Get().FindConsoleVariable(TEXT("r.GenerateMeshDistanceFields"));
 	bool bOriginalGenerateMeshDistanceField = StaticMesh->bGenerateMeshDistanceField;
 	
-	//Always triangulate the original mesh description after we import it
-	FMeshDescription* MeshDescription = StaticMesh->GetMeshDescription(LODIndex);
-	if (MeshDescription)
-	{
-		MeshDescription->TriangulateMesh();
-	}
 	bool bUserCancel = false;
 	if (StaticMesh->FixLODRequiresAdjacencyInformation(LODIndex))
 	{
