Asset Registry Generator
- Do not strip empty chunks as we may add files to them during staging

diff --git a/Engine/Source/Editor/UnrealEd/Private/Commandlets/AssetRegistryGenerator.cpp b/Engine/Source/Editor/UnrealEd/Private/Commandlets/AssetRegistryGenerator.cpp
--- a/Engine/Source/Editor/UnrealEd/Private/Commandlets/AssetRegistryGenerator.cpp
+++ b/Engine/Source/Editor/UnrealEd/Private/Commandlets/AssetRegistryGenerator.cpp
@@ -339,7 +339,7 @@ bool FAssetRegistryGenerator::GenerateStreamingInstallManifest(int64 InExtraFlav
 				int32 PakchunkIndex = GetPakchunkIndex(ChunkID);
 				if (PakchunkIndex >= FinalChunkManifests.Num())
 				{
-					FinalChunkManifests.AddZeroed(PakchunkIndex - FinalChunkManifests.Num() + 1);
+					FinalChunkManifests.AddDefaulted(PakchunkIndex - FinalChunkManifests.Num() + 1);
 				}
 				checkf(PakchunkIndex < FinalChunkManifests.Num(), TEXT("Chunk %i out of range. %i manifests available"), PakchunkIndex, FinalChunkManifests.Num() - 1);
 				checkf(FinalChunkManifests[PakchunkIndex] == nullptr, TEXT("Manifest already exists for chunk %i"), PakchunkIndex);
@@ -368,7 +368,7 @@ bool FAssetRegistryGenerator::GenerateStreamingInstallManifest(int64 InExtraFlav
 		ChunkLayerFile->Serialize(TCHAR_TO_ANSI(*LayerString), LayerString.Len());
 
 		// Is this chunk empty?
-		if (!FinalChunkManifests[PakchunkIndex] || FinalChunkManifests[PakchunkIndex]->Num() == 0)
+		if (!FinalChunkManifests[PakchunkIndex])
 		{
 			continue;
 		}
@@ -626,7 +626,8 @@ bool FAssetRegistryGenerator::SaveManifests(FSandboxPlatformFile* InSandboxFile,
 		// Collect all unique chunk indices and map all files to their chunks
 		for (int32 PakchunkIndex = 0; PakchunkIndex < FinalChunkManifests.Num(); ++PakchunkIndex)
 		{
-			if (FinalChunkManifests[PakchunkIndex] && FinalChunkManifests[PakchunkIndex]->Num())
+			check(FinalChunkManifests[PakchunkIndex]);
+			if (FinalChunkManifests[PakchunkIndex]->Num())
 			{
 				PakchunkIndicesInUse.Add(PakchunkIndex);
 				for (auto& Filename : *FinalChunkManifests[PakchunkIndex])
@@ -1418,7 +1419,7 @@ bool FAssetRegistryGenerator::GenerateAssetChunkInformationCSV(const FString& Ou
 	TArray<TUniquePtr<FArchive>> ChunkFiles;
 	if (bWriteIndividualFiles)
 	{
-		for (int32 PakchunkIndex = 0, ChunkNum = FinalChunkManifests.Num(); PakchunkIndex < ChunkNum; ++PakchunkIndex)
+		for (int32 PakchunkIndex = 0; PakchunkIndex < FinalChunkManifests.Num(); ++PakchunkIndex)
 		{
 			FArchive* ChunkFile = IFileManager::Get().CreateFileWriter(*FPaths::Combine(*OutputPath, *FString::Printf(TEXT("Chunks%dInfo.csv"), PakchunkIndex)));
 			if (ChunkFile == nullptr)
@@ -1667,12 +1668,12 @@ void FAssetRegistryGenerator::FixupPackageDependenciesForChunks(FSandboxPlatform
 
 	for (int32 PakchunkIndex = 0, MaxPakchunk = ChunkManifests.Num(); PakchunkIndex < MaxPakchunk; ++PakchunkIndex)
 	{
-		FinalChunkManifests.Add(nullptr);
+		FinalChunkManifests.Add(new FChunkPackageSet());
 		if (!ChunkManifests[PakchunkIndex])
 		{
 			continue;
 		}
-		FinalChunkManifests[PakchunkIndex] = new FChunkPackageSet();
+		
 		for (auto It = ChunkManifests[PakchunkIndex]->CreateConstIterator(); It; ++It)
 		{
 			AddPackageAndDependenciesToChunk(FinalChunkManifests[PakchunkIndex], It.Key(), It.Value(), PakchunkIndex, InSandboxFile);
@@ -1689,12 +1690,10 @@ void FAssetRegistryGenerator::FixupPackageDependenciesForChunks(FSandboxPlatform
 	//Once complete, Add any remaining assets (that are not assigned to a chunk) to the first chunk.
 	if (FinalChunkManifests.Num() == 0)
 	{
-		FinalChunkManifests.Add(nullptr);
-	}
-	if (!FinalChunkManifests[0])
-	{
-		FinalChunkManifests[0] = new FChunkPackageSet();
+		FinalChunkManifests.Add(new FChunkPackageSet());
 	}
+	check(FinalChunkManifests[0]);
+	
 	// Copy the remaining assets
 	auto RemainingAssets = UnassignedPackageSet;
 	for (auto It = RemainingAssets.CreateConstIterator(); It; ++It)
@@ -1728,17 +1727,18 @@ void FAssetRegistryGenerator::FixupPackageDependenciesForChunks(FSandboxPlatform
 	for (int32 PakchunkIndex = 0, MaxPakchunk = ChunkManifests.Num(); PakchunkIndex < MaxPakchunk; ++PakchunkIndex)
 	{
 		const int32 ChunkManifestNum = ChunkManifests[PakchunkIndex] ? ChunkManifests[PakchunkIndex]->Num() : 0;
-		const int32 FinalChunkManifestNum = FinalChunkManifests[PakchunkIndex] ? FinalChunkManifests[PakchunkIndex]->Num() : 0;
+		const int32 FinalChunkManifestNum = FinalChunkManifests[PakchunkIndex]->Num();
 		UE_LOG(LogAssetRegistryGenerator, Log, TEXT("Chunk: %i, Started with %i packages, Final after dependency resolve: %i"), PakchunkIndex, ChunkManifestNum, FinalChunkManifestNum);
 	}
 	
 	// Fix up the asset registry to reflect this chunk layout
 	for (int32 PakchunkIndex = 0 ; PakchunkIndex < FinalChunkManifests.Num(); ++PakchunkIndex)
 	{
-		if (PakchunkIndex >= FinalChunkManifests.Num() || !FinalChunkManifests[PakchunkIndex])
+		if (PakchunkIndex >= FinalChunkManifests.Num())
 		{
 			continue;
 		}
+		check(FinalChunkManifests[PakchunkIndex]);
 		for (const TPair<FName, FString>& Asset : *FinalChunkManifests[PakchunkIndex])
 		{
 			const TArray<const FAssetData*> AssetIndexArray = State.GetAssetsByPackageName(Asset.Key);
