Remove a unexpected warning from AssetRegUtil commandlet.

diff --git a/Engine/Source/Editor/UnrealEd/Private/Commandlets/AssetRegUtilCommandlet.cpp b/Engine/Source/Editor/UnrealEd/Private/Commandlets/AssetRegUtilCommandlet.cpp
--- a/Engine/Source/Editor/UnrealEd/Private/Commandlets/AssetRegUtilCommandlet.cpp
+++ b/Engine/Source/Editor/UnrealEd/Private/Commandlets/AssetRegUtilCommandlet.cpp
@@ -532,6 +532,8 @@ int32 UAssetRegUtilCommandlet::Main(const FString& CmdLineParams)
 		FString OldOrderFilePath;
 		if (FParse::Value(*CmdLineParams, TEXT("OldFileOpenOrderFile="), OldOrderFilePath, false))
 		{
+			UE_LOG(LogAssetRegUtil, Display, TEXT("Generating partially-updated order file."));
+
 			FString OutPartialOrderFilePath;
 			if (!FParse::Value(*CmdLineParams, TEXT("PartialFileOpenOrderOutput="), OutPartialOrderFilePath, false))
 			{
@@ -541,10 +543,6 @@ int32 UAssetRegUtilCommandlet::Main(const FString& CmdLineParams)
 
 			GeneratePartiallyUpdatedOrderFile(OldOrderFilePath, ReorderOutput, OutPartialOrderFilePath, PatchSizePerfBalanceFactor);
 		}
-		else
-		{
-			UE_LOG(LogAssetRegUtil, Error, TEXT("Invalid or missing arguments.  Skipping generating partially-updated order file."));
-		}
 	}
 
 	return 0;
