Fix crash when saving over existing level and pressing esc #rb lauren.barnes

diff --git a/Engine/Source/Editor/UnrealEd/Private/FileHelpers.cpp b/Engine/Source/Editor/UnrealEd/Private/FileHelpers.cpp
--- a/Engine/Source/Editor/UnrealEd/Private/FileHelpers.cpp
+++ b/Engine/Source/Editor/UnrealEd/Private/FileHelpers.cpp
@@ -688,6 +688,15 @@ static bool SaveWorld(UWorld* World,
 					{
 						Package = DuplicatedWorld->GetOutermost();
 					}
+					else
+					{
+						// Avoid assert during rename when duplicate fails
+						if (!Package->Rename(*NewPackageName, NULL, REN_Test))
+						{
+							FMessageDialog::Open(EAppMsgType::Ok, FText::Format(NSLOCTEXT("UnrealEd", "Error_OverwriteMapCleanup", "Unable to overwrite existing package {0}."), FText::FromString(NewPackageName)));
+							return false;
+						}
+					}
 				}
 
 				if (!DuplicatedWorld)
diff --git a/Engine/Source/Editor/UnrealEd/Private/ObjectTools.cpp b/Engine/Source/Editor/UnrealEd/Private/ObjectTools.cpp
--- a/Engine/Source/Editor/UnrealEd/Private/ObjectTools.cpp
+++ b/Engine/Source/Editor/UnrealEd/Private/ObjectTools.cpp
@@ -495,7 +495,7 @@ namespace ObjectTools
 				}
 			}
 
-			const int32 NumObjectsDeleted = ObjectTools::DeleteObjects(ObjectsToDelete, bPromptToOverwrite);
+			const int32 NumObjectsDeleted = ObjectTools::DeleteObjects(ObjectsToDelete, bPromptToOverwrite, EAllowCancelDuringDelete::CancelNotAllowed);
 
 			// Remove all packages that we added to the root set above.
 			for ( auto PkgIt = DeletedObjectPackages.CreateConstIterator(); PkgIt; ++PkgIt )
@@ -2089,7 +2089,7 @@ namespace ObjectTools
 		return false;
 	}
 
-	int32 DeleteObjects( const TArray< UObject* >& InObjectsToDelete, bool bShowConfirmation )
+	int32 DeleteObjects( const TArray< UObject* >& InObjectsToDelete, bool bShowConfirmation, EAllowCancelDuringDelete AllowCancelDuringDelete )
 	{
 		const FScopedBusyCursor BusyCursor;
 
@@ -2183,24 +2183,30 @@ namespace ObjectTools
 		}
 		
 		bool bUserCanceled = false;
-
-		GWarn->BeginSlowTask(NSLOCTEXT("UnrealEd", "VerifyingDelete", "Verifying Delete"), true, true);
+		const bool bAllowCancelDuringDelete = (AllowCancelDuringDelete == EAllowCancelDuringDelete::AllowCancel);
+		GWarn->BeginSlowTask(NSLOCTEXT("UnrealEd", "VerifyingDelete", "Verifying Delete"), true, bAllowCancelDuringDelete);
 		while ( !bUserCanceled && DeleteModel->GetState() != FAssetDeleteModel::Finished )
 		{
 			DeleteModel->Tick(0);
 			GWarn->StatusUpdate((int32)( DeleteModel->GetProgress() * 100 ), 100, DeleteModel->GetProgressText());
 
-			bUserCanceled = GWarn->ReceivedUserCancel();
+			if (bAllowCancelDuringDelete)
+			{
+				bUserCanceled = GWarn->ReceivedUserCancel();
+			}
 		}
 		GWarn->EndSlowTask();
 
 		if ( bUserCanceled )
 		{
+			UE_LOG(LogUObjectGlobals, Warning, TEXT("User canceled delete operation"));
 			return 0;
 		}
 
 		if ( !DeleteModel->DoDelete() )
 		{
+			UE_LOG(LogUObjectGlobals, Warning, TEXT("Could not delete"));
+
 			//@todo ndarnell explain why the delete failed?  Maybe we should show the delete UI
 			// when this fails?
 		}
diff --git a/Engine/Source/Editor/UnrealEd/Public/ObjectTools.h b/Engine/Source/Editor/UnrealEd/Public/ObjectTools.h
--- a/Engine/Source/Editor/UnrealEd/Public/ObjectTools.h
+++ b/Engine/Source/Editor/UnrealEd/Public/ObjectTools.h
@@ -127,6 +127,12 @@ namespace ObjectTools
 		bool IsValid() const;
 	};
 
+	enum class EAllowCancelDuringDelete : uint8
+	{
+		AllowCancel,
+		CancelNotAllowed
+	};
+
 	/**
 	 * Handles fully loading packages for a set of passed in objects.
 	 *
@@ -272,7 +278,7 @@ namespace ObjectTools
 	 *
 	 * @return The number of objects successfully deleted
 	 */
-	UNREALED_API int32 DeleteObjects( const TArray< UObject* >& ObjectsToDelete, bool bShowConfirmation = true );
+	UNREALED_API int32 DeleteObjects( const TArray< UObject* >& ObjectsToDelete, bool bShowConfirmation = true, EAllowCancelDuringDelete AllowCancelDuringDelete = EAllowCancelDuringDelete::AllowCancel);
 
 	/**
 	* Deletes the list of objects without checking if they are still being used.  This should not be called directly
