Fix bugit not working in editor

diff --git a/Engine/Source/Editor/UnrealEd/Classes/Editor/EditorEngine.h b/Engine/Source/Editor/UnrealEd/Classes/Editor/EditorEngine.h
--- a/Engine/Source/Editor/UnrealEd/Classes/Editor/EditorEngine.h
+++ b/Engine/Source/Editor/UnrealEd/Classes/Editor/EditorEngine.h
@@ -985,6 +985,7 @@ public:
 	bool	HandleDumpPublicCommand( const TCHAR* Str, FOutputDevice& Ar );
 	bool	HandleJumpToCommand( const TCHAR* Str, FOutputDevice& Ar );
 	bool	HandleBugItGoCommand( const TCHAR* Str, FOutputDevice& Ar );
+	bool	HandleBugItCommand(const TCHAR* Str, FOutputDevice& Ar);
 	bool	HandleTagSoundsCommand( const TCHAR* Str, FOutputDevice& Ar );
 	bool	HandlecheckSoundsCommand( const TCHAR* Str, FOutputDevice& Ar );
 	bool	HandleFixupBadAnimNotifiersCommand( const TCHAR* Str, FOutputDevice& Ar );
diff --git a/Engine/Source/Editor/UnrealEd/Private/EditorServer.cpp b/Engine/Source/Editor/UnrealEd/Private/EditorServer.cpp
--- a/Engine/Source/Editor/UnrealEd/Private/EditorServer.cpp
+++ b/Engine/Source/Editor/UnrealEd/Private/EditorServer.cpp
@@ -5808,6 +5808,10 @@ bool UEditorEngine::Exec( UWorld* InWorld, const TCHAR* Stream, FOutputDevice& A
 	{
 		return HandleJumpToCommand( Str, Ar );
 	}
+	else if (FParse::Command(&Str, TEXT("BugIt")))
+	{
+		return HandleBugItCommand(Str, Ar);
+	}
 	else if( FParse::Command(&Str,TEXT("BugItGo")) )
 	{
 		return HandleBugItGoCommand( Str, Ar );
@@ -6282,6 +6286,25 @@ bool UEditorEngine::HandleJumpToCommand( const TCHAR* Str, FOutputDevice& Ar )
 	return true;
 }
 
+
+bool UEditorEngine::HandleBugItCommand(const TCHAR* Str, FOutputDevice& Ar)
+{
+	if(GCurrentLevelEditingViewportClient && !PlayWorld)
+	{
+		FVector ViewLocation = GCurrentLevelEditingViewportClient->GetViewLocation();
+		FRotator ViewRotation = GCurrentLevelEditingViewportClient->GetViewRotation();
+
+		FString GoString = FString::Printf(TEXT("BugItGo %f %f %f %f %f %f"), ViewLocation.X, ViewLocation.Y, ViewLocation.Z, ViewRotation.Pitch, ViewRotation.Yaw, ViewRotation.Roll);
+		UE_LOG(LogEditorServer, Log, TEXT("%s"), *GoString);
+
+		FPlatformApplicationMisc::ClipboardCopy(*GoString);
+
+		return true;
+	}
+
+	return false;
+}
+
 bool UEditorEngine::HandleBugItGoCommand( const TCHAR* Str, FOutputDevice& Ar )
 {
 	if (PlayWorld)
@@ -6316,6 +6339,8 @@ bool UEditorEngine::HandleBugItGoCommand( const TCHAR* Str, FOutputDevice& Ar )
 	Stream = GetFROTATORSpaceDelimited( Stream, Rot, 1.0f );
 	if( Stream != NULL )
 	{
+		// Zero out roll as the editor camera should not roll 
+		Rot.Roll = 0;
 		for(FLevelEditorViewportClient* ViewportClient : GetLevelViewportClients())
 		{
 			ViewportClient->SetViewRotation( Rot );
