Revert 6954065 that caused problems with actor pivot locations

diff --git a/Engine/Source/Editor/UnrealEd/Private/EditorModeManager.cpp b/Engine/Source/Editor/UnrealEd/Private/EditorModeManager.cpp
--- a/Engine/Source/Editor/UnrealEd/Private/EditorModeManager.cpp
+++ b/Engine/Source/Editor/UnrealEd/Private/EditorModeManager.cpp
@@ -520,46 +520,6 @@ bool FEditorModeTools::IsOnlyVisibleActiveMode(FEditorModeID InMode) const
 	return true;
 }
 
-void FEditorModeTools::UpdateModeWidgetLocation()
-{
-	if (!bIsTracking)
-	{
-		USelection* SelectedActors = GetSelectedActors();
-
-		if (SelectedActors && SelectedActors->Num() > 0)
-		{
-			AActor* Actor = Cast<AActor>(SelectedActors->GetSelectedObject(SelectedActors->Num() - 1));
-
-			if (Actor)
-			{
-				FVector CurrentActorLocation = Actor->GetActorLocation();
-
-				if (Actor == LastSelectedActor)
-				{
-					if (!(LastSelectedActorLocation - CurrentActorLocation).IsNearlyZero())
-					{
-						for (const auto& Mode : ActiveModes)
-						{
-							if (Mode->UsesTransformWidget())
-							{
-								SetPivotLocation(CurrentActorLocation, false);
-								LastSelectedActorLocation = CurrentActorLocation;
-								break;
-							}
-						}
-
-					}
-					// lbarnes: need to handle UEdMode gizmos here
-				}
-				else
-				{
-					LastSelectedActor = Actor;
-				}
-			}
-		}
-	}
-}
-
 void FEditorModeTools::OnEditorSelectionChanged(UObject* NewSelection)
 {
 	if(NewSelection == GetSelectedActors())
diff --git a/Engine/Source/Editor/UnrealEd/Public/EditorModeManager.h b/Engine/Source/Editor/UnrealEd/Public/EditorModeManager.h
--- a/Engine/Source/Editor/UnrealEd/Public/EditorModeManager.h
+++ b/Engine/Source/Editor/UnrealEd/Public/EditorModeManager.h
@@ -103,10 +103,6 @@ protected:
 private:
 	void RebuildModeToolBar();
 	void SpawnOrUpdateModeToolbar();
-
-	/** Updates the widget's position for changes in location during simulation */
-	void UpdateModeWidgetLocation();
-
 public:
 
 	/**
@@ -585,11 +581,6 @@ protected:
 
 	/** if true the current selection has a scene component */
 	bool bSelectionHasSceneComponent;
-
-	/** Compare for transformation of selected actor to mode widget */
-	TWeakObjectPtr<AActor>	LastSelectedActor;
-	FVector					LastSelectedActorLocation;
-
 private:
 	struct FEdModeToolbarRow
 	{
