Remove stray UpdateModeWidgetLocation that should have been removed in CL# 9422979
#rnx

diff --git a/Engine/Source/Editor/UnrealEd/Private/EditorModeManager.cpp b/Engine/Source/Editor/UnrealEd/Private/EditorModeManager.cpp
--- a/Engine/Source/Editor/UnrealEd/Private/EditorModeManager.cpp
+++ b/Engine/Source/Editor/UnrealEd/Private/EditorModeManager.cpp
@@ -1448,8 +1448,6 @@ void FEditorModeTools::Tick( FEditorViewportClient* ViewportClient, float DeltaT
 		UEdMode* Mode = ActiveScriptableModes[ModeIndex];
 		Mode->Tick(ViewportClient, DeltaTime);
 	}
-
-	UpdateModeWidgetLocation();
 }
 
 /** Notifies all active modes of any change in mouse movement */
