Back out changelist 9371208. Suspected cause of long pak times

diff --git a/Engine/Source/Editor/UnrealEd/Private/Commandlets/AssetRegistryGenerator.cpp b/Engine/Source/Editor/UnrealEd/Private/Commandlets/AssetRegistryGenerator.cpp
--- a/Engine/Source/Editor/UnrealEd/Private/Commandlets/AssetRegistryGenerator.cpp
+++ b/Engine/Source/Editor/UnrealEd/Private/Commandlets/AssetRegistryGenerator.cpp
@@ -339,7 +339,7 @@ bool FAssetRegistryGenerator::GenerateStreamingInstallManifest(int64 InExtraFlav
 				int32 PakchunkIndex = GetPakchunkIndex(ChunkID);
 				if (PakchunkIndex >= FinalChunkManifests.Num())
 				{
-					FinalChunkManifests.AddDefaulted(PakchunkIndex - FinalChunkManifests.Num() + 1);
+					FinalChunkManifests.AddZeroed(PakchunkIndex - FinalChunkManifests.Num() + 1);
 				}
 				checkf(PakchunkIndex < FinalChunkManifests.Num(), TEXT("Chunk %i out of range. %i manifests available"), PakchunkIndex, FinalChunkManifests.Num() - 1);
 				checkf(FinalChunkManifests[PakchunkIndex] == nullptr, TEXT("Manifest already exists for chunk %i"), PakchunkIndex);
@@ -368,7 +368,7 @@ bool FAssetRegistryGenerator::GenerateStreamingInstallManifest(int64 InExtraFlav
 		ChunkLayerFile->Serialize(TCHAR_TO_ANSI(*LayerString), LayerString.Len());
 
 		// Is this chunk empty?
-		if (!FinalChunkManifests[PakchunkIndex])
+		if (!FinalChunkManifests[PakchunkIndex] || FinalChunkManifests[PakchunkIndex]->Num() == 0)
 		{
 			continue;
 		}
@@ -626,8 +626,7 @@ bool FAssetRegistryGenerator::SaveManifests(FSandboxPlatformFile* InSandboxFile,
 		// Collect all unique chunk indices and map all files to their chunks
 		for (int32 PakchunkIndex = 0; PakchunkIndex < FinalChunkManifests.Num(); ++PakchunkIndex)
 		{
-			check(FinalChunkManifests[PakchunkIndex]);
-			if (FinalChunkManifests[PakchunkIndex]->Num())
+			if (FinalChunkManifests[PakchunkIndex] && FinalChunkManifests[PakchunkIndex]->Num())
 			{
 				PakchunkIndicesInUse.Add(PakchunkIndex);
 				for (auto& Filename : *FinalChunkManifests[PakchunkIndex])
@@ -1441,7 +1440,7 @@ bool FAssetRegistryGenerator::GenerateAssetChunkInformationCSV(const FString& Ou
 	TArray<TUniquePtr<FArchive>> ChunkFiles;
 	if (bWriteIndividualFiles)
 	{
-		for (int32 PakchunkIndex = 0; PakchunkIndex < FinalChunkManifests.Num(); ++PakchunkIndex)
+		for (int32 PakchunkIndex = 0, ChunkNum = FinalChunkManifests.Num(); PakchunkIndex < ChunkNum; ++PakchunkIndex)
 		{
 			FArchive* ChunkFile = IFileManager::Get().CreateFileWriter(*FPaths::Combine(*OutputPath, *FString::Printf(TEXT("Chunks%dInfo.csv"), PakchunkIndex)));
 			if (ChunkFile == nullptr)
@@ -1690,12 +1689,12 @@ void FAssetRegistryGenerator::FixupPackageDependenciesForChunks(FSandboxPlatform
 
 	for (int32 PakchunkIndex = 0, MaxPakchunk = ChunkManifests.Num(); PakchunkIndex < MaxPakchunk; ++PakchunkIndex)
 	{
-		FinalChunkManifests.Add(new FChunkPackageSet());
+		FinalChunkManifests.Add(nullptr);
 		if (!ChunkManifests[PakchunkIndex])
 		{
 			continue;
 		}
-		
+		FinalChunkManifests[PakchunkIndex] = new FChunkPackageSet();
 		for (auto It = ChunkManifests[PakchunkIndex]->CreateConstIterator(); It; ++It)
 		{
 			AddPackageAndDependenciesToChunk(FinalChunkManifests[PakchunkIndex], It.Key(), It.Value(), PakchunkIndex, InSandboxFile);
@@ -1712,10 +1711,12 @@ void FAssetRegistryGenerator::FixupPackageDependenciesForChunks(FSandboxPlatform
 	//Once complete, Add any remaining assets (that are not assigned to a chunk) to the first chunk.
 	if (FinalChunkManifests.Num() == 0)
 	{
-		FinalChunkManifests.Add(new FChunkPackageSet());
+		FinalChunkManifests.Add(nullptr);
+	}
+	if (!FinalChunkManifests[0])
+	{
+		FinalChunkManifests[0] = new FChunkPackageSet();
 	}
-	check(FinalChunkManifests[0]);
-	
 	// Copy the remaining assets
 	auto RemainingAssets = UnassignedPackageSet;
 	for (auto It = RemainingAssets.CreateConstIterator(); It; ++It)
@@ -1749,18 +1750,17 @@ void FAssetRegistryGenerator::FixupPackageDependenciesForChunks(FSandboxPlatform
 	for (int32 PakchunkIndex = 0, MaxPakchunk = ChunkManifests.Num(); PakchunkIndex < MaxPakchunk; ++PakchunkIndex)
 	{
 		const int32 ChunkManifestNum = ChunkManifests[PakchunkIndex] ? ChunkManifests[PakchunkIndex]->Num() : 0;
-		const int32 FinalChunkManifestNum = FinalChunkManifests[PakchunkIndex]->Num();
+		const int32 FinalChunkManifestNum = FinalChunkManifests[PakchunkIndex] ? FinalChunkManifests[PakchunkIndex]->Num() : 0;
 		UE_LOG(LogAssetRegistryGenerator, Log, TEXT("Chunk: %i, Started with %i packages, Final after dependency resolve: %i"), PakchunkIndex, ChunkManifestNum, FinalChunkManifestNum);
 	}
 	
 	// Fix up the asset registry to reflect this chunk layout
 	for (int32 PakchunkIndex = 0 ; PakchunkIndex < FinalChunkManifests.Num(); ++PakchunkIndex)
 	{
-		if (PakchunkIndex >= FinalChunkManifests.Num())
+		if (PakchunkIndex >= FinalChunkManifests.Num() || !FinalChunkManifests[PakchunkIndex])
 		{
 			continue;
 		}
-		check(FinalChunkManifests[PakchunkIndex]);
 		for (const TPair<FName, FString>& Asset : *FinalChunkManifests[PakchunkIndex])
 		{
 			const TArray<const FAssetData*> AssetIndexArray = State.GetAssetsByPackageName(Asset.Key);
