Rebuild HLOD Clusters in the nightly rebuild process

The nightly process was previously setup to only generate proxy meshes. Given the amount of changes going on for S11, it makes sense to also rebuild HLOD clusters to ensure consistency. This will save a ton of time to Bryce.

Also added logging to better visualize that HLOD build options were properly set from the command line.

diff --git a/Engine/Source/Editor/UnrealEd/Private/Commandlets/ContentCommandlets.cpp b/Engine/Source/Editor/UnrealEd/Private/Commandlets/ContentCommandlets.cpp
--- a/Engine/Source/Editor/UnrealEd/Private/Commandlets/ContentCommandlets.cpp
+++ b/Engine/Source/Editor/UnrealEd/Private/Commandlets/ContentCommandlets.cpp
@@ -937,6 +937,17 @@ int32 UResavePackagesCommandlet::Main( const FString& Params )
 	bForceEnableHLODForLevel = HLODOptions.Contains("ForceEnableHLOD");
 	bForceSingleClusterForLevel = HLODOptions.Contains("ForceSingleCluster");
 
+	if (bShouldBuildHLOD)
+	{
+		UE_LOG(LogContentCommandlet, Display, TEXT("Rebuilding HLODs... Options are:"));
+		UE_LOG(LogContentCommandlet, Display, TEXT("  [%s] Clusters"), bGenerateClusters ? TEXT("X") : TEXT(" "));
+		UE_LOG(LogContentCommandlet, Display, TEXT("  [%s] Proxies"), bGenerateMeshProxies ? TEXT("X") : TEXT(" "));
+		UE_LOG(LogContentCommandlet, Display, TEXT("  [%s] ForceClusters"), bForceClusterGeneration ? TEXT("X") : TEXT(" "));
+		UE_LOG(LogContentCommandlet, Display, TEXT("  [%s] ForceProxies"), bForceProxyGeneration ? TEXT("X") : TEXT(" "));
+		UE_LOG(LogContentCommandlet, Display, TEXT("  [%s] ForceEnableHLOD"), bForceEnableHLODForLevel ? TEXT("X") : TEXT(" "));
+		UE_LOG(LogContentCommandlet, Display, TEXT("  [%s] ForceSingleCluster"), bForceSingleClusterForLevel ? TEXT("X") : TEXT(" "));
+	}
+
 	ForceHLODSetupAsset = FString();
 	FParse::Value(*Params, TEXT("ForceHLODSetupAsset="), ForceHLODSetupAsset);
 
