Fix crash after clicking launch if failed and message log was open to 'Packaging Results'

diff --git a/Engine/Source/Editor/UnrealEd/Classes/Editor/EditorEngine.h b/Engine/Source/Editor/UnrealEd/Classes/Editor/EditorEngine.h
--- a/Engine/Source/Editor/UnrealEd/Classes/Editor/EditorEngine.h
+++ b/Engine/Source/Editor/UnrealEd/Classes/Editor/EditorEngine.h
@@ -3084,7 +3084,7 @@ protected:
 	void HandleStageStarted(const FString& InStage, TWeakPtr<SNotificationItem> NotificationItemPtr);
 	void HandleStageCompleted(const FString& InStage, double StageTime, bool bHasCode, TWeakPtr<SNotificationItem> NotificationItemPtr);
 	void HandleLaunchCanceled(double TotalTime, bool bHasCode, TWeakPtr<SNotificationItem> NotificationItemPtr);
-	void HandleLaunchCompleted(bool Succeeded, double TotalTime, int32 ErrorCode, bool bHasCode, TWeakPtr<SNotificationItem> NotificationItemPtr, TSharedPtr<class FMessageLog> MessageLog);
+	void HandleLaunchCompleted(bool Succeeded, double TotalTime, int32 ErrorCode, bool bHasCode, TWeakPtr<SNotificationItem> NotificationItemPtr);
 
 	// Handle requests from slate application to open assets.
 	bool HandleOpenAsset(UObject* Asset);
diff --git a/Engine/Source/Editor/UnrealEd/Private/PlayLevel.cpp b/Engine/Source/Editor/UnrealEd/Private/PlayLevel.cpp
--- a/Engine/Source/Editor/UnrealEd/Private/PlayLevel.cpp
+++ b/Engine/Source/Editor/UnrealEd/Private/PlayLevel.cpp
@@ -112,6 +112,7 @@
 #include "Widgets/Text/STextBlock.h"
 #include "Widgets/SBoxPanel.h"
 #include "Subsystems/AssetEditorSubsystem.h"
+#include "Async/Async.h"
 
 
 DEFINE_LOG_CATEGORY_STATIC(LogPlayLevel, Log, All);
@@ -1696,7 +1697,7 @@ void UEditorEngine::HandleLaunchCanceled(double TotalTime, bool bHasCode, TWeakP
 	bPlayUsingLauncher = false;	
 }
 
-void UEditorEngine::HandleLaunchCompleted(bool Succeeded, double TotalTime, int32 ErrorCode, bool bHasCode, TWeakPtr<SNotificationItem> NotificationItemPtr, TSharedPtr<class FMessageLog> MessageLog)
+void UEditorEngine::HandleLaunchCompleted(bool Succeeded, double TotalTime, int32 ErrorCode, bool bHasCode, TWeakPtr<SNotificationItem> NotificationItemPtr)
 {
 	const FString DummyIOSDeviceName(FString::Printf(TEXT("All_iOS_On_%s"), FPlatformProcess::ComputerName()));
 	const FString DummyTVOSDeviceName(FString::Printf(TEXT("All_tvOS_On_%s"), FPlatformProcess::ComputerName()));
@@ -1739,12 +1740,17 @@ void UEditorEngine::HandleLaunchCompleted(bool Succeeded, double TotalTime, int3
 			CompletionMsg = LOCTEXT("LauncherTaskFailed", "Launch failed!");
 		}
 		
-		MessageLog->Error()
-			->AddToken(FTextToken::Create(CompletionMsg))
-			->AddToken(FTextToken::Create(FText::FromString(FEditorAnalytics::TranslateErrorCode(ErrorCode))));
+		AsyncTask(ENamedThreads::GameThread, [=]
+		{
+			FMessageLog MessageLog("PackagingResults");
+
+			MessageLog.Error()
+				->AddToken(FTextToken::Create(CompletionMsg))
+				->AddToken(FTextToken::Create(FText::FromString(FEditorAnalytics::TranslateErrorCode(ErrorCode))));
 
-		// flush log, because it won't be destroyed until the notification popup closes
-		MessageLog->NumMessages(EMessageSeverity::Info);
+			// flush log, because it won't be destroyed until the notification popup closes
+			MessageLog.NumMessages(EMessageSeverity::Info);
+		});
 
 		TGraphTask<FLauncherNotificationTask>::CreateTask().ConstructAndDispatchWhenReady(
 			NotificationItemPtr,
@@ -2179,13 +2185,11 @@ void UEditorEngine::PlayUsingLauncher()
 		TWeakPtr<SNotificationItem> NotificationItemPtr(NotificationItem);
 		if (GEditor->LauncherWorker.IsValid() && GEditor->LauncherWorker->GetStatus() != ELauncherWorkerStatus::Completed)
 		{
-			TSharedPtr<FMessageLog> MessageLog = MakeShareable(new FMessageLog("PackagingResults"));
-
 			GEditor->PlayEditorSound(TEXT("/Engine/EditorSounds/Notifications/CompileStart_Cue.CompileStart_Cue"));
 			GEditor->LauncherWorker->OnOutputReceived().AddStatic(HandleOutputReceived);
 			GEditor->LauncherWorker->OnStageStarted().AddUObject(this, &UEditorEngine::HandleStageStarted, NotificationItemPtr);
 			GEditor->LauncherWorker->OnStageCompleted().AddUObject(this, &UEditorEngine::HandleStageCompleted, bPlayUsingLauncherHasCode, NotificationItemPtr);
-			GEditor->LauncherWorker->OnCompleted().AddUObject(this, &UEditorEngine::HandleLaunchCompleted, bPlayUsingLauncherHasCode, NotificationItemPtr, MessageLog);
+			GEditor->LauncherWorker->OnCompleted().AddUObject(this, &UEditorEngine::HandleLaunchCompleted, bPlayUsingLauncherHasCode, NotificationItemPtr);
 			GEditor->LauncherWorker->OnCanceled().AddUObject(this, &UEditorEngine::HandleLaunchCanceled, bPlayUsingLauncherHasCode, NotificationItemPtr);
 		}
 		else
