Fix disabled section recovery after a reimport on a inline reduced LOD. Also fix an issue with morph LOD remap when there is a cloth section not at the end of the sections
#rnx

diff --git a/Engine/Plugins/Experimental/SkeletalReduction/Source/Private/SkeletalMeshReductionPlugin.cpp b/Engine/Plugins/Experimental/SkeletalReduction/Source/Private/SkeletalMeshReductionPlugin.cpp
--- a/Engine/Plugins/Experimental/SkeletalReduction/Source/Private/SkeletalMeshReductionPlugin.cpp
+++ b/Engine/Plugins/Experimental/SkeletalReduction/Source/Private/SkeletalMeshReductionPlugin.cpp
@@ -2031,7 +2031,7 @@ void FQuadricSkeletalMeshReduction::ReduceSkeletalMesh(USkeletalMesh& SkeletalMe
 				{
 					continue;
 				}
-				if (SectionData.OriginalDataSectionIndex == ImportedModelLOD.Sections[SectionIndex].OriginalDataSectionIndex)
+				if (SectionData.MaterialIndex == ImportedModelLOD.Sections[SectionIndex].MaterialIndex)
 				{
 					//Restore the LOD material map for this section
 					if (SectionData.MaterialMap != INDEX_NONE)
diff --git a/Engine/Source/Developer/SkeletalMeshUtilitiesCommon/Private/LODUtilities.cpp b/Engine/Source/Developer/SkeletalMeshUtilitiesCommon/Private/LODUtilities.cpp
--- a/Engine/Source/Developer/SkeletalMeshUtilitiesCommon/Private/LODUtilities.cpp
+++ b/Engine/Source/Developer/SkeletalMeshUtilitiesCommon/Private/LODUtilities.cpp
@@ -869,8 +869,8 @@ void FLODUtilities::ApplyMorphTargetsToLOD(USkeletalMesh* SkeletalMesh, int32 So
 	check(!TargetSectionMatchBaseIndex.Contains(INDEX_NONE));
 	TArray<FSoftSkinVertex> BaseVertices;
 	TArray<FSoftSkinVertex> TargetVertices;
-	BaseLODModel.GetNonClothVertices(BaseVertices);
-	TargetLODModel.GetNonClothVertices(TargetVertices);
+	BaseLODModel.GetVertices(BaseVertices);
+	TargetLODModel.GetVertices(TargetVertices);
 	//Create the base triangle indices per section
 	TArray<TArray<uint32>> BaseTriangleIndices;
 	int32 SectionCount = BaseLODModel.Sections.Num();
@@ -878,16 +878,13 @@ void FLODUtilities::ApplyMorphTargetsToLOD(USkeletalMesh* SkeletalMesh, int32 So
 	for (int32 SectionIndex = 0; SectionIndex < SectionCount; ++SectionIndex)
 	{
 		const FSkelMeshSection& Section = BaseLODModel.Sections[SectionIndex];
-		if (Section.ClothingData.AssetGuid.IsValid())
-		{
-			continue;
-		}
 		uint32 TriangleCount = Section.NumTriangles;
 		for (uint32 TriangleIndex = 0; TriangleIndex < TriangleCount; ++TriangleIndex)
 		{
 			for (uint32 PointIndex = 0; PointIndex < 3; PointIndex++)
 			{
-				BaseTriangleIndices[SectionIndex].Add(BaseLODModel.IndexBuffer[Section.BaseIndex + ((TriangleIndex * 3) + PointIndex)]);
+				uint32 IndexBufferValue = BaseLODModel.IndexBuffer[Section.BaseIndex + ((TriangleIndex * 3) + PointIndex)];
+				BaseTriangleIndices[SectionIndex].Add(IndexBufferValue);
 			}
 		}
 	}
diff --git a/Engine/Source/Editor/UnrealEd/Private/Factories/SkeletalMeshImport.cpp b/Engine/Source/Editor/UnrealEd/Private/Factories/SkeletalMeshImport.cpp
--- a/Engine/Source/Editor/UnrealEd/Private/Factories/SkeletalMeshImport.cpp
+++ b/Engine/Source/Editor/UnrealEd/Private/Factories/SkeletalMeshImport.cpp
@@ -234,7 +234,9 @@ ExistingSkelMeshData* SaveExistingSkelMeshData(USkeletalMesh* ExistingSkelMesh,
 	{
 		return nullptr;
 	}
+
 	bool ReimportSpecificLOD = (ReimportLODIndex > 0) && ExistingSkelMesh->GetLODNum() > ReimportLODIndex;
+	int32 SafeReimportLODIndex = ReimportLODIndex < 0 ? 0 : ReimportLODIndex;
 
 	ExistingMeshDataPtr = new ExistingSkelMeshData();
 		
@@ -255,17 +257,61 @@ ExistingSkelMeshData* SaveExistingSkelMeshData(USkeletalMesh* ExistingSkelMesh,
 
 	for (int32 LodIndex = 0; LodIndex < ImportedResource->LODModels.Num(); ++LodIndex)
 	{
-		const FSkeletalMeshLODModel& LODModel = ImportedResource->LODModels[LodIndex];
+		FSkeletalMeshLODModel OriginalLODModel;
+		FSkeletalMeshLODModel* BackupLODModel = &(ImportedResource->LODModels[LodIndex]);
+		if (LodIndex == SafeReimportLODIndex && (ImportedResource->OriginalReductionSourceMeshData.IsValidIndex(SafeReimportLODIndex) && !ImportedResource->OriginalReductionSourceMeshData[SafeReimportLODIndex]->IsEmpty()))
+		{
+			FSkeletalMeshLODInfo* LODInfo = ExistingSkelMesh->GetLODInfo(SafeReimportLODIndex);
+			TMap<FString, TArray<FMorphTargetDelta>> TempLODMorphTargetData;
+			//Get the before reduce LODModel, this lod model contain all the possible sections
+			ImportedResource->OriginalReductionSourceMeshData[SafeReimportLODIndex]->LoadReductionData(OriginalLODModel, TempLODMorphTargetData, ExistingSkelMesh);
+			//If there was section that was remove by the reduction (Disabled in the original data, zero triangle after reduction, GenerateUpTo settings...),
+			//we have to use the original section data and apply the section data that was modified after the reduction
+			if (OriginalLODModel.Sections.Num() > BackupLODModel->Sections.Num())
+			{
+				TArray<bool> OriginalMatched;
+				OriginalMatched.AddZeroed(OriginalLODModel.Sections.Num());
+				//Now apply the after reduce settings change, but we need to match the section since there can be reduced one
+				for (int32 ReduceSectionIndex = 0; ReduceSectionIndex < BackupLODModel->Sections.Num(); ++ReduceSectionIndex)
+				{
+					const FSkelMeshSection& ReduceSection = BackupLODModel->Sections[ReduceSectionIndex];
+					for (int32 OriginalSectionIndex = 0; OriginalSectionIndex < OriginalLODModel.Sections.Num(); ++OriginalSectionIndex)
+					{
+						if (OriginalMatched[OriginalSectionIndex])
+						{
+							continue;
+						}
+						FSkelMeshSection& OriginalSection = OriginalLODModel.Sections[OriginalSectionIndex];
+						if ((OriginalSection.bDisabled) || (OriginalSection.GenerateUpToLodIndex != INDEX_NONE && OriginalSection.GenerateUpToLodIndex < SafeReimportLODIndex))
+						{
+							continue;
+						}
+
+						if (ReduceSection.MaterialIndex == OriginalSection.MaterialIndex)
+						{
+							OriginalMatched[OriginalSectionIndex] = true;
+							OriginalSection.bDisabled = ReduceSection.bDisabled;
+							OriginalSection.bCastShadow = ReduceSection.bCastShadow;
+							OriginalSection.bRecomputeTangent = ReduceSection.bRecomputeTangent;
+							OriginalSection.GenerateUpToLodIndex = ReduceSection.GenerateUpToLodIndex;
+						}
+					}
+				}
+				//Use the OriginalLODModel
+				BackupLODModel = &OriginalLODModel;
+			}
+		}
 		ExistingMeshDataPtr->ExistingImportMeshLodSectionMaterialData.AddZeroed();
-		for (int32 SectionIndex = 0; SectionIndex < LODModel.Sections.Num(); ++SectionIndex)
+		check(ExistingMeshDataPtr->ExistingImportMeshLodSectionMaterialData.IsValidIndex(LodIndex));
+
+		for (int32 SectionIndex = 0; SectionIndex < BackupLODModel->Sections.Num(); ++SectionIndex)
 		{
-			const FSkelMeshSourceSectionUserData& UserSectionData = LODModel.UserSectionsData.FindChecked(LODModel.Sections[SectionIndex].OriginalDataSectionIndex);
-			int32 SectionMaterialIndex = LODModel.Sections[SectionIndex].MaterialIndex;
-			bool SectionCastShadow = UserSectionData.bCastShadow;
-			bool SectionRecomputeTangents = UserSectionData.bRecomputeTangent;
-			int32 GenerateUpTo = UserSectionData.GenerateUpToLodIndex;
-			bool bDisabled = UserSectionData.bDisabled;
-			bool bBoneChunkedSection = LODModel.Sections[SectionIndex].ChunkedParentSectionIndex != INDEX_NONE;
+			int32 SectionMaterialIndex = BackupLODModel->Sections[SectionIndex].MaterialIndex;
+			bool SectionCastShadow = BackupLODModel->Sections[SectionIndex].bCastShadow;
+			bool SectionRecomputeTangents = BackupLODModel->Sections[SectionIndex].bRecomputeTangent;
+			int32 GenerateUpTo = BackupLODModel->Sections[SectionIndex].GenerateUpToLodIndex;
+			bool bDisabled = BackupLODModel->Sections[SectionIndex].bDisabled;
+			bool bBoneChunkedSection = BackupLODModel->Sections[SectionIndex].ChunkedParentSectionIndex != INDEX_NONE;
 			if (ExistingMeshDataPtr->ExistingImportMaterialOriginalNameData.IsValidIndex(SectionMaterialIndex))
 			{
 				ExistingMeshDataPtr->ExistingImportMeshLodSectionMaterialData[LodIndex].Add(ExistingMeshLodSectionData(ExistingMeshDataPtr->ExistingImportMaterialOriginalNameData[SectionMaterialIndex], SectionCastShadow, SectionRecomputeTangents, GenerateUpTo, bDisabled, bBoneChunkedSection));
@@ -285,7 +331,7 @@ ExistingSkelMeshData* SaveExistingSkelMeshData(USkeletalMesh* ExistingSkelMesh,
 		ExistingSkelMesh->GetLODNum() == ImportedResource->LODModels.Num() )
 	{
 		int32 OffsetReductionLODIndex = 0;
-		FSkeletalMeshLODInfo* LODInfo = ExistingSkelMesh->GetLODInfo( ReimportLODIndex < 0 ? 0 : ReimportLODIndex);
+		FSkeletalMeshLODInfo* LODInfo = ExistingSkelMesh->GetLODInfo( SafeReimportLODIndex);
 		ExistingMeshDataPtr->bIsReimportLODReduced = (LODInfo && LODInfo->bHasBeenSimplified);
 		if (ExistingMeshDataPtr->bIsReimportLODReduced)
 		{
@@ -658,7 +704,7 @@ void RestoreExistingSkelMeshData(ExistingSkelMeshData* MeshData, USkeletalMesh*
 	{
 		return;
 	}
-
+	int32 SafeReimportLODIndex = ReimportLODIndex < 0 ? 0 : ReimportLODIndex;
 	SkeletalMesh->MinLod = MeshData->MinLOD;
 	SkeletalMesh->DisableBelowMinLodStripping = MeshData->DisableBelowMinLodStripping;
 	SkeletalMesh->bSupportLODStreaming = MeshData->bSupportLODStreaming;
@@ -735,21 +781,8 @@ void RestoreExistingSkelMeshData(ExistingSkelMeshData* MeshData, USkeletalMesh*
 		SkeletalMesh->LODSettings->SetLODSettingsToMesh(SkeletalMesh, 0);
 	}
 
-	//Copy back the reimport LOD specific data
-	if (SkeletalMesh->GetLODInfoArray().IsValidIndex(ReimportLODIndex))
-	{
-		FSkeletalMeshLODInfo& BaseLODInfo = SkeletalMesh->GetLODInfoArray()[ReimportLODIndex];
-		//Restore the build setting first
-		BaseLODInfo.BuildSettings = MeshData->ExistingBaseLODInfo.BuildSettings;
-		if (MeshData->bIsReimportLODReduced)
-		{
-			//Restore the reimport LOD reduction settings
-			BaseLODInfo.ReductionSettings = MeshData->ExistingBaseLODInfo.ReductionSettings;
-		}
-	}
-
 	//Do everything we need for base LOD re-import
-	if (ReimportLODIndex < 1)
+	if (SafeReimportLODIndex < 1)
 	{
 		// this is not ideal. Ideally we'll have to save only diff with indicating which joints, 
 		// but for now, we allow them to keep the previous pose IF the element count is same
@@ -954,13 +987,13 @@ void RestoreExistingSkelMeshData(ExistingSkelMeshData* MeshData, USkeletalMesh*
 	if (MeshData->UseMaterialNameSlotWorkflow)
 	{
 		//Restore the base LOD materialMap the LODs LODMaterialMap are restore differently
-		if (ReimportLODIndex < 1 && SkeletalMesh->GetLODInfoArray().IsValidIndex(ReimportLODIndex))
+		if (SafeReimportLODIndex < 1 && SkeletalMesh->GetLODInfoArray().IsValidIndex(SafeReimportLODIndex))
 		{
-			FSkeletalMeshLODInfo& BaseLODInfo = SkeletalMesh->GetLODInfoArray()[ReimportLODIndex];
+			FSkeletalMeshLODInfo& BaseLODInfo = SkeletalMesh->GetLODInfoArray()[SafeReimportLODIndex];
 			//Restore the Base MaterialMap
-			if (SkeletalMesh->GetImportedModel() && SkeletalMesh->GetImportedModel()->LODModels.IsValidIndex(ReimportLODIndex))
+			if (SkeletalMesh->GetImportedModel() && SkeletalMesh->GetImportedModel()->LODModels.IsValidIndex(SafeReimportLODIndex))
 			{
-				FSkeletalMeshLODModel& LODModel = SkeletalMesh->GetImportedModel()->LODModels[ReimportLODIndex];
+				FSkeletalMeshLODModel& LODModel = SkeletalMesh->GetImportedModel()->LODModels[SafeReimportLODIndex];
 				for (int32 SectionIndex = 0; SectionIndex < LODModel.Sections.Num(); ++SectionIndex)
 				{
 					int32 MaterialIndex = LODModel.Sections[SectionIndex].MaterialIndex;
@@ -980,7 +1013,7 @@ void RestoreExistingSkelMeshData(ExistingSkelMeshData* MeshData, USkeletalMesh*
 				}
 			}
 		}
-		FSkeletalMeshLODModel &NewSkelMeshLodModel = SkeletalMeshImportedModel->LODModels[ReimportLODIndex];
+		FSkeletalMeshLODModel &NewSkelMeshLodModel = SkeletalMeshImportedModel->LODModels[SafeReimportLODIndex];
 		int32 ParentSectionIndex = 0;
 		//Restore the section changes from the old import data
 		for (int32 SectionIndex = 0; SectionIndex < NewSkelMeshLodModel.Sections.Num(); SectionIndex++)
@@ -994,13 +1027,13 @@ void RestoreExistingSkelMeshData(ExistingSkelMeshData* MeshData, USkeletalMesh*
 				if (SkeletalMesh->Materials.IsValidIndex(RemapMaterial[NewMeshSectionMaterialIndex]))
 				{
 					NewSkelMeshLodModel.Sections[SectionIndex].MaterialIndex = RemapMaterial[NewMeshSectionMaterialIndex];
-					if (MeshData->ExistingImportMeshLodSectionMaterialData[ReimportLODIndex].IsValidIndex(RemapMaterial[NewMeshSectionMaterialIndex]))
+					if (MeshData->ExistingImportMeshLodSectionMaterialData[SafeReimportLODIndex].IsValidIndex(RemapMaterial[NewMeshSectionMaterialIndex]))
 					{
-						NewSkelMeshLodModel.Sections[SectionIndex].bCastShadow = MeshData->ExistingImportMeshLodSectionMaterialData[ReimportLODIndex][RemapMaterial[NewMeshSectionMaterialIndex]].bCastShadow;
-						NewSkelMeshLodModel.Sections[SectionIndex].bRecomputeTangent = MeshData->ExistingImportMeshLodSectionMaterialData[ReimportLODIndex][RemapMaterial[NewMeshSectionMaterialIndex]].bRecomputeTangents;
-						NewSkelMeshLodModel.Sections[SectionIndex].GenerateUpToLodIndex = MeshData->ExistingImportMeshLodSectionMaterialData[ReimportLODIndex][RemapMaterial[NewMeshSectionMaterialIndex]].GenerateUpTo;
-						NewSkelMeshLodModel.Sections[SectionIndex].bDisabled = MeshData->ExistingImportMeshLodSectionMaterialData[ReimportLODIndex][RemapMaterial[NewMeshSectionMaterialIndex]].bDisabled;
-						bool bBoneChunkedSection = MeshData->ExistingImportMeshLodSectionMaterialData[ReimportLODIndex][RemapMaterial[NewMeshSectionMaterialIndex]].bBoneChunkedSection;
+						NewSkelMeshLodModel.Sections[SectionIndex].bCastShadow = MeshData->ExistingImportMeshLodSectionMaterialData[SafeReimportLODIndex][RemapMaterial[NewMeshSectionMaterialIndex]].bCastShadow;
+						NewSkelMeshLodModel.Sections[SectionIndex].bRecomputeTangent = MeshData->ExistingImportMeshLodSectionMaterialData[SafeReimportLODIndex][RemapMaterial[NewMeshSectionMaterialIndex]].bRecomputeTangents;
+						NewSkelMeshLodModel.Sections[SectionIndex].GenerateUpToLodIndex = MeshData->ExistingImportMeshLodSectionMaterialData[SafeReimportLODIndex][RemapMaterial[NewMeshSectionMaterialIndex]].GenerateUpTo;
+						NewSkelMeshLodModel.Sections[SectionIndex].bDisabled = MeshData->ExistingImportMeshLodSectionMaterialData[SafeReimportLODIndex][RemapMaterial[NewMeshSectionMaterialIndex]].bDisabled;
+						bool bBoneChunkedSection = MeshData->ExistingImportMeshLodSectionMaterialData[SafeReimportLODIndex][RemapMaterial[NewMeshSectionMaterialIndex]].bBoneChunkedSection;
 						if (!bBoneChunkedSection)
 						{
 							//Set the new Parent Index
@@ -1018,28 +1051,28 @@ void RestoreExistingSkelMeshData(ExistingSkelMeshData* MeshData, USkeletalMesh*
 				}
 			}
 
-			if (MeshData->LastImportMeshLodSectionMaterialData.Num() < 1 || !MeshData->LastImportMeshLodSectionMaterialData.IsValidIndex(ReimportLODIndex) || MeshData->LastImportMeshLodSectionMaterialData[ReimportLODIndex].Num() <= SectionIndex ||
-				MeshData->ExistingImportMeshLodSectionMaterialData.Num() < 1 || !MeshData->ExistingImportMeshLodSectionMaterialData.IsValidIndex(ReimportLODIndex) || MeshData->ExistingImportMeshLodSectionMaterialData[ReimportLODIndex].Num() <= SectionIndex)
+			if (MeshData->LastImportMeshLodSectionMaterialData.Num() < 1 || !MeshData->LastImportMeshLodSectionMaterialData.IsValidIndex(SafeReimportLODIndex) || MeshData->LastImportMeshLodSectionMaterialData[SafeReimportLODIndex].Num() <= SectionIndex ||
+				MeshData->ExistingImportMeshLodSectionMaterialData.Num() < 1 || !MeshData->ExistingImportMeshLodSectionMaterialData.IsValidIndex(SafeReimportLODIndex) || MeshData->ExistingImportMeshLodSectionMaterialData[SafeReimportLODIndex].Num() <= SectionIndex)
 			{
 				break;
 			}
 
 			FName CurrentSectionImportedMaterialName = SkeletalMesh->Materials[NewSkelMeshLodModel.Sections[SectionIndex].MaterialIndex].ImportedMaterialSlotName;
-			for (int32 ExistSectionIndex = 0; ExistSectionIndex < MeshData->ExistingImportMeshLodSectionMaterialData[ReimportLODIndex].Num(); ++ExistSectionIndex)
+			for (int32 ExistSectionIndex = 0; ExistSectionIndex < MeshData->ExistingImportMeshLodSectionMaterialData[SafeReimportLODIndex].Num(); ++ExistSectionIndex)
 			{
-				if (!MeshData->LastImportMeshLodSectionMaterialData[ReimportLODIndex].IsValidIndex(ExistSectionIndex) || !MeshData->ExistingImportMeshLodSectionMaterialData[ReimportLODIndex].IsValidIndex(ExistSectionIndex))
+				if (!MeshData->LastImportMeshLodSectionMaterialData[SafeReimportLODIndex].IsValidIndex(ExistSectionIndex) || !MeshData->ExistingImportMeshLodSectionMaterialData[SafeReimportLODIndex].IsValidIndex(ExistSectionIndex))
 				{
 					continue;
 				}
 				//Get the Last imported skelmesh section slot import name
-				FName OriginalImportMeshSectionSlotName = MeshData->LastImportMeshLodSectionMaterialData[ReimportLODIndex][ExistSectionIndex];
+				FName OriginalImportMeshSectionSlotName = MeshData->LastImportMeshLodSectionMaterialData[SafeReimportLODIndex][ExistSectionIndex];
 				if (OriginalImportMeshSectionSlotName != CurrentSectionImportedMaterialName)
 				{
 					continue;
 				}
 
 				//Get the current skelmesh section slot import name
-				FName ExistMeshSectionSlotName = MeshData->ExistingImportMeshLodSectionMaterialData[ReimportLODIndex][ExistSectionIndex].ImportedMaterialSlotName;
+				FName ExistMeshSectionSlotName = MeshData->ExistingImportMeshLodSectionMaterialData[SafeReimportLODIndex][ExistSectionIndex].ImportedMaterialSlotName;
 				if (ExistMeshSectionSlotName != OriginalImportMeshSectionSlotName)
 				{
 					//The last import slot name match the New import slot name, but the Exist slot name is different then the last import slot name.
@@ -1062,18 +1095,16 @@ void RestoreExistingSkelMeshData(ExistingSkelMeshData* MeshData, USkeletalMesh*
 		NewSkelMeshLodModel.SyncronizeUserSectionsDataArray(true);
 	}
 
-	if (!bImportSkinningOnly && !MeshData->bIsReimportLODReduced)
+	//Copy back the reimport LOD specific data
+	if (SkeletalMesh->GetLODInfoArray().IsValidIndex(SafeReimportLODIndex))
 	{
-		if (ReimportLODIndex < 0)
-		{
-			if (SkeletalMeshImportedModel->OriginalReductionSourceMeshData.IsValidIndex(0))
-			{
-				SkeletalMeshImportedModel->OriginalReductionSourceMeshData[0]->EmptyBulkData();
-			}
-		}
-		else if (SkeletalMeshImportedModel->OriginalReductionSourceMeshData.IsValidIndex(ReimportLODIndex))
+		FSkeletalMeshLODInfo& BaseLODInfo = SkeletalMesh->GetLODInfoArray()[SafeReimportLODIndex];
+		//Restore the build setting first
+		BaseLODInfo.BuildSettings = MeshData->ExistingBaseLODInfo.BuildSettings;
+		if (MeshData->bIsReimportLODReduced)
 		{
-			SkeletalMeshImportedModel->OriginalReductionSourceMeshData[ReimportLODIndex]->EmptyBulkData();
+			//Restore the reimport LOD reduction settings
+			BaseLODInfo.ReductionSettings = MeshData->ExistingBaseLODInfo.ReductionSettings;
 		}
 	}
 
@@ -1089,6 +1120,15 @@ void RestoreExistingSkelMeshData(ExistingSkelMeshData* MeshData, USkeletalMesh*
 		UserDataObject->Rename(nullptr, SkeletalMesh, REN_DontCreateRedirectors | REN_DoNotDirty);
 		SkeletalMesh->AddAssetUserData(UserDataObject);
 	}
+
+	if (!bImportSkinningOnly && !MeshData->bIsReimportLODReduced)
+	{
+		if (SkeletalMeshImportedModel->OriginalReductionSourceMeshData.IsValidIndex(SafeReimportLODIndex))
+		{
+			SkeletalMeshImportedModel->OriginalReductionSourceMeshData[SafeReimportLODIndex]->EmptyBulkData();
+		}
+	}
+
 	//Copy mesh changed delegate data
 	SkeletalMesh->GetOnMeshChanged() = MeshData->ExistingOnMeshChanged;
 }
