Protect against bad state and log a message in the thumbnail renderer

diff --git a/Engine/Source/Editor/UnrealEd/Private/ThumbnailRendering/SoundWaveThumbnailRenderer.cpp b/Engine/Source/Editor/UnrealEd/Private/ThumbnailRendering/SoundWaveThumbnailRenderer.cpp
--- a/Engine/Source/Editor/UnrealEd/Private/ThumbnailRendering/SoundWaveThumbnailRenderer.cpp
+++ b/Engine/Source/Editor/UnrealEd/Private/ThumbnailRendering/SoundWaveThumbnailRenderer.cpp
@@ -74,11 +74,25 @@ void USoundWaveThumbnailRenderer::Draw(UObject* Object, int32 X, int32 Y, uint32
 					// Make sure we have a mono file here
 					if (SoundWave->NumChannels > 2)
 					{
-						check(WaveInfo.pChannels && *WaveInfo.pChannels == 1);
+						bool bLooksSane = (WaveInfo.pChannels && *WaveInfo.pChannels == 1);
+						if (!bLooksSane)
+						{
+							UE_LOG(LogAudio, Warning, TEXT("USoundWaveThumbnailRenderer::Draw: Not rendering thumbnail for Wave '%s', as its supposed to be mono and has '%d' channels."), 
+								*GetNameSafe(SoundWave), WaveInfo.pChannels ? *WaveInfo.pChannels : -1);
+
+							continue;
+						}
 					}
 					else
 					{
-						check(WaveInfo.pChannels && *WaveInfo.pChannels == SoundWave->NumChannels);
+						bool bLooksSane = WaveInfo.pChannels && *WaveInfo.pChannels == SoundWave->NumChannels;
+						if (!bLooksSane)
+						{
+							UE_LOG(LogAudio, Warning, TEXT("USoundWaveThumbnailRenderer::Draw: Not rendering thumbnail for Wave '%s', as pChannels('%d') != SoundWave->NumChannels('%d')"),
+								*GetNameSafe(SoundWave), WaveInfo.pChannels ? *WaveInfo.pChannels : -1, SoundWave->NumChannels);
+							
+							continue;
+						}
 					}
 
 					// Sample count
