Allow package order randomization on cooks to allow overlapping buildfarm builds to more easily share DDC results
Fix bug in materialinstance postload that was precaching shaders too early to work correctly with cooker-side batching

diff --git a/Engine/Source/Editor/UnrealEd/Private/CookOnTheFlyServer.cpp b/Engine/Source/Editor/UnrealEd/Private/CookOnTheFlyServer.cpp
--- a/Engine/Source/Editor/UnrealEd/Private/CookOnTheFlyServer.cpp
+++ b/Engine/Source/Editor/UnrealEd/Private/CookOnTheFlyServer.cpp
@@ -7145,8 +7145,10 @@ void UCookOnTheFlyServer::StartCookByTheBook( const FCookByTheBookStartupOptions
 		UE_LOG(LogCook, Warning, TEXT("No files found."));
 	}
 
-	if (FParse::Param(FCommandLine::Get(), TEXT("DIFFONLY")) && !FParse::Param(FCommandLine::Get(), TEXT("DIFFNORANDCOOK")))
+	if (FParse::Param(FCommandLine::Get(), TEXT("RANDOMPACKAGEORDER")) || 
+		(FParse::Param(FCommandLine::Get(), TEXT("DIFFONLY")) && !FParse::Param(FCommandLine::Get(), TEXT("DIFFNORANDCOOK"))))
 	{
+		UE_LOG(LogCook, Log, TEXT("Randomizing package order."));
 		//randomize the array, taking the Array_Shuffle approach, in order to help bring cooking determinism issues to the surface.
 		for (int32 FileIndex = 0; FileIndex < FilesInPath.Num(); ++FileIndex)
 		{
diff --git a/Engine/Source/Runtime/Engine/Private/Materials/MaterialInstance.cpp b/Engine/Source/Runtime/Engine/Private/Materials/MaterialInstance.cpp
--- a/Engine/Source/Runtime/Engine/Private/Materials/MaterialInstance.cpp
+++ b/Engine/Source/Runtime/Engine/Private/Materials/MaterialInstance.cpp
@@ -3581,7 +3581,7 @@ void UMaterialInstance::PostLoad()
 
 		// Make sure static parameters are up to date and shaders are cached for the current platform
 		InitStaticPermutation();
-#if WITH_EDITOR
+#if WITH_EDITOR && 0
 		// enable caching in postload for derived data cache commandlet and cook by the book
 		ITargetPlatformManagerModule* TPM = GetTargetPlatformManager();
 		if (TPM && (TPM->RestrictFormatsToRuntimeOnly() == false)) 
