Fix CookOnTheFly server crash caused by new package store cooking

diff --git a/Engine/Source/Editor/UnrealEd/Classes/CookOnTheSide/CookOnTheFlyServer.h b/Engine/Source/Editor/UnrealEd/Classes/CookOnTheSide/CookOnTheFlyServer.h
--- a/Engine/Source/Editor/UnrealEd/Classes/CookOnTheSide/CookOnTheFlyServer.h
+++ b/Engine/Source/Editor/UnrealEd/Classes/CookOnTheSide/CookOnTheFlyServer.h
@@ -689,6 +689,11 @@ private:
 	*/
 	void InitializePackageStore(const TArray<FName>& TargetPlatformNames);
 
+	/**
+	* Finalize the package store
+	*/
+	void FinalizePackageStore();
+
 	/**
 	* Initialize all target platforms
 	*/
diff --git a/Engine/Source/Editor/UnrealEd/Private/CookOnTheFlyServer.cpp b/Engine/Source/Editor/UnrealEd/Private/CookOnTheFlyServer.cpp
--- a/Engine/Source/Editor/UnrealEd/Private/CookOnTheFlyServer.cpp
+++ b/Engine/Source/Editor/UnrealEd/Private/CookOnTheFlyServer.cpp
@@ -1986,7 +1986,7 @@ bool UCookOnTheFlyServer::IsCookByTheBookMode() const
 
 bool UCookOnTheFlyServer::IsUsingPackageStore() const
 {
-	return CookByTheBookOptions->bPackageStore;
+	return IsCookByTheBookMode() && CookByTheBookOptions->bPackageStore;
 }
 
 bool UCookOnTheFlyServer::IsCookOnTheFlyMode() const
@@ -6332,18 +6332,7 @@ void UCookOnTheFlyServer::CookByTheBookFinished()
 
 	UPackage::WaitForAsyncFileWrites();
 	
-	if (IsCookByTheBookMode() && IsUsingPackageStore())
-	{
-		SCOPE_TIMER(SavingNameMap);
-		UE_LOG(LogCook, Display, TEXT("Saving name map(s)..."));
-		const TArray<ITargetPlatform*>& TargetPlatforms = GetCookingTargetPlatforms();
-		for (int32 PlatformIndex = 0; PlatformIndex < TargetPlatforms.Num(); ++PlatformIndex)
-		{
-			INameMapSaver& NameMapSaver = SavePackageContexts[PlatformIndex]->HeaderSaver.NameMapSaver;
-			NameMapSaver.End();
-		}
-		UE_LOG(LogCook, Display, TEXT("Done saving name map(s)"));
-	}
+	FinalizePackageStore();
 
 	GetDerivedDataCacheRef().WaitForQuiescence(true);
 	
@@ -6718,13 +6707,9 @@ void UCookOnTheFlyServer::InitializeSandbox()
 
 void UCookOnTheFlyServer::InitializePackageStore(const TArray<FName>& TargetPlatformNames)
 {
-	// TODO: should ideally support cook-on-the-fly
-
-	if (!(IsCookByTheBookMode() && IsUsingPackageStore()))
+	// TODO: should ideally support all cook modes
+	if (!IsUsingPackageStore())
 	{
-		// set up reusable buffers for non-package store case???
-		SavePackageContexts.SetNum(TargetPlatformNames.Num());
-
 		return;
 	}
 
@@ -6748,6 +6733,25 @@ void UCookOnTheFlyServer::InitializePackageStore(const TArray<FName>& TargetPlat
 	}
 }
 
+void UCookOnTheFlyServer::FinalizePackageStore()
+{
+	if (!IsUsingPackageStore())
+	{
+		return;
+	}
+
+	SCOPE_TIMER(FinalizePackageStore);
+
+	UE_LOG(LogCook, Display, TEXT("Saving name map(s)..."));
+	const TArray<ITargetPlatform*>& TargetPlatforms = GetCookingTargetPlatforms();
+	for (int32 PlatformIndex = 0; PlatformIndex < TargetPlatforms.Num(); ++PlatformIndex)
+	{
+		INameMapSaver& NameMapSaver = SavePackageContexts[PlatformIndex]->HeaderSaver.NameMapSaver;
+		NameMapSaver.End();
+	}
+	UE_LOG(LogCook, Display, TEXT("Done saving name map(s)"));
+}
+
 void UCookOnTheFlyServer::InitializeTargetPlatforms()
 {
 	const TArray<ITargetPlatform*>& TargetPlatforms = GetCookingTargetPlatforms();
