Revert unique naming for morph target change, the artist do not want this behavior

diff --git a/Engine/Source/Editor/UnrealEd/Private/Fbx/FbxSkeletalMeshImport.cpp b/Engine/Source/Editor/UnrealEd/Private/Fbx/FbxSkeletalMeshImport.cpp
--- a/Engine/Source/Editor/UnrealEd/Private/Fbx/FbxSkeletalMeshImport.cpp
+++ b/Engine/Source/Editor/UnrealEd/Private/Fbx/FbxSkeletalMeshImport.cpp
@@ -4204,9 +4204,6 @@ void UnFbx::FFbxImporter::ImportMorphTargetsInternal( TArray<FbxNode*>& SkelMesh
 
 	GWarn->BeginSlowTask( NSLOCTEXT("FbxImporter", "BeginGeneratingMorphModelsTask", "Generating Morph Models"), true);
 
-	//Use this TMap to store the name of the blendshape map by the fbx uniqueID. This allow to assign unique name to each blend shape.
-	TMap<uint64, FString> UniqueIDToName;
-
 	// For each morph in FBX geometries, we create one morph target for the Unreal skeletal mesh
 	for (int32 NodeIndex = 0; NodeIndex < SkelMeshNodeArray.Num(); NodeIndex++)
 	{
@@ -4266,36 +4263,6 @@ void UnFbx::FFbxImporter::ImportMorphTargetsInternal( TArray<FbxNode*>& SkelMesh
 								}
 							}
 
-							uint64 UniqueID = Shape->GetUniqueID();
-							if (UniqueIDToName.Contains(UniqueID))
-							{
-								ShapeName = UniqueIDToName.FindChecked(UniqueID);
-							}
-							else
-							{
-								//In case the shape do not have a unique name
-								//make it unique
-								int32 NameIndex = 0;
-								FString CurrentShapeName = ShapeName;
-								bool bNameNotUnique = ShapeNameToShapeArray.Contains(CurrentShapeName);
-								while(bNameNotUnique)
-								{
-									//Append ShapeX to the shape name (same has Maya)
-									if (NameIndex == 0)
-									{
-										CurrentShapeName = ShapeName + TEXT("Shape");
-										NameIndex++;
-									}
-									else
-									{
-										CurrentShapeName = ShapeName + TEXT("Shape") + FString::FromInt(NameIndex++);
-									}
-									bNameNotUnique = ShapeNameToShapeArray.Contains(CurrentShapeName);
-								}
-								ShapeName = CurrentShapeName;
-								UniqueIDToName.Add(UniqueID, ShapeName);
-							}
-
 							TArray<FbxShape*> & ShapeArray = ShapeNameToShapeArray.FindOrAdd(ShapeName);
 							if (ShapeArray.Num() == 0)
 							{
