Added an editor experimental setting BaseClassesToAllowRecompilingDuringPlayInEditor to allow iterating on arbitrary Blueprints during PIE (any blueprint deriving from one of these base classes will be allowed to recompile during Play-in-Editor)

diff --git a/Engine/Source/Editor/UnrealEd/Classes/Settings/EditorExperimentalSettings.h b/Engine/Source/Editor/UnrealEd/Classes/Settings/EditorExperimentalSettings.h
--- a/Engine/Source/Editor/UnrealEd/Classes/Settings/EditorExperimentalSettings.h
+++ b/Engine/Source/Editor/UnrealEd/Classes/Settings/EditorExperimentalSettings.h
@@ -61,6 +61,17 @@ public:
 	/** Enables "Find and Replace All" tool in the MyBlueprint window for variables */
 	UPROPERTY(EditAnywhere, config, Category = Blueprints, meta = (DisplayName = "Find and Replace All References Tool"))
 	bool bEnableFindAndReplaceReferences;
+	
+protected:
+	/** Any blueprint deriving from one of these base classes will be allowed to recompile during Play-in-Editor */
+	UPROPERTY(EditAnywhere, config, Category = Blueprints, meta=(AllowAbstract))
+	TArray<TSoftClassPtr<UObject>> BaseClassesToAllowRecompilingDuringPlayInEditor;
+
+	UPROPERTY(Transient)
+	mutable TArray<UClass*> ResolvedBaseClassesToAllowRecompilingDuringPlayInEditor;
+
+public:
+	bool IsClassAllowedToRecompileDuringPIE(UClass* TestClass) const;
 
 	/** Should arrows indicating data/execution flow be drawn halfway along wires? */
 	UPROPERTY(/*EditAnywhere - deprecated (moved into UBlueprintEditorSettings), */config/*, Category=Blueprints, meta=(DisplayName="Draw midpoint arrows in Blueprints")*/)
diff --git a/Engine/Source/Editor/UnrealEd/Private/Settings/SettingsClasses.cpp b/Engine/Source/Editor/UnrealEd/Private/Settings/SettingsClasses.cpp
--- a/Engine/Source/Editor/UnrealEd/Private/Settings/SettingsClasses.cpp
+++ b/Engine/Source/Editor/UnrealEd/Private/Settings/SettingsClasses.cpp
@@ -155,6 +155,35 @@ UEditorExperimentalSettings::UEditorExperimentalSettings( const FObjectInitializ
 {
 }
 
+bool UEditorExperimentalSettings::IsClassAllowedToRecompileDuringPIE(UClass* TestClass) const
+{
+	if (TestClass != nullptr)
+	{
+		// Rebuild the list if necessary (if the list was edited, either number of entires or value, ResolvedBaseClassesToAllowRecompilingDuringPlayInEditor will get reset below)
+		if (ResolvedBaseClassesToAllowRecompilingDuringPlayInEditor.Num() != BaseClassesToAllowRecompilingDuringPlayInEditor.Num())
+		{
+			ResolvedBaseClassesToAllowRecompilingDuringPlayInEditor.Reset();
+			for (TSoftClassPtr<UObject> BaseClassPtr : BaseClassesToAllowRecompilingDuringPlayInEditor)
+			{
+				if (UClass* BaseClass = BaseClassPtr.Get())
+				{
+					ResolvedBaseClassesToAllowRecompilingDuringPlayInEditor.Add(BaseClass);
+				}
+			}
+		}
+
+		// See if the test class matches any of the enabled base classes
+		for (UClass* BaseClass : ResolvedBaseClassesToAllowRecompilingDuringPlayInEditor)
+		{
+			if ((BaseClass != nullptr) && TestClass->IsChildOf(BaseClass))
+			{
+				return true;
+			}
+		}
+	}
+	return false;
+}
+
 void UEditorExperimentalSettings::PostInitProperties()
 {
 	CVarEditorHDRSupport->Set(bHDREditor ? 1 : 0, ECVF_SetByProjectSetting);
@@ -194,6 +223,8 @@ void UEditorExperimentalSettings::PostEditChangeProperty( struct FPropertyChange
 		SaveConfig();
 	}
 
+	ResolvedBaseClassesToAllowRecompilingDuringPlayInEditor.Reset();
+
 	SettingChangedEvent.Broadcast(Name);
 }
 
diff --git a/Engine/Source/Runtime/Engine/Classes/Engine/Blueprint.h b/Engine/Source/Runtime/Engine/Classes/Engine/Blueprint.h
--- a/Engine/Source/Runtime/Engine/Classes/Engine/Blueprint.h
+++ b/Engine/Source/Runtime/Engine/Classes/Engine/Blueprint.h
@@ -799,7 +799,7 @@ public:
 	virtual bool AlwaysCompileOnLoad() const { return false; }
 
 	/** Some Blueprints (and classes) can recompile while we are debugging a live session. This function controls whether this can occur. */
-	virtual bool CanRecompileWhilePlayingInEditor() const { return false; }
+	virtual bool CanRecompileWhilePlayingInEditor() const;
 
 	/**
 	 * Check whether this blueprint can be nativized or not
diff --git a/Engine/Source/Runtime/Engine/Private/Blueprint.cpp b/Engine/Source/Runtime/Engine/Private/Blueprint.cpp
--- a/Engine/Source/Runtime/Engine/Private/Blueprint.cpp
+++ b/Engine/Source/Runtime/Engine/Private/Blueprint.cpp
@@ -21,6 +21,7 @@
 #if WITH_EDITOR
 #include "BlueprintCompilationManager.h"
 #include "Editor/UnrealEd/Classes/Settings/ProjectPackagingSettings.h"
+#include "Editor/UnrealEd/Classes/Settings/EditorExperimentalSettings.h"
 #include "Engine/SimpleConstructionScript.h"
 #include "Engine/SCS_Node.h"
 #include "Kismet2/BlueprintEditorUtils.h"
@@ -889,6 +890,11 @@ void UBlueprint::GetReparentingRules(TSet< const UClass* >& AllowedChildrenOfCla
 
 }
 
+bool UBlueprint::CanRecompileWhilePlayingInEditor() const
+{
+	return GetDefault<UEditorExperimentalSettings>()->IsClassAllowedToRecompileDuringPIE(ParentClass);
+}
+
 void UBlueprint::GetAssetRegistryTags(TArray<FAssetRegistryTag>& OutTags) const
 {
 	// We use Generated instead of Skeleton because the CDO data is more accurate on Generated
