Run the spline migration tool in the nightly build process

Added logic to avoid dirtying assets/levels when no change are performed - tool is now much faster to execute
Includes patch to ensure materials overrides by splines are up to date with their components parts.

diff --git a/Engine/Source/Editor/UnrealEd/Private/Commandlets/ContentCommandlets.cpp b/Engine/Source/Editor/UnrealEd/Private/Commandlets/ContentCommandlets.cpp
--- a/Engine/Source/Editor/UnrealEd/Private/Commandlets/ContentCommandlets.cpp
+++ b/Engine/Source/Editor/UnrealEd/Private/Commandlets/ContentCommandlets.cpp
@@ -1383,7 +1383,7 @@ void UResavePackagesCommandlet::CheckoutAndSavePackage(UPackage* Package, TArray
 	check(Package);
 
 	FString PackageFilename;
-	if (FPackageName::TryConvertLongPackageNameToFilename(Package->GetName(), PackageFilename, FPackageName::GetAssetPackageExtension()))
+	if (FPackageName::TryConvertLongPackageNameToFilename(Package->GetName(), PackageFilename, Package->ContainsMap() ? FPackageName::GetMapPackageExtension() : FPackageName::GetAssetPackageExtension()))
 	{
 		if (IFileManager::Get().FileExists(*PackageFilename))
 		{
