Android appends the texture compression used to the platform name which causes the cooker to not cook plugin content due to direct matching supported platform and target platform names. As a temporary fix, this change matches against platform names beginning with 'Android'.

diff --git a/Engine/Source/Editor/UnrealEd/Private/CookOnTheFlyServer.cpp b/Engine/Source/Editor/UnrealEd/Private/CookOnTheFlyServer.cpp
--- a/Engine/Source/Editor/UnrealEd/Private/CookOnTheFlyServer.cpp
+++ b/Engine/Source/Editor/UnrealEd/Private/CookOnTheFlyServer.cpp
@@ -6724,6 +6724,7 @@ void UCookOnTheFlyServer::DiscoverPlatformSpecificNeverCookPackages(
 	TArray<FAssetData> PluginAssets;
 	FARFilter PluginARFilter;
 	FString PluginPackagePath;
+	static const FString AndroidPlatformName(TEXT("Android"));
 
 	TArray<TSharedRef<IPlugin>> AllContentPlugins = IPluginManager::Get().GetEnabledPluginsWithContent();
 	for (TSharedRef<IPlugin> Plugin : AllContentPlugins)
@@ -6736,11 +6737,27 @@ void UCookOnTheFlyServer::DiscoverPlatformSpecificNeverCookPackages(
 			continue;
 		}
 
+		// Android platforms append a texture compression string to the end of the platform name so we can't do a direct comparison
+		bool bSupportsAndroid = false;
+		for (const FString& PlatformName : Descriptor.SupportedTargetPlatforms)
+		{
+			if (PlatformName.Contains(AndroidPlatformName))
+			{
+				bSupportsAndroid = true;
+				break;
+			}
+		}
+
 		// find any unsupported target platforms for this plugin
 		PluginUnsupportedTargetPlatforms.Reset();
 		for (int32 I = 0, Count = TargetPlatformNames.Num(); I < Count; ++I)
 		{
-			if (!Descriptor.SupportedTargetPlatforms.Contains(TargetPlatformStrings[I]))
+			// Partial match for the Android platform family
+			if (bSupportsAndroid && TargetPlatformStrings[I].Contains(AndroidPlatformName))
+			{
+				continue;
+			}
+			else if (!Descriptor.SupportedTargetPlatforms.Contains(TargetPlatformStrings[I]))
 			{
 				PluginUnsupportedTargetPlatforms.Add(TargetPlatformNames[I]);
 			}
