Merging //UE4/Dev-Main to Dev-VR (//UE4/Dev-VR)

diff --git a/Engine/Plugins/Runtime/ResonanceAudio/Source/ResonanceAudio/Private/ResonanceAudioReverb.h b/Engine/Plugins/Runtime/ResonanceAudio/Source/ResonanceAudio/Private/ResonanceAudioReverb.h
--- a/Engine/Plugins/Runtime/ResonanceAudio/Source/ResonanceAudio/Private/ResonanceAudioReverb.h
+++ b/Engine/Plugins/Runtime/ResonanceAudio/Source/ResonanceAudio/Private/ResonanceAudioReverb.h
@@ -113,6 +113,7 @@ namespace ResonanceAudio
 		virtual void OnReleaseSource(const uint32 SourceId) override;
 		virtual class FSoundEffectSubmix* GetEffectSubmix(class USoundSubmix* Submix) override;
 		virtual void ProcessSourceAudio(const FAudioPluginSourceInputData& InputData, FAudioPluginSourceOutputData& OutputData) override;
+		virtual bool DoesReverbOverrideMasterReverb() const override { return false; }
 
 		void SetResonanceAudioApi(vraudio::VrAudioApi* InResonanceAudioApi) { ResonanceAudioApi = InResonanceAudioApi; };
 		void SetPreset(UResonanceAudioReverbPluginPreset* InPreset);
diff --git a/Engine/Source/Developer/MeshUtilities/Private/MeshUtilities.cpp b/Engine/Source/Developer/MeshUtilities/Private/MeshUtilities.cpp
--- a/Engine/Source/Developer/MeshUtilities/Private/MeshUtilities.cpp
+++ b/Engine/Source/Developer/MeshUtilities/Private/MeshUtilities.cpp
@@ -2785,7 +2785,7 @@ public:
 			LODModel.IndexBuffer.SetIndices(CombinedIndices, bNeeds32BitIndices ? EIndexBufferStride::Force32Bit : EIndexBufferStride::Force16Bit);
 			
 			// Build the reversed index buffer.
-			if (InOutModels[0].BuildSettings.bBuildReversedIndexBuffer)
+			if (LODModel.AdditionalIndexBuffers && InOutModels[0].BuildSettings.bBuildReversedIndexBuffer)
 			{
 				TArray<uint32> InversedIndices;
 				const int32 IndexCount = CombinedIndices.Num();
@@ -2823,7 +2823,7 @@ public:
 			}
 
 			// Build the inversed depth only index buffer.
-			if (InOutModels[0].BuildSettings.bBuildReversedIndexBuffer)
+			if (LODModel.AdditionalIndexBuffers && InOutModels[0].BuildSettings.bBuildReversedIndexBuffer)
 			{
 				TArray<uint32> ReversedDepthOnlyIndices;
 				const int32 IndexCount = DepthOnlyIndices.Num();
@@ -2836,6 +2836,7 @@ public:
 			}
 
 			// Build a list of wireframe edges in the static mesh.
+			if (LODModel.AdditionalIndexBuffers)
 			{
 				TArray<FMeshEdgeDef> Edges;
 				TArray<uint32> WireframeIndices;
@@ -2852,7 +2853,7 @@ public:
 			}
 
 			// Build the adjacency index buffer used for tessellation.
-			if (InOutModels[0].BuildSettings.bBuildAdjacencyBuffer)
+			if (LODModel.AdditionalIndexBuffers && InOutModels[0].BuildSettings.bBuildAdjacencyBuffer)
 			{
 				TArray<uint32> AdjacencyIndices;
 
diff --git a/Engine/Source/Editor/Kismet/Private/BlueprintCompilationManager.cpp b/Engine/Source/Editor/Kismet/Private/BlueprintCompilationManager.cpp
--- a/Engine/Source/Editor/Kismet/Private/BlueprintCompilationManager.cpp
+++ b/Engine/Source/Editor/Kismet/Private/BlueprintCompilationManager.cpp
@@ -1959,7 +1959,8 @@ void FBlueprintCompilationManagerImpl::ReinstanceBatch(TArray<FReinstancingJob>&
 
 	for(UObject* ArchetypeReferencer : ArchetypeReferencers)
 	{
-		FArchiveReplaceObjectRef<UObject> ReplaceInCDOAr(ArchetypeReferencer, OldArchetypeToNewArchetype, false, false, false, false, true);
+		UPackage* NewPackage = ArchetypeReferencer->GetOutermost();
+		FArchiveReplaceOrClearExternalReferences<UObject> ReplaceInCDOAr(ArchetypeReferencer, OldArchetypeToNewArchetype, NewPackage);
 	}
 }
 
diff --git a/Engine/Source/Editor/UnrealEd/Private/Kismet2/KismetReinstanceUtilities.cpp b/Engine/Source/Editor/UnrealEd/Private/Kismet2/KismetReinstanceUtilities.cpp
--- a/Engine/Source/Editor/UnrealEd/Private/Kismet2/KismetReinstanceUtilities.cpp
+++ b/Engine/Source/Editor/UnrealEd/Private/Kismet2/KismetReinstanceUtilities.cpp
@@ -31,7 +31,6 @@
 #include "BlueprintEditor.h"
 #include "Engine/Selection.h"
 #include "BlueprintEditorSettings.h"
-#include "Algo/RemoveIf.h"
 
 extern COREUOBJECT_API bool GBlueprintUseCompilationManager;
 
@@ -89,24 +88,6 @@ struct FReplaceReferenceHelper
 			BP_SCOPED_COMPILER_EVENT_STAT(EKismetReinstancerStats_FindReferencers);
 
 			Targets = FArchiveHasReferences::GetAllReferencers(SourceObjects, ObjectsThatShouldUseOldStuff);
-
-			// Make sure we don't update properties in old objects, as they
-			// may take ownership of objects referenced in new objects (e.g.
-			// delete components owned by new actors)
-			int32 RangeEnd = Algo::RemoveIf(Targets, [&ObjectsToReplace](UObject* Obj)
-			{
-				for (UObject* ObjectToReplace : ObjectsToReplace)
-				{
-					if (ObjectToReplace == Obj || Obj->IsIn(ObjectToReplace))
-					{
-						return true;
-					}
-				}
-
-				return false;
-			});
-
-			Targets.RemoveAtSwap(RangeEnd, Targets.Num() - RangeEnd, false);
 		}
 
 		{
@@ -114,45 +95,51 @@ struct FReplaceReferenceHelper
 
 			for (UObject* Obj : Targets)
 			{
-				// The class for finding and replacing weak references.
-				// We can't relay on "standard" weak references replacement as
-				// it depends on FSoftObjectPath::ResolveObject, which
-				// tries to find the object with the stored path. It is
-				// impossible, cause above we deleted old actors (after
-				// spawning new ones), so during objects traverse we have to
-				// find FSoftObjectPath with the raw given path taken
-				// before deletion of old actors and fix them.
-				class ReferenceReplace : public FArchiveReplaceObjectRef < UObject >
+				// Make sure we don't update properties in old objects, as they
+				// may take ownership of objects referenced in new objects (e.g.
+				// delete components owned by new actors)
+				if (!ObjectsToReplace.Contains(Obj))
 				{
-				public:
-					ReferenceReplace(UObject* InSearchObject, const TMap<UObject*, UObject*>& InReplacementMap, const TMap<FSoftObjectPath, UObject*>& InWeakReferencesMap)
-						: FArchiveReplaceObjectRef<UObject>(InSearchObject, InReplacementMap, false, false, false, true), WeakReferencesMap(InWeakReferencesMap)
+					// The class for finding and replacing weak references.
+					// We can't relay on "standard" weak references replacement as
+					// it depends on FSoftObjectPath::ResolveObject, which
+					// tries to find the object with the stored path. It is
+					// impossible, cause above we deleted old actors (after
+					// spawning new ones), so during objects traverse we have to
+					// find FSoftObjectPath with the raw given path taken
+					// before deletion of old actors and fix them.
+					class ReferenceReplace : public FArchiveReplaceObjectRef < UObject >
 					{
-						SerializeSearchObject();
-					}
-
-					FArchive& operator<<(FSoftObjectPath& Ref) override
-					{
-						const UObject*const* PtrToObjPtr = WeakReferencesMap.Find(Ref);
-
-						if (PtrToObjPtr != nullptr)
+					public:
+						ReferenceReplace(UObject* InSearchObject, const TMap<UObject*, UObject*>& InReplacementMap, const TMap<FSoftObjectPath, UObject*>& InWeakReferencesMap)
+							: FArchiveReplaceObjectRef<UObject>(InSearchObject, InReplacementMap, false, false, false, true), WeakReferencesMap(InWeakReferencesMap)
 						{
-							Ref = *PtrToObjPtr;
+							SerializeSearchObject();
 						}
 
-						return *this;
-					}
+						FArchive& operator<<(FSoftObjectPath& Ref) override
+						{
+							const UObject*const* PtrToObjPtr = WeakReferencesMap.Find(Ref);
 
-					FArchive& operator<<(FSoftObjectPtr& Ref) override
-					{
-						return operator<<(Ref.GetUniqueID());
-					}
+							if (PtrToObjPtr != nullptr)
+							{
+								Ref = *PtrToObjPtr;
+							}
+
+							return *this;
+						}
 
-				private:
-					const TMap<FSoftObjectPath, UObject*>& WeakReferencesMap;
-				};
+						FArchive& operator<<(FSoftObjectPtr& Ref) override
+						{
+							return operator<<(Ref.GetUniqueID());
+						}
 
-				ReferenceReplace ReplaceAr(Obj, OldToNewInstanceMap, ReinstancedObjectsWeakReferenceMap);
+					private:
+						const TMap<FSoftObjectPath, UObject*>& WeakReferencesMap;
+					};
+
+					ReferenceReplace ReplaceAr(Obj, OldToNewInstanceMap, ReinstancedObjectsWeakReferenceMap);
+				}
 			}
 		}
 	}
@@ -1704,7 +1691,9 @@ static void ReplaceObjectHelper(UObject*& OldObject, UClass* OldClass, UObject*&
 		}
 	}
 
-	// NOTE: Don't mark the old object as pending kill yet! Otherwise, any weak references to this object will be invalidated too soon and won't get replaced.
+	OldObject->RemoveFromRoot();
+	OldObject->MarkPendingKill();
+
 	OldToNewInstanceMap.Add(OldObject, NewUObject);
 
 	if (bIsComponent)
@@ -1861,7 +1850,7 @@ static void ReplaceActorHelper(UObject* OldObject, UClass* OldClass, UObject*& N
 		GEditor->Layers->DisassociateActorFromLayers(OldActor);
 	}
 
-	// NOTE: Don't destroy the Actor yet! Otherwise, any weak references to this Actor will be invalidated too soon and won't get replaced.
+	World->EditorDestroyActor(OldActor, /*bShouldModifyLevel =*/true);
 	OldToNewInstanceMap.Add(OldActor, NewActor);
 }
 
@@ -1916,6 +1905,18 @@ void FBlueprintCompileReinstancer::ReplaceInstancesOfClass_Inner(TMap<UClass*, U
 		OnObjectsReplacedHandle = GEditor->OnObjectsReplaced().AddRaw(&ObjectRemappingHelper, &FObjectRemappingHelper::OnObjectsReplaced);
 	}
 
+	auto UpdateObjectBeingDebugged = [](UObject* InOldObject, UObject* InNewObject)
+	{
+		if (UBlueprint* OldObjBlueprint = Cast<UBlueprint>(InOldObject->GetClass()->ClassGeneratedBy))
+		{
+			const UObject* DebugObj = OldObjBlueprint->GetObjectBeingDebugged(EGetObjectOrWorldBeingDebuggedFlags::IgnorePendingKill);
+			if (DebugObj == InOldObject)
+			{
+				OldObjBlueprint->SetObjectBeingDebugged(InNewObject);
+			}
+		}
+	};
+
 	{
 		TArray<UObject*> ObjectsToReplace;
 
@@ -1966,7 +1967,7 @@ void FBlueprintCompileReinstancer::ReplaceInstancesOfClass_Inner(TMap<UClass*, U
 					{
 						UObject* NewUObject = nullptr;
 						ReplaceObjectHelper(OldObject, OldClass, NewUObject, NewClass, OldToNewInstanceMap, OldToNewNameMap, OldObjIndex, ObjectsToReplace, PotentialEditorsForRefreshing, OwnersToRerunConstructionScript, &FDirectAttachChildrenAccessor::Get, bIsComponent, bArchetypesAreUpToDate);
-						
+						UpdateObjectBeingDebugged(OldObject, NewUObject);
 						ObjectsReplaced.Add(OldObject);
 
 						if (bLogConversions)
@@ -2038,7 +2039,7 @@ void FBlueprintCompileReinstancer::ReplaceInstancesOfClass_Inner(TMap<UClass*, U
 							// Actors that are not in a level cannot be reconstructed, sequencer team decided to reinstance these as normal objects:
 							ReplaceObjectHelper(OldObject, OldClass, NewUObject, NewClass, OldToNewInstanceMap, OldToNewNameMap, OldObjIndex, ObjectsToReplace, PotentialEditorsForRefreshing, OwnersToRerunConstructionScript, &FDirectAttachChildrenAccessor::Get, false, bArchetypesAreUpToDate);
 						}
-
+						UpdateObjectBeingDebugged(OldObject, NewUObject);
 						ObjectsReplaced.Add(OldObject);
 
 						if (bLogConversions)
diff --git a/Engine/Source/Editor/UnrealEd/Private/UnrealEdMisc.cpp b/Engine/Source/Editor/UnrealEd/Private/UnrealEdMisc.cpp
--- a/Engine/Source/Editor/UnrealEd/Private/UnrealEdMisc.cpp
+++ b/Engine/Source/Editor/UnrealEd/Private/UnrealEdMisc.cpp
@@ -614,6 +614,11 @@ void FUnrealEdMisc::InitEngineAnalytics()
 * @EventParam IsDebugger (bool) Whether the debugger is currently present
 * @EventParam WasDebuggerPresent (bool) Whether the debugger was present previously
 * @EventParam IsInVRMode (bool) If the current heartbeat occurred while VR mode was active
+* @EventParam bIsEnterprise (bool) If the editor has an enterprise project loaded
+* @EventParam Is5MinIdle (bool) Whether the user is idle, using the 5 minute methodology
+* @EventParam Is30MinIdle (bool) Whether the user is idle, using the 30 minute methodology
+* @EventParam IsInPIE (bool) Whether the user is/was in PIE during this heartbeat
+* @EventParam RealIntervalSec (bool) The real time since the last heartbeat
 *
 * @Source Editor
 *
@@ -634,6 +639,13 @@ void FUnrealEdMisc::EditorAnalyticsHeartbeat()
 	//////////////////////////////////////////////////////////////////////////////////////////////////////////////////////////////////////////////
 	const double HeartbeatInvervalSec = (double)(UnrealEdMiscDefs::HeartbeatIntervalSeconds);
 	const double Now = FPlatformTime::Seconds();
+
+	static double LastTimeInPIE = 0;
+	if (FPlayWorldCommandCallbacks::IsInPIE())
+	{
+		LastTimeInPIE = Now;
+	}
+
 	// Initialize to ensure the first time this is called we will DEF execute the heartbeat.
 	static double LastHeartbeatTime = Now - HeartbeatInvervalSec;
 	// allow a little bit of slop in the timing in case the event is only firing slightly too soon, even though that is technically impossible.
@@ -654,10 +666,10 @@ void FUnrealEdMisc::EditorAnalyticsHeartbeat()
 	const bool bInVRMode = IVREditorModule::Get().IsVREditorModeActive();
 	const bool bIsEnterprise = IProjectManager::Get().IsEnterpriseProject();
 
-	double LastInteractionTime = FSlateApplication::Get().GetLastUserInteractionTime();
+	const double LastInteractionTime = FSlateApplication::Get().GetLastUserInteractionTime();
 	
 	// Did the user interact since the last heartbeat
-	bool bIdle = LastInteractionTime < LastHeartbeatTime;
+	const bool bIdle = LastInteractionTime < LastHeartbeatTime;
 	
 	extern ENGINE_API float GAverageFPS;
 
@@ -678,9 +690,22 @@ void FUnrealEdMisc::EditorAnalyticsHeartbeat()
 	Attributes.Add(FAnalyticsEventAttribute(TEXT("IsInVRMode"), bInVRMode));
 	Attributes.Add(FAnalyticsEventAttribute(TEXT("Enterprise"), bIsEnterprise));
 
+	// In the 5 Min and 30 Min idle case we treat longer periods of time as none idle post user interaction
+	const bool b5MinIdle = LastInteractionTime + (5.0f * 60.0f) < LastHeartbeatTime;
+	Attributes.Add(FAnalyticsEventAttribute(TEXT("Is5MinIdle"), b5MinIdle));
+	const bool b30MinIdle = LastInteractionTime + (30.0f * 60.0f) < LastHeartbeatTime;
+	Attributes.Add(FAnalyticsEventAttribute(TEXT("Is30MinIdle"), b30MinIdle));
+
+	// InPIE bool, added a 5 second grace period to include PIE shutdown issues in the "InPIE" window
+	const bool bInPIE = LastTimeInPIE + 5.0f > LastHeartbeatTime;
+	Attributes.Add(FAnalyticsEventAttribute(TEXT("IsInPIE"), bInPIE));
+
+	// Interval seconds actually track the entirety of the time since last heartbeat.
+	Attributes.Add(FAnalyticsEventAttribute(TEXT("RealIntervalSec"), Now - LastHeartbeatTime));
+
 	FEngineAnalytics::GetProvider().RecordEvent(TEXT("Editor.Usage.Heartbeat"), Attributes);
 	
-	LastHeartbeatTime = FPlatformTime::Seconds();
+	LastHeartbeatTime = Now;
 }
 
 void FUnrealEdMisc::TickAssetAnalytics()
diff --git a/Engine/Source/Programs/UnrealBuildTool/Platform/IOS/IOSToolChain.cs b/Engine/Source/Programs/UnrealBuildTool/Platform/IOS/IOSToolChain.cs
--- a/Engine/Source/Programs/UnrealBuildTool/Platform/IOS/IOSToolChain.cs
+++ b/Engine/Source/Programs/UnrealBuildTool/Platform/IOS/IOSToolChain.cs
@@ -1848,11 +1848,11 @@ namespace UnrealBuildTool
 
 				if (AppName == "UE4Game" || AppName == "UE4Client" || Target.ProjectFile == null || Target.ProjectFile.IsUnderDirectory(UnrealBuildTool.EngineDirectory))
 				{
-					GenerateProjectFiles(Target.ProjectFile, new string[] { "-platforms=" + (Target.Platform == UnrealTargetPlatform.IOS ? "IOS" : "TVOS"), "-NoIntellIsense", (Target.Platform == UnrealTargetPlatform.IOS ? "-iosdeployonly" : "-tvosdeployonly"), "-ignorejunk", (Target.bForDistribution ? "-distribution" : "-development"), "-bundleID=" + BundleID });
+					GenerateProjectFiles(Target.ProjectFile, new string[] { "-platforms=" + (Target.Platform == UnrealTargetPlatform.IOS ? "IOS" : "TVOS"), "-NoIntellIsense", (Target.Platform == UnrealTargetPlatform.IOS ? "-iosdeployonly" : "-tvosdeployonly"), "-ignorejunk", (Target.bForDistribution ? "-distribution" : "-development"), "-bundleID=" + BundleID, "-includetemptargets" });
 				}
 				else
 				{
-					GenerateProjectFiles(Target.ProjectFile, new string[] { "-platforms=" + (Target.Platform == UnrealTargetPlatform.IOS ? "IOS" : "TVOS"), "-NoIntellIsense", (Target.Platform == UnrealTargetPlatform.IOS ? "-iosdeployonly" : "-tvosdeployonly"), "-ignorejunk", (Target.bForDistribution ? "-distribution" : "-development"), String.Format("-project={0}", Target.ProjectFile), "-game", "-bundleID=" + BundleID });
+					GenerateProjectFiles(Target.ProjectFile, new string[] { "-platforms=" + (Target.Platform == UnrealTargetPlatform.IOS ? "IOS" : "TVOS"), "-NoIntellIsense", (Target.Platform == UnrealTargetPlatform.IOS ? "-iosdeployonly" : "-tvosdeployonly"), "-ignorejunk", (Target.bForDistribution ? "-distribution" : "-development"), String.Format("-project={0}", Target.ProjectFile), "-game", "-bundleID=" + BundleID, "-includetemptargets" });
 				}
 				// Make sure it exists
 				if (!DirectoryReference.Exists(XcodeWorkspaceDir))
diff --git a/Engine/Source/Programs/UnrealHeaderTool/Private/HeaderParser.cpp b/Engine/Source/Programs/UnrealHeaderTool/Private/HeaderParser.cpp
--- a/Engine/Source/Programs/UnrealHeaderTool/Private/HeaderParser.cpp
+++ b/Engine/Source/Programs/UnrealHeaderTool/Private/HeaderParser.cpp
@@ -6700,10 +6700,6 @@ void FHeaderParser::CompileFunctionDeclaration(FClasses& AllClasses)
 	ProcessFunctionSpecifiers(FuncInfo, SpecifiersFound, MetaData);
 
 	const bool bClassGeneratedFromBP = FClass::IsDynamic(GetCurrentClass());
-	if ((FuncInfo.FunctionFlags & FUNC_NetServer) && !(FuncInfo.FunctionFlags & FUNC_NetValidate) && !bClassGeneratedFromBP)
-	{
-		FError::Throwf(TEXT("Server RPC missing 'WithValidation' keyword in the UFUNCTION() declaration statement.  Required for security purposes."));
-	}
 
 	if ((0 != (FuncInfo.FunctionExportFlags & FUNCEXPORT_CustomThunk)) && !MetaData.Contains("CustomThunk"))
 	{
diff --git a/Engine/Source/Runtime/ApplicationCore/Private/Windows/Accessibility/WindowsUIAManager.cpp b/Engine/Source/Runtime/ApplicationCore/Private/Windows/Accessibility/WindowsUIAManager.cpp
--- a/Engine/Source/Runtime/ApplicationCore/Private/Windows/Accessibility/WindowsUIAManager.cpp
+++ b/Engine/Source/Runtime/ApplicationCore/Private/Windows/Accessibility/WindowsUIAManager.cpp
@@ -167,6 +167,7 @@ void FWindowsUIAManager::OnEventRaised(TSharedRef<IAccessibleWidget> Widget, EAc
 				UiaRaiseAutomationEvent(&ScopedProvider.Provider, UIA_Invoke_InvokedEventId);
 			}
 			break;
+#if WINVER >= 0x0A00
 		case EAccessibleEvent::Notification:
 		{
 			typedef HRESULT(WINAPI* UiaRaiseNotificationEventFunc)(IRawElementProviderSimple*, NotificationKind, NotificationProcessing, BSTR, BSTR);
@@ -188,6 +189,7 @@ void FWindowsUIAManager::OnEventRaised(TSharedRef<IAccessibleWidget> Widget, EAc
 			}
 			break;
 		}
+#endif
 		// IMPORTANT: Calling UiaRaiseStructureChangedEvent seems to raise our per-frame timing for accessibility by
 		// over 10x (.04ms to .7ms, tested by clicking on the "All Classes" button in the modes panel of the Editor).
 		// For now, I'm disabling this until we figure out if it's absolutely necessary.
diff --git a/Engine/Source/Runtime/AudioMixer/Private/AudioMixerDevice.cpp b/Engine/Source/Runtime/AudioMixer/Private/AudioMixerDevice.cpp
--- a/Engine/Source/Runtime/AudioMixer/Private/AudioMixerDevice.cpp
+++ b/Engine/Source/Runtime/AudioMixer/Private/AudioMixerDevice.cpp
@@ -642,6 +642,8 @@ namespace Audio
 			FSoundEffectSubmixInitData InitData;
 			InitData.SampleRate = GetSampleRate();
 	
+			bool bIsMasterReverbBypassed = false;
+
 			// Setup the master reverb plugin
 			if (ReverbPluginInterface.IsValid() && MasterSubmixInstances[EMasterSubmixType::ReverbPlugin].IsValid())
 			{
@@ -656,8 +658,11 @@ namespace Audio
 				MasterReverbPluginSubmix->AddSoundEffectSubmix(ReverbPluginId, ReverbPluginEffectSubmix);
 				MasterReverbPluginSubmix->SetParentSubmix(MasterSubmixInstance);
 				MasterSubmixInstance->AddChildSubmix(MasterReverbPluginSubmix);
+
+				bIsMasterReverbBypassed = ReverbPluginInterface->DoesReverbOverrideMasterReverb();
 			}
-			else if (MasterSubmixInstances[EMasterSubmixType::Reverb].IsValid())
+					
+			if (MasterSubmixInstances[EMasterSubmixType::Reverb].IsValid() && !bIsMasterReverbBypassed)
 			{
 				// Setup the master reverb only if we don't have a reverb plugin
 				USoundEffectSubmixPreset* ReverbPreset = nullptr;
diff --git a/Engine/Source/Runtime/AudioMixer/Private/AudioMixerSource.cpp b/Engine/Source/Runtime/AudioMixer/Private/AudioMixerSource.cpp
--- a/Engine/Source/Runtime/AudioMixer/Private/AudioMixerSource.cpp
+++ b/Engine/Source/Runtime/AudioMixer/Private/AudioMixerSource.cpp
@@ -12,14 +12,6 @@
 // Link to "Audio" profiling category
 CSV_DECLARE_CATEGORY_MODULE_EXTERN(AUDIOMIXER_API, Audio);
 
-static int32 DisableHRTFCvar = 0;
-FAutoConsoleVariableRef CVarDisableHRTF(
-	TEXT("au.DisableHRTF"),
-	DisableHRTFCvar,
-	TEXT("Disables HRTF\n")
-	TEXT("0: Not Disabled, 1: Disabled"),
-	ECVF_Default);
-
 static int32 UseListenerOverrideForSpreadCVar = 0;
 FAutoConsoleVariableRef CVarUseListenerOverrideForSpread(
 	TEXT("au.UseListenerOverrideForSpread"),
@@ -123,6 +115,7 @@ namespace Audio
 			InitParams.NumInputFrames = NumFrames;
 			InitParams.SourceVoice = MixerSourceVoice;
 			InitParams.bUseHRTFSpatialization = UseObjectBasedSpatialization();
+			InitParams.bIsExternalSend = MixerDevice->bSpatializationIsExternalSend;
 			InitParams.bIsAmbisonics = WaveInstance->bIsAmbisonics;
 
 			if (InitParams.bIsAmbisonics)
@@ -943,13 +936,11 @@ namespace Audio
 					MixerSourceVoice->SetSubmixSendInfo(MixerDevice->GetMasterReverbPluginSubmix(), ReverbSendLevel);
 				}
 			}
-			else
+
+			// Send the source audio to the master reverb
+			if (!AudioDevice->IsReverbPluginBypassingMasterReverb() && MixerDevice->GetMasterReverbSubmix().IsValid())
 			{
-				// Send the source audio to the master reverb
-				if (MixerDevice->GetMasterReverbSubmix().IsValid())
-				{
-					MixerSourceVoice->SetSubmixSendInfo(MixerDevice->GetMasterReverbSubmix(), ReverbSendLevel);
-				}
+				MixerSourceVoice->SetSubmixSendInfo(MixerDevice->GetMasterReverbSubmix(), ReverbSendLevel);
 			}
 		}
 
@@ -1183,9 +1174,8 @@ namespace Audio
 	bool FMixerSource::UseObjectBasedSpatialization() const
 	{
 		return (Buffer->NumChannels <= MixerDevice->MaxChannelsSupportedBySpatializationPlugin &&
-			AudioDevice->IsSpatializationPluginEnabled() &&
-			DisableHRTFCvar == 0 &&
-			WaveInstance->SpatializationMethod == ESoundSpatializationAlgorithm::SPATIALIZATION_HRTF);
+				AudioDevice->IsSpatializationPluginEnabled() &&
+				WaveInstance->SpatializationMethod == ESoundSpatializationAlgorithm::SPATIALIZATION_HRTF);
 	}
 
 	bool FMixerSource::UseSpatializationPlugin() const
diff --git a/Engine/Source/Runtime/AudioMixer/Private/AudioMixerSourceManager.cpp b/Engine/Source/Runtime/AudioMixer/Private/AudioMixerSourceManager.cpp
--- a/Engine/Source/Runtime/AudioMixer/Private/AudioMixerSourceManager.cpp
+++ b/Engine/Source/Runtime/AudioMixer/Private/AudioMixerSourceManager.cpp
@@ -496,6 +496,7 @@ namespace Audio
 		SourceInfo.bIsStopping = false;
 		SourceInfo.bIsBusy = false;
 		SourceInfo.bUseHRTFSpatializer = false;
+		SourceInfo.bIsExternalSend = false;
 		SourceInfo.bUseOcclusionPlugin = false;
 		SourceInfo.bUseReverbPlugin = false;
 		SourceInfo.bUseModulationPlugin = false;
@@ -631,6 +632,7 @@ namespace Audio
 			SourceInfo.bIsDone = false;
 			SourceInfo.bIsLastBuffer = false;
 			SourceInfo.bUseHRTFSpatializer = InitParams.bUseHRTFSpatialization;
+			SourceInfo.bIsExternalSend = InitParams.bIsExternalSend;
 			SourceInfo.bIsVorbis = InitParams.bIsVorbis;
 			SourceInfo.AudioComponentID = InitParams.AudioComponentID;
 
@@ -775,7 +777,7 @@ namespace Audio
 
 			// Initialize a new downmix data:
 			check(SourceId < SourceInfos.Num());
-			const int32 SourceInputChannels = SourceInfo.bUseHRTFSpatializer ? 2 : SourceInfo.NumInputChannels;
+			const int32 SourceInputChannels = (SourceInfo.bUseHRTFSpatializer && !SourceInfo.bIsExternalSend) ? 2 : SourceInfo.NumInputChannels;
 			FSourceDownmixData& DownmixData = InitializeDownmixForSource(SourceId, SourceInputChannels, MixerDevice->GetDeviceOutputChannels(), NumOutputFrames);
 
 			for (int32 i = 0; i < InitParams.SubmixSends.Num(); ++i)
diff --git a/Engine/Source/Runtime/AudioMixer/Private/AudioMixerSourceManager.h b/Engine/Source/Runtime/AudioMixer/Private/AudioMixerSourceManager.h
--- a/Engine/Source/Runtime/AudioMixer/Private/AudioMixerSourceManager.h
+++ b/Engine/Source/Runtime/AudioMixer/Private/AudioMixerSourceManager.h
@@ -109,6 +109,7 @@ namespace Audio
 		uint64 AudioComponentID;
 		uint8 bPlayEffectChainTails : 1;
 		uint8 bUseHRTFSpatialization : 1;
+		uint8 bIsExternalSend : 1;
 		uint8 bIsDebugMode : 1;
 		uint8 bOutputToBusOnly : 1;
 		uint8 bIsVorbis : 1;
@@ -133,6 +134,7 @@ namespace Audio
 			, AudioComponentID(0)
 			, bPlayEffectChainTails(false)
 			, bUseHRTFSpatialization(false)
+			, bIsExternalSend(false)
 			, bIsDebugMode(false)
 			, bOutputToBusOnly(false)
 			, bIsVorbis(false)
@@ -491,6 +493,7 @@ namespace Audio
 			uint8 bHasStarted:1;
 			uint8 bIsBusy:1;
 			uint8 bUseHRTFSpatializer:1;
+			uint8 bIsExternalSend:1;
 			uint8 bUseOcclusionPlugin:1;
 			uint8 bUseReverbPlugin:1;
 			uint8 bUseModulationPlugin:1;
diff --git a/Engine/Source/Runtime/CoreUObject/Private/Serialization/ArchiveHasReferences.cpp b/Engine/Source/Runtime/CoreUObject/Private/Serialization/ArchiveHasReferences.cpp
--- a/Engine/Source/Runtime/CoreUObject/Private/Serialization/ArchiveHasReferences.cpp
+++ b/Engine/Source/Runtime/CoreUObject/Private/Serialization/ArchiveHasReferences.cpp
@@ -12,7 +12,6 @@ FArchiveHasReferences::FArchiveHasReferences(UObject* InTarget, const TSet<UObje
 	check(InTarget);
 
 	ArIsObjectReferenceCollector = true;
-	ArIsModifyingWeakAndStrongReferences = true;
 	InTarget->Serialize(*this);
 
 	class FArchiveProxyCollector : public FReferenceCollector
diff --git a/Engine/Source/Runtime/Engine/Classes/Engine/Blueprint.h b/Engine/Source/Runtime/Engine/Classes/Engine/Blueprint.h
--- a/Engine/Source/Runtime/Engine/Classes/Engine/Blueprint.h
+++ b/Engine/Source/Runtime/Engine/Classes/Engine/Blueprint.h
@@ -658,11 +658,9 @@ public:
 
 protected:
 	/** Current object being debugged for this blueprint */
-	UPROPERTY(transient, duplicatetransient)
 	TWeakObjectPtr< class UObject > CurrentObjectBeingDebugged;
 
 	/** Current world being debugged for this blueprint */
-	UPROPERTY(transient, duplicatetransient)
 	TWeakObjectPtr< class UWorld > CurrentWorldBeingDebugged;
 
 	/** Delegate called when the debug object is set */
diff --git a/Engine/Source/Runtime/Engine/Private/ActiveSound.cpp b/Engine/Source/Runtime/Engine/Private/ActiveSound.cpp
--- a/Engine/Source/Runtime/Engine/Private/ActiveSound.cpp
+++ b/Engine/Source/Runtime/Engine/Private/ActiveSound.cpp
@@ -1635,6 +1635,10 @@ void FActiveSound::ApplyAttenuation(FSoundParseParameters& ParseParams, const FL
 	{
 		ParseParams.SpatializationMethod = ESoundSpatializationAlgorithm::SPATIALIZATION_HRTF;
 	}
+	else if (Settings->SpatializationAlgorithm == ESoundSpatializationAlgorithm::SPATIALIZATION_HRTF && AudioDevice->IsHRTFDisabled())
+	{
+		ParseParams.SpatializationMethod = ESoundSpatializationAlgorithm::SPATIALIZATION_Default;
+	}
 	else
 	{
 		ParseParams.SpatializationMethod = Settings->SpatializationAlgorithm;
diff --git a/Engine/Source/Runtime/Engine/Private/AudioDebug.cpp b/Engine/Source/Runtime/Engine/Private/AudioDebug.cpp
--- a/Engine/Source/Runtime/Engine/Private/AudioDebug.cpp
+++ b/Engine/Source/Runtime/Engine/Private/AudioDebug.cpp
@@ -277,11 +277,9 @@ void FAudioDebugger::DrawDebugInfo(const FActiveSound& ActiveSound, const TArray
 		float DisplayValue = 0.0f;
 		if (ActiveSoundVisualizeModeCVar == 1 || ActiveSoundVisualizeModeCVar == 2)
 		{
-			DisplayValue = ActiveSound.GetVolume();
-			if (ThisSoundsWaveInstances.Num() == 1)
+			for (FWaveInstance* WaveInstance : ThisSoundsWaveInstances)
 			{
-				const FWaveInstance* WaveInstance = ThisSoundsWaveInstances[0];
-				DisplayValue = WaveInstance->GetVolumeWithDistanceAttenuation() * WaveInstance->GetDynamicVolume();
+				DisplayValue = FMath::Max(DisplayValue, WaveInstance->GetVolumeWithDistanceAttenuation() * WaveInstance->GetDynamicVolume());
 			}
 		}
 		else if (ActiveSoundVisualizeModeCVar == 3)
diff --git a/Engine/Source/Runtime/Engine/Private/AudioDevice.cpp b/Engine/Source/Runtime/Engine/Private/AudioDevice.cpp
--- a/Engine/Source/Runtime/Engine/Private/AudioDevice.cpp
+++ b/Engine/Source/Runtime/Engine/Private/AudioDevice.cpp
@@ -145,6 +145,12 @@ FAutoConsoleVariableRef CVarEnableBinauralAudioForAllSpatialSounds(
 	TEXT("Toggles binaural audio rendering for all spatial sounds if binaural rendering is available.\n"),
 	ECVF_Default);
 
+static int32 DisableBinauralSpatializationCVar = 0;
+FAutoConsoleVariableRef CVarDisableBinauralSpatialization(
+	TEXT("au.DisableBinauralSpatialization"),
+	DisableBinauralSpatializationCVar,
+	TEXT("Disables binaural spatialization.\n"),
+	ECVF_Default);
 
 using FVirtualLoopPair = TPair<FActiveSound*, FAudioVirtualLoop>;
 
@@ -184,6 +190,7 @@ FAudioDevice::FAudioDevice()
 	, PlatformAudioHeadroom(1.0f)
 	, DefaultReverbSendLevel(0.0f)
 	, bHRTFEnabledForAll_OnGameThread(false)
+	, bHRTFDisabled_OnGameThread(false)
 	, bGameWasTicking(true)
 	, bDisableAudioCaching(false)
 	, bIsAudioDeviceHardwareInitialized(false)
@@ -198,9 +205,11 @@ FAudioDevice::FAudioDevice()
 	, bSpatializationInterfaceEnabled(false)
 	, bOcclusionInterfaceEnabled(false)
 	, bReverbInterfaceEnabled(false)
+	, bReverbPluginBypassesMasterReverb(true)
 	, bModulationInterfaceEnabled(false)
 	, bPluginListenersInitialized(false)
 	, bHRTFEnabledForAll(false)
+	, bHRTFDisabled(false)
 	, bIsDeviceMuted(false)
 	, bIsInitialized(false)
 	, AudioClock(0.0)
@@ -351,9 +360,9 @@ bool FAudioDevice::Init(int32 InMaxChannels)
 	{
 		ReverbPluginInterface = ReverbPluginFactory->CreateNewReverbPlugin(this);
 		bReverbInterfaceEnabled = true;
+		bReverbPluginBypassesMasterReverb = ReverbPluginInterface->DoesReverbOverrideMasterReverb();
 		bReverbIsExternalSend = ReverbPluginFactory->IsExternalSend();
 		UE_LOG(LogAudio, Log, TEXT("Audio Reverb Plugin: %s"), *(ReverbPluginFactory->GetDisplayName()));
-
 	}
 	else
 	{
@@ -1708,6 +1717,17 @@ bool FAudioDevice::IsHRTFEnabledForAll() const
 	return (bHRTFEnabledForAll_OnGameThread || EnableBinauralAudioForAllSpatialSoundsCVar == 1) && IsSpatializationPluginEnabled();
 }
 
+bool FAudioDevice::IsHRTFDisabled() const
+{
+	if (IsInAudioThread())
+	{
+		return (bHRTFDisabled || DisableBinauralSpatializationCVar == 1);
+	}
+
+	check(IsInGameThread());
+	return (bHRTFDisabled_OnGameThread || DisableBinauralSpatializationCVar == 1);
+}
+
 void FAudioDevice::SetMixDebugState(EDebugState InDebugState)
 {
 	if (!IsInAudioThread())
diff --git a/Engine/Source/Runtime/Engine/Private/Particles/ParticleSystemRender.cpp b/Engine/Source/Runtime/Engine/Private/Particles/ParticleSystemRender.cpp
--- a/Engine/Source/Runtime/Engine/Private/Particles/ParticleSystemRender.cpp
+++ b/Engine/Source/Runtime/Engine/Private/Particles/ParticleSystemRender.cpp
@@ -1870,7 +1870,7 @@ void FDynamicMeshEmitterData::GetDynamicMeshElementsEmitter(const FParticleSyste
 
 					if (bIsWireframe)
 					{
-						if (LODModel.AdditionalIndexBuffers->WireframeIndexBuffer.IsInitialized()
+						if (LODModel.AdditionalIndexBuffers && LODModel.AdditionalIndexBuffers->WireframeIndexBuffer.IsInitialized()
 							&& !(RHISupportsTessellation(ShaderPlatform) && Mesh.VertexFactory->GetType()->SupportsTessellationShaders()))
 						{
 							Mesh.Type = PT_LineList;
diff --git a/Engine/Source/Runtime/Engine/Public/AudioDevice.h b/Engine/Source/Runtime/Engine/Public/AudioDevice.h
--- a/Engine/Source/Runtime/Engine/Public/AudioDevice.h
+++ b/Engine/Source/Runtime/Engine/Public/AudioDevice.h
@@ -950,6 +950,25 @@ public:
 		}, GET_STATID(STAT_SetHRTFEnabledForAll));
 	}
 
+	/** Whether or not HRTF is disabled. */
+	bool IsHRTFDisabled() const;
+
+	void SetHRTFDisabled(bool InIsHRTFDisabled)
+	{
+		const bool bNewHRTFDisabled = InIsHRTFDisabled;
+
+		bHRTFDisabled_OnGameThread = bNewHRTFDisabled;
+
+		DECLARE_CYCLE_STAT(TEXT("FAudioThreadTask.SetHRTFDisabled"), STAT_SetHRTFDisabled, STATGROUP_AudioThreadCommands);
+
+		FAudioDevice* AudioDevice = this;
+		FAudioThread::RunCommandOnAudioThread([AudioDevice, bNewHRTFDisabled]()
+		{
+			AudioDevice->bHRTFDisabled = bNewHRTFDisabled;
+
+		}, GET_STATID(STAT_SetHRTFDisabled));
+	}
+
 	void SetSpatializationInterfaceEnabled(bool InbSpatializationInterfaceEnabled)
 	{
 		FAudioThread::SuspendAudioThread();
@@ -1085,6 +1104,12 @@ public:
 		return bReverbInterfaceEnabled;
 	}
 
+	/** Whether or not the reverb plugin is bypassing the master reverb. */
+	bool IsReverbPluginBypassingMasterReverb() const
+	{
+		return bReverbPluginBypassesMasterReverb;
+	}
+
 	static bool IsReverbPluginLoaded()
 	{
 		if (FAudioDevice* MainAudioDevice = GEngine->GetMainAudioDevice())
@@ -1679,7 +1704,10 @@ private:
 	FActivatedReverb HighestPriorityActivatedReverb;
 
 	/** Gamethread representation of whether HRTF is enabled for all 3d sounds. (not bitpacked to avoid thread issues) */
-	bool bHRTFEnabledForAll_OnGameThread:1;
+	bool bHRTFEnabledForAll_OnGameThread;
+
+	/** Gamethread representation of whether HRTF is disabbled for all 3d sounds. */
+	bool bHRTFDisabled_OnGameThread;
 
 	uint8 bGameWasTicking:1;
 
@@ -1714,14 +1742,18 @@ private:
 	uint8 bSpatializationInterfaceEnabled:1;
 	uint8 bOcclusionInterfaceEnabled:1;
 	uint8 bReverbInterfaceEnabled:1;
+	uint8 bReverbPluginBypassesMasterReverb:1;
 	uint8 bModulationInterfaceEnabled:1;
 
 	/** Whether or not we've initialized plugin listeners array. */
 	uint8 bPluginListenersInitialized:1;
 
-	/** Whether HRTF is enabled for all 3d sounds. */
+	/** Whether HRTF is enabled for all 3d sounds. This will automatically make all 3d mono sounds HRTF spatialized. */
 	uint8 bHRTFEnabledForAll:1;
 
+	/** Whether or not HRTF is disabled. This will make any sounds which are set to HRTF spatialize to spatialize with panning. */
+	uint8 bHRTFDisabled:1;
+
 	/** Whether the audio device is active (current audio device in-focus in PIE) */
 	uint8 bIsDeviceMuted:1;
 
diff --git a/Engine/Source/Runtime/Engine/Public/IAudioExtensionPlugin.h b/Engine/Source/Runtime/Engine/Public/IAudioExtensionPlugin.h
--- a/Engine/Source/Runtime/Engine/Public/IAudioExtensionPlugin.h
+++ b/Engine/Source/Runtime/Engine/Public/IAudioExtensionPlugin.h
@@ -641,6 +641,9 @@ public:
 	virtual void ProcessSourceAudio(const FAudioPluginSourceInputData& InputData, FAudioPluginSourceOutputData& OutputData)
 	{
 	}
+
+	/** Returns whether or not the plugin reverb overrides the master reverb. If true, then the built in reverb will be uninitialized and bypassed. */
+	virtual bool DoesReverbOverrideMasterReverb() const { return true; }
 };
 
 /************************************************************************/
diff --git a/Engine/Source/Runtime/MovieScene/Private/Channels/MovieSceneFloatChannel.cpp b/Engine/Source/Runtime/MovieScene/Private/Channels/MovieSceneFloatChannel.cpp
--- a/Engine/Source/Runtime/MovieScene/Private/Channels/MovieSceneFloatChannel.cpp
+++ b/Engine/Source/Runtime/MovieScene/Private/Channels/MovieSceneFloatChannel.cpp
@@ -787,15 +787,7 @@ void FMovieSceneFloatChannel::AutoSetTangents(float Tension)
 			NewTangent = (NextKey.Value - ThisKey.Value) / NextTimeDiff;
 			ThisKey.Tangent.LeaveTangent = NewTangent;
 		}
-		else if (PrevKey.InterpMode == RCIM_Constant || ThisKey.InterpMode == RCIM_Constant)
-		{
-			if (PrevKey.InterpMode != RCIM_Cubic)
-			{
-				ThisKey.Tangent.ArriveTangent = 0.f;
-			}
 
-			ThisKey.Tangent.LeaveTangent  = 0.0f;
-		}
 		
 	}
 }
diff --git a/Engine/Source/Runtime/MovieScene/Private/MovieSceneFolder.cpp b/Engine/Source/Runtime/MovieScene/Private/MovieSceneFolder.cpp
--- a/Engine/Source/Runtime/MovieScene/Private/MovieSceneFolder.cpp
+++ b/Engine/Source/Runtime/MovieScene/Private/MovieSceneFolder.cpp
@@ -173,7 +173,7 @@ void UMovieSceneFolder::PostLoad()
 				ChildMasterTracks.RemoveAt(ChildMasterTrackIndex);
 				ChildMasterTrackIndex--;
 
-				UE_LOG(LogMovieScene, Warning, TEXT("Folder (%s) contained a reference to a Master Track (%s) that no longer exists in the sequence, removing."), *GetFolderName().ToString(), *GetNameSafe(ChildTrack));
+				UE_LOG(LogMovieScene, Warning, TEXT("Folder (%s) in Sequence (%s) contained a reference to a Master Track (%s) that no longer exists in the sequence, removing."), *GetFolderName().ToString(), *OwningScene->GetPathName(), *GetNameSafe(ChildTrack));
 			}
 		}
 
@@ -186,7 +186,7 @@ void UMovieSceneFolder::PostLoad()
 				ChildObjectBindings.RemoveAt(ChildObjectBindingIndex);
 				ChildObjectBindingIndex--;
 
-				UE_LOG(LogMovieScene, Warning, TEXT("Folder (%s) contained a reference to an Object Binding (%s) that no longer exists in the sequence, removing."), *GetFolderName().ToString(), *ChildBinding.ToString());
+				UE_LOG(LogMovieScene, Warning, TEXT("Folder (%s) in Sequence (%s) contained a reference to an Object Binding (%s) that no longer exists in the sequence, removing."), *GetFolderName().ToString(), *OwningScene->GetPathName(), *ChildBinding.ToString());
 			}
 		}
 
@@ -204,7 +204,7 @@ void UMovieSceneFolder::PostLoad()
 				ChildFolders.RemoveAt(ChildFolderIndex);
 				ChildFolderIndex--;
 
-				UE_LOG(LogMovieScene, Warning, TEXT("Folder (%s) contained a reference to an Folder (%s) that exists in multiple places in the sequence, removing."), *GetFolderName().ToString(), *ChildFolders[ChildFolderIndex]->GetFolderName().ToString());
+				UE_LOG(LogMovieScene, Warning, TEXT("Folder (%s) in Sequence (%s) contained a reference to an Folder (%s) that exists in multiple places in the sequence, removing."), *GetFolderName().ToString(), *OwningScene->GetPathName(), *ChildFolders[ChildFolderIndex]->GetFolderName().ToString());
 			}
 		}
 	}
