XR tracking system ends up calling OnBeginPlay twice, once near the beginning of PlayInEditor and then again in CreatePIEGameInstance which is called by PlayInEditor. This ends up running derived callback code twice which is not the expected behavior. I've removed the PlayInEditor call and let CreatePIEGameInstance be the sole call site for OnBeginPlay. This means we call it slightly later than before, but the code inbetween didn't seem like something we needed to call OnBeginPlay before.

diff --git a/Engine/Source/Editor/UnrealEd/Private/PlayLevel.cpp b/Engine/Source/Editor/UnrealEd/Private/PlayLevel.cpp
--- a/Engine/Source/Editor/UnrealEd/Private/PlayLevel.cpp
+++ b/Engine/Source/Editor/UnrealEd/Private/PlayLevel.cpp
@@ -2459,11 +2459,6 @@ void UEditorEngine::PlayInEditor( UWorld* InWorld, bool bInSimulateInEditor, FPl
 		OutputLogErrorsToMessageLogProxyPtr = MakeShareable(new FOutputLogErrorsToMessageLogProxy());
 	}
 
-	if (GEngine->XRSystem.IsValid() && !bInSimulateInEditor)
-	{
-		GEngine->XRSystem->OnBeginPlay(*GEngine->GetWorldContextFromWorld(InWorld));
-	}
-
 	// remember old GWorld
 	EditorWorld = InWorld;
 
