Make it possible to move directional light marked as atmospheric light with the ctrl+L keyboard shortcut in the level editor only for now.

diff --git a/Engine/Source/Editor/UnrealEd/Private/LevelEditorViewport.cpp b/Engine/Source/Editor/UnrealEd/Private/LevelEditorViewport.cpp
--- a/Engine/Source/Editor/UnrealEd/Private/LevelEditorViewport.cpp
+++ b/Engine/Source/Editor/UnrealEd/Private/LevelEditorViewport.cpp
@@ -23,6 +23,7 @@
 #include "Factories/MaterialFactoryNew.h"
 #include "Editor/GroupActor.h"
 #include "Components/DecalComponent.h"
+#include "Components/DirectionalLightComponent.h"
 #include "Components/InstancedStaticMeshComponent.h"
 #include "Components/ModelComponent.h"
 #include "Kismet2/ComponentEditorUtils.h"
@@ -2176,6 +2177,9 @@ static FMatrix GPerspViewMatrix;
 
 void FLevelEditorViewportClient::Tick(float DeltaTime)
 {
+	CachedLastMouseX = Viewport->GetMouseX();
+	CachedLastMouseY = Viewport->GetMouseY();
+
 	if (bWasEditorCameraCut && bEditorCameraCut)
 	{
 		bEditorCameraCut = false;
@@ -2206,6 +2210,8 @@ void FLevelEditorViewportClient::Tick(float DeltaTime)
 	}
 
 	UpdateViewForLockedActor(DeltaTime);
+
+	UserIsControllingSunLightTimer = FMath::Max(UserIsControllingSunLightTimer - DeltaTime, 0.0f);
 }
 
 void FLevelEditorViewportClient::UpdateViewForLockedActor(float DeltaTime)
@@ -2661,6 +2667,42 @@ bool FLevelEditorViewportClient::InputKey(FViewport* InViewport, int32 Controlle
 		return true;
 	}
 
+	if (InputState.IsCtrlButtonPressed() && Key == EKeys::L)
+	{
+		UWorld* ViewportWorld = GetWorld();
+		for (TObjectIterator<UDirectionalLightComponent> ComponentIt; ComponentIt; ++ComponentIt)
+		{
+			if (ComponentIt->GetWorld() == ViewportWorld)
+			{
+				UDirectionalLightComponent* SunLight = *ComponentIt;
+
+				int32 mouseDeltaX = HitX - CachedLastMouseX;
+				int32 mouseDeltaY = HitY - CachedLastMouseY;
+				if (!SunLight->IsUsedAsAtmosphereSunLight() || !SunLight->bVisible)
+					continue;
+
+				FTransform ComponentTransform = SunLight->GetComponentTransform();
+				FQuat SunRotation = ComponentTransform.GetRotation();
+				// Rotate around up axis (yaw)
+				FVector UpVector = FVector(0, 0, 1);
+				SunRotation = FQuat(UpVector, float(mouseDeltaX)*0.02f) * SunRotation;
+				// Sun Zenith rotation (pitch)
+				FVector PitchRotationAxis = FVector::CrossProduct(SunRotation.GetForwardVector(), UpVector);
+				PitchRotationAxis.Normalize();
+				SunRotation = FQuat(PitchRotationAxis, float(mouseDeltaY)*0.02f) * SunRotation;
+
+				ComponentTransform.SetRotation(SunRotation);
+				SunLight->SetWorldTransform(ComponentTransform);
+
+				// Stop on the first encountered light
+				UserControlledSunLightMatrix = ComponentTransform;
+				UserIsControllingSunLightTimer = 3.0f;
+				// Only manipulate a single light, the first one.
+				return true;
+			}
+		}
+	}
+
 	bool bHandled = FEditorViewportClient::InputKey(InViewport,ControllerId,Key,Event,AmountDepressed,bGamepad);
 
 	// Handle input for the player height preview mode. 
@@ -4293,6 +4335,55 @@ void FLevelEditorViewportClient::Draw(const FSceneView* View,FPrimitiveDrawInter
 		}
 	}
 
+	if (UserIsControllingSunLightTimer > 0.0f)
+	{
+		// Draw a gizmo helping to figure out where is the sun when moving it using a shortcut.
+		FQuat ViewRotation = FQuat(GetViewRotation());
+		FVector ViewPosition = GetViewLocation();
+		const float GizmoDistance = 50.0f;
+		const float GizmoSideOffset = 15.0f;
+		const float GizmoRadius = 10.0f;
+		const float ThicknessLight = 0.05f;
+		const float ThicknessBold = 0.2f;
+
+		// Always draw the gizmo right in in front of the camera with a little side shift.
+		const FVector X(1.0f, 0.0f, 0.0f);
+		const FVector Y(0.0f, 1.0f, 0.0f);
+		const FVector Z(0.0f, 0.0f, 1.0f);
+		const FVector Base = ViewPosition + GizmoDistance * ViewRotation.GetForwardVector() + GizmoSideOffset * (-ViewRotation.GetUpVector() + ViewRotation.GetRightVector());
+
+		// Draw world main axis
+		FRotator IdentityX(0.0f, 0.0f, 0.0f);
+		FRotator IdentityY(0.0f, 90.0f, 0.0f);
+		FRotator IdentityZ(90.0f, 0.0f, 0.0f);
+		DrawDirectionalArrow(PDI, FQuatRotationTranslationMatrix(FQuat(IdentityX), Base), FColor(255, 0, 0, 127), GizmoRadius, 0.3f, SDPG_World, ThicknessBold);
+		DrawDirectionalArrow(PDI, FQuatRotationTranslationMatrix(FQuat(IdentityY), Base), FColor(0, 255, 0, 127), GizmoRadius, 0.3f, SDPG_World, ThicknessBold);
+		DrawDirectionalArrow(PDI, FQuatRotationTranslationMatrix(FQuat(IdentityZ), Base), FColor(0, 0, 255, 127), GizmoRadius, 0.3f, SDPG_World, ThicknessBold);
+
+		// Render polar coordinate circles
+		DrawCircle(PDI, Base, X, Y, FLinearColor(0.2f, 0.2f, 1.0f), GizmoRadius, 32, SDPG_World, ThicknessBold);
+		DrawCircle(PDI, Base, X, Y, FLinearColor(0.2f, 0.2f, 0.75f), GizmoRadius*0.75f, 32, SDPG_World, ThicknessLight);
+		DrawCircle(PDI, Base, X, Y, FLinearColor(0.2f, 0.2f, 0.50f), GizmoRadius*0.50f, 32, SDPG_World, ThicknessLight);
+		DrawCircle(PDI, Base, X, Y, FLinearColor(0.2f, 0.2f, 0.25f), GizmoRadius*0.25f, 32, SDPG_World, ThicknessLight);
+		DrawArc(PDI, Base, Z, Y, -90.0f, 90.0f, GizmoRadius, 32, FLinearColor(1.0f, 0.2f, 0.2f), SDPG_World);
+		DrawArc(PDI, Base, Z, X, -90.0f, 90.0f, GizmoRadius, 32, FLinearColor(0.2f, 1.0f, 0.2f), SDPG_World);
+
+		// Draw the sun incoming light direction. The arrow is offset outward to help depth perception when it intersects with other gizmo elements.
+		const FLinearColor ArrowColor = -UserControlledSunLightMatrix.GetRotation().GetForwardVector() * 0.5f + 0.5f;
+		const FVector ArrowOrigin = Base - UserControlledSunLightMatrix.GetRotation().GetForwardVector()*GizmoRadius*1.25;
+		const FQuatRotationTranslationMatrix ArrowToWorld(UserControlledSunLightMatrix.GetRotation(), ArrowOrigin);
+		DrawDirectionalArrow(PDI, ArrowToWorld, ArrowColor, GizmoRadius, 0.3f, SDPG_World, ThicknessBold);
+
+		// Now draw x, y and z axis to help getting a sense of depth when look at the vectors on screen.
+		FVector SunArrowTip = -UserControlledSunLightMatrix.GetRotation().GetForwardVector()*GizmoRadius;
+		FVector P0 = Base + SunArrowTip * FVector(1.0f, 0.0f, 0.0f);
+		FVector P1 = Base + SunArrowTip * FVector(1.0f, 1.0f, 0.0f);
+		FVector P2 = Base + SunArrowTip * FVector(1.0f, 1.0f, 1.0f);
+		PDI->DrawLine(Base, P0, FLinearColor(1.0f, 0.0f, 0.0f), SDPG_World, ThicknessLight);
+		PDI->DrawLine(P0, P1, FLinearColor(0.0f, 1.0f, 0.0f), SDPG_World, ThicknessLight);
+		PDI->DrawLine(P1, P2, FLinearColor(0.0f, 0.0f, 1.0f), SDPG_World, ThicknessLight);
+	}
+
 	Mark.Pop();
 }
 
diff --git a/Engine/Source/Editor/UnrealEd/Public/EditorViewportClient.h b/Engine/Source/Editor/UnrealEd/Public/EditorViewportClient.h
--- a/Engine/Source/Editor/UnrealEd/Public/EditorViewportClient.h
+++ b/Engine/Source/Editor/UnrealEd/Public/EditorViewportClient.h
@@ -1531,6 +1531,15 @@ protected:
 	uint32 CachedMouseX;
 	uint32 CachedMouseY;
 
+	/** Represents the last mouse position. It is constantly updated on tick so it can also be the current position. */
+	int32 CachedLastMouseX = 0;
+	int32 CachedLastMouseY = 0;
+
+	/** True is the use is controling the light via a*/
+	float UserIsControllingSunLightTimer = 0.0f;
+	FTransform UserControlledSunLightMatrix;
+
+
 	// -1, -1 if not set
 	FIntPoint CurrentMousePos;
 
