Fixed sun move shortcut slopiness.
 - CachedLastMouse is now updated in editor view mouse move
 - Sun orientation is now update in mousemove of the level editor for a smoother result.

diff --git a/Engine/Source/Editor/UnrealEd/Private/EditorViewportClient.cpp b/Engine/Source/Editor/UnrealEd/Private/EditorViewportClient.cpp
--- a/Engine/Source/Editor/UnrealEd/Private/EditorViewportClient.cpp
+++ b/Engine/Source/Editor/UnrealEd/Private/EditorViewportClient.cpp
@@ -5209,6 +5209,9 @@ void FEditorViewportClient::MouseMove(FViewport* InViewport,int32 x, int32 y)
 
 	// Let the current editor mode know about the mouse movement.
 	ModeTools->MouseMove(this, Viewport, x, y);
+
+	CachedLastMouseX = x;
+	CachedLastMouseY = y;
 }
 
 void FEditorViewportClient::MouseLeave(FViewport* InViewport)
diff --git a/Engine/Source/Editor/UnrealEd/Private/LevelEditorViewport.cpp b/Engine/Source/Editor/UnrealEd/Private/LevelEditorViewport.cpp
--- a/Engine/Source/Editor/UnrealEd/Private/LevelEditorViewport.cpp
+++ b/Engine/Source/Editor/UnrealEd/Private/LevelEditorViewport.cpp
@@ -2177,9 +2177,6 @@ static FMatrix GPerspViewMatrix;
 
 void FLevelEditorViewportClient::Tick(float DeltaTime)
 {
-	CachedLastMouseX = Viewport->GetMouseX();
-	CachedLastMouseY = Viewport->GetMouseY();
-
 	if (bWasEditorCameraCut && bEditorCameraCut)
 	{
 		bEditorCameraCut = false;
@@ -2669,39 +2666,11 @@ bool FLevelEditorViewportClient::InputKey(FViewport* InViewport, int32 Controlle
 
 	if (InputState.IsCtrlButtonPressed() && Key == EKeys::L)
 	{
-		UWorld* ViewportWorld = GetWorld();
-		for (TObjectIterator<UDirectionalLightComponent> ComponentIt; ComponentIt; ++ComponentIt)
-		{
-			if (ComponentIt->GetWorld() == ViewportWorld)
-			{
-				UDirectionalLightComponent* SunLight = *ComponentIt;
-
-				int32 mouseDeltaX = HitX - CachedLastMouseX;
-				int32 mouseDeltaY = HitY - CachedLastMouseY;
-				if (!SunLight->IsUsedAsAtmosphereSunLight() || !SunLight->bVisible)
-					continue;
-
-				FTransform ComponentTransform = SunLight->GetComponentTransform();
-				FQuat SunRotation = ComponentTransform.GetRotation();
-				// Rotate around up axis (yaw)
-				FVector UpVector = FVector(0, 0, 1);
-				SunRotation = FQuat(UpVector, float(mouseDeltaX)*0.02f) * SunRotation;
-				// Sun Zenith rotation (pitch)
-				FVector PitchRotationAxis = FVector::CrossProduct(SunRotation.GetForwardVector(), UpVector);
-				PitchRotationAxis.Normalize();
-				SunRotation = FQuat(PitchRotationAxis, float(mouseDeltaY)*0.02f) * SunRotation;
-
-				ComponentTransform.SetRotation(SunRotation);
-				SunLight->SetWorldTransform(ComponentTransform);
-
-				// Stop on the first encountered light
-				UserControlledSunLightMatrix = ComponentTransform;
-				UserIsControllingSunLightTimer = 3.0f;
-				// Only manipulate a single light, the first one.
-				return true;
-			}
-		}
+		bUserIsControllingSunLight = true;
+		UserIsControllingSunLightTimer = 3.0f; // Keep the widget open for a few seconds even when not tweaking the sun light
+		return true;
 	}
+	bUserIsControllingSunLight = false;
 
 	bool bHandled = FEditorViewportClient::InputKey(InViewport,ControllerId,Key,Event,AmountDepressed,bGamepad);
 
@@ -4001,6 +3970,44 @@ EMouseCursor::Type FLevelEditorViewportClient::GetCursor(FViewport* InViewport,i
 
 }
 
+void FLevelEditorViewportClient::MouseMove(FViewport* InViewport, int32 x, int32 y)
+{
+	if (bUserIsControllingSunLight)
+	{
+		UWorld* ViewportWorld = GetWorld();
+		for (TObjectIterator<UDirectionalLightComponent> ComponentIt; ComponentIt; ++ComponentIt)
+		{
+			if (ComponentIt->GetWorld() == ViewportWorld)
+			{
+				UDirectionalLightComponent* SunLight = *ComponentIt;
+
+				int32 mouseDeltaX = x - CachedLastMouseX;
+				int32 mouseDeltaY = y - CachedLastMouseY;
+				if (!SunLight->IsUsedAsAtmosphereSunLight() || !SunLight->bVisible)
+					continue;
+
+				FTransform ComponentTransform = SunLight->GetComponentTransform();
+				FQuat SunRotation = ComponentTransform.GetRotation();
+				// Rotate around up axis (yaw)
+				FVector UpVector = FVector(0, 0, 1);
+				SunRotation = FQuat(UpVector, float(mouseDeltaX)*0.01f) * SunRotation;
+				// Sun Zenith rotation (pitch)
+				FVector PitchRotationAxis = FVector::CrossProduct(SunRotation.GetForwardVector(), UpVector);
+				PitchRotationAxis.Normalize();
+				SunRotation = FQuat(PitchRotationAxis, float(mouseDeltaY)*0.01f) * SunRotation;
+
+				ComponentTransform.SetRotation(SunRotation);
+				SunLight->SetWorldTransform(ComponentTransform);
+
+				// Only manipulate a single light, the first one.
+				UserControlledSunLightMatrix = ComponentTransform;
+			}
+		}
+	}
+
+	FEditorViewportClient::MouseMove(InViewport, x, y);
+}
+
 /**
  * Called when the mouse is moved while a window input capture is in effect
  *
diff --git a/Engine/Source/Editor/UnrealEd/Public/EditorViewportClient.h b/Engine/Source/Editor/UnrealEd/Public/EditorViewportClient.h
--- a/Engine/Source/Editor/UnrealEd/Public/EditorViewportClient.h
+++ b/Engine/Source/Editor/UnrealEd/Public/EditorViewportClient.h
@@ -1536,6 +1536,7 @@ protected:
 	int32 CachedLastMouseY = 0;
 
 	/** True is the use is controling the light via a*/
+	bool bUserIsControllingSunLight = false;
 	float UserIsControllingSunLightTimer = 0.0f;
 	FTransform UserControlledSunLightMatrix;
 
diff --git a/Engine/Source/Editor/UnrealEd/Public/LevelEditorViewport.h b/Engine/Source/Editor/UnrealEd/Public/LevelEditorViewport.h
--- a/Engine/Source/Editor/UnrealEd/Public/LevelEditorViewport.h
+++ b/Engine/Source/Editor/UnrealEd/Public/LevelEditorViewport.h
@@ -156,7 +156,8 @@ public:
 	virtual bool InputKey(FViewport* Viewport, int32 ControllerId, FKey Key, EInputEvent Event, float AmountDepressed = 1.f, bool bGamepad=false) override;
 	virtual bool InputAxis(FViewport* Viewport, int32 ControllerId, FKey Key, float Delta, float DeltaTime, int32 NumSamples=1, bool bGamepad=false) override;
 	virtual EMouseCursor::Type GetCursor(FViewport* Viewport,int32 X,int32 Y) override;
-	virtual void CapturedMouseMove( FViewport* InViewport, int32 InMouseX, int32 InMouseY ) override;
+	virtual void CapturedMouseMove(FViewport* InViewport, int32 InMouseX, int32 InMouseY) override;
+	virtual void MouseMove(FViewport* InViewport, int32 x, int32 y) override;
 	virtual void Tick(float DeltaSeconds) override;
 	virtual bool InputWidgetDelta( FViewport* Viewport, EAxisList::Type CurrentAxis, FVector& Drag, FRotator& Rot, FVector& Scale ) override;
 	virtual TSharedPtr<FDragTool> MakeDragTool( EDragTool::Type DragToolType ) override;
