Fixed logic used to select the AtmosphereComponent to use in lightmass. (only the one in the actual world we build)

diff --git a/Engine/Source/Editor/UnrealEd/Private/StaticLightingSystem/StaticLightingSystem.cpp b/Engine/Source/Editor/UnrealEd/Private/StaticLightingSystem/StaticLightingSystem.cpp
--- a/Engine/Source/Editor/UnrealEd/Private/StaticLightingSystem/StaticLightingSystem.cpp
+++ b/Engine/Source/Editor/UnrealEd/Private/StaticLightingSystem/StaticLightingSystem.cpp
@@ -2029,7 +2029,7 @@ void FStaticLightingSystem::GatherScene()
 	for (TObjectIterator<UAtmosphericFogComponent> It; It; ++It)
 	{
 		UAtmosphericFogComponent* AtmosphericFog = *It;
-		if (!AtmosphericFog->IsPendingKill() && AtmosphericFog->bVisible)
+		if (AtmosphericFog->GetOwner() && World->ContainsActor(AtmosphericFog->GetOwner()) && !AtmosphericFog->IsPendingKill() && ShouldOperateOnLevel(AtmosphericFog->GetOwner()->GetLevel()))
 		{
 			LightmassExporter->SetAtmosphericComponent(AtmosphericFog);
 			break;	// We only register the first we find
