Fixed a couple compile errors when using clang

diff --git a/Engine/Source/Developer/AssetTools/Private/AssetVtConversion.cpp b/Engine/Source/Developer/AssetTools/Private/AssetVtConversion.cpp
--- a/Engine/Source/Developer/AssetTools/Private/AssetVtConversion.cpp
+++ b/Engine/Source/Developer/AssetTools/Private/AssetVtConversion.cpp
@@ -1,5 +1,3 @@
-#pragma once
-
 #include "AssetVtConversion.h"
 #include "Engine/Texture2D.h"
 #include "Misc/PackageName.h"
diff --git a/Engine/Source/Editor/UnrealEd/Private/Factories/EditorFactories.cpp b/Engine/Source/Editor/UnrealEd/Private/Factories/EditorFactories.cpp
--- a/Engine/Source/Editor/UnrealEd/Private/Factories/EditorFactories.cpp
+++ b/Engine/Source/Editor/UnrealEd/Private/Factories/EditorFactories.cpp
@@ -3118,7 +3118,7 @@ void FImportImage::Init2DWithMips(int32 InSizeX, int32 InSizeY, int32 InNumMips,
 {
 	SizeX = InSizeX;
 	SizeY = InSizeY;
-	NumMips = NumMips;
+	NumMips = InNumMips;
 	Format = InFormat;
 
 	int32 TotalSize = 0;
@@ -3165,7 +3165,7 @@ bool UTextureFactory::ImportImage(const uint8* Buffer, uint32 Length, FFeedbackC
 	{
 		if (!IsImportResolutionValid(PngImageWrapper->GetWidth(), PngImageWrapper->GetHeight(), bAllowNonPowerOfTwo, Warn))
 		{
-			return nullptr;
+			return false;
 		}
 
 		// Select the texture's source format
@@ -3313,7 +3313,7 @@ bool UTextureFactory::ImportImage(const uint8* Buffer, uint32 Length, FFeedbackC
 
 		if (!IsImportResolutionValid(Width, Height, bAllowNonPowerOfTwo, Warn))
 		{
-			return nullptr;
+			return false;
 		}
 
 		// Select the texture's source format
@@ -3330,7 +3330,7 @@ bool UTextureFactory::ImportImage(const uint8* Buffer, uint32 Length, FFeedbackC
 		if (TextureFormat == TSF_Invalid)
 		{
 			Warn->Logf(ELogVerbosity::Error, TEXT("EXR file contains data in an unsupported format."));
-			return nullptr;
+			return false;
 		}
 
 		const TArray<uint8>* Raw = nullptr;
diff --git a/Engine/Source/Runtime/Engine/Private/Materials/HLSLMaterialTranslator.h b/Engine/Source/Runtime/Engine/Private/Materials/HLSLMaterialTranslator.h
--- a/Engine/Source/Runtime/Engine/Private/Materials/HLSLMaterialTranslator.h
+++ b/Engine/Source/Runtime/Engine/Private/Materials/HLSLMaterialTranslator.h
@@ -3788,7 +3788,7 @@ protected:
 				);
 	}
 
-	static TCHAR* GetVTAddressMode(TextureAddress Address)
+	static const TCHAR* GetVTAddressMode(TextureAddress Address)
 	{
 		switch (Address)
 		{
diff --git a/Engine/Source/Runtime/Engine/Private/TextureDerivedData.cpp b/Engine/Source/Runtime/Engine/Private/TextureDerivedData.cpp
--- a/Engine/Source/Runtime/Engine/Private/TextureDerivedData.cpp
+++ b/Engine/Source/Runtime/Engine/Private/TextureDerivedData.cpp
@@ -438,7 +438,7 @@ static void GetTextureBuildSettings(
 	OutBuildSettings.ChromaKeyThreshold = Texture.ChromaKeyThreshold;
 	OutBuildSettings.CompressionQuality = Texture.CompressionQuality - 1; // translate from enum's 0 .. 5 to desired compression (-1 .. 4, where -1 is default while 0 .. 4 are actual quality setting override)
 	// TODO - get default value from config/CVAR/LODGroup?
-	OutBuildSettings.LossyCompressionAmount = (Texture.LossyCompressionAmount == TLCA_Default) ? TLCA_Lowest : Texture.LossyCompressionAmount;
+	OutBuildSettings.LossyCompressionAmount = (Texture.LossyCompressionAmount == TLCA_Default) ? TLCA_Lowest : Texture.LossyCompressionAmount.GetValue();
 
 	// For virtual texturing we take the address mode into consideration
 	if (OutBuildSettings.bVirtualStreamable)
