Removing small part of HDR display mode in order to fix startup crash. Entire CL was reverted in ue4-main, so will be cleaned up after the next pull from ue-main.#rnx

diff --git a/Engine/Source/Editor/UnrealEd/Private/EditorEngine.cpp b/Engine/Source/Editor/UnrealEd/Private/EditorEngine.cpp
--- a/Engine/Source/Editor/UnrealEd/Private/EditorEngine.cpp
+++ b/Engine/Source/Editor/UnrealEd/Private/EditorEngine.cpp
@@ -1992,14 +1992,6 @@ bool UEditorEngine::UpdateSingleViewportClient(FEditorViewportClient* InViewport
 {
 	bool bUpdatedNonRealtimeViewport = false;
 
-	// When rendering the viewport we need to know whether the final result will be shown on a HDR display. This affects the final post processing step
-	FSceneViewport *SceneViewport = static_cast<FSceneViewport*>(InViewportClient->Viewport);
-	TSharedPtr<SWindow> Window = SceneViewport->FindWindow();
-	if (Window)
-	{
-		InViewportClient->Viewport->SetHDRMode(Window->GetIsHDR());
-	}
-
 	// Always submit view information for content streaming 
 	// otherwise content for editor view can be streamed out if there are other views (ex: thumbnails)
 	if (InViewportClient->IsPerspective())
