Fixed compiler error on Mac introduced by CL 6676473 (MSVC silently translates boolean argument <false> to NULL pointer).

#rnx

diff --git a/Engine/Source/Editor/MaterialEditor/Private/MaterialStats.cpp b/Engine/Source/Editor/MaterialEditor/Private/MaterialStats.cpp
--- a/Engine/Source/Editor/MaterialEditor/Private/MaterialStats.cpp
+++ b/Engine/Source/Editor/MaterialEditor/Private/MaterialStats.cpp
@@ -183,7 +183,7 @@ bool FShaderPlatformSettings::CheckShaders()
 
 				FMaterialShaderMapId ShaderMapId;
 				Data.MaterialResourcesStats->GetShaderMapId(PlatformShaderID, ShaderMapId);
-				Data.MaterialResourcesStats->CacheShaders(ShaderMapId, PlatformShaderID, false);
+				Data.MaterialResourcesStats->CacheShaders(ShaderMapId, PlatformShaderID);
 
 				Data.bCompilingShaders = true;
 				Data.bUpdateShaderCode = true;
diff --git a/Engine/Source/Editor/UnrealEd/Private/EditorBuildUtils.cpp b/Engine/Source/Editor/UnrealEd/Private/EditorBuildUtils.cpp
--- a/Engine/Source/Editor/UnrealEd/Private/EditorBuildUtils.cpp
+++ b/Engine/Source/Editor/UnrealEd/Private/EditorBuildUtils.cpp
@@ -1200,7 +1200,7 @@ bool FEditorBuildUtils::EditorBuildMaterialTextureStreamingData(UPackage* Packag
 				FMaterialResource* Resource = Material->GetMaterialResource(FeatureLevel);
 				if (Resource)
 				{
-					Resource->CacheShaders(GMaxRHIShaderPlatform, false);
+					Resource->CacheShaders(GMaxRHIShaderPlatform);
 					Materials.Add(Material);
 				}
 			}
@@ -1387,7 +1387,7 @@ bool FEditorBuildUtils::CompileShadersComplexityViewMode(EMaterialQualityLevel::
 
 		FMaterialShaderMapId ResourceId;
 		SpecialResource->GetShaderMapId(ShaderPlatform, ResourceId);
-		SpecialResource->CacheShaders(ResourceId, ShaderPlatform, false);
+		SpecialResource->CacheShaders(ResourceId, ShaderPlatform);
 
 		OfflineShaderResources.Add(SpecialResource);
 	}
