- Proper handling of UDIM naming convension (name.#### is standard, rather than name_####)
- All UDIM page textures must use the same base name
- UDIM pages smaller than VT tile size are automatically sized up

diff --git a/Engine/Source/Editor/UnrealEd/Private/Factories/EditorFactories.cpp b/Engine/Source/Editor/UnrealEd/Private/Factories/EditorFactories.cpp
--- a/Engine/Source/Editor/UnrealEd/Private/Factories/EditorFactories.cpp
+++ b/Engine/Source/Editor/UnrealEd/Private/Factories/EditorFactories.cpp
@@ -3881,21 +3881,25 @@ bool UTextureFactory::DoesSupportClass(UClass* Class)
 	return Class == UTexture2D::StaticClass() || Class == UTextureCube::StaticClass();
 }
 
-static int32 ParseUDIMName(const FString& Name)
+static int32 ParseUDIMName(const FString& Name, FString& OutRootName)
 {
-	int32 UnderscoreIndex = INDEX_NONE;
-	if (!Name.FindLastChar('_', UnderscoreIndex))
+	int32 SeparatorIndex = INDEX_NONE;
+	if (!Name.FindLastChar('.', SeparatorIndex))
 	{
-		return INDEX_NONE;
+		// '.' is the standard UDIM separator, but we'll accept '_' as well
+		if (!Name.FindLastChar('_', SeparatorIndex))
+		{
+			return INDEX_NONE;
+		}
 	}
-	if (UnderscoreIndex + 5 != Name.Len())
+	if (SeparatorIndex + 5 != Name.Len())
 	{
 		return INDEX_NONE;
 	}
-	const TCHAR Digit0 = Name[UnderscoreIndex + 4];
-	const TCHAR Digit1 = Name[UnderscoreIndex + 3];
-	const TCHAR Digit2 = Name[UnderscoreIndex + 2];
-	const TCHAR Digit3 = Name[UnderscoreIndex + 1];
+	const TCHAR Digit0 = Name[SeparatorIndex + 4];
+	const TCHAR Digit1 = Name[SeparatorIndex + 3];
+	const TCHAR Digit2 = Name[SeparatorIndex + 2];
+	const TCHAR Digit3 = Name[SeparatorIndex + 1];
 	if (Digit0 < '0' || Digit0 > '9') return INDEX_NONE;
 	if (Digit1 < '0' || Digit1 > '9') return INDEX_NONE;
 	if (Digit2 < '0' || Digit2 > '9') return INDEX_NONE;
@@ -3908,6 +3912,7 @@ static int32 ParseUDIMName(const FString& Name)
 		return INDEX_NONE;
 	}
 
+	OutRootName = Name.Left(SeparatorIndex);
 	return Value;
 }
 
@@ -4042,14 +4047,15 @@ UObject* UTextureFactory::FactoryCreateBinary
 
 	// Check to see if we should import a series of textures as UDIM
 	const FString FilenameNoExtension = FPaths::GetBaseFilename(CurrentFilename);
-	const int32 BaseUDIMIndex = ParseUDIMName(FilenameNoExtension);
+	FString BaseUDIMName;
+	const int32 BaseUDIMIndex = ParseUDIMName(FilenameNoExtension, BaseUDIMName);
 	if (BaseUDIMIndex != INDEX_NONE)
 	{
 		TMap<int32, FString> UDIMIndexToFile;
 		UDIMIndexToFile.Add(BaseUDIMIndex, CurrentFilename);
 
 		const FString Path = FPaths::GetPath(CurrentFilename);
-		const FString UDIMFilter = Path / TEXT("*_*.*");
+		const FString UDIMFilter = (Path / BaseUDIMName) + TEXT("*");
 
 		TArray<FString> UDIMFiles;
 		IFileManager::Get().FindFiles(UDIMFiles, *UDIMFilter, true, false);
@@ -4058,8 +4064,9 @@ UObject* UTextureFactory::FactoryCreateBinary
 		{
 			if (!CurrentFilename.EndsWith(UDIMFile) && FactoryCanImport(UDIMFile))
 			{
-				const int32 UDIMIndex = ParseUDIMName(FPaths::GetBaseFilename(UDIMFile));
-				if (!UDIMIndexToFile.Contains(UDIMIndex))
+				FString UDIMName;
+				const int32 UDIMIndex = ParseUDIMName(FPaths::GetBaseFilename(UDIMFile), UDIMName);
+				if (!UDIMIndexToFile.Contains(UDIMIndex) && UDIMName == BaseUDIMName)
 				{
 					UDIMIndexToFile.Add(UDIMIndex, Path / UDIMFile);
 				}
@@ -4218,7 +4225,9 @@ UObject* UTextureFactory::FactoryCreateBinary
 		// however for a new texture platform data may not be generated yet, and for an reimport of a texture this is the size of the
 		// old texture. 
 		// Using source size gives one small caveat. It looks at the size before mipmap power of two padding adjustment.
-		if (Texture2D->Source.GetSizeX()*Texture2D->Source.GetSizeY() >= virtualTextureAutoEnableThresholdPixels)
+		// Textures with more than 1 block (UDIM textures) must be imported as VT
+		if (Texture->Source.GetNumBlocks() > 1 ||
+			Texture2D->Source.GetSizeX()*Texture2D->Source.GetSizeY() >= virtualTextureAutoEnableThresholdPixels)
 		{
 			Texture2D->VirtualTextureStreaming = true;
 		}
diff --git a/Engine/Source/Runtime/Engine/Private/VT/VirtualTextureDataBuilder.cpp b/Engine/Source/Runtime/Engine/Private/VT/VirtualTextureDataBuilder.cpp
--- a/Engine/Source/Runtime/Engine/Private/VT/VirtualTextureDataBuilder.cpp
+++ b/Engine/Source/Runtime/Engine/Private/VT/VirtualTextureDataBuilder.cpp
@@ -877,15 +877,27 @@ void FVirtualTextureDataBuilder::BuildSourcePixels(const FTextureSourceData& Sou
 				TBSettings.MipGenSettings = TMGS_SimpleAverage;
 			}
 
+			// For multi-block images, we may have scaled the max block size to be tile-sized, but individual blocks may still be smaller than 1 tile
+			// These need to be scaled up as well (scaling up individual blocks has the effect of reducing the block's mip-bias)
+			int32 LocalBlockSizeScale = BlockSizeScale;
+			while (SourceMips[0].SizeX * LocalBlockSizeScale < TileSize || SourceMips[0].SizeY * LocalBlockSizeScale < TileSize)
+			{
+				check(BlockData.MipBias > 0u);
+				--BlockData.MipBias;
+				LocalBlockSizeScale *= 2;
+			}
+
 			// Use the texture compressor module to do all the hard work
 			TArray<FCompressedImage2D> CompressedMips;
 			bool bBuildTextureResult = false;
-			if (BlockSizeScale == 1)
+			if (LocalBlockSizeScale == 1)
 			{
 				bBuildTextureResult = Compressor->BuildTexture(SourceMips, *CompositeSourceMips, TBSettings, CompressedMips);
 			}
 			else
 			{
+				// Need to generate scaled source images before building mips
+				// Typically this is only needed to scale very small source images to be at least tile-sized, so performance shouldn't be a big concern here
 				TArray<FImage> ScaledSourceMips;
 				TArray<FImage> ScaledCompositeMips;
 				ScaledSourceMips.Reserve(SourceMips.Num());
@@ -893,13 +905,13 @@ void FVirtualTextureDataBuilder::BuildSourcePixels(const FTextureSourceData& Sou
 				for (const FImage& SrcMip : SourceMips)
 				{
 					FImage* ScaledMip = new(ScaledSourceMips) FImage;
-					SrcMip.ResizeTo(*ScaledMip, SrcMip.SizeX * BlockSizeScale, SrcMip.SizeY * BlockSizeScale, SrcMip.Format, SrcMip.GammaSpace);
+					SrcMip.ResizeTo(*ScaledMip, SrcMip.SizeX * LocalBlockSizeScale, SrcMip.SizeY * LocalBlockSizeScale, SrcMip.Format, SrcMip.GammaSpace);
 				}
 
 				for (const FImage& SrcMip : *CompositeSourceMips)
 				{
 					FImage* ScaledMip = new(ScaledCompositeMips) FImage;
-					SrcMip.ResizeTo(*ScaledMip, SrcMip.SizeX * BlockSizeScale, SrcMip.SizeY * BlockSizeScale, SrcMip.Format, SrcMip.GammaSpace);
+					SrcMip.ResizeTo(*ScaledMip, SrcMip.SizeX * LocalBlockSizeScale, SrcMip.SizeY * LocalBlockSizeScale, SrcMip.Format, SrcMip.GammaSpace);
 				}
 
 				bBuildTextureResult = Compressor->BuildTexture(ScaledSourceMips, ScaledCompositeMips, TBSettings, CompressedMips);
