Don't include UDIM number in name of texture asset imported as UDIM

diff --git a/Engine/Source/Editor/UnrealEd/Private/Factories/EditorFactories.cpp b/Engine/Source/Editor/UnrealEd/Private/Factories/EditorFactories.cpp
--- a/Engine/Source/Editor/UnrealEd/Private/Factories/EditorFactories.cpp
+++ b/Engine/Source/Editor/UnrealEd/Private/Factories/EditorFactories.cpp
@@ -3931,11 +3931,51 @@ UObject* UTextureFactory::FactoryCreateBinary
 {
 	check(Type);
 
-	GEditor->GetEditorSubsystem<UImportSubsystem>()->BroadcastAssetPreImport(this, Class, InParent, Name, Type);
+	FName TextureName = Name;
+
+	// Check to see if we should import a series of textures as UDIM
+	// Need to do this first, as this step affects the final name of the created texture asset
+	TMap<int32, FString> UDIMIndexToFile;
+	{
+		const FString FilenameNoExtension = FPaths::GetBaseFilename(CurrentFilename);
+		FString BaseUDIMName;
+		const int32 BaseUDIMIndex = ParseUDIMName(FilenameNoExtension, BaseUDIMName);
+		if (BaseUDIMIndex != INDEX_NONE)
+		{
+			UDIMIndexToFile.Add(BaseUDIMIndex, CurrentFilename);
+
+			const FString Path = FPaths::GetPath(CurrentFilename);
+			const FString UDIMFilter = (Path / BaseUDIMName) + TEXT("*");
+
+			TArray<FString> UDIMFiles;
+			IFileManager::Get().FindFiles(UDIMFiles, *UDIMFilter, true, false);
+
+			for (const FString& UDIMFile : UDIMFiles)
+			{
+				if (!CurrentFilename.EndsWith(UDIMFile) && FactoryCanImport(UDIMFile))
+				{
+					FString UDIMName;
+					const int32 UDIMIndex = ParseUDIMName(FPaths::GetBaseFilename(UDIMFile), UDIMName);
+					if (!UDIMIndexToFile.Contains(UDIMIndex) && UDIMName == BaseUDIMName)
+					{
+						UDIMIndexToFile.Add(UDIMIndex, Path / UDIMFile);
+					}
+				}
+			}
+			if (UDIMIndexToFile.Num() > 1)
+			{
+				// Found multiple UDIM pages, so import as UDIM texture
+				// Exclude UDIM number from the name of the UE4 texture asset we create
+				TextureName = *BaseUDIMName;
+			}
+		}
+	}
+
+	GEditor->GetEditorSubsystem<UImportSubsystem>()->BroadcastAssetPreImport(this, Class, InParent, TextureName, Type);
 
 	// if the texture already exists, remember the user settings
-	UTexture* ExistingTexture = FindObject<UTexture>( InParent, *Name.ToString() );
-	UTexture2D* ExistingTexture2D = FindObject<UTexture2D>( InParent, *Name.ToString() );
+	UTexture* ExistingTexture = FindObject<UTexture>( InParent, *TextureName.ToString() );
+	UTexture2D* ExistingTexture2D = FindObject<UTexture2D>( InParent, *TextureName.ToString() );
 
 	TextureAddress						ExistingAddressX	= TA_Wrap;
 	TextureAddress						ExistingAddressY	= TA_Wrap;
@@ -3970,7 +4010,7 @@ UObject* UTextureFactory::FactoryCreateBinary
 	{
 		DisplayOverwriteOptionsDialog(FText::Format(
 			NSLOCTEXT("TextureFactory", "ImportOverwriteWarning", "You are about to import '{0}' over an existing texture."),
-			FText::FromName(Name)));
+			FText::FromName(TextureName)));
 
 		switch( OverwriteYesOrNoToAllState )
 		{
@@ -4044,45 +4084,15 @@ UObject* UTextureFactory::FactoryCreateBinary
 	FTextureReferenceReplacer RefReplacer(ExistingTexture);
 
 	UTexture* Texture = nullptr;
-
-	// Check to see if we should import a series of textures as UDIM
-	const FString FilenameNoExtension = FPaths::GetBaseFilename(CurrentFilename);
-	FString BaseUDIMName;
-	const int32 BaseUDIMIndex = ParseUDIMName(FilenameNoExtension, BaseUDIMName);
-	if (BaseUDIMIndex != INDEX_NONE)
+	if (UDIMIndexToFile.Num() > 1)
 	{
-		TMap<int32, FString> UDIMIndexToFile;
-		UDIMIndexToFile.Add(BaseUDIMIndex, CurrentFilename);
-
-		const FString Path = FPaths::GetPath(CurrentFilename);
-		const FString UDIMFilter = (Path / BaseUDIMName) + TEXT("*");
-
-		TArray<FString> UDIMFiles;
-		IFileManager::Get().FindFiles(UDIMFiles, *UDIMFilter, true, false);
-
-		for (const FString& UDIMFile : UDIMFiles)
-		{
-			if (!CurrentFilename.EndsWith(UDIMFile) && FactoryCanImport(UDIMFile))
-			{
-				FString UDIMName;
-				const int32 UDIMIndex = ParseUDIMName(FPaths::GetBaseFilename(UDIMFile), UDIMName);
-				if (!UDIMIndexToFile.Contains(UDIMIndex) && UDIMName == BaseUDIMName)
-				{
-					UDIMIndexToFile.Add(UDIMIndex, Path / UDIMFile);
-				}
-			}
-		}
-
-		if (UDIMIndexToFile.Num() > 1)
-		{
-			Texture = ImportTextureUDIM(Class, InParent, Name, Flags, Type, UDIMIndexToFile, Warn);
-		}
+		// Import UDIM texture
+		Texture = ImportTextureUDIM(Class, InParent, TextureName, Flags, Type, UDIMIndexToFile, Warn);
 	}
-
-	if (!Texture)
+	else
 	{
 		// Not a UDIM, import a regular texture
-		Texture = ImportTexture(Class, InParent, Name, Flags, Type, Buffer, BufferEnd, Warn);
+		Texture = ImportTexture(Class, InParent, TextureName, Flags, Type, Buffer, BufferEnd, Warn);
 		if (Texture)
 		{
 			Texture->AssetImportData->Update(CurrentFilename, FileHash.IsValid() ? &FileHash : nullptr);
@@ -4266,7 +4276,7 @@ UObject* UTextureFactory::FactoryCreateBinary
 	if( bCreateMaterial )
 	{
 		// Create the package for the material
-		const FString MaterialName = FString::Printf( TEXT("%s_Mat"), *Name.ToString() );
+		const FString MaterialName = FString::Printf( TEXT("%s_Mat"), *TextureName.ToString() );
 		const FString MaterialPackageName = FPackageName::GetLongPackagePath(InParent->GetName()) + TEXT("/") + MaterialName;
 		UPackage* MaterialPackage = CreatePackage(nullptr, *MaterialPackageName);
 
