- Ensure UDIM texture's package is named to match texture object (need to remove UDIM number from package name)
- Without this materials created from texture will be named based on package
- May cause other inconsistencies

diff --git a/Engine/Source/Editor/UnrealEd/Private/Factories/EditorFactories.cpp b/Engine/Source/Editor/UnrealEd/Private/Factories/EditorFactories.cpp
--- a/Engine/Source/Editor/UnrealEd/Private/Factories/EditorFactories.cpp
+++ b/Engine/Source/Editor/UnrealEd/Private/Factories/EditorFactories.cpp
@@ -3967,6 +3967,18 @@ UObject* UTextureFactory::FactoryCreateBinary
 				// Found multiple UDIM pages, so import as UDIM texture
 				// Exclude UDIM number from the name of the UE4 texture asset we create
 				TextureName = *BaseUDIMName;
+
+				// Need to rename the package to match the new texture name, since package was already created
+				// Package name will be the same as the object name, except will contain additional path information,
+				// so we take the existing package name, then extract the UDIM index in order to preserve the path
+				FString PackageName;
+				InParent->GetName(PackageName);
+
+				FString PackageUDIMName;
+				const int32 PackageUDIMIndex = ParseUDIMName(PackageName, PackageUDIMName);
+				check(PackageUDIMIndex == BaseUDIMIndex);
+				check(PackageUDIMName.EndsWith(BaseUDIMName, ESearchCase::CaseSensitive));
+				verify(InParent->Rename(*PackageUDIMName, nullptr, REN_DontCreateRedirectors));
 			}
 		}
 	}
