Check for existing package when renaming UDIM texture package

diff --git a/Engine/Source/Editor/UnrealEd/Private/Factories/EditorFactories.cpp b/Engine/Source/Editor/UnrealEd/Private/Factories/EditorFactories.cpp
--- a/Engine/Source/Editor/UnrealEd/Private/Factories/EditorFactories.cpp
+++ b/Engine/Source/Editor/UnrealEd/Private/Factories/EditorFactories.cpp
@@ -3978,7 +3978,19 @@ UObject* UTextureFactory::FactoryCreateBinary
 				const int32 PackageUDIMIndex = ParseUDIMName(PackageName, PackageUDIMName);
 				check(PackageUDIMIndex == BaseUDIMIndex);
 				check(PackageUDIMName.EndsWith(BaseUDIMName, ESearchCase::CaseSensitive));
-				verify(InParent->Rename(*PackageUDIMName, nullptr, REN_DontCreateRedirectors));
+
+				// In normal case, higher level code would have already checked for duplicate package name
+				// But since we're changing package name here, check to see if package with the new name already exists...
+				// If it does, code later in this method will prompt user to overwrite the existing asset
+				UPackage* ExistingPackage = FindPackage(InParent->GetOuter(), *PackageUDIMName);
+				if (ExistingPackage)
+				{
+					InParent = ExistingPackage;
+				}
+				else
+				{
+					verify(InParent->Rename(*PackageUDIMName, nullptr, REN_DontCreateRedirectors));
+				}
 			}
 		}
 	}
