Fixed post merge compilation issues

diff --git a/Engine/Plugins/Developer/Concert/MultiUserClient/Source/MultiUserClient/Private/Widgets/SConcertBrowser.cpp b/Engine/Plugins/Developer/Concert/MultiUserClient/Source/MultiUserClient/Private/Widgets/SConcertBrowser.cpp
--- a/Engine/Plugins/Developer/Concert/MultiUserClient/Source/MultiUserClient/Private/Widgets/SConcertBrowser.cpp
+++ b/Engine/Plugins/Developer/Concert/MultiUserClient/Source/MultiUserClient/Private/Widgets/SConcertBrowser.cpp
@@ -1851,7 +1851,6 @@ void SConcertSessionBrowser::Construct(const FArguments& InArgs, IConcertClientP
 	{
 		return; // Don't build the UI if ConcertClient is not available.
 	}
-		ResumeSessionDef.OnClicked.BindSP(const_cast<SConcertBrowser*>(this), &SConcertBrowser::OnClickResumeSession);
 
 	Controller = MakeUnique<FConcertBrowserController>(InConcertClient);
 
diff --git a/Engine/Source/Editor/UnrealEd/Private/Factories/EditorFactories.cpp b/Engine/Source/Editor/UnrealEd/Private/Factories/EditorFactories.cpp
--- a/Engine/Source/Editor/UnrealEd/Private/Factories/EditorFactories.cpp
+++ b/Engine/Source/Editor/UnrealEd/Private/Factories/EditorFactories.cpp
@@ -3462,14 +3462,16 @@ bool UTextureFactory::ImportImage(const uint8* Buffer, uint32 Length, FFeedbackC
 				NewV,
 				TSF_BGRA8
 			);
-			const int32 DestSize = Texture->Source.CalcMipSize(0);
+
+			uint8* Dest = OutImage.RawData.GetData();
+
 			// Doing a memset to make sure the alpha channel is set to 0xff since we only have 3 color planes.
-			FMemory::Memset(Dest, 0xff, DestSize);
+			FMemory::Memset(Dest, 0xff, NewU * NewV * sizeof(uint8));
 
 			// Copy upside-down scanlines.
 			Buffer += 128;
 			int32 CountU = FMath::Min<int32>(PCX->BytesPerLine, NewU);
-			for (int32 IndexU = 0; IndexU < CountU; IndexU++)
+			for (int32 i = 0; i < NewV; i++)
 			{
 				// We need to decode image one line per time building RGB image color plane by color plane.
 				int32 RunLength, Overflow = 0;
@@ -3483,7 +3485,7 @@ bool UTextureFactory::ImportImage(const uint8* Buffer, uint32 Length, FFeedbackC
 							Color = *Buffer++;
 							if ((Color & 0xc0) == 0xc0)
 							{
-								RunLength = FMath::Min((Color & 0x3f), CountU - IndexU);
+								RunLength = FMath::Min((Color & 0x3f), CountU - j);
 								Overflow = (Color & 0x3f) - RunLength;
 								Color = *Buffer++;
 							}
@@ -3492,17 +3494,17 @@ bool UTextureFactory::ImportImage(const uint8* Buffer, uint32 Length, FFeedbackC
 						}
 						else
 						{
-							RunLength = FMath::Min(Overflow, CountU - IndexU);
+							RunLength = FMath::Min(Overflow, CountU - j);
 							Overflow = Overflow - RunLength;
 						}
 
-						//checkf(((IndexV*NewU + RunLength) * 4 + ColorPlane) < (Texture->Source.CalcMipSize(0)),
+						//checkf(((i*NewU + RunLength) * 4 + ColorPlane) < (Texture->Source.CalcMipSize(0)),
 						//	TEXT("RLE going off the end of buffer"));
-						for (int32 k = IndexU; k < IndexU + RunLength; k++)
+						for (int32 k = j; k < j + RunLength; k++)
 						{
-							Dest[(IndexV*NewU + k) * 4 + ColorPlane] = Color;
+							Dest[(i*NewU + k) * 4 + ColorPlane] = Color;
 						}
-						IndexU += RunLength - 1;
+						j += RunLength - 1;
 					}
 				}
 			}
diff --git a/Engine/Source/Runtime/Renderer/Private/PostProcess/PostProcessMotionBlur.cpp b/Engine/Source/Runtime/Renderer/Private/PostProcess/PostProcessMotionBlur.cpp
--- a/Engine/Source/Runtime/Renderer/Private/PostProcess/PostProcessMotionBlur.cpp
+++ b/Engine/Source/Runtime/Renderer/Private/PostProcess/PostProcessMotionBlur.cpp
@@ -127,7 +127,7 @@ bool IsMotionBlurScatterRequired(const FViewInfo& View, const FScreenPassTexture
 	const bool bIsScatterRequiredByVelocityLength = VelocityMaxInTiles > TileDistanceMaxGathered;
 
 	// Cinematic is paused.
-	const bool bInPausedCinematic = (ViewState && ViewState->bSequencerIsPaused);
+	const bool bInPausedCinematic = (ViewState && ViewState->SequencerState == ESS_Paused);
 
 	// Use the scatter approach if requested by cvar or we're in a paused cinematic (higher quality).
 	const bool bIsScatterRequiredByUser = CVarMotionBlurScatter.GetValueOnRenderThread() == 1 || bInPausedCinematic;
diff --git a/Engine/Source/Runtime/Renderer/Private/SceneCaptureRendering.cpp b/Engine/Source/Runtime/Renderer/Private/SceneCaptureRendering.cpp
--- a/Engine/Source/Runtime/Renderer/Private/SceneCaptureRendering.cpp
+++ b/Engine/Source/Runtime/Renderer/Private/SceneCaptureRendering.cpp
@@ -888,6 +888,7 @@ void FScene::UpdateSceneCaptureContents(USceneCaptureComponentCube* CaptureCompo
 			for (int32 faceidx = 0; faceidx < (int32)ECubeFace::CubeFace_MAX; faceidx++)
 			{
 				const ECubeFace TargetFace = (ECubeFace)faceidx;
+				const FVector Location = CaptureComponent->GetComponentToWorld().GetTranslation();
 
 				FMatrix ViewRotationMatrix;
 
