Fixing post-merge PCX Image alpha

diff --git a/Engine/Source/Editor/UnrealEd/Private/Factories/EditorFactories.cpp b/Engine/Source/Editor/UnrealEd/Private/Factories/EditorFactories.cpp
--- a/Engine/Source/Editor/UnrealEd/Private/Factories/EditorFactories.cpp
+++ b/Engine/Source/Editor/UnrealEd/Private/Factories/EditorFactories.cpp
@@ -3466,7 +3466,7 @@ bool UTextureFactory::ImportImage(const uint8* Buffer, uint32 Length, FFeedbackC
 			uint8* Dest = OutImage.RawData.GetData();
 
 			// Doing a memset to make sure the alpha channel is set to 0xff since we only have 3 color planes.
-			FMemory::Memset(Dest, 0xff, NewU * NewV * sizeof(uint8));
+			FMemory::Memset(Dest, 0xff, NewU * NewV * FTextureSource::GetBytesPerPixel(OutImage.Format));
 
 			// Copy upside-down scanlines.
 			Buffer += 128;
