- Respect SRBG/TextureCompression settings for UDIM textures
- Require all UDIM pages to use the same file extension (not strictly required, but seems cleaner, avoids some potential of mismatched formats)

diff --git a/Engine/Source/Editor/UnrealEd/Private/Factories/EditorFactories.cpp b/Engine/Source/Editor/UnrealEd/Private/Factories/EditorFactories.cpp
--- a/Engine/Source/Editor/UnrealEd/Private/Factories/EditorFactories.cpp
+++ b/Engine/Source/Editor/UnrealEd/Private/Factories/EditorFactories.cpp
@@ -3662,6 +3662,8 @@ UTexture* UTextureFactory::ImportTextureUDIM(UClass* Class, UObject* InParent, F
 	SourceFileNames.Reserve(UDIMIndexToFile.Num());
 
 	ETextureSourceFormat Format = TSF_Invalid;
+	TextureCompressionSettings TCSettings = TC_MAX;
+	bool bSRGB = false;
 	for (const auto& It : UDIMIndexToFile)
 	{
 		const FString& TexturePath = It.Value;
@@ -3676,10 +3678,17 @@ UTexture* UTextureFactory::ImportTextureUDIM(UClass* Class, UObject* InParent, F
 				if (Format == TSF_Invalid)
 				{
 					Format = Image.Format;
+					bSRGB = Image.SRGB;
+				}
+
+				if (TCSettings == TC_MAX)
+				{
+					// Should we somehow try to combine different compression settings? Is that ever useful/needed?
+					TCSettings = Image.CompressionSettings;
 				}
 
 				// Deal with mismatched formats somehow?  convert?
-				if (ensure(Format == Image.Format))
+				if (ensure(Format == Image.Format && bSRGB == Image.SRGB))
 				{
 					const int32 UDIMIndex = It.Key;
 					FTextureSourceBlock* Block = new(SourceBlocks) FTextureSourceBlock();
@@ -3695,7 +3704,7 @@ UTexture* UTextureFactory::ImportTextureUDIM(UClass* Class, UObject* InParent, F
 				}
 				else
 				{
-					Warn->Logf(ELogVerbosity::Warning, TEXT("Mismatched UDIM image formats."));
+					Warn->Logf(ELogVerbosity::Warning, TEXT("Mismatched UDIM image formats, skipping file \"%s\""), *TexturePath);
 				}
 			}
 		}
@@ -3715,6 +3724,8 @@ UTexture* UTextureFactory::ImportTextureUDIM(UClass* Class, UObject* InParent, F
 
 	UTexture2D* Texture = CreateTexture2D(InParent, Name, Flags);
 	Texture->Source.InitBlocked(&Format, SourceBlocks.GetData(), 1, SourceBlocks.Num(), SourceImageData.GetData());
+	Texture->CompressionSettings = TCSettings;
+	Texture->SRGB = bSRGB;
 
 	for (int32 FileIndex = 0; FileIndex < SourceFileNames.Num(); ++FileIndex)
 	{
@@ -3965,8 +3976,9 @@ UObject* UTextureFactory::FactoryCreateBinary
 		{
 			UDIMIndexToFile.Add(BaseUDIMIndex, CurrentFilename);
 
+			// Filter for other potential UDIM pages, with the same base name and file extension
 			const FString Path = FPaths::GetPath(CurrentFilename);
-			const FString UDIMFilter = (Path / BaseUDIMName) + TEXT("*");
+			const FString UDIMFilter = (Path / BaseUDIMName) + TEXT("*") + FPaths::GetExtension(CurrentFilename, true);
 
 			TArray<FString> UDIMFiles;
 			IFileManager::Get().FindFiles(UDIMFiles, *UDIMFilter, true, false);
