Fixed crash when selecting Landscape mode and then selecting Foliage mode. Landscape mode automatically creates an actor and selects it. Switching to Foliage mode destroys the actor, but an invalid weak reference is still in the selected actors list. Just checking Num() is not enough, as GetSelectedObject has additional clauses that return nullptr.

The fix uses GetBottom which is a bit more robust, and also uses Cast instead of CastChecked so nullptr is accepted and skipped.

It might be possible to skip this code when not in realtime mode (since a lot of selected actors might be a performance issue?), but I'm unsure of that.

diff --git a/Engine/Source/Editor/UnrealEd/Private/EditorModeManager.cpp b/Engine/Source/Editor/UnrealEd/Private/EditorModeManager.cpp
--- a/Engine/Source/Editor/UnrealEd/Private/EditorModeManager.cpp
+++ b/Engine/Source/Editor/UnrealEd/Private/EditorModeManager.cpp
@@ -166,11 +166,9 @@ void FEditorModeTools::UpdateModeWidgetLocation()
 	double starttime = FPlatformTime::Seconds();
 
 	USelection* SelectedActors = GetSelectedActors();
-
-	if (SelectedActors && SelectedActors->Num() > 0)
+	if (SelectedActors)
 	{
-		AActor* Actor = CastChecked<AActor>(SelectedActors->GetSelectedObject(SelectedActors->Num() - 1));
-
+		AActor* Actor = Cast<AActor>(SelectedActors->GetBottom(AActor::StaticClass()));
 		if (Actor)
 		{
 			FVector CurrentActorLocation = Actor->GetActorLocation();
