Fixes merge conflict in EditorModeManager.cpp

diff --git a/Engine/Source/Editor/UnrealEd/Private/EditorModeManager.cpp b/Engine/Source/Editor/UnrealEd/Private/EditorModeManager.cpp
--- a/Engine/Source/Editor/UnrealEd/Private/EditorModeManager.cpp
+++ b/Engine/Source/Editor/UnrealEd/Private/EditorModeManager.cpp
@@ -168,7 +168,7 @@ void FEditorModeTools::UpdateModeWidgetLocation()
 	USelection* SelectedActors = GetSelectedActors();
 	if (SelectedActors)
 	{
-		AActor* Actor = CastChecked<AActor>(SelectedActors->GetBottom(AActor::StaticClass()));
+		AActor* Actor = Cast<AActor>(SelectedActors->GetBottom(AActor::StaticClass()));
 		if (Actor)
 		{
 			FVector CurrentActorLocation = Actor->GetActorLocation();
