Fixed Linux compilation warning

diff --git a/Engine/Source/Editor/UnrealEd/Private/Lightmass/Lightmass.cpp b/Engine/Source/Editor/UnrealEd/Private/Lightmass/Lightmass.cpp
--- a/Engine/Source/Editor/UnrealEd/Private/Lightmass/Lightmass.cpp
+++ b/Engine/Source/Editor/UnrealEd/Private/Lightmass/Lightmass.cpp
@@ -1038,7 +1038,7 @@ void FLightmassExporter::WriteLights( int32 Channel )
 	}
 
 	// Compute a mapping between directional light and trnasmittance to apply. For each AtmosphereSunLightIndex, the brightest lights is kept.
-	if (AtmosphericFogComponent && AtmosphericFogComponent->bAtmosphereAffectsSunIlluminance || SkyAtmosphereComponent)
+	if ((AtmosphericFogComponent && AtmosphericFogComponent->bAtmosphereAffectsSunIlluminance) || SkyAtmosphereComponent)
 	{
 		for (int32 LightIndex = 0; LightIndex < DirectionalLights.Num(); ++LightIndex)
 		{
