Fix B8G8R8A8_UNORM (PF_B8G8R8A8) UAVs not supported by DX 11.0. It is now replaced with PF_R8G8B8A8. Legacy data will be converted when loaded.

diff --git a/Engine/Source/Editor/UnrealEd/Private/Lightmass/ImportVolumetricLightmap.cpp b/Engine/Source/Editor/UnrealEd/Private/Lightmass/ImportVolumetricLightmap.cpp
--- a/Engine/Source/Editor/UnrealEd/Private/Lightmass/ImportVolumetricLightmap.cpp
+++ b/Engine/Source/Editor/UnrealEd/Private/Lightmass/ImportVolumetricLightmap.cpp
@@ -978,6 +978,29 @@ void BuildIndirectionTexture(
 	}
 }
 
+inline void ConvertBGRA8ToRGBA8ForLayer(FVolumetricLightmapDataLayer& Layer)
+{
+	if (Layer.Format == PF_B8G8R8A8)
+	{
+		for (int32 PixelIndex = 0; PixelIndex < Layer.Data.Num() / GPixelFormats[PF_B8G8R8A8].BlockBytes; PixelIndex++)
+		{
+			FColor Color;
+
+			Color.B = Layer.Data[PixelIndex * 4 + 0];
+			Color.G = Layer.Data[PixelIndex * 4 + 1];
+			Color.R = Layer.Data[PixelIndex * 4 + 2];
+			Color.A = Layer.Data[PixelIndex * 4 + 3];
+
+			Layer.Data[PixelIndex * 4 + 0] = Color.R;
+			Layer.Data[PixelIndex * 4 + 1] = Color.G;
+			Layer.Data[PixelIndex * 4 + 2] = Color.B;
+			Layer.Data[PixelIndex * 4 + 3] = Color.A;
+		}
+
+		Layer.Format = PF_R8G8B8A8;
+	}
+}
+
 // For debugging
 bool bStitchDetailBricksWithLowDensityNeighbors = true;
 bool bCopyPaddingFromUniqueData = true;
@@ -1027,6 +1050,8 @@ void FLightmassProcessor::ImportVolumetricLightmap()
 	{
 		CurrentLevelData.InitializeOnImport(FBox(VolumetricLightmapSettings.VolumeMin, VolumetricLightmapSettings.VolumeMin + VolumetricLightmapSettings.VolumeSize), BrickSize);
 
+		// Temporarily assign format as PF_B8G8R8A8 since that matches FColor and makes our operations easier
+		// Will be converted to PF_R8G8B8A8 by ConvertBGRA8ToRGBA8ForLayer() later
 		CurrentLevelData.BrickData.AmbientVector.Format = PF_FloatR11G11B10;
 		CurrentLevelData.BrickData.SkyBentNormal.Format = PF_B8G8R8A8;
 		CurrentLevelData.BrickData.DirectionalLightShadowing.Format = PF_G8;
@@ -1319,6 +1344,8 @@ void FLightmassProcessor::ImportVolumetricLightmap()
 
 				SubLevelData.InitializeOnImport(FBox(VolumetricLightmapSettings.VolumeMin, VolumetricLightmapSettings.VolumeMin + VolumetricLightmapSettings.VolumeSize), BrickSize);
 				{
+					// Temporarily assign format as PF_B8G8R8A8 since that matches FColor and makes our operations easier
+					// Will be converted to PF_R8G8B8A8 by ConvertBGRA8ToRGBA8ForLayer() later
 					SubLevelData.BrickData.AmbientVector.Format = PF_FloatR11G11B10;
 					SubLevelData.BrickData.SkyBentNormal.Format = PF_B8G8R8A8;
 					SubLevelData.BrickData.DirectionalLightShadowing.Format = PF_G8;
@@ -1515,6 +1542,26 @@ void FLightmassProcessor::ImportVolumetricLightmap()
 			}
 
 		}
+
+		{
+			for (auto Pair : SubLevelTotalBricks)
+			{
+				FGuid LevelGuid = Pair.Key;
+
+				ULevel* SubLevelStorageLevel = System.LightingScenario ? System.LightingScenario : FindLevel(LevelGuid);
+				UMapBuildDataRegistry* SubLevelRegistry = SubLevelStorageLevel->GetOrCreateMapBuildData();
+				FPrecomputedVolumetricLightmapData& SubLevelData = *SubLevelRegistry->GetLevelPrecomputedVolumetricLightmapBuildData(FindLevel(LevelGuid)->LevelBuildDataId);
+
+				ConvertBGRA8ToRGBA8ForLayer(SubLevelData.BrickData.SkyBentNormal);
+
+				for (int32 i = 0; i < ARRAY_COUNT(SubLevelData.BrickData.SHCoefficients); i++)
+				{
+					ConvertBGRA8ToRGBA8ForLayer(SubLevelData.BrickData.SHCoefficients[i]);
+				}
+
+				ConvertBGRA8ToRGBA8ForLayer(SubLevelData.BrickData.LQLightDirection);
+			}
+		}
 	}
 
 	{
diff --git a/Engine/Source/Runtime/Engine/Private/PrecomputedVolumetricLightmap.cpp b/Engine/Source/Runtime/Engine/Private/PrecomputedVolumetricLightmap.cpp
--- a/Engine/Source/Runtime/Engine/Private/PrecomputedVolumetricLightmap.cpp
+++ b/Engine/Source/Runtime/Engine/Private/PrecomputedVolumetricLightmap.cpp
@@ -147,6 +147,29 @@ private:
 
 TGlobalResource<FVolumetricLightmapBrickAtlas> GVolumetricLightmapBrickAtlas = TGlobalResource<FVolumetricLightmapBrickAtlas>();
 
+inline void ConvertBGRA8ToRGBA8ForLayer(FVolumetricLightmapDataLayer& Layer)
+{
+	if (Layer.Format == PF_B8G8R8A8)
+	{
+		for (int32 PixelIndex = 0; PixelIndex < Layer.Data.Num() / GPixelFormats[PF_B8G8R8A8].BlockBytes; PixelIndex++)
+		{
+			FColor Color;
+
+			Color.B = Layer.Data[PixelIndex * 4 + 0];
+			Color.G = Layer.Data[PixelIndex * 4 + 1];
+			Color.R = Layer.Data[PixelIndex * 4 + 2];
+			Color.A = Layer.Data[PixelIndex * 4 + 3];
+
+			Layer.Data[PixelIndex * 4 + 0] = Color.R;
+			Layer.Data[PixelIndex * 4 + 1] = Color.G;
+			Layer.Data[PixelIndex * 4 + 2] = Color.B;
+			Layer.Data[PixelIndex * 4 + 3] = Color.A;
+		}
+
+		Layer.Format = PF_R8G8B8A8;
+	}
+}
+
 FArchive& operator<<(FArchive& Ar,FVolumetricLightmapDataLayer& Layer)
 {
 	Ar << Layer.Data;
@@ -163,6 +186,8 @@ FArchive& operator<<(FArchive& Ar,FVolumetricLightmapDataLayer& Layer)
 		FString PixelFormatString;
 		Ar << PixelFormatString;
 		Layer.Format = (EPixelFormat)PixelFormatEnum->GetValueByName(*PixelFormatString);
+
+		ConvertBGRA8ToRGBA8ForLayer(Layer);
 	}
 	else if (Ar.IsSaving())
 	{
