Compile fix for TextureArray CL

diff --git a/Engine/Source/Editor/UnrealEd/Classes/Factories/Texture2DArrayFactory.h b/Engine/Source/Editor/UnrealEd/Classes/Factories/Texture2DArrayFactory.h
--- a/Engine/Source/Editor/UnrealEd/Classes/Factories/Texture2DArrayFactory.h
+++ b/Engine/Source/Editor/UnrealEd/Classes/Factories/Texture2DArrayFactory.h
@@ -7,6 +7,8 @@
 #include "Factories/Factory.h"
 #include "Texture2DArrayFactory.generated.h"
 
+class UTexture2D;
+
 /** Factory for creating volume texture */
 UCLASS(hidecategories=Object, MinimalAPI)
 class UTexture2DArrayFactory : public UFactory
