Case sensitive include compile error on linux for texture array CL.

diff --git a/Engine/Source/Editor/UnrealEd/Classes/Factories/Texture2dArrayThumbnailRenderer.h b/Engine/Source/Editor/UnrealEd/Classes/Factories/Texture2dArrayThumbnailRenderer.h
--- a/Engine/Source/Editor/UnrealEd/Classes/Factories/Texture2dArrayThumbnailRenderer.h
+++ b/Engine/Source/Editor/UnrealEd/Classes/Factories/Texture2dArrayThumbnailRenderer.h
@@ -10,7 +10,7 @@
 #include "CoreMinimal.h"
 #include "UObject/ObjectMacros.h"
 #include "ThumbnailRendering/TextureThumbnailRenderer.h"
-#include "Texture2DArrayThumbnailRenderer.generated.h"
+#include "Texture2dArrayThumbnailRenderer.generated.h"
 
 UCLASS()
 class UTexture2DArrayThumbnailRenderer : public UTextureThumbnailRenderer
diff --git a/Engine/Source/Editor/UnrealEd/Private/ThumbnailRendering/Texture2DArrayThumbnailRenderer.cpp b/Engine/Source/Editor/UnrealEd/Private/ThumbnailRendering/Texture2DArrayThumbnailRenderer.cpp
--- a/Engine/Source/Editor/UnrealEd/Private/ThumbnailRendering/Texture2DArrayThumbnailRenderer.cpp
+++ b/Engine/Source/Editor/UnrealEd/Private/ThumbnailRendering/Texture2DArrayThumbnailRenderer.cpp
@@ -1,6 +1,6 @@
 // Copyright 1998-2019 Epic Games, Inc. All Rights Reserved.
 
-#include "Factories/Texture2DArrayThumbnailRenderer.h"
+#include "Factories/Texture2dArrayThumbnailRenderer.h"
 #include "Engine/Texture2DArray.h"
 #include "ThumbnailRendering/ThumbnailManager.h"
 
