DR - Fix static analysis

diff --git a/Engine/Source/Editor/UnrealEd/Private/Texture2DPreview.cpp b/Engine/Source/Editor/UnrealEd/Private/Texture2DPreview.cpp
--- a/Engine/Source/Editor/UnrealEd/Private/Texture2DPreview.cpp
+++ b/Engine/Source/Editor/UnrealEd/Private/Texture2DPreview.cpp
@@ -112,8 +112,8 @@ public:
 			SetShaderValue(RHICmdList, GetPixelShader(), NumSlices, NumSlicesData);
 		}
 
-		SetShaderValue(RHICmdList, GetPixelShader(),TextureComponentReplicate,TextureValue->bGreyScaleFormat ? FLinearColor(1,0,0,0) : FLinearColor(0,0,0,0));
-		SetShaderValue(RHICmdList, GetPixelShader(),TextureComponentReplicateAlpha,TextureValue->bGreyScaleFormat ? FLinearColor(1,0,0,0) : FLinearColor(0,0,0,1));
+		SetShaderValue(RHICmdList, GetPixelShader(),TextureComponentReplicate, (TextureValue && TextureValue->bGreyScaleFormat) ? FLinearColor(1,0,0,0) : FLinearColor(0,0,0,0));
+		SetShaderValue(RHICmdList, GetPixelShader(),TextureComponentReplicateAlpha, (TextureValue && TextureValue->bGreyScaleFormat) ? FLinearColor(1,0,0,0) : FLinearColor(0,0,0,1));
 	}
 
 	virtual bool Serialize(FArchive& Ar) override
