Merging //UE4/Main@9388041 to Dev-RenderPlat-Staging

diff --git a/Engine/Source/Editor/UnrealEd/Private/Fbx/FbxSkeletalMeshImport.cpp b/Engine/Source/Editor/UnrealEd/Private/Fbx/FbxSkeletalMeshImport.cpp
--- a/Engine/Source/Editor/UnrealEd/Private/Fbx/FbxSkeletalMeshImport.cpp
+++ b/Engine/Source/Editor/UnrealEd/Private/Fbx/FbxSkeletalMeshImport.cpp
@@ -27,6 +27,9 @@
 #include "Animation/MorphTarget.h"
 #include "PhysicsAssetUtils.h"
 #include "Rendering/SkeletalMeshModel.h"
+#include "Rendering/SkeletalMeshRenderData.h"
+#include "Rendering/SkeletalMeshLODRenderData.h"
+#include "Misc/ScopedSlowTask.h"
 
 #include "SkelImport.h"
 #include "Logging/TokenizedMessage.h"
@@ -52,8 +55,9 @@
 #include "Misc/FbxErrors.h"
 #include "PhysicsEngine/PhysicsAsset.h"
 #include "Engine/SkeletalMeshSocket.h"
-#include "Assets/ClothingAsset.h"
+#include "ClothingAsset.h"
 #include "LODUtilities.h"
+#include "IMeshBuilderModule.h"
 
 #define LOCTEXT_NAMESPACE "FBXImpoter"
 
@@ -1433,7 +1437,7 @@ bool UnFbx::FFbxImporter::FillSkeletalMeshImportData(TArray<FbxNode*>& NodeArray
 
 	if (ImportOptions->bPreserveSmoothingGroups)
 	{
-		bool bDuplicateUnSmoothWedges = (ImportOptions->NormalGenerationMethod != EFBXNormalGenerationMethod::MikkTSpace);
+		bool bDuplicateUnSmoothWedges = false; //We deprecate legacy build
 		DoUnSmoothVerts(*SkelMeshImportDataPtr, bDuplicateUnSmoothWedges);
 	}
 	else
@@ -1754,28 +1758,40 @@ USkeletalMesh* UnFbx::FFbxImporter::ImportSkeletalMesh(FImportSkeletalMeshArgs &
 		return nullptr;
 	}
 
+	//Stack the PostEditChange call, it will call post edit change when it will go out of scope
+	FScopedSkeletalMeshPostEditChange ScopedPostEditChange(ExistingSkelMesh);
+
+	ESkeletalMeshGeoImportVersions GeoImportVersion = ESkeletalMeshGeoImportVersions::LatestVersion;
+	ESkeletalMeshSkinningImportVersions SkinningImportVersion = ESkeletalMeshSkinningImportVersions::LatestVersion;
 	//Adjust the import data from the import options
 	if (ExistingSkelMesh != nullptr)
 	{
+		int32 SafeLODIndex = ImportSkeletalMeshArgs.LodIndex < 0 ? 0 : ImportSkeletalMeshArgs.LodIndex;
+		if (ExistingSkelMesh->GetImportedModel() && ExistingSkelMesh->GetImportedModel()->LODModels.IsValidIndex(SafeLODIndex))
+		{
+			const FRawSkeletalMeshBulkData& RawSkeletalMeshBulkData = ExistingSkelMesh->GetImportedModel()->LODModels[SafeLODIndex].RawSkeletalMeshBulkData;
+			GeoImportVersion = RawSkeletalMeshBulkData.GeoImportVersion;
+			SkinningImportVersion = RawSkeletalMeshBulkData.SkinningImportVersion;
+		}
+
 		if ((ImportOptions->bImportAsSkeletalSkinning || ImportOptions->bImportAsSkeletalGeometry))
 		{
 			if (ImportOptions->bImportAsSkeletalSkinning)
 			{
 				//Replace geometry import data by original existing skel mesh geometry data
-				ReplaceSkeletalMeshGeometryImportData(ExistingSkelMesh, SkelMeshImportDataPtr, ImportSkeletalMeshArgs.LodIndex);
+				ReplaceSkeletalMeshGeometryImportData(ExistingSkelMesh, SkelMeshImportDataPtr, SafeLODIndex);
 			}
 			else if (ImportOptions->bImportAsSkeletalGeometry)
 			{
 				//Replace skinning import data by original existing skel mesh skinning data
-				ReplaceSkeletalMeshSkinningImportData(ExistingSkelMesh, SkelMeshImportDataPtr, ImportSkeletalMeshArgs.LodIndex);
+				ReplaceSkeletalMeshSkinningImportData(ExistingSkelMesh, SkelMeshImportDataPtr, SafeLODIndex);
 			}
 		}
 
 		//Reuse the vertex color in case we are not re-importing weights only and there is some vertexcolor in the existing skeleMesh
 		if (ExistingSkelMesh->bHasVertexColors && ImportOptions->VertexColorImportOption == EVertexColorImportOption::Ignore)
 		{
-			int32 LODIndex = ImportSkeletalMeshArgs.LodIndex < 0 ? 0 : ImportSkeletalMeshArgs.LodIndex;
-			RemapSkeletalMeshVertexColorToImportData(ExistingSkelMesh, LODIndex, SkelMeshImportDataPtr);
+			RemapSkeletalMeshVertexColorToImportData(ExistingSkelMesh, SafeLODIndex, SkelMeshImportDataPtr);
 		}
 	}
 
@@ -1783,15 +1799,6 @@ USkeletalMesh* UnFbx::FFbxImporter::ImportSkeletalMesh(FImportSkeletalMeshArgs &
 
 	// Create initial bounding box based on expanded version of reference pose for meshes without physics assets. Can be overridden by artist.
 	FBox BoundingBox(SkelMeshImportDataPtr->Points.GetData(), SkelMeshImportDataPtr->Points.Num());
-	FBox Temp = BoundingBox;
-	FVector MidMesh = 0.5f*(Temp.Min + Temp.Max);
-	BoundingBox.Min = Temp.Min + 1.0f*(Temp.Min - MidMesh);
-	BoundingBox.Max = Temp.Max + 1.0f*(Temp.Max - MidMesh);
-	// Tuck up the bottom as this rarely extends lower than a reference pose's (e.g. having its feet on the floor).
-	// Maya has Y in the vertical, other packages have Z.
-	//BEN const int32 CoordToTuck = bAssumeMayaCoordinates ? 1 : 2;
-	//BEN BoundingBox.Min[CoordToTuck]	= Temp.Min[CoordToTuck] + 0.1f*(Temp.Min[CoordToTuck] - MidMesh[CoordToTuck]);
-	BoundingBox.Min[2] = Temp.Min[2] + 0.1f*(Temp.Min[2] - MidMesh[2]);
 	const FVector BoundingBoxSize = BoundingBox.GetSize();
 
 	if (SkelMeshImportDataPtr->Points.Num() > 2 && BoundingBoxSize.X < THRESH_POINTS_ARE_SAME && BoundingBoxSize.Y < THRESH_POINTS_ARE_SAME && BoundingBoxSize.Z < THRESH_POINTS_ARE_SAME)
@@ -1805,18 +1812,8 @@ USkeletalMesh* UnFbx::FFbxImporter::ImportSkeletalMesh(FImportSkeletalMeshArgs &
 		return nullptr;
 	}
 
-	// Array of re-import contexts for components using this mesh
-	// Will unregister before import, then re-register afterwards
-	TIndirectArray<FComponentReregisterContext> ComponentContexts;
-	TUniquePtr< FSkinnedMeshComponentRecreateRenderStateContext > SkinnedMeshComponentRecreateRenderStateContext;
-
-	FSkinnedMeshComponentRecreateRenderStateContext* RecreateExistingRenderStateContext = ExistingSkelMesh ? new FSkinnedMeshComponentRecreateRenderStateContext(ExistingSkelMesh, false) : nullptr;
-	//Backup the data before importing the new one
 	if (ExistingSkelMesh)
 	{
-		//Release the Renderdata resources before reimporting the skeletal mesh
-		SkinnedMeshComponentRecreateRenderStateContext = MakeUnique<FSkinnedMeshComponentRecreateRenderStateContext>(ExistingSkelMesh);
-
 		ExistingSkelMesh->PreEditChange(NULL);
 		//The backup of the skeletal mesh data empty the LOD array in the ImportedResource of the skeletal mesh
 		//If the import fail after this step the editor can crash when updating the bone later since the LODModel will not exist anymore
@@ -1825,12 +1822,18 @@ USkeletalMesh* UnFbx::FFbxImporter::ImportSkeletalMesh(FImportSkeletalMeshArgs &
 
 	if (SkeletalMesh == nullptr)
 	{
+		//Backup the PostEditchangeStackCounter before overriding the skeletalmesh
+		int32 PostEditChangeStackCounter = ExistingSkelMesh != nullptr ? ExistingSkelMesh->GetPostEditChangeStackCounter() : 0;
 		// Create the new mesh after saving the old data, since it will replace the old skeletalmesh
 		// This should happen only when doing a re-import. Otherwise the skeletal mesh is create before.
 		SkeletalMesh = NewObject<USkeletalMesh>(ImportSkeletalMeshArgs.InParent, ImportSkeletalMeshArgs.Name, ImportSkeletalMeshArgs.Flags);
+		SkeletalMesh->SetPostEditChangeStackCounter(PostEditChangeStackCounter);
 	}
 
+
 	SkeletalMesh->PreEditChange(NULL);
+	//Dirty the DDC Key for any imported Skeletal Mesh
+	SkeletalMesh->InvalidateDeriveDataCacheGUID();
 
 	FSkeletalMeshModel *ImportedResource = SkeletalMesh->GetImportedModel();
 	check(ImportedResource->LODModels.Num() == 0);
@@ -1858,9 +1861,24 @@ USkeletalMesh* UnFbx::FFbxImporter::ImportSkeletalMesh(FImportSkeletalMeshArgs &
 	// process bone influences from import data
 	ProcessImportMeshInfluences(*SkelMeshImportDataPtr);
 
-	//Store the original fbx import data
+	//Store the original fbx import data the SkelMeshImportDataPtr should not be modified after this
 	LODModel.RawSkeletalMeshBulkData.SaveRawMesh(*SkelMeshImportDataPtr);
-
+	if (ImportOptions->bImportAsSkeletalSkinning)
+	{
+		LODModel.RawSkeletalMeshBulkData.SkinningImportVersion = ESkeletalMeshSkinningImportVersions::LatestVersion;
+		LODModel.RawSkeletalMeshBulkData.GeoImportVersion = GeoImportVersion;
+	}
+	else if(ImportOptions->bImportAsSkeletalGeometry)
+	{
+		LODModel.RawSkeletalMeshBulkData.GeoImportVersion = ESkeletalMeshGeoImportVersions::LatestVersion;
+		LODModel.RawSkeletalMeshBulkData.SkinningImportVersion = SkinningImportVersion;
+	}
+	else
+	{
+		//We reimport both
+		LODModel.RawSkeletalMeshBulkData.SkinningImportVersion = ESkeletalMeshSkinningImportVersions::LatestVersion;
+		LODModel.RawSkeletalMeshBulkData.GeoImportVersion = ESkeletalMeshGeoImportVersions::LatestVersion;
+	}
 	SkeletalMesh->ResetLODInfo();
 	FSkeletalMeshLODInfo& NewLODInfo = SkeletalMesh->AddLODInfo();
 	NewLODInfo.ReductionSettings.NumOfTrianglesPercentage = 1.0f;
@@ -1886,38 +1904,61 @@ USkeletalMesh* UnFbx::FFbxImporter::ImportSkeletalMesh(FImportSkeletalMeshArgs &
 		TArray<int32> LODPointToRawMap;
 		SkelMeshImportDataPtr->CopyLODImportData(LODPoints,LODWedges,LODFaces,LODInfluences,LODPointToRawMap);
 
-		IMeshUtilities::MeshBuildOptions BuildOptions;
-		BuildOptions.OverlappingThresholds = ImportOptions->OverlappingThresholds;
-		BuildOptions.bComputeNormals = !ImportOptions->ShouldImportNormals() || !SkelMeshImportDataPtr->bHasNormals;
-		BuildOptions.bComputeTangents = !ImportOptions->ShouldImportTangents() || !SkelMeshImportDataPtr->bHasTangents;
+		FSkeletalMeshBuildSettings BuildOptions;
+		BuildOptions.bBuildAdjacencyBuffer = true;
+		BuildOptions.bRecomputeNormals = !ImportOptions->ShouldImportNormals() || !SkelMeshImportDataPtr->bHasNormals;
+		BuildOptions.bRecomputeTangents = !ImportOptions->ShouldImportTangents() || !SkelMeshImportDataPtr->bHasTangents;
 		BuildOptions.bUseMikkTSpace = (ImportOptions->NormalGenerationMethod == EFBXNormalGenerationMethod::MikkTSpace) && (!ImportOptions->ShouldImportNormals() || !ImportOptions->ShouldImportTangents());
 		BuildOptions.bComputeWeightedNormals = ImportOptions->bComputeWeightedNormals;
-		BuildOptions.bRemoveDegenerateTriangles = false;
-
-		IMeshUtilities& MeshUtilities = FModuleManager::Get().LoadModuleChecked<IMeshUtilities>("MeshUtilities");
-		
-		TArray<FText> WarningMessages;
-		TArray<FName> WarningNames;
-		// Create actual rendering data.
-		bool bBuildSuccess = MeshUtilities.BuildSkeletalMesh(ImportedResource->LODModels[0],SkeletalMesh->RefSkeleton,LODInfluences, LODWedges,LODFaces,LODPoints,LODPointToRawMap,BuildOptions, &WarningMessages, &WarningNames);
+		BuildOptions.bRemoveDegenerates = ImportOptions->bRemoveDegenerates;
+		BuildOptions.ThresholdPosition = ImportOptions->OverlappingThresholds.ThresholdPosition;
+		BuildOptions.ThresholdTangentNormal = ImportOptions->OverlappingThresholds.ThresholdTangentNormal;
+		BuildOptions.ThresholdUV = ImportOptions->OverlappingThresholds.ThresholdUV;
 
-		// temporary hack of message/names, should be one token or a struct
-		if(WarningMessages.Num() > 0 && WarningNames.Num() == WarningMessages.Num())
+		bool bUseLegacyMeshUtilities = false;
+		bool bBuildSuccess = false;
+		if (bUseLegacyMeshUtilities)
 		{
-			EMessageSeverity::Type MessageSeverity = bBuildSuccess ? EMessageSeverity::Warning : EMessageSeverity::Error;
+			IMeshUtilities::MeshBuildOptions LegacyBuildOptions;
+			LegacyBuildOptions.FillOptions(BuildOptions);
 
-			for(int32 MessageIdx = 0; MessageIdx<WarningMessages.Num(); ++MessageIdx)	
+			IMeshUtilities& MeshUtilities = FModuleManager::Get().LoadModuleChecked<IMeshUtilities>("MeshUtilities");
+
+			TArray<FText> WarningMessages;
+			TArray<FName> WarningNames;
+			// Create actual rendering data.
+			bBuildSuccess = MeshUtilities.BuildSkeletalMesh(ImportedResource->LODModels[0], SkeletalMesh->RefSkeleton, LODInfluences, LODWedges, LODFaces, LODPoints, LODPointToRawMap, LegacyBuildOptions, &WarningMessages, &WarningNames);
+
+			// temporary hack of message/names, should be one token or a struct
+			if (WarningMessages.Num() > 0 && WarningNames.Num() == WarningMessages.Num())
 			{
-				AddTokenizedErrorMessage(FTokenizedMessage::Create(MessageSeverity, WarningMessages[MessageIdx]), WarningNames[MessageIdx]);
+				EMessageSeverity::Type MessageSeverity = bBuildSuccess ? EMessageSeverity::Warning : EMessageSeverity::Error;
+
+				for (int32 MessageIdx = 0; MessageIdx < WarningMessages.Num(); ++MessageIdx)
+				{
+					AddTokenizedErrorMessage(FTokenizedMessage::Create(MessageSeverity, WarningMessages[MessageIdx]), WarningNames[MessageIdx]);
+				}
 			}
 		}
+		else
+		{
+			//The imported LOD is always 0 here, the LOD custom import will import the LOD alone(in a temporary skeletalmesh) and add it to the base skeletal mesh later
+			const int32 LODIndex = 0;
+			check(SkeletalMesh->GetLODInfo(LODIndex) != nullptr);
+			//Set the build options
+			SkeletalMesh->GetLODInfo(LODIndex)->BuildSettings = BuildOptions;
+			//New MeshDescription build process
+			IMeshBuilderModule& MeshBuilderModule = FModuleManager::Get().LoadModuleChecked<IMeshBuilderModule>("MeshBuilder");
+			//We must build the LODModel so we can restore properly the mesh, but we do not have to regenerate LODs
+			bBuildSuccess = MeshBuilderModule.BuildSkeletalMesh(SkeletalMesh, LODIndex, false);
+		}
 
 		if( !bBuildSuccess )
 		{
 			SkeletalMesh->MarkPendingKill();
 			return NULL;
 		}
-
+		
 		// Presize the per-section shadow casting array with the number of sections in the imported LOD.
 		const int32 NumSections = LODModel.Sections.Num();
 
@@ -1975,18 +2016,12 @@ USkeletalMesh* UnFbx::FFbxImporter::ImportSkeletalMesh(FImportSkeletalMeshArgs &
 		}
 
 		SkeletalMesh->CalculateInvRefMatrices();
-		SkeletalMesh->PostEditChange();
-		SkeletalMesh->MarkPackageDirty();
-
-		// Now iterate over all skeletal mesh components re-initialising them.
-		for (TObjectIterator<USkinnedMeshComponent> It; It; ++It)
+		if ((!SkeletalMesh->GetResourceForRendering() || !SkeletalMesh->GetResourceForRendering()->LODRenderData.IsValidIndex(0)) && ImportOptions->bCreatePhysicsAsset)
 		{
-			USkinnedMeshComponent* SkinComp = *It;
-			if (SkinComp->SkeletalMesh == SkeletalMesh)
-			{
-				ComponentContexts.Add(new FComponentReregisterContext(SkinComp));
-			}
+			//We need to have a valid render data to create physic asset
+			SkeletalMesh->Build();
 		}
+		SkeletalMesh->MarkPackageDirty();
 	}
 
 	if(ImportSkeletalMeshArgs.LodIndex == 0)
@@ -2142,13 +2177,6 @@ USkeletalMesh* UnFbx::FFbxImporter::ImportSkeletalMesh(FImportSkeletalMeshArgs &
 		ImportNodeCustomProperties(SkeletalMesh->Skeleton, SkeletonNode, true);
 	}
 
-	// ComponentContexts will now go out of scope, causing components to be re-registered
-
-	if (RecreateExistingRenderStateContext)
-	{
-		delete RecreateExistingRenderStateContext;
-	}
-
 	return SkeletalMesh;
 }
 
@@ -2352,8 +2380,6 @@ USkeletalMesh* UnFbx::FFbxImporter::ReimportSkeletalMesh(USkeletalMesh* Mesh, UF
 		return NULL;
 	}
 
-	char MeshName[1024];
-	FCStringAnsi::Strcpy(MeshName, TCHAR_TO_UTF8(*Mesh->GetName()));
 	TArray<FbxNode*>* FbxNodes = NULL;
 	USkeletalMesh* NewMesh = NULL;
 
@@ -2362,16 +2388,35 @@ USkeletalMesh* UnFbx::FFbxImporter::ReimportSkeletalMesh(USkeletalMesh* Mesh, UF
 	bool Old_ImportTextures = ImportOptions->bImportTextures;
 	bool Old_ImportAnimations = ImportOptions->bImportAnimations;
 
+	auto CleanUpImportOptionsData = [&]()
+	{
+		ImportOptions->bImportRigidMesh = Old_ImportRigidMesh;
+		ImportOptions->bImportMaterials = Old_ImportMaterials;
+		ImportOptions->bImportTextures = Old_ImportTextures;
+		ImportOptions->bImportAnimations = Old_ImportAnimations;
+	};
+
 	// support to update rigid animation mesh
 	ImportOptions->bImportRigidMesh = true;
 
-	// Unbind any existing clothing assets before we reimport the geometry
-	TArray<ClothingAssetUtils::FClothingAssetMeshBinding> ClothingBindings;
-	ClothingAssetUtils::GetMeshClothingAssetBindings(Mesh, ClothingBindings);
-
-	for(ClothingAssetUtils::FClothingAssetMeshBinding& Binding : ClothingBindings)
+	UFbxSkeletalMeshImportData* SKImportData = Cast<UFbxSkeletalMeshImportData>(Mesh->AssetImportData);
+	if (SKImportData)
 	{
-		Binding.Asset->UnbindFromSkeletalMesh(Mesh, Binding.LODIndex);
+		FSkeletalMeshLODInfo* LODInfo = Mesh->GetLODInfo(0);
+		if(LODInfo && Mesh->GetImportedModel() && Mesh->GetImportedModel()->LODModels.IsValidIndex(0))
+		{
+			const FSkeletalMeshLODModel& LODModel = Mesh->GetImportedModel()->LODModels[0];
+			if (!LODModel.RawSkeletalMeshBulkData.IsBuildDataAvailable())
+			{
+				//Set the build settings
+				LODInfo->BuildSettings.bBuildAdjacencyBuffer = true;
+				LODInfo->BuildSettings.bComputeWeightedNormals = SKImportData->bComputeWeightedNormals;
+				LODInfo->BuildSettings.bRecomputeNormals = SKImportData->NormalImportMethod == EFBXNormalImportMethod::FBXNIM_ComputeNormals;
+				LODInfo->BuildSettings.bRecomputeTangents = SKImportData->NormalImportMethod != EFBXNormalImportMethod::FBXNIM_ImportNormalsAndTangents;
+				LODInfo->BuildSettings.bRemoveDegenerates = true;
+				LODInfo->BuildSettings.bUseMikkTSpace = SKImportData->NormalGenerationMethod == EFBXNormalGenerationMethod::MikkTSpace;
+			}
+		}
 	}
 
 	// get meshes in Fbx file
@@ -2427,26 +2472,53 @@ USkeletalMesh* UnFbx::FFbxImporter::ReimportSkeletalMesh(USkeletalMesh* Mesh, UF
 			if (FbxNodes != nullptr)
 				break;
 		}
-		if (FbxNodes == nullptr)
-			return nullptr;
 	}
 	else
 	{
 		// if there is only one mesh, use it without name checking 
 		// (because the "Used As Full Name" option enables users name the Unreal mesh by themselves
-		if (FbxSkelMeshArray.Num() > 0)
+		if (FbxSkelMeshArray.Num() == 1)
 		{
 			FbxNodes = FbxSkelMeshArray[0];
 		}
-		else
+		else if (FbxSkelMeshArray.Num() > 1)
 		{
-			// @todo - FBX Importing - We need proper support for reimport if the file contains more than one skeletal mesh. 
-
+			//Search a set corresponding
+			for (TArray<FbxNode*>*SkeletalMeshNodes : FbxSkelMeshArray)
+			{
+				FbxNode* Node = GetMeshNodesFromName(Mesh->GetName(), *SkeletalMeshNodes);
+				if (Node != nullptr)
+				{
+					FbxNodes = SkeletalMeshNodes;
+					break;
+				}
+			}
+			if (FbxNodes == nullptr)
+			{
+				//Re import the first skeletalmesh found in the fbx file
+				FbxNodes = FbxSkelMeshArray[0];
+			}
 		}
 	}
-
+	
 	if (FbxNodes)
 	{
+		//Posteditchange will be call at the end of the reimport which will reallocate the render ressource
+		FScopedSkeletalMeshPostEditChange ScopedPostEditChange(Mesh);
+
+		FScopedSlowTask SlowTask(100.0f, NSLOCTEXT("SkeltalMeshReimportFactory", "SkeletalMeshReimportTask", "Reimporting skeletal mesh..."));
+		SlowTask.MakeDialog();
+
+		//The fbx file is valid we can reimport the skeletal mesh
+		Mesh->PreEditChange(nullptr);
+
+		// Unbind any existing clothing assets before we reimport the geometry
+		TArray<ClothingAssetUtils::FClothingAssetMeshBinding> ClothingBindings;
+		FLODUtilities::UnbindClothingAndBackup(Mesh, ClothingBindings);
+
+		//Reimport must force a new guid
+		Mesh->InvalidateDeriveDataCacheGUID();
+
 		for (int32 LODCounter = 1; LODCounter < Mesh->GetLODNum(); ++LODCounter)
 		{
 			//Reset all import with base mesh flags
@@ -2491,8 +2563,10 @@ USkeletalMesh* UnFbx::FFbxImporter::ReimportSkeletalMesh(USkeletalMesh* Mesh, UF
 		int32 LODIndex;
 		int32 SuccessfulLodIndex = 0;
 		bool bOperationCanceled = false;
+		float ProgressStep = 90.0 / MaxLODLevel;
 		for (LODIndex = 0; LODIndex < MaxLODLevel; LODIndex++)
 		{
+			SlowTask.EnterProgressFrame(ProgressStep);
 			int32 ImportedSuccessfulLodIndex = INDEX_NONE;
 			TArray<FbxNode*> SkelMeshNodeArray;
 			for (int32 j = 0; j < (*FbxNodes).Num(); j++)
@@ -2542,8 +2616,9 @@ USkeletalMesh* UnFbx::FFbxImporter::ReimportSkeletalMesh(USkeletalMesh* Mesh, UF
 				NewMesh = ImportSkeletalMesh( ImportSkeletalMeshArgs );
 				if (bOperationCanceled)
 				{
-					// User cancelled, clean up and return
-					return nullptr;
+					// User cancel, clean up and return
+					NewMesh = nullptr;
+					break;
 				}
 				if (NewMesh)
 				{
@@ -2590,7 +2665,7 @@ USkeletalMesh* UnFbx::FFbxImporter::ReimportSkeletalMesh(USkeletalMesh* Mesh, UF
 				ImportSkeletalMeshArgs.OutData = &OutData;
 
 				UObject *LODObject = ImportSkeletalMesh( ImportSkeletalMeshArgs );
-				bool bImportSucceeded = !bOperationCanceled && ImportSkeletalMeshLOD( Cast<USkeletalMesh>(LODObject), BaseSkeletalMesh, SuccessfulLodIndex, false, nullptr, TemplateImportData);
+				bool bImportSucceeded = !bOperationCanceled && ImportSkeletalMeshLOD( Cast<USkeletalMesh>(LODObject), BaseSkeletalMesh, SuccessfulLodIndex, TemplateImportData);
 
 				if (bImportSucceeded)
 				{
@@ -2626,47 +2701,41 @@ USkeletalMesh* UnFbx::FFbxImporter::ReimportSkeletalMesh(USkeletalMesh* Mesh, UF
 
 			if (NewMesh)
 			{
-				// import morph target
 				if ((ImportOptions->bImportSkeletalMeshLODs || LODIndex == 0) &&
 					ImportOptions->bImportMorph &&
 					ImportedSuccessfulLodIndex != INDEX_NONE &&
 					NewMesh->GetImportedModel() &&
-					NewMesh->GetImportedModel()->LODModels.IsValidIndex(LODIndex))
+					NewMesh->GetImportedModel()->LODModels.IsValidIndex(ImportedSuccessfulLodIndex))
 				{
-					// @fixme: @question : where do they import this morph? where to? What morph target sets?
 					ImportFbxMorphTarget(SkelMeshNodeArray, NewMesh, ImportedSuccessfulLodIndex, OutData);
 				}
 			}
 		}
-		if (NewMesh)
+
+		auto RebindClothing = [&ClothingBindings](USkeletalMesh* SkelMesh)
 		{
-			FSkeletalMeshModel* NewMeshResource = NewMesh->GetImportedModel();
-			if(NewMeshResource && ClothingBindings.Num() > 0)
+			FSkeletalMeshModel* MeshResource = SkelMesh->GetImportedModel();
+			if (MeshResource && ClothingBindings.Num() > 0)
 			{
-				NewMesh->PreEditChange(nullptr);
+				FLODUtilities::RestoreClothingFromBackup(SkelMesh, ClothingBindings);
+			}
+		};
 
-				for(ClothingAssetUtils::FClothingAssetMeshBinding& Binding : ClothingBindings)
-				{
-					if(NewMeshResource->LODModels.IsValidIndex(Binding.LODIndex) &&
-						NewMeshResource->LODModels[Binding.LODIndex].Sections.IsValidIndex(Binding.SectionIndex))
-					{
-						Binding.Asset->BindToSkeletalMesh(NewMesh, Binding.LODIndex, Binding.SectionIndex, Binding.AssetInternalLodIndex);
-					}
-				}
 
-				NewMesh->PostEditChange();
-			}
+		if (!bOperationCanceled && NewMesh)
+		{
+			RebindClothing(NewMesh);
 
 			//Update the import data so we can re-import correctly
 			UpdateSkeletalMeshImportData(NewMesh, TemplateImportData, INDEX_NONE, &ImportMaterialOriginalNameData, &ImportMeshLodData);
-			//If we have import some morph target we have to rebuild the render resources since morph target are now using GPU
-			if (NewMesh->MorphTargets.Num() > 0)
-			{
-				NewMesh->ReleaseResources();
-				//Rebuild the resources with a post edit change since we have added some morph targets
-				NewMesh->PostEditChange();
-			}
 		}
+		else if(bOperationCanceled && Mesh != nullptr)
+		{
+			//Operation was cancel
+			//rebind the clothing on the original mesh
+			RebindClothing(Mesh);
+		}
+		SlowTask.EnterProgressFrame(10.0f);
 	}
 	else
 	{
@@ -2674,10 +2743,7 @@ USkeletalMesh* UnFbx::FFbxImporter::ReimportSkeletalMesh(USkeletalMesh* Mesh, UF
 		AddTokenizedErrorMessage(FTokenizedMessage::Create(EMessageSeverity::Warning, FText::Format(LOCTEXT("FbxSkeletaLMeshimport_NoFBXMeshMatch", "No FBX mesh matches the Unreal mesh '{0}'."), FText::FromString(Mesh->GetName()))), FFbxErrors::Generic_Mesh_MeshNotFound);
 	}
 
-	ImportOptions->bImportRigidMesh = Old_ImportRigidMesh;
-	ImportOptions->bImportMaterials = Old_ImportMaterials;
-	ImportOptions->bImportTextures = Old_ImportTextures;
-	ImportOptions->bImportAnimations = Old_ImportAnimations;
+	CleanUpImportOptionsData();
 
 	return NewMesh;
 }
@@ -3411,7 +3477,7 @@ bool UnFbx::FFbxImporter::FillSkelMeshImporterFromFbx( FSkeletalMeshImportData&
 			//
 			// normals, tangents and binormals
 			//
-			if( ImportOptions->ShouldImportNormals() && bHasNormalInformation )
+			if( bHasNormalInformation )
 			{
 				int32 TmpIndex = LocalIndex*3 + VertexIndex;
 				//normals may have different reference and mapping mode than tangents and binormals
@@ -3425,7 +3491,7 @@ bool UnFbx::FFbxImporter::FillSkelMeshImporterFromFbx( FSkeletalMeshImportData&
 
 				FbxVector4 TempValue;
 
-				if( ImportOptions->ShouldImportTangents() && bHasTangentInformation )
+				if( bHasTangentInformation )
 				{
 					TempValue = LayerElementTangent->GetDirectArray().GetAt(TangentMapIndex);
 					TempValue = TotalMatrixForNormal.MultT(TempValue);
@@ -3722,6 +3788,9 @@ bool UnFbx::FFbxImporter::FillSkelMeshImporterFromFbx( FSkeletalMeshImportData&
 
 void UnFbx::FFbxImporter::InsertNewLODToBaseSkeletalMesh(USkeletalMesh* InSkeletalMesh, USkeletalMesh* BaseSkeletalMesh, int32 DesiredLOD, UFbxSkeletalMeshImportData* TemplateImportData)
 {
+	//Scope a post edit change
+	FScopedSkeletalMeshPostEditChange ScopedPostEditChange(BaseSkeletalMesh);
+
 	FSkeletalMeshModel* ImportedResource = InSkeletalMesh->GetImportedModel();
 	FSkeletalMeshModel* DestImportedResource = BaseSkeletalMesh->GetImportedModel();
 
@@ -3781,6 +3850,13 @@ void UnFbx::FFbxImporter::InsertNewLODToBaseSkeletalMesh(USkeletalMesh* InSkelet
 
 	// Set up LODMaterialMap to number of materials in new mesh.
 	FSkeletalMeshLODInfo& LODInfo = *(BaseSkeletalMesh->GetLODInfo(DesiredLOD));
+
+	//Copy the build settings
+	if(InSkeletalMesh->GetLODInfo(0))
+	{
+		const FSkeletalMeshLODInfo& ImportedLODInfo = *(InSkeletalMesh->GetLODInfo(0));
+		LODInfo.BuildSettings = ImportedLODInfo.BuildSettings;
+	}
 	
 	LODInfo.LODMaterialMap.Empty();
 	// Now set up the material mapping array.
@@ -3874,12 +3950,9 @@ void UnFbx::FFbxImporter::InsertNewLODToBaseSkeletalMesh(USkeletalMesh* InSkelet
 			}
 		}
 	}
-
-	// rebuild vertex buffers and reinit RHI resources
-	BaseSkeletalMesh->PostEditChange();
 }
 
-bool UnFbx::FFbxImporter::ImportSkeletalMeshLOD(USkeletalMesh* InSkeletalMesh, USkeletalMesh* BaseSkeletalMesh, int32 DesiredLOD, bool bNeedToReregister, TArray<UActorComponent*>* ReregisterAssociatedComponents, UFbxSkeletalMeshImportData* TemplateImportData)
+bool UnFbx::FFbxImporter::ImportSkeletalMeshLOD(USkeletalMesh* InSkeletalMesh, USkeletalMesh* BaseSkeletalMesh, int32 DesiredLOD, UFbxSkeletalMeshImportData* TemplateImportData)
 {
 	check(InSkeletalMesh);
 	check(BaseSkeletalMesh);
@@ -3943,6 +4016,8 @@ bool UnFbx::FFbxImporter::ImportSkeletalMeshLOD(USkeletalMesh* InSkeletalMesh, U
 		}
 	}
 
+	FScopedSkeletalMeshPostEditChange ScopedPostEditChange(BaseSkeletalMesh);
+
 	//Skeleton do not match and we are importing skeletal geometry only so apply the base LOD Skinning
 	if (bApplyBaseSkinning)
 	{
@@ -3950,8 +4025,10 @@ bool UnFbx::FFbxImporter::ImportSkeletalMeshLOD(USkeletalMesh* InSkeletalMesh, U
 		//Apply the base LOD skinning to the new LOD geometry, DestImportedResource here is the base lod and is the source for the apply skinning.
 		//The Dest here mean we merge the just import LOD to the base LodModel.
 		SkeletalMeshHelper::ApplySkinning(InSkeletalMesh, DestImportedResource->LODModels[0], ImportedResource->LODModels[0]);
-		IMeshUtilities& MeshUtilities = FModuleManager::Get().LoadModuleChecked<IMeshUtilities>("MeshUtilities");
-		MeshUtilities.RemoveBonesFromMesh(BaseSkeletalMesh, DesiredLOD, NULL);
+
+		//Excluded bones will be remove when we rebuild the asset
+		//IMeshUtilities& MeshUtilities = FModuleManager::Get().LoadModuleChecked<IMeshUtilities>("MeshUtilities");
+		//MeshUtilities.RemoveBonesFromMesh(BaseSkeletalMesh, DesiredLOD, NULL);
 	}
 
 	FSkeletalMeshLODModel& NewLODModel = ImportedResource->LODModels[0];
@@ -4114,32 +4191,7 @@ bool UnFbx::FFbxImporter::ImportSkeletalMeshLOD(USkeletalMesh* InSkeletalMesh, U
 		}
 	}
 
-	if (bNeedToReregister)
-	{
-		// Shut down the skeletal mesh component that is previewing this mesh.
-		if (ReregisterAssociatedComponents)
-		{
-			FMultiComponentReregisterContext ReregisterContext(*ReregisterAssociatedComponents);
-			// wait until resources are released
-			FlushRenderingCommands();
-
-			InsertNewLODToBaseSkeletalMesh(InSkeletalMesh, BaseSkeletalMesh, DesiredLOD, TemplateImportData);
-
-			// ReregisterContexts go out of scope here, reregistering associated components with the scene.
-		}
-		else
-		{
-			TComponentReregisterContext<USkinnedMeshComponent> ComponentReregisterContext;
-
-			InsertNewLODToBaseSkeletalMesh(InSkeletalMesh, BaseSkeletalMesh, DesiredLOD, TemplateImportData);
-
-			// ReregisterContexts go out of scope here, reregistering skel components with the scene.
-		}
-	}
-	else
-	{
-		InsertNewLODToBaseSkeletalMesh(InSkeletalMesh, BaseSkeletalMesh, DesiredLOD, TemplateImportData);
-	}
+	InsertNewLODToBaseSkeletalMesh(InSkeletalMesh, BaseSkeletalMesh, DesiredLOD, TemplateImportData);
 
 	return true;
 }
@@ -4280,9 +4332,21 @@ void UnFbx::FFbxImporter::ImportMorphTargetsInternal( TArray<FbxNode*>& SkelMesh
 		check(BaseImportData.MorphTargetNames.Num() == BaseImportData.MorphTargets.Num() && BaseImportData.MorphTargetNames.Num() == BaseImportData.MorphTargetModifiedPoints.Num());
 	}
 
-	//Build MorphTargets
-	// Calculate overlapping corners and tangents
-	FLODUtilities::BuildMorphTargets(BaseSkelMesh, BaseImportData, LODIndex, ImportOptions->ShouldImportNormals(), ImportOptions->ShouldImportTangents(), (ImportOptions->NormalGenerationMethod == EFBXNormalGenerationMethod::MikkTSpace));
+	if (BaseSkelMesh->GetImportedModel() && BaseSkelMesh->GetImportedModel()->LODModels.IsValidIndex(LODIndex))
+	{
+		//If we can build the skeletal mesh there is no need to build the morph target now, all the necessary build morph target data was copied before.
+		if (!BaseSkelMesh->GetImportedModel()->LODModels[LODIndex].RawSkeletalMeshBulkData.IsBuildDataAvailable())
+		{
+			//Build MorphTargets
+			FLODUtilities::BuildMorphTargets(
+				BaseSkelMesh,
+				BaseImportData,
+				LODIndex,
+				ImportOptions->ShouldImportNormals(),
+				ImportOptions->ShouldImportTangents(),
+				(ImportOptions->NormalGenerationMethod == EFBXNormalGenerationMethod::MikkTSpace));
+		}
+	}
 
 	GWarn->EndSlowTask();
 }
@@ -4290,6 +4354,9 @@ void UnFbx::FFbxImporter::ImportMorphTargetsInternal( TArray<FbxNode*>& SkelMesh
 // Import Morph target
 void UnFbx::FFbxImporter::ImportFbxMorphTarget(TArray<FbxNode*> &SkelMeshNodeArray, USkeletalMesh* BaseSkelMesh, int32 LODIndex, FSkeletalMeshImportData &BaseSkeletalMeshImportData)
 {
+	//Stack the PostEditChange call, it will call post edit change when it will go out of scope
+	FScopedSkeletalMeshPostEditChange ScopedPostEditChange(BaseSkelMesh);
+
 	bool bHasMorph = false;
 	int32 NodeIndex;
 	// check if there are morph in this geometry
@@ -4308,7 +4375,12 @@ void UnFbx::FFbxImporter::ImportFbxMorphTarget(TArray<FbxNode*> &SkelMeshNodeArr
 	
 	if (bHasMorph)
 	{
-		ImportMorphTargetsInternal( SkelMeshNodeArray, BaseSkelMesh, LODIndex, BaseSkeletalMeshImportData );
+		ImportMorphTargetsInternal( SkelMeshNodeArray, BaseSkelMesh, LODIndex, BaseSkeletalMeshImportData);
+		//Save the rawMesh
+		if (BaseSkelMesh->GetImportedModel() && BaseSkelMesh->GetImportedModel()->LODModels.IsValidIndex(LODIndex))
+		{
+			BaseSkelMesh->GetImportedModel()->LODModels[LODIndex].RawSkeletalMeshBulkData.SaveRawMesh(BaseSkeletalMeshImportData);
+		}
 	}
 }
 
diff --git a/Engine/Source/Editor/UnrealEd/Private/FbxMeshUtils.cpp b/Engine/Source/Editor/UnrealEd/Private/FbxMeshUtils.cpp
--- a/Engine/Source/Editor/UnrealEd/Private/FbxMeshUtils.cpp
+++ b/Engine/Source/Editor/UnrealEd/Private/FbxMeshUtils.cpp
@@ -33,8 +33,9 @@
 #include "Framework/Notifications/NotificationManager.h"
 #include "Widgets/Notifications/SNotificationList.h"
 
-#include "Assets/ClothingAsset.h"
+#include "ClothingAsset.h"
 #include "SkinWeightsUtilities.h"
+#include "LODUtilities.h"
 
 DEFINE_LOG_CATEGORY_STATIC(LogExportMeshUtils, Log, All);
 
@@ -280,10 +281,13 @@ namespace FbxMeshUtils
 			return false;
 		}
 
+		FScopedSkeletalMeshPostEditChange ScopePostEditChange(SelectedSkelMesh);
+
 		//If the imported LOD already exist, we will need to reimport all the skin weight profiles
 		bool bMustReimportAlternateSkinWeightProfile = false;
 
 		// Get a list of all the clothing assets affecting this LOD so we can re-apply later
+		TArray<ClothingAssetUtils::FClothingAssetMeshBinding> ClothingBindings;
 		TArray<UClothingAssetBase*> ClothingAssetsInUse;
 		TArray<int32> ClothingAssetSectionIndices;
 		TArray<int32> ClothingAssetInternalLodIndices;
@@ -292,47 +296,16 @@ namespace FbxMeshUtils
 		if(ImportedResource && ImportedResource->LODModels.IsValidIndex(LODLevel))
 		{
 			bMustReimportAlternateSkinWeightProfile = true;
-			FSkeletalMeshLODModel& LodModel = ImportedResource->LODModels[LODLevel];
-
-			const int32 NumSections = LodModel.Sections.Num();
-
-			for(int32 SectionIndex = 0; SectionIndex < NumSections; ++SectionIndex)
-			{
-				FSkelMeshSection& Section = LodModel.Sections[SectionIndex];
-
-				if(Section.HasClothingData())
-				{
-					UClothingAssetBase* AssetInUse = SelectedSkelMesh->GetSectionClothingAsset(LODLevel, SectionIndex);
-					ClothingAssetsInUse.Add(AssetInUse);
-					ClothingAssetSectionIndices.Add(SectionIndex);
-					ClothingAssetInternalLodIndices.Add(Section.ClothingData.AssetLodIndex);
-				}
-			}
-		}
-
-		// Remove our clothing assets while we import this LOD
-		for(UClothingAssetBase* ClothingAsset : ClothingAssetsInUse)
-		{
-			ClothingAsset->UnbindFromSkeletalMesh(SelectedSkelMesh, LODLevel);
+			FLODUtilities::UnbindClothingAndBackup(SelectedSkelMesh, ClothingBindings, LODLevel);
 		}
 
 		//Lambda to call to re-apply the clothing
-		auto ReapplyClothing = [&SelectedSkelMesh, &ClothingAssetsInUse, &ClothingAssetSectionIndices, &ClothingAssetInternalLodIndices, &ImportedResource, &LODLevel]()
+		auto ReapplyClothing = [&SelectedSkelMesh, &ClothingBindings, &ImportedResource, &LODLevel]()
 		{
-			// Re-apply our clothing assets
-			int32 NumClothingAssetsToApply = ClothingAssetsInUse.Num();
 			if (ImportedResource && ImportedResource->LODModels.IsValidIndex(LODLevel))
 			{
-				FSkeletalMeshLODModel& LodModel = ImportedResource->LODModels[LODLevel];
-				for (int32 AssetIndex = 0; AssetIndex < NumClothingAssetsToApply; ++AssetIndex)
-				{
-					// Only if the same section exists
-					if (LodModel.Sections.IsValidIndex(ClothingAssetSectionIndices[AssetIndex]))
-					{
-						UClothingAssetBase* AssetToApply = ClothingAssetsInUse[AssetIndex];
-						AssetToApply->BindToSkeletalMesh(SelectedSkelMesh, LODLevel, ClothingAssetSectionIndices[AssetIndex], ClothingAssetInternalLodIndices[AssetIndex]);
-					}
-				}
+				// Re-apply our clothing assets
+				FLODUtilities::RestoreClothingFromBackup(SelectedSkelMesh, ClothingBindings, LODLevel);
 			}
 		};
 
@@ -395,6 +368,24 @@ namespace FbxMeshUtils
 			TArray<FString> LODStrings;
 			TArray<FbxNode*>* MeshObject = NULL;;
 
+			//Set the build options if the BuildDat is not available so it is the same option we use to import the LOD
+			if (ImportedResource && ImportedResource->LODModels.IsValidIndex(LODLevel) && !ImportedResource->LODModels[LODLevel].RawSkeletalMeshBulkData.IsBuildDataAvailable())
+			{
+				FSkeletalMeshLODInfo* LODInfo = SelectedSkelMesh->GetLODInfo(LODLevel);
+				if (LODInfo)
+				{
+					LODInfo->BuildSettings.bBuildAdjacencyBuffer = true;
+					LODInfo->BuildSettings.bRecomputeNormals = !ImportOptions->ShouldImportNormals();
+					LODInfo->BuildSettings.bRecomputeTangents = !ImportOptions->ShouldImportTangents();
+					LODInfo->BuildSettings.bUseMikkTSpace = (ImportOptions->NormalGenerationMethod == EFBXNormalGenerationMethod::MikkTSpace) && (!ImportOptions->ShouldImportNormals() || !ImportOptions->ShouldImportTangents());
+					LODInfo->BuildSettings.bComputeWeightedNormals = ImportOptions->bComputeWeightedNormals;
+					LODInfo->BuildSettings.bRemoveDegenerates = ImportOptions->bRemoveDegenerates;
+					LODInfo->BuildSettings.ThresholdPosition = ImportOptions->OverlappingThresholds.ThresholdPosition;
+					LODInfo->BuildSettings.ThresholdTangentNormal = ImportOptions->OverlappingThresholds.ThresholdTangentNormal;
+					LODInfo->BuildSettings.ThresholdUV = ImportOptions->OverlappingThresholds.ThresholdUV;
+				}
+			}
+
 			// Populate the mesh array
 			FFbxImporter->FillFbxSkelMeshArrayInScene(FFbxImporter->Scene->GetRootNode(), MeshArray, false, ImportOptions->bImportAsSkeletalGeometry || ImportOptions->bImportAsSkeletalSkinning, ImportOptions->bImportScene);
 
@@ -524,12 +515,8 @@ namespace FbxMeshUtils
 
 				TempSkelMesh = (USkeletalMesh*)FFbxImporter->ImportSkeletalMesh( ImportSkeletalMeshArgs );
 				// Add the new imported LOD to the existing model (check skeleton compatibility)
-				if( TempSkelMesh  && FFbxImporter->ImportSkeletalMeshLOD(TempSkelMesh, SelectedSkelMesh, SelectedLOD, true, nullptr, TempAssetImportData))
+				if( TempSkelMesh  && FFbxImporter->ImportSkeletalMeshLOD(TempSkelMesh, SelectedSkelMesh, SelectedLOD, TempAssetImportData))
 				{
-					TComponentReregisterContext<USkinnedMeshComponent> ReregisterContext;
-					SelectedSkelMesh->ReleaseResources();
-					SelectedSkelMesh->ReleaseResourcesFence.Wait();
-					
 					//Update the import data for this lod
 					UnFbx::FFbxImporter::UpdateSkeletalMeshImportData(SelectedSkelMesh, nullptr, SelectedLOD, &ImportMaterialOriginalNameData, &ImportMeshLodData);
 
@@ -555,13 +542,9 @@ namespace FbxMeshUtils
 					{
 						//We cannot use anymore the FFbxImporter after the cleanup
 						CleanUpScene();
-						FSkinWeightsUtilities::ReimportAlternateSkinWeight(SelectedSkelMesh, SelectedLOD, false);
+						FSkinWeightsUtilities::ReimportAlternateSkinWeight(SelectedSkelMesh, SelectedLOD);
 					}
 
-					SelectedSkelMesh->PostEditChange();
-					SelectedSkelMesh->InitResources();
-					SelectedSkelMesh->MarkPackageDirty();
-
 					// Notification of success
 					FNotificationInfo NotificationInfo(FText::GetEmpty());
 					NotificationInfo.Text = FText::Format(NSLOCTEXT("UnrealEd", "LODImportSuccessful", "Mesh for LOD {0} imported successfully!"), FText::AsNumber(SelectedLOD));
@@ -633,7 +616,7 @@ namespace FbxMeshUtils
 		return ChosenFilname;
 	}
 
-	bool ImportMeshLODDialog(class UObject* SelectedMesh, int32 LODLevel)
+	bool ImportMeshLODDialog(class UObject* SelectedMesh, int32 LODLevel, bool bNotifyCB /*= true*/)
 	{
 		if(!SelectedMesh)
 		{
@@ -732,7 +715,7 @@ namespace FbxMeshUtils
 			FMessageDialog::Open(EAppMsgType::Ok, FText::Format(LOCTEXT("LODImport_Failure", "Failed to import LOD{0}"), FText::AsNumber(LODLevel)));
 		}
 
-		if (bImportSuccess)
+		if (bImportSuccess && bNotifyCB)
 		{
 			if (SkeletalMesh)
 			{
diff --git a/Engine/Source/Editor/UnrealEd/Private/Kismet2/BlueprintEditorUtils.cpp b/Engine/Source/Editor/UnrealEd/Private/Kismet2/BlueprintEditorUtils.cpp
--- a/Engine/Source/Editor/UnrealEd/Private/Kismet2/BlueprintEditorUtils.cpp
+++ b/Engine/Source/Editor/UnrealEd/Private/Kismet2/BlueprintEditorUtils.cpp
@@ -6920,7 +6920,11 @@ void FBlueprintEditorUtils::UpdateOutOfDateCompositeWithOuter(UBlueprint* Bluepr
 		}
 		else if (UAnimStateTransitionNode* TransitionNode = Cast<UAnimStateTransitionNode>(Node))
 		{
-			BoundGraph = TransitionNode->BoundGraph;
+			// Shared rule graphs are not necessarily outered to their node, so skip them here
+			if(!TransitionNode->bSharedRules)
+			{
+				BoundGraph = TransitionNode->BoundGraph;
+			}
 		}
 		else if (UAnimGraphNode_StateMachineBase* StateMachineNode = Cast<UAnimGraphNode_StateMachineBase>(Node))
 		{
diff --git a/Engine/Source/Editor/UnrealEd/Private/SSkeletonWidget.cpp b/Engine/Source/Editor/UnrealEd/Private/SSkeletonWidget.cpp
--- a/Engine/Source/Editor/UnrealEd/Private/SSkeletonWidget.cpp
+++ b/Engine/Source/Editor/UnrealEd/Private/SSkeletonWidget.cpp
@@ -37,6 +37,7 @@ void SSkeletonListWidget::Construct(const FArguments& InArgs)
 	AssetPickerConfig.Filter.ClassNames.Add(USkeleton::StaticClass()->GetFName());
 	AssetPickerConfig.OnAssetSelected = FOnAssetSelected::CreateSP(this, &SSkeletonListWidget::SkeletonSelectionChanged);
 	AssetPickerConfig.InitialAssetViewType = EAssetViewType::Column;
+	AssetPickerConfig.SelectionMode = ESelectionMode::Single;
 	AssetPickerConfig.bShowPathInColumnView = true;
 	AssetPickerConfig.bShowTypeInColumnView = false;
 
diff --git a/Engine/Source/Editor/UnrealEd/Public/ApexClothingUtils.h b/Engine/Source/Editor/UnrealEd/Public/ApexClothingUtils.h
--- a/Engine/Source/Editor/UnrealEd/Public/ApexClothingUtils.h
+++ b/Engine/Source/Editor/UnrealEd/Public/ApexClothingUtils.h
@@ -20,7 +20,7 @@ namespace ApexClothingUtils
 {
 	UNREALED_API FString PromptForClothingFile();
 
-	// Prompt the user to select an APEX file that will be imported to the new UClothingAsset format
+	// Prompt the user to select an APEX file that will be imported to the new UClothingAssetCommon format
 	UNREALED_API void PromptAndImportClothing(USkeletalMesh* SkelMesh);
 
 	// Given a buffer, build an apex clothing asset
diff --git a/Engine/Source/Programs/UnrealBuildTool/Configuration/ModuleRules.cs b/Engine/Source/Programs/UnrealBuildTool/Configuration/ModuleRules.cs
--- a/Engine/Source/Programs/UnrealBuildTool/Configuration/ModuleRules.cs
+++ b/Engine/Source/Programs/UnrealBuildTool/Configuration/ModuleRules.cs
@@ -999,23 +999,23 @@ namespace UnrealBuildTool
 			PublicDependencyModuleNames.Add("PhysicsCore");
 
 			bool bUseNonPhysXInterface = Target.bUseChaos == true;
+            PublicIncludePathModuleNames.AddRange(
+                new string[] {
+                    "Chaos",
+					"FieldSystemCore"
+                }
+            );
+            PublicDependencyModuleNames.AddRange(
+				new string[] {
+					"Chaos",
+					"FieldSystemCore"
+                }
+            );
 
             // 
             if (Target.bCompileChaos == true || Target.bUseChaos == true)
             {
                 PublicDefinitions.Add("INCLUDE_CHAOS=1");
-                PublicIncludePathModuleNames.AddRange(
-                    new string[] {
-                        "Chaos",
-						"FieldSystemCore"
-                    }
-                );
-                PublicDependencyModuleNames.AddRange(
-                  new string[] {
-                        "Chaos",
-						"FieldSystemCore"
-                  }
-                );
             }
             else
             {
@@ -1038,6 +1038,7 @@ namespace UnrealBuildTool
 			{
 				// Disable non-physx interfaces
 				PublicDefinitions.Add("WITH_CHAOS=0");
+				PublicDefinitions.Add("WITH_CHAOS_CLOTHING=0");
 
 				// 
 				// WITH_CHAOS_NEEDS_TO_BE_FIXED
@@ -1116,7 +1117,6 @@ namespace UnrealBuildTool
 				PublicDefinitions.Add("PHYSICS_INTERFACE_PHYSX=0");
 				PublicDefinitions.Add("WITH_APEX=0");
 				PublicDefinitions.Add("WITH_APEX_CLOTHING=0");
-				PublicDefinitions.Add("WITH_CLOTH_COLLISION_DETECTION=0");
 				PublicDefinitions.Add(string.Format("WITH_PHYSX_COOKING={0}", Target.bBuildEditor && Target.bCompilePhysX ? 1 : 0));  // without APEX, we only need cooking in editor builds
 				PublicDefinitions.Add("WITH_NVCLOTH=0");
 
@@ -1125,6 +1125,8 @@ namespace UnrealBuildTool
 					PublicDefinitions.Add("WITH_CHAOS=1");
 					PublicDefinitions.Add("WITH_CHAOS_NEEDS_TO_BE_FIXED=1");
 					PublicDefinitions.Add("COMPILE_ID_TYPES_AS_INTS=0");
+					PublicDefinitions.Add("WITH_CHAOS_CLOTHING=1");
+					PublicDefinitions.Add("WITH_CLOTH_COLLISION_DETECTION=1");
 					
 					PublicIncludePathModuleNames.AddRange(
 						new string[] {
@@ -1142,6 +1144,8 @@ namespace UnrealBuildTool
 				{
 					PublicDefinitions.Add("WITH_CHAOS=0");
 					PublicDefinitions.Add("WITH_CHAOS_NEEDS_TO_BE_FIXED=0");
+					PublicDefinitions.Add("WITH_CHAOS_CLOTHING=0");
+					PublicDefinitions.Add("WITH_CLOTH_COLLISION_DETECTION=0");
 				}
 			}
 
diff --git a/Engine/Source/Programs/UnrealBuildTool/Configuration/TargetRules.cs b/Engine/Source/Programs/UnrealBuildTool/Configuration/TargetRules.cs
--- a/Engine/Source/Programs/UnrealBuildTool/Configuration/TargetRules.cs
+++ b/Engine/Source/Programs/UnrealBuildTool/Configuration/TargetRules.cs
@@ -342,7 +342,7 @@ namespace UnrealBuildTool
 		/// Whether to enable the mesh editor.
 		/// </summary>
 		[RequiresUniqueBuildEnvironment]
-		public bool bEnableMeshEditor = false;
+		public bool bEnableMeshEditor = true;
 
 		/// <summary>
 		/// Whether to compile the Chaos physics plugin.
@@ -362,6 +362,12 @@ namespace UnrealBuildTool
 		[RequiresUniqueBuildEnvironment]
 		public bool bUseChaosChecked = false;
 
+		/// <summary>
+		/// Whether to compile in chaos memory tracking features
+		/// </summary>
+		[RequiresUniqueBuildEnvironment]
+		public bool bUseChaosMemoryTracking = false;
+
 		/// <summary>
 		/// Whether scene query acceleration is done by UE4. The physx scene query structure is still created, but we do not use it.
 		/// </summary>
@@ -1871,6 +1877,11 @@ namespace UnrealBuildTool
 			get { return Inner.bUseChaos; }
 		}
 
+		public bool bUseChaosMemoryTracking
+		{
+			get { return Inner.bUseChaosMemoryTracking; }
+		}
+
 		public bool bUseChaosChecked
 		{
 			get { return Inner.bUseChaosChecked; }
diff --git a/Engine/Source/Programs/UnrealHeaderTool/Resources/UHTDebugging/DocumentationObjects.h b/Engine/Source/Programs/UnrealHeaderTool/Resources/UHTDebugging/DocumentationObjects.h
--- a/Engine/Source/Programs/UnrealHeaderTool/Resources/UHTDebugging/DocumentationObjects.h
+++ b/Engine/Source/Programs/UnrealHeaderTool/Resources/UHTDebugging/DocumentationObjects.h
@@ -13,4 +13,17 @@ UCLASS(meta=(DocumentationPolicy = "Strict"))
 class UClassToDocument : public UObject
 {
 	GENERATED_BODY()
+
+	UClassToDocument()
+		: MyProperty(0.f)
+	{
+	}
+
+	// A property for testing the policy
+	UPROPERTY(meta = (UIMin = "0.0", UIMax = "1.0"))
+	float MyProperty;
+
+	// A function to test the policy
+	UFUNCTION()
+	void MyFunction() {}
 };
diff --git a/Engine/Source/Runtime/AnimGraphRuntime/Private/AnimNodes/AnimNode_BlendListBase.cpp b/Engine/Source/Runtime/AnimGraphRuntime/Private/AnimNodes/AnimNode_BlendListBase.cpp
--- a/Engine/Source/Runtime/AnimGraphRuntime/Private/AnimNodes/AnimNode_BlendListBase.cpp
+++ b/Engine/Source/Runtime/AnimGraphRuntime/Private/AnimNodes/AnimNode_BlendListBase.cpp
@@ -113,10 +113,6 @@ void FAnimNode_BlendListBase::Update_AnyThread(const FAnimationUpdateContext& Co
 				{
 					InertializationNode->Request(BlendTime[ChildIndex]);
 				}
-				else
-				{
-					UE_LOG(LogAnimation, Error, TEXT("FAnimNode_BlendListBase: No inertialization context found for transition"));
-				}
 				
 				RemainingBlendTime = 0.0f;
 			}
diff --git a/Engine/Source/Runtime/AnimGraphRuntime/Private/AnimNodes/AnimNode_MultiWayBlend.cpp b/Engine/Source/Runtime/AnimGraphRuntime/Private/AnimNodes/AnimNode_MultiWayBlend.cpp
--- a/Engine/Source/Runtime/AnimGraphRuntime/Private/AnimNodes/AnimNode_MultiWayBlend.cpp
+++ b/Engine/Source/Runtime/AnimGraphRuntime/Private/AnimNodes/AnimNode_MultiWayBlend.cpp
@@ -83,6 +83,8 @@ void FAnimNode_MultiWayBlend::UpdateCachedAlphas()
 			}
 		}
 	}
+
+	ensure(Poses.Num() == CachedAlphas.Num());
 }
 
 void FAnimNode_MultiWayBlend::Update_AnyThread(const FAnimationUpdateContext& Context)
diff --git a/Engine/Source/Runtime/AnimGraphRuntime/Private/AnimSequencerInstanceProxy.cpp b/Engine/Source/Runtime/AnimGraphRuntime/Private/AnimSequencerInstanceProxy.cpp
--- a/Engine/Source/Runtime/AnimGraphRuntime/Private/AnimSequencerInstanceProxy.cpp
+++ b/Engine/Source/Runtime/AnimGraphRuntime/Private/AnimSequencerInstanceProxy.cpp
@@ -92,7 +92,6 @@ void FAnimSequencerInstanceProxy::InitAnimTrack(UAnimSequenceBase* InAnimSequenc
 			}
 
 			const int32 PoseIndex = BlendNode.AddPose() - 1;
-			BlendNode.UpdateCachedAlphas();
 
 			// add the new entry to map
 			FSequencerPlayerAnimSequence* NewPlayerState = new FSequencerPlayerAnimSequence();
diff --git a/Engine/Source/Runtime/AnimGraphRuntime/Private/BoneControllers/AnimNode_RigidBody.cpp b/Engine/Source/Runtime/AnimGraphRuntime/Private/BoneControllers/AnimNode_RigidBody.cpp
--- a/Engine/Source/Runtime/AnimGraphRuntime/Private/BoneControllers/AnimNode_RigidBody.cpp
+++ b/Engine/Source/Runtime/AnimGraphRuntime/Private/BoneControllers/AnimNode_RigidBody.cpp
@@ -82,6 +82,21 @@ void FAnimNode_RigidBody::GatherDebugData(FNodeDebugData& DebugData)
 	}
 }
 
+FTransform SpaceToWorldTransform(ESimulationSpace Space, const FTransform& ComponentToWorld, const FTransform& BaseBoneTM)
+{
+	switch (Space)
+	{
+	case ESimulationSpace::ComponentSpace: 
+		return ComponentToWorld;
+	case ESimulationSpace::WorldSpace: 
+		return FTransform::Identity;
+	case ESimulationSpace::BaseBoneSpace:
+		return BaseBoneTM * ComponentToWorld;
+	default:
+		return FTransform::Identity;
+	}
+}
+
 FVector WorldVectorToSpaceNoScale(ESimulationSpace Space, const FVector& WorldDir, const FTransform& ComponentToWorld, const FTransform& BaseBoneTM)
 {
 	switch(Space)
@@ -252,6 +267,11 @@ void FAnimNode_RigidBody::EvaluateSkeletalControl_AnyThread(FComponentSpacePoseC
 		}
 		const FTransform BaseBoneTM = Output.Pose.GetComponentSpaceTransform(BaseBoneRef.GetCompactPoseIndex(BoneContainer));
 
+#if !UE_BUILD_SHIPPING
+		// Only used for debug draw...
+		PhysicsSimulation->SetSimulationSpaceTransform(SpaceToWorldTransform(SimulationSpace, CompWorldSpaceTM, BaseBoneTM));
+#endif
+
 		// Initialize potential new bodies because of LOD change.
 		if (ResetSimulatedTeleportType == ETeleportType::None && bCheckForBodyTransformInit)
 		{
diff --git a/Engine/Source/Runtime/AnimGraphRuntime/Private/BoneControllers/AnimNode_RigidBody_Chaos.cpp b/Engine/Source/Runtime/AnimGraphRuntime/Private/BoneControllers/AnimNode_RigidBody_Chaos.cpp
--- a/Engine/Source/Runtime/AnimGraphRuntime/Private/BoneControllers/AnimNode_RigidBody_Chaos.cpp
+++ b/Engine/Source/Runtime/AnimGraphRuntime/Private/BoneControllers/AnimNode_RigidBody_Chaos.cpp
@@ -22,8 +22,6 @@
 
 #define LOCTEXT_NAMESPACE "ImmediatePhysics"
 
-#if INCLUDE_CHAOS
-
 TAutoConsoleVariable<int32> CVarEnableChaosRigidBodyNode(TEXT("p.ChaosRigidBodyNode"), 1, TEXT("Enables/disables chaos rigid body node updates and evaluations"), ECVF_Scalability);
 TAutoConsoleVariable<int32> CVarChaosRigidBodyLODThreshold(TEXT("p.ChaosRigidBodyLODThreshold"), -1, TEXT("Max LOD that chaos rigid body node is allowed to run on. Provides a global threshold that overrides per-node the LODThreshold property. -1 means no override."), ECVF_Scalability);
 
@@ -1059,8 +1057,6 @@ bool FAnimNode_RigidBody_Chaos::IsValidToEvaluate(const USkeleton* Skeleton, con
 	return BaseBoneRef.IsValidToEvaluate(RequiredBones);
 }
 
-#endif // INCLUDE_CHAOS
-
 #if WITH_EDITORONLY_DATA
 void FAnimNode_RigidBody_Chaos::PostSerialize(const FArchive& Ar)
 {
diff --git a/Engine/Source/Runtime/AnimGraphRuntime/Private/CommonAnimationLibrary.cpp b/Engine/Source/Runtime/AnimGraphRuntime/Private/CommonAnimationLibrary.cpp
--- a/Engine/Source/Runtime/AnimGraphRuntime/Private/CommonAnimationLibrary.cpp
+++ b/Engine/Source/Runtime/AnimGraphRuntime/Private/CommonAnimationLibrary.cpp
@@ -109,8 +109,8 @@ float CommonAnimationLibrary::ScalarEasing(float Value, const FRuntimeFloatCurve
 
 FVector CommonAnimationLibrary::RetargetSingleLocation(
 	FVector Location,
-	FTransform Source,
-	FTransform Target,
+	const FTransform& Source,
+	const FTransform& Target,
 	const FRuntimeFloatCurve& CustomCurve,
 	EEasingFuncType EasingType,
 	bool bFlipEasing,
@@ -156,9 +156,9 @@ FVector CommonAnimationLibrary::RetargetSingleLocation(
 }
 
 FQuat CommonAnimationLibrary::RetargetSingleRotation(
-	FQuat Rotation,
-	FTransform Source,
-	FTransform Target,
+	const FQuat& RotationIn,
+	const FTransform& Source,
+	const FTransform& Target,
 	const FRuntimeFloatCurve& CustomCurve,
 	EEasingFuncType EasingType,
 	bool bFlipEasing,
@@ -172,6 +172,8 @@ FQuat CommonAnimationLibrary::RetargetSingleRotation(
 	float TargetMaximum
 )
 {
+	FQuat Rotation = RotationIn;
+
 	float SourceDelta = SourceMaximum - SourceMinimum;
 	float TargetDelta = TargetMaximum - TargetMinimum;
 	if (FMath::IsNearlyEqual(SourceDelta, 0.f) || FMath::IsNearlyEqual(TargetDelta, 0.f))
diff --git a/Engine/Source/Runtime/AnimGraphRuntime/Public/AnimNodes/AnimNode_MultiWayBlend.h b/Engine/Source/Runtime/AnimGraphRuntime/Public/AnimNodes/AnimNode_MultiWayBlend.h
--- a/Engine/Source/Runtime/AnimGraphRuntime/Public/AnimNodes/AnimNode_MultiWayBlend.h
+++ b/Engine/Source/Runtime/AnimGraphRuntime/Public/AnimNodes/AnimNode_MultiWayBlend.h
@@ -12,8 +12,6 @@ struct ANIMGRAPHRUNTIME_API FAnimNode_MultiWayBlend : public FAnimNode_Base
 {
 	GENERATED_USTRUCT_BODY()
 
-	friend struct FAnimSequencerInstanceProxy;
-
 public:
 
 	UPROPERTY(EditAnywhere, BlueprintReadWrite, Category = Links)
@@ -54,6 +52,7 @@ public:
 	{
 		Poses.AddDefaulted();
 		DesiredAlphas.AddZeroed();
+		UpdateCachedAlphas();
 
 		return Poses.Num();
 	}
@@ -61,12 +60,14 @@ public:
 	void RemovePose(int32 PoseIndex)
 	{
 		Poses.RemoveAt(PoseIndex);
+		CachedAlphas.RemoveAt(PoseIndex);
 		DesiredAlphas.RemoveAt(PoseIndex);
 	}
 
 	void ResetPoses()
 	{
 		Poses.Reset();
+		CachedAlphas.Reset();
 		DesiredAlphas.Reset();
 	}
 
diff --git a/Engine/Source/Runtime/AnimGraphRuntime/Public/BoneControllers/AnimNode_RigidBody_Chaos.h b/Engine/Source/Runtime/AnimGraphRuntime/Public/BoneControllers/AnimNode_RigidBody_Chaos.h
--- a/Engine/Source/Runtime/AnimGraphRuntime/Public/BoneControllers/AnimNode_RigidBody_Chaos.h
+++ b/Engine/Source/Runtime/AnimGraphRuntime/Public/BoneControllers/AnimNode_RigidBody_Chaos.h
@@ -143,8 +143,6 @@ private:
 	bool bComponentSpaceSimulation_DEPRECATED;	//use SimulationSpace
 #endif
 
-#if INCLUDE_CHAOS
-
 	// FAnimNode_SkeletalControlBase interface
 	virtual void InitializeBoneReferences(const FBoneContainer& RequiredBones) override;
 	// End of FAnimNode_SkeletalControlBase interface
@@ -252,46 +250,8 @@ private:
 	FBlendedHeapCurve CapturedFrozenCurves;
 
 	FVector PreviousComponentLinearVelocity;
-#endif
 };
 
-#if !INCLUDE_CHAOS
-// Stub out public API and safe-initialize all public properties if Chaos is not compiled.
-inline FAnimNode_RigidBody_Chaos::FAnimNode_RigidBody_Chaos()
-	: OverridePhysicsAsset(nullptr)
-	, OverrideWorldGravity(FVector::ZeroVector)
-	, ExternalForce(FVector::ZeroVector)
-	, ComponentLinearAccScale(FVector::ZeroVector)
-	, ComponentLinearVelScale(FVector::ZeroVector)
-	, ComponentAppliedLinearAccClamp(FVector::ZeroVector)
-	, CachedBoundsScale(false)
-	, BaseBoneRef()
-	, OverlapChannel(ECollisionChannel::ECC_WorldStatic)
-	, SimulationSpace(ESimulationSpace::ComponentSpace)
-	, bForceDisableCollisionBetweenConstraintBodies(false)
-	, bEnableWorldGeometry(false)
-	, bOverrideWorldGravity(false)
-	, bTransferBoneVelocities(0)
-	, bFreezeIncomingPoseOnStart(0)
-	, bClampLinearTranslationLimitToRefPose(0)
-#if WITH_EDITORONLY_DATA
-	, bComponentSpaceSimulation_DEPRECATED(false)
-#endif
-{}
-inline FAnimNode_RigidBody_Chaos::~FAnimNode_RigidBody_Chaos() {}
-inline void FAnimNode_RigidBody_Chaos::FAnimNode_RigidBody_Chaos::GatherDebugData(FNodeDebugData& DebugData) { Super::GatherDebugData(DebugData); }
-inline void FAnimNode_RigidBody_Chaos::UpdateComponentPose_AnyThread(const FAnimationUpdateContext& Context) { Super::UpdateComponentPose_AnyThread(Context); }
-inline void FAnimNode_RigidBody_Chaos::EvaluateComponentPose_AnyThread(FComponentSpacePoseContext& Output) { Super::EvaluateComponentPose_AnyThread(Output); }
-inline void FAnimNode_RigidBody_Chaos::EvaluateSkeletalControl_AnyThread(FComponentSpacePoseContext& Output, TArray<FBoneTransform>& OutBoneTransforms) { Super::EvaluateSkeletalControl_AnyThread(Output, OutBoneTransforms); }
-inline void FAnimNode_RigidBody_Chaos::OnInitializeAnimInstance(const FAnimInstanceProxy* InProxy, const UAnimInstance* InAnimInstance) { Super::OnInitializeAnimInstance(InProxy, InAnimInstance); }
-inline void FAnimNode_RigidBody_Chaos::PreUpdate(const UAnimInstance* InAnimInstance) { Super::PreUpdate(InAnimInstance); }
-inline void FAnimNode_RigidBody_Chaos::UpdateInternal(const FAnimationUpdateContext& Context) { Super::UpdateInternal(Context); }
-inline bool FAnimNode_RigidBody_Chaos::IsValidToEvaluate(const USkeleton* Skeleton, const FBoneContainer& RequiredBones) { return true; }
-inline bool FAnimNode_RigidBody_Chaos::NeedsDynamicReset() const { return false; }
-inline void FAnimNode_RigidBody_Chaos::ResetDynamics(ETeleportType InTeleportType) {}
-inline int32 FAnimNode_RigidBody_Chaos::GetLODThreshold() const { return 0; }
-#endif
-
 #if WITH_EDITORONLY_DATA
 template<>
 struct TStructOpsTypeTraits<FAnimNode_RigidBody_Chaos> : public TStructOpsTypeTraitsBase2<FAnimNode_RigidBody_Chaos>
diff --git a/Engine/Source/Runtime/AnimGraphRuntime/Public/CommonAnimationLibrary.h b/Engine/Source/Runtime/AnimGraphRuntime/Public/CommonAnimationLibrary.h
--- a/Engine/Source/Runtime/AnimGraphRuntime/Public/CommonAnimationLibrary.h
+++ b/Engine/Source/Runtime/AnimGraphRuntime/Public/CommonAnimationLibrary.h
@@ -45,8 +45,8 @@ namespace CommonAnimationLibrary
 	 */
 	ANIMGRAPHRUNTIME_API FVector RetargetSingleLocation(
 		FVector Location,
-		FTransform Source,
-		FTransform Target,
+		const FTransform& Source,
+		const FTransform& Target,
 		const FRuntimeFloatCurve& CustomCurve,
 		EEasingFuncType EasingType = EEasingFuncType::Linear,
 		bool bFlipEasing = false,
@@ -77,9 +77,9 @@ namespace CommonAnimationLibrary
 	 *	@param EasingType The type of easing to apply
 	 */
 	ANIMGRAPHRUNTIME_API FQuat RetargetSingleRotation(
-		FQuat Rotation,
-		FTransform Source,
-		FTransform Target,
+		const FQuat& Rotation,
+		const FTransform& Source,
+		const FTransform& Target,
 		const FRuntimeFloatCurve& CustomCurve,
 		EEasingFuncType EasingType = EEasingFuncType::Linear,
 		bool bFlipEasing = false,
diff --git a/Engine/Source/Runtime/ClothingSystemRuntimeCommon/Private/ClothLODData.cpp b/Engine/Source/Runtime/ClothingSystemRuntimeCommon/Private/ClothLODData.cpp
new file mode 100644
--- /dev/null
+++ b/Engine/Source/Runtime/ClothingSystemRuntimeCommon/Private/ClothLODData.cpp
@@ -0,0 +1,66 @@
+// Copyright 1998-2019 Epic Games, Inc. All Rights Reserved.
+
+#include "ClothLODData.h"
+
+
+UClothLODDataBase::UClothLODDataBase(const FObjectInitializer& Init)
+	: PhysicalMeshData(nullptr)
+{}
+
+UClothLODDataBase::~UClothLODDataBase()
+{}
+
+#if WITH_EDITORONLY_DATA
+void UClothLODDataBase::GetParameterMasksForTarget(
+	const uint8 InTarget, 
+	TArray<FPointWeightMap*>& OutMasks)
+{
+	for(FPointWeightMap& Mask : ParameterMasks)
+	{
+		if(Mask.CurrentTarget == InTarget)
+		{
+			OutMasks.Add(&Mask);
+		}
+	}
+}
+#endif // WITH_EDITORONLY_DATA
+
+void UClothLODDataBase::Serialize(FArchive& Ar)
+{
+	Super::Serialize(Ar);
+
+	// Ryan - do we need to do this tagged property business still?
+
+	// Serialize normal tagged data
+	//UScriptStruct* Struct = UClothLODDataNv::StaticStruct();
+	//UClass* Class = UClothLODDataNv::StaticClass();
+	/*if (!Ar.IsCountingMemory())
+	{
+		//Struct->SerializeTaggedProperties(Ar, (uint8*)this, Struct, nullptr);
+		Class->SerializeTaggedProperties(Ar, (uint8*)this, Class, nullptr);
+	}*/
+
+	// Serialize the mesh to mesh data (not a USTRUCT)
+	Ar << TransitionUpSkinData
+	   << TransitionDownSkinData;
+}
+
+#if WITH_EDITOR
+void UClothLODDataBase::PushWeightsToMesh()
+{
+	if (PhysicalMeshData)
+	{
+		PhysicalMeshData->ClearParticleParameters();
+		for (const FPointWeightMap &Weights : ParameterMasks)
+		{
+			if (Weights.bEnabled)
+			{
+				if (TArray<float>* TargetArray = PhysicalMeshData->GetFloatArray(Weights.CurrentTarget))
+				{
+					*TargetArray = Weights.GetValueArray();
+				}
+			}
+		}
+	}
+}
+#endif
diff --git a/Engine/Source/Runtime/ClothingSystemRuntimeCommon/Private/ClothingAsset.cpp b/Engine/Source/Runtime/ClothingSystemRuntimeCommon/Private/ClothingAsset.cpp
new file mode 100644
--- /dev/null
+++ b/Engine/Source/Runtime/ClothingSystemRuntimeCommon/Private/ClothingAsset.cpp
@@ -0,0 +1,692 @@
+#include "ClothingAsset.h"
+
+#include "Utils/ClothingMeshUtils.h"
+
+#if WITH_EDITOR
+#include "Engine/SkeletalMesh.h"
+#endif
+
+#include "PhysicsEngine/PhysicsAsset.h"
+
+#include "Widgets/Notifications/SNotificationList.h"
+#include "Framework/Notifications/NotificationManager.h"
+#include "Rendering/SkeletalMeshModel.h"
+#include "Rendering/SkeletalMeshRenderData.h"
+
+#include "Components/SkeletalMeshComponent.h"
+
+#include "ClothingSimulationInteractor.h"
+#include "ComponentReregisterContext.h"
+#include "UObject/UObjectIterator.h"
+
+#include "GPUSkinPublicDefs.h"
+#include "GPUSkinVertexFactory.h"
+
+DEFINE_LOG_CATEGORY(LogClothingAsset)
+#define LOCTEXT_NAMESPACE "ClothingAsset"
+
+//==============================================================================
+// ClothingAssetUtils
+//==============================================================================
+
+void ClothingAssetUtils::GetMeshClothingAssetBindings(
+	USkeletalMesh* InSkelMesh, 
+	TArray<FClothingAssetMeshBinding>& OutBindings)
+{
+	OutBindings.Empty();
+
+	if(!InSkelMesh)
+	{
+		return;
+	}
+
+	if(FSkeletalMeshRenderData* Resource = InSkelMesh->GetResourceForRendering())
+	{
+		const int32 NumLods = Resource->LODRenderData.Num();
+
+		for(int32 LodIndex = 0; LodIndex < NumLods; ++LodIndex)
+		{
+			TArray<FClothingAssetMeshBinding> LodBindings;
+			GetMeshClothingAssetBindings(InSkelMesh, LodBindings, LodIndex);
+
+			OutBindings.Append(LodBindings);
+		}
+	}
+}
+
+void ClothingAssetUtils::GetMeshClothingAssetBindings(
+	USkeletalMesh* InSkelMesh, 
+	TArray<FClothingAssetMeshBinding>& OutBindings, 
+	int32 InLodIndex)
+{
+	OutBindings.Empty();
+
+	if(!InSkelMesh)
+	{
+		return;
+	}
+
+	if(FSkeletalMeshRenderData* Resource = InSkelMesh->GetResourceForRendering())
+	{
+		if(Resource->LODRenderData.IsValidIndex(InLodIndex))
+		{
+			FSkeletalMeshLODRenderData& LodData = Resource->LODRenderData[InLodIndex];
+
+			const int32 NumSections = LodData.RenderSections.Num();
+
+			for(int32 SectionIndex = 0; SectionIndex < NumSections; ++SectionIndex)
+			{
+				FSkelMeshRenderSection& Section = LodData.RenderSections[SectionIndex];
+
+				if(Section.HasClothingData())
+				{
+					UClothingAssetCommon* SectionAsset = Cast<UClothingAssetCommon>(InSkelMesh->GetSectionClothingAsset(InLodIndex, SectionIndex));
+
+					if(SectionAsset)
+					{
+						// This is the original section of a clothing section pair
+						OutBindings.AddDefaulted();
+						FClothingAssetMeshBinding& Binding = OutBindings.Last();
+
+						Binding.Asset = SectionAsset;
+						Binding.LODIndex = InLodIndex;
+						Binding.SectionIndex = SectionIndex;
+						Binding.AssetInternalLodIndex = Section.ClothingData.AssetLodIndex;
+					}
+				}
+			}
+		}
+	}
+}
+
+#if WITH_EDITOR
+void ClothingAssetUtils::ClearSectionClothingData(FSkelMeshSection& InSection)
+{
+	InSection.ClothingData.AssetGuid = FGuid();
+	InSection.ClothingData.AssetLodIndex = INDEX_NONE;
+	InSection.CorrespondClothAssetIndex = INDEX_NONE;
+
+	InSection.ClothMappingData.Empty();
+}
+#endif
+
+//==============================================================================
+// UClothingAssetCommon
+//==============================================================================
+
+UClothingAssetCommon::UClothingAssetCommon(const FObjectInitializer& ObjectInitializer)
+	: Super(ObjectInitializer)
+	, PhysicsAsset(nullptr)
+	, ClothSimConfig(nullptr)
+	, ReferenceBoneIndex(0)
+	, CustomData(nullptr)
+{}
+
+#if WITH_EDITOR
+
+void Warn(const FText& Error)
+{
+	FNotificationInfo Info(Error);
+	Info.ExpireDuration = 5.0f;
+	FSlateNotificationManager::Get().AddNotification(Info);
+
+	UE_LOG(LogClothingAsset, Warning, TEXT("%s"), *Error.ToString());
+}
+
+bool UClothingAssetCommon::BindToSkeletalMesh(
+	USkeletalMesh* InSkelMesh, 
+	const int32 InMeshLodIndex, 
+	const int32 InSectionIndex, 
+	const int32 InAssetLodIndex, 
+	const bool bCallPostEditChange)
+{
+	// If we've been added to the wrong mesh
+	if(InSkelMesh != GetOuter())
+	{
+		FText Error = FText::Format(LOCTEXT(
+			"Error_WrongMesh", 
+			"Failed to bind clothing asset {0} as the provided mesh is not the owner of this asset."), 
+			FText::FromString(GetName()));
+		Warn(Error);
+		return false;
+	}
+
+	// If we don't have clothing data
+	if(!ClothLodData.IsValidIndex(InAssetLodIndex))
+	{
+		FText Error = FText::Format(LOCTEXT(
+			"Error_NoClothingLod", 
+			"Failed to bind clothing asset {0} LOD{1} as LOD{2} does not exist."), 
+			FText::FromString(GetName()), 
+			FText::AsNumber(InAssetLodIndex), 
+			FText::AsNumber(InAssetLodIndex));
+		Warn(Error);
+		return false;
+	}
+
+	// If we don't have a mesh
+	if(!InSkelMesh)
+	{
+		FText Error = FText::Format(LOCTEXT(
+			"Error_NoMesh", 
+			"Failed to bind clothing asset {0} as provided skel mesh does not exist."), 
+			FText::FromString(GetName()));
+		Warn(Error);
+		return false;
+	}
+
+	// If the mesh LOD index is invalid
+	if(!InSkelMesh->GetImportedModel()->LODModels.IsValidIndex(InMeshLodIndex))
+	{
+		FText Error = FText::Format(LOCTEXT(
+			"Error_InvalidMeshLOD", 
+			"Failed to bind clothing asset {0} as mesh LOD{1} does not exist."), 
+			FText::FromString(GetName()), 
+			FText::AsNumber(InMeshLodIndex));
+		Warn(Error);
+		return false;
+	}
+
+	const int32 NumMapEntries = LodMap.Num();
+	for(int MapIndex = 0; MapIndex < NumMapEntries; ++MapIndex)
+	{
+		const int32& MappedLod = LodMap[MapIndex];
+		if(MappedLod == InAssetLodIndex)
+		{
+			FText Error = FText::Format(LOCTEXT(
+				"Error_LodMapped", 
+				"Failed to bind clothing asset {0} LOD{1} as LOD{2} is already mapped to mesh LOD{3}."), 
+				FText::FromString(GetName()), 
+				FText::AsNumber(InAssetLodIndex), 
+				FText::AsNumber(InAssetLodIndex), 
+				FText::AsNumber(MapIndex));
+			Warn(Error);
+			return false;
+		}
+	}
+
+	if(LodMap.IsValidIndex(InMeshLodIndex) && LodMap[InMeshLodIndex] != INDEX_NONE)
+	{
+		// Already mapped
+		return false;
+	}
+
+	BuildSelfCollisionData();
+	CalculateReferenceBoneIndex();
+
+	// Grab the clothing and skel lod data
+	UClothLODDataBase* LodData = ClothLodData[InAssetLodIndex];
+	FSkeletalMeshLODModel& SkelLod = InSkelMesh->GetImportedModel()->LODModels[InMeshLodIndex];
+
+	FSkelMeshSection& OriginalSection = SkelLod.Sections[InSectionIndex];
+
+	// Data for mesh to mesh binding
+	TArray<FMeshToMeshVertData> MeshToMeshData;
+	TArray<FVector> RenderPositions;
+	TArray<FVector> RenderNormals;
+	TArray<FVector> RenderTangents;
+
+	RenderPositions.Reserve(OriginalSection.SoftVertices.Num());
+	RenderNormals.Reserve(OriginalSection.SoftVertices.Num());
+	RenderTangents.Reserve(OriginalSection.SoftVertices.Num());
+
+	// Original data to weight to the clothing simulation mesh
+	for(FSoftSkinVertex& UnrealVert : OriginalSection.SoftVertices)
+	{
+		RenderPositions.Add(UnrealVert.Position);
+		RenderNormals.Add(UnrealVert.TangentZ);
+		RenderTangents.Add(UnrealVert.TangentX);
+	}
+
+	TArrayView<uint32> IndexView(SkelLod.IndexBuffer);
+	IndexView.Slice(OriginalSection.BaseIndex, OriginalSection.NumTriangles * 3);
+
+	ClothingMeshUtils::ClothMeshDesc TargetMesh(RenderPositions, RenderNormals, IndexView);
+	ClothingMeshUtils::ClothMeshDesc SourceMesh(
+		LodData->PhysicalMeshData->Vertices, 
+		LodData->PhysicalMeshData->Normals, 
+		LodData->PhysicalMeshData->Indices);
+
+	ClothingMeshUtils::GenerateMeshToMeshSkinningData(MeshToMeshData, TargetMesh, &RenderTangents, SourceMesh);
+
+	if(MeshToMeshData.Num() == 0)
+	{
+		// Failed to generate skinning data, the function above will have notified
+		// with the cause of the failure, so just exit
+		return false;
+	}
+
+	// Calculate fixed verts
+	for(FMeshToMeshVertData& VertData : MeshToMeshData)
+	{
+		if(LodData->PhysicalMeshData->IsFixed(
+			VertData.SourceMeshVertIndices[0], 
+			VertData.SourceMeshVertIndices[1], 
+			VertData.SourceMeshVertIndices[2])) // Default threshold is 0.1, not 0.0.  Using 0.1 for consistency.
+		{
+			VertData.SourceMeshVertIndices[3] = 0xFFFF;
+		}
+	}
+
+	// We have to copy the bone map to verify we don't exceed the maximum while adding the clothing bones
+	TArray<FBoneIndexType> TempBoneMap = OriginalSection.BoneMap;
+	for(FName& BoneName : UsedBoneNames)
+	{
+		const int32 BoneIndex = InSkelMesh->RefSkeleton.FindBoneIndex(BoneName);
+		if(BoneIndex != INDEX_NONE)
+		{
+			TempBoneMap.AddUnique(BoneIndex);
+		}
+	}
+	
+	// Verify number of bones against current capabilities
+	if(TempBoneMap.Num() > FGPUBaseSkinVertexFactory::GetMaxGPUSkinBones())
+	{
+		// Failed to apply as we've exceeded the number of bones we can skin
+		FText Error = FText::Format(LOCTEXT(
+			"Error_TooManyBones", 
+			"Failed to bind clothing asset {0} LOD{1} as this causes the section to require {2} bones. "
+			"The maximum per section is currently {3}."), 
+			FText::FromString(GetName()), 
+			FText::AsNumber(InAssetLodIndex), 
+			FText::AsNumber(TempBoneMap.Num()), 
+			FText::AsNumber(FGPUBaseSkinVertexFactory::GetMaxGPUSkinBones()));
+		Warn(Error);
+		return false;
+	}
+
+	// After verifying copy the new bone map to the section
+	OriginalSection.BoneMap = TempBoneMap;
+
+	// Array of re-import contexts for components using this mesh
+	TIndirectArray<FComponentReregisterContext> ComponentContexts;
+	for (TObjectIterator<USkeletalMeshComponent> It; It; ++It)
+	{
+		USkeletalMeshComponent* Component = *It;
+		if (Component && !Component->IsTemplate() && Component->SkeletalMesh == InSkelMesh)
+		{
+			ComponentContexts.Add(new FComponentReregisterContext(Component));
+		}
+	}
+
+	// Ready to apply the changes
+	InSkelMesh->PreEditChange(nullptr);
+
+	// calculate LOD verts before adding our new section
+	uint32 NumLodVertices = 0;
+	for(const FSkelMeshSection& CurSection : SkelLod.Sections)
+	{
+		NumLodVertices += CurSection.GetNumVertices();
+	}
+
+	// Set the asset index, used during rendering to pick the correct sim mesh buffer
+	int32 AssetIndex = INDEX_NONE;
+	check(InSkelMesh->MeshClothingAssets.Find(this, AssetIndex));
+	OriginalSection.CorrespondClothAssetIndex = AssetIndex;
+
+	// sim properties
+	OriginalSection.ClothMappingData = MeshToMeshData;
+	OriginalSection.ClothingData.AssetGuid = AssetGuid;
+	OriginalSection.ClothingData.AssetLodIndex = InAssetLodIndex;
+
+	bool bRequireBoneChange = false;
+	for(FBoneIndexType& BoneIndex : OriginalSection.BoneMap)
+	{
+		if(!SkelLod.RequiredBones.Contains(BoneIndex))
+		{
+			bRequireBoneChange = true;
+			if(InSkelMesh->RefSkeleton.IsValidIndex(BoneIndex))
+			{
+				SkelLod.RequiredBones.Add(BoneIndex);
+				SkelLod.ActiveBoneIndices.AddUnique(BoneIndex);
+			}
+		}
+	}
+	if(bRequireBoneChange)
+	{
+		SkelLod.RequiredBones.Sort();
+		InSkelMesh->RefSkeleton.EnsureParentsExistAndSort(SkelLod.ActiveBoneIndices);		
+	}
+
+	if(CustomData)
+	{
+		CustomData->BindToSkeletalMesh(InSkelMesh, InMeshLodIndex, InSectionIndex, InAssetLodIndex);
+	}
+
+	// Make sure the LOD map is always big enough for the asset to use.
+	// This shouldn't grow to an unwieldy size but maybe consider compacting later.
+	while(LodMap.Num() - 1 < InMeshLodIndex)
+	{
+		LodMap.Add(INDEX_NONE);
+	}
+
+	LodMap[InMeshLodIndex] = InAssetLodIndex;
+
+	if (bCallPostEditChange)
+	{
+		InSkelMesh->PostEditChange();
+	}
+
+	return true;
+	// ComponentContexts goes out of scope, causing components to be re-registered
+}
+
+void UClothingAssetCommon::UnbindFromSkeletalMesh(USkeletalMesh* InSkelMesh)
+{
+	if(FSkeletalMeshModel* Mesh = InSkelMesh->GetImportedModel())
+	{
+		const int32 NumLods = Mesh->LODModels.Num();
+		for(int32 LodIndex = 0; LodIndex < NumLods; ++LodIndex)
+		{
+			UnbindFromSkeletalMesh(InSkelMesh, LodIndex);
+		}
+	}
+}
+
+void UClothingAssetCommon::UnbindFromSkeletalMesh(
+	USkeletalMesh* InSkelMesh, 
+	const int32 InMeshLodIndex)
+{
+	bool bChangedMesh = false;
+
+	// Find the chunk(s) we created
+	if(FSkeletalMeshModel* Mesh = InSkelMesh->GetImportedModel())
+	{
+		if(!Mesh->LODModels.IsValidIndex(InMeshLodIndex))
+		{
+			FText Error = FText::Format(LOCTEXT(
+				"Error_UnbindNoMeshLod", 
+				"Failed to remove clothing asset {0} from mesh LOD{1} as that LOD doesn't exist."), 
+				FText::FromString(GetName()), 
+				FText::AsNumber(InMeshLodIndex));
+			Warn(Error);
+
+			return;
+		}
+
+		FSkeletalMeshLODModel& LodModel = Mesh->LODModels[InMeshLodIndex];
+
+		for(int32 SectionIdx = LodModel.Sections.Num() - 1; SectionIdx >= 0; --SectionIdx)
+		{
+			FSkelMeshSection& Section = LodModel.Sections[SectionIdx];
+			if(Section.HasClothingData() && Section.ClothingData.AssetGuid == AssetGuid)
+			{
+				if(!bChangedMesh)
+				{
+					InSkelMesh->PreEditChange(nullptr);
+				}
+				ClothingAssetUtils::ClearSectionClothingData(Section);
+				bChangedMesh = true;
+			}
+		}
+
+		// Clear the LOD map entry for this asset LOD, after a unbind we must be able to bind any asset
+		if (LodMap.IsValidIndex(InMeshLodIndex))
+		{
+			LodMap[InMeshLodIndex] = INDEX_NONE;
+			bChangedMesh = true;
+		}
+	}
+
+	// If the mesh changed we need to re-register any components that use it to reflect the changes
+	if(bChangedMesh)
+	{
+		InSkelMesh->PostEditChange();
+
+		for(TObjectIterator<USkeletalMeshComponent> It; It; ++It)
+		{
+			USkeletalMeshComponent* MeshComponent = *It;
+			if(MeshComponent &&
+			   !MeshComponent->IsTemplate() &&
+			   MeshComponent->SkeletalMesh == InSkelMesh)
+			{
+				MeshComponent->ReregisterComponent();
+			}
+		}
+	}
+}
+
+void UClothingAssetCommon::ReregisterComponentsUsingClothing()
+{
+	if(USkeletalMesh* OwnerMesh = Cast<USkeletalMesh>(GetOuter()))
+	{
+		for(TObjectIterator<USkeletalMeshComponent> It; It; ++It)
+		{
+			if(USkeletalMeshComponent* Component = *It)
+			{
+				if(Component->SkeletalMesh == OwnerMesh)
+				{
+					FComponentReregisterContext Context(Component);
+					// Context goes out of scope, causing Component to be re-registered
+				}
+			}
+		}
+	}
+}
+
+void UClothingAssetCommon::ForEachInteractorUsingClothing(TFunction<void(UClothingSimulationInteractor*)> Func)
+{
+	if(USkeletalMesh* OwnerMesh = Cast<USkeletalMesh>(GetOuter()))
+	{
+		for(TObjectIterator<USkeletalMeshComponent> It; It; ++It)
+		{
+			if(USkeletalMeshComponent* Component = *It)
+			{
+				if(Component->SkeletalMesh == OwnerMesh)
+				{
+					if(UClothingSimulationInteractor* CurInteractor = Component->GetClothingSimulationInteractor())
+					{
+						Func(CurInteractor);
+					}
+				}
+			}
+		}
+	}
+}
+
+void UClothingAssetCommon::ApplyParameterMasks()
+{
+	for(UClothLODDataBase* Lod : ClothLodData)
+	{
+		Lod->PushWeightsToMesh();
+	}
+	InvalidateCachedData();
+}
+
+void UClothingAssetCommon::BuildLodTransitionData()
+{
+	const int32 NumLods = GetNumLods();
+	for(int32 LodIndex = 0; LodIndex < NumLods; ++LodIndex)
+	{
+		const bool bHasPrevLod = LodIndex > 0;
+		const bool bHasNextLod = LodIndex < NumLods - 1;
+
+		UClothLODDataBase* CurrentLod = ClothLodData[LodIndex];
+		check(CurrentLod->PhysicalMeshData);
+		UClothPhysicalMeshDataBase& CurrentPhysMesh = *CurrentLod->PhysicalMeshData;
+
+		UClothLODDataBase* PrevLod = bHasPrevLod ? ClothLodData[LodIndex - 1] : nullptr;
+		UClothLODDataBase* NextLod = bHasNextLod ? ClothLodData[LodIndex + 1] : nullptr;
+
+		const int32 CurrentLodNumVerts = CurrentPhysMesh.Vertices.Num();
+
+		ClothingMeshUtils::ClothMeshDesc CurrentMeshDesc(CurrentPhysMesh.Vertices, CurrentPhysMesh.Normals, CurrentPhysMesh.Indices);
+
+		if(PrevLod)
+		{
+			check(PrevLod->PhysicalMeshData);
+			UClothPhysicalMeshDataBase& PrevPhysMesh = *PrevLod->PhysicalMeshData;
+			CurrentLod->TransitionUpSkinData.Empty(CurrentLodNumVerts);
+			ClothingMeshUtils::ClothMeshDesc PrevMeshDesc(PrevPhysMesh.Vertices, PrevPhysMesh.Normals, PrevPhysMesh.Indices);
+			ClothingMeshUtils::GenerateMeshToMeshSkinningData(CurrentLod->TransitionUpSkinData, CurrentMeshDesc, nullptr, PrevMeshDesc);
+		}
+		if(NextLod)
+		{
+			check(NextLod->PhysicalMeshData);
+			UClothPhysicalMeshDataBase& NextPhysMesh = *NextLod->PhysicalMeshData;
+			CurrentLod->TransitionDownSkinData.Empty(CurrentLodNumVerts);
+			ClothingMeshUtils::ClothMeshDesc NextMeshDesc(NextPhysMesh.Vertices, NextPhysMesh.Normals, NextPhysMesh.Indices);
+			ClothingMeshUtils::GenerateMeshToMeshSkinningData(CurrentLod->TransitionDownSkinData, CurrentMeshDesc, nullptr, NextMeshDesc);
+		}
+	}
+}
+
+#endif // WITH_EDITOR
+
+void UClothingAssetCommon::RefreshBoneMapping(USkeletalMesh* InSkelMesh)
+{
+	// No mesh, can't remap
+	if(!InSkelMesh)
+	{
+		return;
+	}
+
+	if(UsedBoneNames.Num() != UsedBoneIndices.Num())
+	{
+		UsedBoneIndices.Reset();
+		UsedBoneIndices.AddDefaulted(UsedBoneNames.Num());
+	}
+
+	// Repopulate the used indices.
+	for(int32 BoneNameIndex = 0; BoneNameIndex < UsedBoneNames.Num(); ++BoneNameIndex)
+	{
+		UsedBoneIndices[BoneNameIndex] = 
+			InSkelMesh->RefSkeleton.FindBoneIndex(
+				UsedBoneNames[BoneNameIndex]);
+	}
+}
+
+void UClothingAssetCommon::CalculateReferenceBoneIndex()
+{
+	// Starts at root
+	ReferenceBoneIndex = 0;
+
+	// Find the root bone for this clothing asset (common bone for all used bones)
+	typedef TArray<int32> BoneIndexArray;
+
+	// List of valid paths to the root bone from each weighted bone
+	TArray<BoneIndexArray> PathsToRoot;
+	
+	USkeletalMesh* OwnerMesh = Cast<USkeletalMesh>(GetOuter());
+
+	if(OwnerMesh)
+	{
+		FReferenceSkeleton& RefSkel = OwnerMesh->RefSkeleton;
+		// First build a list per used bone for it's path to root
+		const int32 NumUsedBones = UsedBoneIndices.Num();
+
+		// List of actually weighted (not just used) bones
+		TArray<int32> WeightedBones;
+
+		for(UClothLODDataBase* CurLod : ClothLodData)
+		{
+			check(CurLod && CurLod->PhysicalMeshData);
+			UClothPhysicalMeshDataBase* MeshData = CurLod->PhysicalMeshData;
+			for(FClothVertBoneData& VertBoneData : MeshData->BoneData)
+			{
+				for(int32 InfluenceIndex = 0; InfluenceIndex < MAX_TOTAL_INFLUENCES; ++InfluenceIndex)
+				{
+					if(VertBoneData.BoneWeights[InfluenceIndex] > SMALL_NUMBER)
+					{
+						WeightedBones.AddUnique(VertBoneData.BoneIndices[InfluenceIndex]);
+					}
+					else
+					{
+						// Hit the last weight (they're sorted)
+						break;
+					}
+				}
+			}
+		}
+
+		const int32 NumWeightedBones = WeightedBones.Num();
+		PathsToRoot.Reserve(NumWeightedBones);
+		
+		// Compute paths to the root bone
+		for(int32 WeightedBoneIndex = 0; WeightedBoneIndex < NumWeightedBones; ++WeightedBoneIndex)
+		{
+			PathsToRoot.AddDefaulted();
+			BoneIndexArray& Path = PathsToRoot.Last();
+			
+			int32 CurrentBone = WeightedBones[WeightedBoneIndex];
+			Path.Add(CurrentBone);
+			
+			while(CurrentBone != 0 && CurrentBone != INDEX_NONE)
+			{
+				CurrentBone = RefSkel.GetParentIndex(CurrentBone);
+				Path.Add(CurrentBone);
+			}
+		}
+
+		// Paths are from leaf->root, we want the other way
+		for(BoneIndexArray& Path : PathsToRoot)
+		{
+			Algo::Reverse(Path);
+		}
+
+		// Verify the last common bone in all paths as the root of the sim space
+		const int32 NumPaths = PathsToRoot.Num();
+		if(NumPaths > 0)
+		{
+			BoneIndexArray& FirstPath = PathsToRoot[0];
+		
+			const int32 FirstPathSize = FirstPath.Num();
+			for(int32 PathEntryIndex = 0; PathEntryIndex < FirstPathSize; ++PathEntryIndex)
+			{
+				const int32 CurrentQueryIndex = FirstPath[PathEntryIndex];
+				bool bValidRoot = true;
+
+				for(int32 PathIndex = 1; PathIndex < NumPaths; ++PathIndex)
+				{
+					if(!PathsToRoot[PathIndex].Contains(CurrentQueryIndex))
+					{
+						bValidRoot = false;
+						break;
+					}
+				}
+
+				if(bValidRoot)
+				{
+					ReferenceBoneIndex = CurrentQueryIndex;
+				}
+				else
+				{
+					// Once we fail to find a valid root we're done.
+					break;
+				}
+			}
+		}
+		else
+		{
+			// Just use root
+			ReferenceBoneIndex = 0;
+		}
+	}
+}
+
+bool UClothingAssetCommon::IsValidLod(int32 InLodIndex)
+{
+	return ClothLodData.IsValidIndex(InLodIndex);
+}
+
+int32 UClothingAssetCommon::GetNumLods()
+{
+	return ClothLodData.Num();
+}
+
+void UClothingAssetCommon::BuildSelfCollisionData()
+{
+	if(!ClothSimConfig || !ClothSimConfig->HasSelfCollision())
+	{
+		// No self collision, can't generate data
+		return;
+	}
+	for(UClothLODDataBase* Lod : ClothLodData)
+	{
+		check(Lod && Lod->PhysicalMeshData);
+		Lod->PhysicalMeshData->BuildSelfCollisionData(ClothSimConfig);
+	}
+}
+
+
+#undef LOCTEXT_NAMESPACE
diff --git a/Engine/Source/Runtime/ClothingSystemRuntimeCommon/Private/Utils/ClothingMeshUtils.cpp b/Engine/Source/Runtime/ClothingSystemRuntimeCommon/Private/Utils/ClothingMeshUtils.cpp
new file mode 100644
--- /dev/null
+++ b/Engine/Source/Runtime/ClothingSystemRuntimeCommon/Private/Utils/ClothingMeshUtils.cpp
@@ -0,0 +1,347 @@
+// Copyright 1998-2019 Epic Games, Inc. All Rights Reserved.
+
+#include "Utils/ClothingMeshUtils.h"
+
+#include "ClothPhysicalMeshData.h"
+
+#include "Math/UnrealMathUtility.h"
+#include "Logging/LogMacros.h"
+
+#if WITH_EDITOR
+#include "Framework/Notifications/NotificationManager.h"
+#include "Widgets/Notifications/SNotificationList.h"
+#endif
+
+DEFINE_LOG_CATEGORY(LogClothingMeshUtils)
+DECLARE_CYCLE_STAT(TEXT("Skin Physics Mesh"), STAT_ClothSkinPhysMesh, STATGROUP_Physics);
+
+#define LOCTEXT_NAMESPACE "ClothingMeshUtils"
+
+namespace ClothingMeshUtils
+{
+	void SkinPhysicsMesh(
+		const TArray<int32>& BoneMap,
+		const UClothPhysicalMeshDataBase& InMesh,
+		const FTransform& RootBoneTransform,
+		const FMatrix* InBoneMatrices,
+		const int32 InNumBoneMatrices,
+		TArray<FVector>& OutPositions,
+		TArray<FVector>& OutNormals)
+	{
+		SCOPE_CYCLE_COUNTER(STAT_ClothSkinPhysMesh);
+
+		// Ignore any user scale. It's already accounted for in our skinning matrices
+		FTransform RootBoneTransformInternal = RootBoneTransform;
+		RootBoneTransformInternal.SetScale3D(FVector(1.0f));
+
+		const uint32 NumVerts = InMesh.Vertices.Num();
+
+		OutPositions.Reset(NumVerts);
+		OutNormals.Reset(NumVerts);
+		OutPositions.AddZeroed(NumVerts);
+		OutNormals.AddZeroed(NumVerts);
+
+		const int32 MaxInfluences = InMesh.MaxBoneWeights;
+		const FMatrix* RESTRICT BoneMatrices = InBoneMatrices;
+
+		for (uint32 VertIndex = 0; VertIndex < NumVerts; ++VertIndex)
+		{
+			// Fixed particle, needs to be skinned
+			const uint16* RESTRICT BoneIndices = InMesh.BoneData[VertIndex].BoneIndices;
+			const float* RESTRICT BoneWeights = InMesh.BoneData[VertIndex].BoneWeights;
+
+			// WARNING - HORRIBLE UNROLLED LOOP + JUMP TABLE BELOW
+			// done this way because this is a pretty tight and perf critical loop. essentially
+			// rather than checking each influence we can just jump into this switch and fall through
+			// everything to compose the final skinned data
+			const FVector& RefParticle = InMesh.Vertices[VertIndex];
+			const FVector& RefNormal = InMesh.Normals[VertIndex];
+			FVector& OutPosition = OutPositions[VertIndex];
+			FVector& OutNormal = OutNormals[VertIndex];
+			switch (InMesh.BoneData[VertIndex].NumInfluences)
+			{
+			default: break;
+			case 8:
+			{
+				const FMatrix& BoneMatrix = BoneMatrices[BoneMap[BoneIndices[7]]];
+				const float Weight = BoneWeights[7];
+
+				OutPosition += BoneMatrix.TransformPosition(RefParticle) * Weight;
+				OutNormal += BoneMatrix.TransformVector(RefNormal) * Weight;
+			}
+			case 7:
+			{
+				const FMatrix& BoneMatrix = BoneMatrices[BoneMap[BoneIndices[6]]];
+				const float Weight = BoneWeights[6];
+
+				OutPosition += BoneMatrix.TransformPosition(RefParticle) * Weight;
+				OutNormal += BoneMatrix.TransformVector(RefNormal) * Weight;
+			}
+			case 6:
+			{
+				const FMatrix& BoneMatrix = BoneMatrices[BoneMap[BoneIndices[5]]];
+				const float Weight = BoneWeights[5];
+
+				OutPosition += BoneMatrix.TransformPosition(RefParticle) * Weight;
+				OutNormal += BoneMatrix.TransformVector(RefNormal) * Weight;
+			}
+			case 5:
+			{
+				const FMatrix& BoneMatrix = BoneMatrices[BoneMap[BoneIndices[4]]];
+				const float Weight = BoneWeights[4];
+
+				OutPosition += BoneMatrix.TransformPosition(RefParticle) * Weight;
+				OutNormal += BoneMatrix.TransformVector(RefNormal) * Weight;
+			}
+			case 4:
+			{
+				const FMatrix& BoneMatrix = BoneMatrices[BoneMap[BoneIndices[3]]];
+				const float Weight = BoneWeights[3];
+
+				OutPosition += BoneMatrix.TransformPosition(RefParticle) * Weight;
+				OutNormal += BoneMatrix.TransformVector(RefNormal) * Weight;
+			}
+			case 3:
+			{
+				const FMatrix& BoneMatrix = BoneMatrices[BoneMap[BoneIndices[2]]];
+				const float Weight = BoneWeights[2];
+
+				OutPosition += BoneMatrix.TransformPosition(RefParticle) * Weight;
+				OutNormal += BoneMatrix.TransformVector(RefNormal) * Weight;
+			}
+			case 2:
+			{
+				const FMatrix& BoneMatrix = BoneMatrices[BoneMap[BoneIndices[1]]];
+				const float Weight = BoneWeights[1];
+
+				OutPosition += BoneMatrix.TransformPosition(RefParticle) * Weight;
+				OutNormal += BoneMatrix.TransformVector(RefNormal) * Weight;
+			}
+			case 1:
+			{
+				const FMatrix& BoneMatrix = BoneMatrices[BoneMap[BoneIndices[0]]];
+				const float Weight = BoneWeights[0];
+
+				OutPosition += BoneMatrix.TransformPosition(RefParticle) * Weight;
+				OutNormal += BoneMatrix.TransformVector(RefNormal) * Weight;
+			}
+			}
+
+			OutPosition = RootBoneTransformInternal.InverseTransformPosition(OutPosition);
+			OutNormal = RootBoneTransformInternal.InverseTransformVector(OutNormal);
+
+			if (OutNormal.SizeSquared() > SMALL_NUMBER)
+			{
+				OutNormal = OutNormal.GetUnsafeNormal();
+			}
+		}
+	}
+
+	/** 
+	 * Gets the best match triangle for a specified position from the triangles in Mesh.
+	 * Performs no validation on the incoming mesh data, the mesh data should be verified
+	 * to be valid before using this function
+	 */
+	int32 GetBestTriangleBaseIndex(const ClothMeshDesc& Mesh, const FVector& Position)
+	{
+		float MinimumDistanceSq = MAX_flt;
+		int32 ClosestBaseIndex = INDEX_NONE;
+
+		const int32 NumTriangles = Mesh.Indices.Num() / 3;
+		for(int32 TriIdx = 0; TriIdx < NumTriangles; ++TriIdx)
+		{
+			int32 TriBaseIdx = TriIdx * 3;
+
+			const uint32 IA = Mesh.Indices[TriBaseIdx + 0];
+			const uint32 IB = Mesh.Indices[TriBaseIdx + 1];
+			const uint32 IC = Mesh.Indices[TriBaseIdx + 2];
+
+			const FVector& A = Mesh.Positions[IA];
+			const FVector& B = Mesh.Positions[IB];
+			const FVector& C = Mesh.Positions[IC];
+
+			FVector PointOnTri = FMath::ClosestPointOnTriangleToPoint(Position, A, B, C);
+			float DistSq = (PointOnTri - Position).SizeSquared();
+
+			if(DistSq < MinimumDistanceSq)
+			{
+				MinimumDistanceSq = DistSq;
+				ClosestBaseIndex = TriBaseIdx;
+			}
+		}
+
+		return ClosestBaseIndex;
+	}
+
+	void GenerateMeshToMeshSkinningData(TArray<FMeshToMeshVertData>& OutSkinningData, const ClothMeshDesc& TargetMesh, const TArray<FVector>* TargetTangents, const ClothMeshDesc& SourceMesh)
+	{
+		if(!TargetMesh.HasValidMesh())
+		{
+			UE_LOG(LogClothingMeshUtils, Warning, TEXT("Failed to generate mesh to mesh skinning data. Invalid Target Mesh."));
+			return;
+		}
+
+		if(!SourceMesh.HasValidMesh())
+		{
+			UE_LOG(LogClothingMeshUtils, Warning, TEXT("Failed to generate mesh to mesh skinning data. Invalid Source Mesh."));
+			return;
+		}
+
+		const int32 NumMesh0Verts = TargetMesh.Positions.Num();
+		const int32 NumMesh0Normals = TargetMesh.Normals.Num();
+		const int32 NumMesh0Tangents = TargetTangents ? TargetTangents->Num() : 0;
+
+		const int32 NumMesh1Verts = SourceMesh.Positions.Num();
+		const int32 NumMesh1Normals = SourceMesh.Normals.Num();
+		const int32 NumMesh1Indices = SourceMesh.Indices.Num();
+
+		// Check we have properly formed triangles
+		check(NumMesh1Indices % 3 == 0);
+
+		const int32 NumMesh1Triangles = NumMesh1Indices / 3;
+
+		// Check mesh data to make sure we have the same number of each element
+		if(NumMesh0Verts != NumMesh0Normals || (TargetTangents && NumMesh0Tangents != NumMesh0Verts))
+		{
+			UE_LOG(LogClothingMeshUtils, Warning, TEXT("Can't generate mesh to mesh skinning data, Mesh0 data is missing verts."));
+			return;
+		}
+
+		if(NumMesh1Verts != NumMesh1Normals)
+		{
+			UE_LOG(LogClothingMeshUtils, Warning, TEXT("Can't generate mesh to mesh skinning data, Mesh1 data is missing verts."));
+			return;
+		}
+
+		OutSkinningData.Reserve(NumMesh0Verts);
+
+		// For all mesh0 verts
+		for(int32 VertIdx0 = 0; VertIdx0 < NumMesh0Verts; ++VertIdx0)
+		{
+			OutSkinningData.AddZeroed();
+			FMeshToMeshVertData& SkinningData = OutSkinningData.Last();
+
+			const FVector& VertPosition = TargetMesh.Positions[VertIdx0];
+			const FVector& VertNormal = TargetMesh.Normals[VertIdx0];
+
+			FVector VertTangent;
+			if(TargetTangents)
+			{
+				VertTangent = (*TargetTangents)[VertIdx0];
+			}
+			else
+			{
+				FVector Tan0, Tan1;
+				VertNormal.FindBestAxisVectors(Tan0, Tan1);
+				VertTangent = Tan0;
+			}
+
+			int32 ClosestTriangleBaseIdx = GetBestTriangleBaseIndex(SourceMesh, VertPosition);
+
+			check(ClosestTriangleBaseIdx != INDEX_NONE);
+
+			const FVector& A = SourceMesh.Positions[SourceMesh.Indices[ClosestTriangleBaseIdx]];
+			const FVector& B = SourceMesh.Positions[SourceMesh.Indices[ClosestTriangleBaseIdx + 1]];
+			const FVector& C = SourceMesh.Positions[SourceMesh.Indices[ClosestTriangleBaseIdx + 2]];
+
+			const FVector& NA = SourceMesh.Normals[SourceMesh.Indices[ClosestTriangleBaseIdx]];
+			const FVector& NB = SourceMesh.Normals[SourceMesh.Indices[ClosestTriangleBaseIdx + 1]];
+			const FVector& NC = SourceMesh.Normals[SourceMesh.Indices[ClosestTriangleBaseIdx + 2]];
+
+			// Before generating the skinning data we need to check for a degenerate triangle.
+			// If we find _any_ degenerate triangles we will notify and fail to generate the skinning data
+			const FVector TriNormal = FVector::CrossProduct(B - A, C - A);
+			if(TriNormal.SizeSquared() < SMALL_NUMBER)
+			{
+				// Failed, we have 2 identical vertices
+				OutSkinningData.Reset();
+
+				// Log and toast
+				FText Error = FText::Format(LOCTEXT("DegenerateTriangleError", "Failed to generate skinning data, found conincident vertices in triangle A={0} B={1} C={2}"), FText::FromString(A.ToString()), FText::FromString(B.ToString()), FText::FromString(C.ToString()));
+
+				UE_LOG(LogClothingMeshUtils, Warning, TEXT("%s"), *Error.ToString());
+
+#if WITH_EDITOR
+				FNotificationInfo Info(Error);
+				Info.ExpireDuration = 5.0f;
+				FSlateNotificationManager::Get().AddNotification(Info);
+#endif
+				return;
+			}
+
+			SkinningData.PositionBaryCoordsAndDist = GetPointBaryAndDist(A, B, C, NA, NB, NC, VertPosition);
+			SkinningData.NormalBaryCoordsAndDist = GetPointBaryAndDist(A, B, C, NA, NB, NC, VertPosition + VertNormal);
+			SkinningData.TangentBaryCoordsAndDist = GetPointBaryAndDist(A, B, C, NA, NB, NC, VertPosition + VertTangent);
+			SkinningData.SourceMeshVertIndices[0] = SourceMesh.Indices[ClosestTriangleBaseIdx];
+			SkinningData.SourceMeshVertIndices[1] = SourceMesh.Indices[ClosestTriangleBaseIdx + 1];
+			SkinningData.SourceMeshVertIndices[2] = SourceMesh.Indices[ClosestTriangleBaseIdx + 2];
+			SkinningData.SourceMeshVertIndices[3] = 0;
+		}
+	}
+
+	FVector4 GetPointBaryAndDist(const FVector& A, const FVector& B, const FVector& C, const FVector& NA, const FVector& NB, const FVector& NC, const FVector& Point)
+	{
+		FPlane TrianglePlane(A, B, C);
+		const FVector PointOnTriPlane = FVector::PointPlaneProject(Point, TrianglePlane);
+		const FVector BaryCoords = FMath::ComputeBaryCentric2D(PointOnTriPlane, A, B, C);
+		const FVector NormalAtPoint = TrianglePlane;
+		FVector TriPointToVert = Point - PointOnTriPlane;
+		TriPointToVert = TriPointToVert.ProjectOnTo(NormalAtPoint);
+		float Dist = TriPointToVert.Size();
+
+		float Sign = TrianglePlane.PlaneDot(Point) < 0.0f ? -1.0f : 1.0f;
+
+		return FVector4(BaryCoords, TrianglePlane.PlaneDot(Point));
+	}
+
+	void GenerateEmbeddedPositions(const ClothMeshDesc& SourceMesh, TArrayView<const FVector> Positions, TArray<FVector4>& OutEmbeddedPositions, TArray<int32>& OutSourceIndices)
+	{
+		if(!SourceMesh.HasValidMesh())
+		{
+			// No valid source mesh
+			return;
+		}
+
+		const int32 NumPositions = Positions.Num();
+
+		OutEmbeddedPositions.Reset();
+		OutEmbeddedPositions.AddUninitialized(NumPositions);
+
+		OutSourceIndices.Reset(NumPositions * 3);
+
+		for(int32 PositionIndex = 0 ; PositionIndex < NumPositions ; ++PositionIndex)
+		{
+			const FVector& Position = Positions[PositionIndex];
+
+			int32 TriBaseIndex = GetBestTriangleBaseIndex(SourceMesh, Position);
+
+			const int32 IA = SourceMesh.Indices[TriBaseIndex];
+			const int32 IB = SourceMesh.Indices[TriBaseIndex + 1];
+			const int32 IC = SourceMesh.Indices[TriBaseIndex + 2];
+
+			const FVector& A = SourceMesh.Positions[IA];
+			const FVector& B = SourceMesh.Positions[IB];
+			const FVector& C = SourceMesh.Positions[IC];
+
+			const FVector& NA = SourceMesh.Normals[IA];
+			const FVector& NB = SourceMesh.Normals[IB];
+			const FVector& NC = SourceMesh.Normals[IC];
+
+			OutEmbeddedPositions[PositionIndex] = GetPointBaryAndDist(A, B, C, NA, NB, NC, Position);
+			OutSourceIndices.Add(IA);
+			OutSourceIndices.Add(IB);
+			OutSourceIndices.Add(IC);
+		}
+	}
+
+	void FVertexParameterMapper::Map(TArrayView<const float> Source, TArray<float>& Dest)
+	{
+		Map(Source, Dest, [](FVector Bary, float A, float B, float C)
+		{
+			return Bary.X * A + Bary.Y * B + Bary.Z * C;
+		});
+	}
+
+}
+
+#undef LOCTEXT_NAMESPACE
diff --git a/Engine/Source/Runtime/ClothingSystemRuntimeCommon/Public/ClothLODData.h b/Engine/Source/Runtime/ClothingSystemRuntimeCommon/Public/ClothLODData.h
new file mode 100644
--- /dev/null
+++ b/Engine/Source/Runtime/ClothingSystemRuntimeCommon/Public/ClothLODData.h
@@ -0,0 +1,50 @@
+// Copyright 1998-2019 Epic Games, Inc. All Rights Reserved.
+#pragma once
+
+#include "ClothPhysicalMeshData.h"
+
+#include "ClothCollisionData.h"
+#include "PointWeightMap.h"
+
+#include "SkeletalMeshTypes.h"
+
+#include "ClothLODData.generated.h"
+
+/** */
+UCLASS()
+class CLOTHINGSYSTEMRUNTIMECOMMON_API UClothLODDataBase : public UObject
+{
+	GENERATED_BODY()
+public:
+	UClothLODDataBase(const FObjectInitializer& Init);
+	virtual ~UClothLODDataBase();
+
+	// Raw phys mesh data
+	UPROPERTY(EditAnywhere, Category = SimMesh)
+	UClothPhysicalMeshDataBase* PhysicalMeshData;
+
+	// Collision primitive and covex data for clothing collisions
+	UPROPERTY(EditAnywhere, Category = Collision)
+	FClothCollisionData CollisionData;
+
+#if WITH_EDITORONLY_DATA
+	// Parameter masks defining the physics mesh masked data
+	UPROPERTY(EditAnywhere, Category = Masks)
+	TArray<FPointWeightMap> ParameterMasks;
+
+	// Get all available parameter masks for the specified target
+	//void GetParameterMasksForTarget(const MaskTarget_PhysMesh& InTarget, TArray<FPointWeightMap*>& OutMasks);
+	void GetParameterMasksForTarget(const uint8 InTarget, TArray<FPointWeightMap*>& OutMasks);
+#endif // WITH_EDITORONLY_DATA
+#if WITH_EDITOR
+	/** Copy \c ParameterMasks to corresponding targets in \c PhysicalMeshData. */
+	void PushWeightsToMesh();
+#endif
+	// Skinning data for transitioning from a higher detail LOD to this one
+	TArray<FMeshToMeshVertData> TransitionUpSkinData;
+
+	// Skinning data for transitioning from a lower detail LOD to this one
+	TArray<FMeshToMeshVertData> TransitionDownSkinData;
+
+	virtual void Serialize(FArchive& Ar) override;
+};
