Add support to script re-import of FBX animation sequences

diff --git a/Engine/Source/Editor/UnrealEd/Private/Fbx/FbxFactory.cpp b/Engine/Source/Editor/UnrealEd/Private/Fbx/FbxFactory.cpp
--- a/Engine/Source/Editor/UnrealEd/Private/Fbx/FbxFactory.cpp
+++ b/Engine/Source/Editor/UnrealEd/Private/Fbx/FbxFactory.cpp
@@ -193,6 +193,7 @@ UObject* UFbxFactory::FactoryCreateFile
 		{
 			UStaticMesh *ExistingStaticMesh = Cast<UStaticMesh>(ExistingObject);
 			USkeletalMesh *ExistingSkeletalMesh = Cast<USkeletalMesh>(ExistingObject);
+			UAnimSequence* ExistingAnimSequence = Cast<UAnimSequence>(ExistingObject);
 			UObject *ObjectToReimport = nullptr;
 			if (ExistingStaticMesh)
 			{
@@ -202,6 +203,10 @@ UObject* UFbxFactory::FactoryCreateFile
 			{
 				ObjectToReimport = ExistingSkeletalMesh;
 			}
+			else if (ExistingAnimSequence)
+			{
+				ObjectToReimport = ExistingAnimSequence;
+			}
 
 			if (ObjectToReimport != nullptr)
 			{
