Fixes for recompiling Blueprints that affect replicated actors during PIE
- Restore UActorChannel::Actor after the replace references archiver runs during reinstancing (otherwise when the channel for the old actor tears down, it destroys the *new* actor on the client instance)
- Call PostActorConstruction on reinstanced actors if they're in a world that has already begun play
- Don't allow BeginPlay processing to be deferred if being run as part of reinstancing
- Prevent a log spam warning about networked renamed actors when reinstancing during PIE

diff --git a/Engine/Source/Editor/UnrealEd/Private/Kismet2/KismetReinstanceUtilities.cpp b/Engine/Source/Editor/UnrealEd/Private/Kismet2/KismetReinstanceUtilities.cpp
--- a/Engine/Source/Editor/UnrealEd/Private/Kismet2/KismetReinstanceUtilities.cpp
+++ b/Engine/Source/Editor/UnrealEd/Private/Kismet2/KismetReinstanceUtilities.cpp
@@ -32,6 +32,8 @@
 #include "BlueprintEditor.h"
 #include "Engine/Selection.h"
 #include "BlueprintEditorSettings.h"
+#include "Engine/NetDriver.h"
+#include "Engine/ActorChannel.h"
 #include "Subsystems/AssetEditorSubsystem.h"
 
 DECLARE_CYCLE_STAT(TEXT("Replace Instances"), EKismetReinstancerStats_ReplaceInstancesOfClass, STATGROUP_KismetReinstancer );
@@ -82,6 +84,17 @@ struct FReplaceReferenceHelper
 			return;
 		}
 
+		// Remember what values were in UActorChannel::Actor so we can restore them later (this should only affect reinstancing during PIE)
+		// We need the old actor channel to tear down cleanly without affecting the new actor
+		TMap<UActorChannel*, AActor*> ActorChannelActorRestorationMap;
+		for (UActorChannel* ActorChannel : TObjectRange<UActorChannel>())
+		{
+			if (OldToNewInstanceMap.Contains(ActorChannel->Actor))
+			{
+				ActorChannelActorRestorationMap.Add(ActorChannel, ActorChannel->Actor);
+			}
+		}
+
 		// Find everything that references these objects
 		TArray<UObject *> Targets;
 		{
@@ -142,6 +155,12 @@ struct FReplaceReferenceHelper
 				}
 			}
 		}
+	
+		// Restore the old UActorChannel::Actor values (undoing what the replace references archiver did above to them)
+		for (const auto& KVP : ActorChannelActorRestorationMap)
+		{
+			KVP.Key->Actor = KVP.Value;
+		}
 	}
 };
 
@@ -503,6 +522,13 @@ public:
 					Actor->ReregisterAllComponents();
 					Actor->RerunConstructionScripts();
 
+					// The reinstancing case doesn't ever explicitly call Actor->FinishSpawning, we've handled the construction script
+					// portion above but still need the PostActorConstruction() case so BeginPlay gets routed correctly while in a BegunPlay world
+					if (World->HasBegunPlay())
+					{
+						Actor->PostActorConstruction();
+					}
+
 					if (SelectedObjecs.Contains(Obj) && GEditor)
 					{
 						GEditor->SelectActor(Actor, /*bInSelected =*/true, /*bNotify =*/true, false, true);
@@ -1044,6 +1070,16 @@ void FActorReplacementHelper::Finalize(const TMap<UObject*, UObject*>& OldToNewI
 		NewActor->ExecuteConstruction(TargetWorldTransform, nullptr, &DummyComponentData);
 	}	
 
+	// The reinstancing case doesn't ever explicitly call Actor->FinishSpawning, we've handled the construction script
+	// portion above but still need the PostActorConstruction() case so BeginPlay gets routed correctly while in a BegunPlay world
+	if (UWorld* World = NewActor->GetWorld())
+	{
+		if (World->HasBegunPlay())
+		{
+			NewActor->PostActorConstruction();
+		}
+	}
+
 	// make sure that the actor is properly hidden if it's in a hidden sublevel:
 	bool bIsInHiddenLevel = false;
 	if (ULevel* Level = NewActor->GetLevel())
diff --git a/Engine/Source/Runtime/Engine/Private/Actor.cpp b/Engine/Source/Runtime/Engine/Private/Actor.cpp
--- a/Engine/Source/Runtime/Engine/Private/Actor.cpp
+++ b/Engine/Source/Runtime/Engine/Private/Actor.cpp
@@ -3133,7 +3133,7 @@ void AActor::PostActorConstruction()
 	}
 
 	// If this is dynamically spawned replicated actor, defer calls to BeginPlay and UpdateOverlaps until replicated properties are deserialized
-	const bool bDeferBeginPlayAndUpdateOverlaps = (bExchangedRoles && RemoteRole == ROLE_Authority);
+	const bool bDeferBeginPlayAndUpdateOverlaps = (bExchangedRoles && RemoteRole == ROLE_Authority) && !GIsReinstancing;
 
 	if (bActorsInitialized)
 	{
diff --git a/Engine/Source/Runtime/Engine/Private/NetworkDriver.cpp b/Engine/Source/Runtime/Engine/Private/NetworkDriver.cpp
--- a/Engine/Source/Runtime/Engine/Private/NetworkDriver.cpp
+++ b/Engine/Source/Runtime/Engine/Private/NetworkDriver.cpp
@@ -2960,6 +2960,15 @@ void UNetDriver::NotifyActorRenamed(AActor* ThisActor, FName PreviousName)
 	const bool bIsActorStatic = !GuidCache->IsDynamicObject(ThisActor);
 	const bool bActorHasRole = ThisActor->GetRemoteRole() != ROLE_None;
 
+#if WITH_EDITOR
+	// When recompiling and reinstancing a Blueprint, we rename the old actor out of the way, which would cause this code to emit a warning during PIE
+	// Since that old actor is about to die (on both client and server) in the reinstancing case, it's safe to skip
+	if (GCompilingBlueprint)
+	{
+		return;
+	}
+#endif
+
 	if (bIsActorStatic && bActorHasRole)
 	{
 		if (bIsServer)
