Fixing warnings in HLOD nightly steps
Do not trigger warnings for hidden levels that have HLOD system disabled

diff --git a/Engine/Source/Editor/UnrealEd/Private/HierarchicalLOD.cpp b/Engine/Source/Editor/UnrealEd/Private/HierarchicalLOD.cpp
--- a/Engine/Source/Editor/UnrealEd/Private/HierarchicalLOD.cpp
+++ b/Engine/Source/Editor/UnrealEd/Private/HierarchicalLOD.cpp
@@ -88,8 +88,10 @@ void FHierarchicalLODBuilder::Build()
 		{
 			BuildClusters(LevelIter, true);
 		}	
-	
-		bVisibleLevelsWarning |= !LevelIter->bIsVisible;	
+		else
+		{
+			bVisibleLevelsWarning |= LevelIter->GetWorldSettings()->bEnableHierarchicalLODSystem;
+		}
 	}
 
 
@@ -118,8 +120,10 @@ void FHierarchicalLODBuilder::PreviewBuild()
 			LevelIter->MarkPackageDirty();
 			BuildClusters(LevelIter, false);
 		}
-
-		bVisibleLevelsWarning |= !LevelIter->bIsVisible;
+		else
+		{
+			bVisibleLevelsWarning |= LevelIter->GetWorldSettings()->bEnableHierarchicalLODSystem;
+		}
 	}
 
 	// Fire map check warnings for hidden levels 
@@ -621,11 +625,14 @@ void FHierarchicalLODBuilder::ClearHLODs()
 
 	for (ULevel* Level : World->GetLevels())
 	{
-		bVisibleLevelsWarning |= !Level->bIsVisible;
 		if (Level->bIsVisible)
 		{
 			DeleteLODActors(Level);
 		}
+		else
+		{
+			bVisibleLevelsWarning |= Level->GetWorldSettings()->bEnableHierarchicalLODSystem;
+		}
 	}
 
 
@@ -642,13 +649,17 @@ void FHierarchicalLODBuilder::ClearHLODs()
 void FHierarchicalLODBuilder::ClearPreviewBuild()
 {
 	bool bVisibleLevelsWarning = false;
+
 	for (ULevel* Level : World->GetLevels())
 	{
-		bVisibleLevelsWarning |= !Level->bIsVisible;
-		if ( Level->bIsVisible )
+		if (Level->bIsVisible)
 		{
 			DeleteLODActors(Level);
 		}		
+		else
+		{
+			bVisibleLevelsWarning |= Level->GetWorldSettings()->bEnableHierarchicalLODSystem;
+		}
 	}
 
 	// Fire map check warnings for hidden levels 
@@ -671,7 +682,7 @@ void FHierarchicalLODBuilder::BuildMeshesForLODActors(bool bForceAll)
 		// Only meshes for clusters that are in a visible level
 		if (!LevelIter->bIsVisible)
 		{
-			bVisibleLevelsWarning = true;
+			bVisibleLevelsWarning |= LevelIter->GetWorldSettings()->bEnableHierarchicalLODSystem;
 			continue;
 		}
 
