CustomVersion thread-safety follow-up fix

diff --git a/Engine/Plugins/Developer/Concert/ConcertMain/Source/Concert/Private/ConcertVersion.cpp b/Engine/Plugins/Developer/Concert/ConcertMain/Source/Concert/Private/ConcertVersion.cpp
--- a/Engine/Plugins/Developer/Concert/ConcertMain/Source/Concert/Private/ConcertVersion.cpp
+++ b/Engine/Plugins/Developer/Concert/ConcertMain/Source/Concert/Private/ConcertVersion.cpp
@@ -116,8 +116,9 @@ void FConcertSessionVersionInfo::Initialize()
 			}
 		}
 	}
-
-	for (const FCustomVersion& EngineCustomVersion : FCurrentCustomVersions::GetAll().GetAllVersions())
+	
+	FCustomVersionContainer AllCurrentVersions = FCurrentCustomVersions::GetAll();
+	for (const FCustomVersion& EngineCustomVersion : AllCurrentVersions.GetAllVersions())
 	{
 		FConcertCustomVersionInfo& CustomVersion = CustomVersions.AddDefaulted_GetRef();
 		CustomVersion.Initialize(EngineCustomVersion);
diff --git a/Engine/Source/Editor/UnrealEd/Private/CookOnTheFlyServer.cpp b/Engine/Source/Editor/UnrealEd/Private/CookOnTheFlyServer.cpp
--- a/Engine/Source/Editor/UnrealEd/Private/CookOnTheFlyServer.cpp
+++ b/Engine/Source/Editor/UnrealEd/Private/CookOnTheFlyServer.cpp
@@ -4497,8 +4497,8 @@ void GetAdditionalCurrentIniVersionStrings( const ITargetPlatform* TargetPlatfor
 		IniVersionMap.Add(TEXT("fastcook"));
 	}
 
-
-	for (const FCustomVersion& CustomVersion : FCurrentCustomVersions::GetAll().GetAllVersions())
+	FCustomVersionContainer AllCurrentVersions = FCurrentCustomVersions::GetAll();
+	for (const FCustomVersion& CustomVersion : AllCurrentVersions.GetAllVersions())
 	{
 		FString CustomVersionString = FString::Printf(TEXT("%s:%s"), *CustomVersion.GetFriendlyName().ToString(), *CustomVersion.Key.ToString());
 		FString CustomVersionValue = FString::Printf(TEXT("%d"), CustomVersion.Version);
