Curve Editor: Speculative fix for when curve owner is null and testing cubic tangents

diff --git a/Engine/Source/Editor/UnrealEd/Private/SCurveEditor.cpp b/Engine/Source/Editor/UnrealEd/Private/SCurveEditor.cpp
--- a/Engine/Source/Editor/UnrealEd/Private/SCurveEditor.cpp
+++ b/Engine/Source/Editor/UnrealEd/Private/SCurveEditor.cpp
@@ -3091,7 +3091,7 @@ SCurveEditor::FSelectedTangent SCurveEditor::HitTestCubicTangents( const FGeomet
 {
 	FSelectedTangent Tangent;
 
-	if( AreCurvesVisible() && CurveOwner->HasRichCurves() )
+	if( AreCurvesVisible() && CurveOwner && CurveOwner->HasRichCurves() )
 	{
 		FTrackScaleInfo ScaleInfo(ViewMinInput.Get(),  ViewMaxInput.Get(), ViewMinOutput.Get(), ViewMaxOutput.Get(), InMyGeometry.GetLocalSize());
 
@@ -3712,7 +3712,7 @@ TArray<SCurveEditor::FSelectedTangent> SCurveEditor::GetEditableTangentsWithinMa
 	MarqueeBox.Max = FVector(MarqueeBottomRight.X, MarqueeBottomRight.Y, 0);
 
 	TArray<FSelectedTangent> TangentsWithinMarquee;
-	if (AreCurvesVisible() && CurveOwner->HasRichCurves())
+	if (AreCurvesVisible() && CurveOwner && CurveOwner->HasRichCurves())
 	{
 		FTrackScaleInfo ScaleInfo(ViewMinInput.Get(), ViewMaxInput.Get(), ViewMinOutput.Get(), ViewMaxOutput.Get(), InMyGeometry.GetLocalSize());
 		for (auto CurveViewModel : CurveViewModels)
