Blueprints: Fixed a crash when attempting to reinstance an actor that did not belong to a UWorld

This is always true of sequencer spawnables that do not have a UWorld ptr

diff --git a/Engine/Source/Editor/UnrealEd/Private/Kismet2/KismetReinstanceUtilities.cpp b/Engine/Source/Editor/UnrealEd/Private/Kismet2/KismetReinstanceUtilities.cpp
--- a/Engine/Source/Editor/UnrealEd/Private/Kismet2/KismetReinstanceUtilities.cpp
+++ b/Engine/Source/Editor/UnrealEd/Private/Kismet2/KismetReinstanceUtilities.cpp
@@ -1921,7 +1921,10 @@ void FBlueprintCompileReinstancer::ReplaceInstancesOfClass_Inner(TMap<UClass*, U
 		{
 			if (AActor* OldActor = Cast<AActor>(ReinstancedPair.Key))
 			{
-				OldActor->GetWorld()->EditorDestroyActor(OldActor, /*bShouldModifyLevel =*/true);
+				if (UWorld* World = OldActor->GetWorld())
+				{
+					World->EditorDestroyActor(OldActor, /*bShouldModifyLevel =*/true);
+				}
 			}
 		}
 	}
