Correctly force saving in PromptForCheckoutAndSave if bAlreadyCheckedOut. Fixes breaking change from CL8370719

diff --git a/Engine/Source/Editor/UnrealEd/Private/FileHelpers.cpp b/Engine/Source/Editor/UnrealEd/Private/FileHelpers.cpp
--- a/Engine/Source/Editor/UnrealEd/Private/FileHelpers.cpp
+++ b/Engine/Source/Editor/UnrealEd/Private/FileHelpers.cpp
@@ -3738,7 +3738,7 @@ FEditorFileUtils::EPromptReturnCode FEditorFileUtils::PromptForCheckoutAndSave(
 			bUserResponse = FEditorFileUtils::PromptToCheckoutPackages(false, PackagesToSave, &PackagesCheckedOutOrMadeWritable, &PackagesNotNeedingCheckout);
 		}
 
-		if( bUserResponse && (PackagesCheckedOutOrMadeWritable.Num() > 0 || PackagesNotNeedingCheckout.Num() > 0) )
+		if( bAlreadyCheckedOut || (bUserResponse && (PackagesCheckedOutOrMadeWritable.Num() > 0 || PackagesNotNeedingCheckout.Num() > 0)) )
 		{
 			TArray<UPackage*> FinalSaveList;
 			
