Optimize object snapshot diffing if we know which properties have changed

This also changes snapshots to only write and diff properties, regardless of whether we give it a specific list

diff --git a/Engine/Source/Editor/UnrealEd/Classes/Editor/Transactor.h b/Engine/Source/Editor/UnrealEd/Classes/Editor/Transactor.h
--- a/Engine/Source/Editor/UnrealEd/Classes/Editor/Transactor.h
+++ b/Engine/Source/Editor/UnrealEd/Classes/Editor/Transactor.h
@@ -258,8 +258,8 @@ protected:
 		void Save( FTransaction* Owner );
 		void Load( FTransaction* Owner );
 		void Finalize( FTransaction* Owner, TSharedPtr<ITransactionObjectAnnotation>& OutFinalizedObjectAnnotation );
-		void Snapshot( FTransaction* Owner );
-		static void Diff( const FTransaction* Owner, const FSerializedObject& OldSerializedObect, const FSerializedObject& NewSerializedObject, FTransactionObjectDeltaChange& OutDeltaChange );
+		void Snapshot( FTransaction* Owner, TArrayView<const UProperty*> Properties );
+		static void Diff( const FTransaction* Owner, const FSerializedObject& OldSerializedObect, const FSerializedObject& NewSerializedObject, FTransactionObjectDeltaChange& OutDeltaChange, const bool bFullDiff = true );
 
 		/** Used by GC to collect referenced objects. */
 		void AddReferencedObjects( FReferenceCollector& Collector );
@@ -354,9 +354,11 @@ protected:
 		public:
 			FWriter(
 				FSerializedObject& InSerializedObject,
-				bool bWantBinarySerialization
+				bool bWantBinarySerialization,
+				TArrayView<const UProperty*> InPropertiesToSerialize = TArrayView<const UProperty*>()
 				)
 				: SerializedObject(InSerializedObject)
+				, PropertiesToSerialize(InPropertiesToSerialize)
 				, Offset(0)
 			{
 				for(int32 ObjIndex = 0; ObjIndex < SerializedObject.ReferencedObjects.Num(); ++ObjIndex)
@@ -381,6 +383,12 @@ protected:
 				Offset = InPos; 
 			}
 
+			virtual bool ShouldSkipProperty(const UProperty* InProperty) const override
+			{
+				return (PropertiesToSerialize.Num() > 0 && !PropertiesToSerialize.Contains(InProperty))
+					|| FArchiveUObject::ShouldSkipProperty(InProperty);
+			}
+
 		private:
 			struct FCachedPropertyKey
 			{
@@ -483,6 +491,7 @@ protected:
 				return (FArchive&)*this << ObjectIndex;
 			}
 			FSerializedObject& SerializedObject;
+			TArrayView<const UProperty*> PropertiesToSerialize;
 			TMap<UObject*, int32> ObjectMap;
 			TMap<FName, int32> NameMap;
 			FCachedPropertyKey CachedSerializedTaggedPropertyKey;
@@ -569,7 +578,7 @@ public:
 	virtual void SaveArray( UObject* Object, FScriptArray* Array, int32 Index, int32 Count, int32 Oper, int32 ElementSize, STRUCT_DC DefaultConstructor, STRUCT_AR Serializer, STRUCT_DTOR Destructor ) override;
 	virtual void StoreUndo( UObject* Object, TUniquePtr<FChange> UndoChange ) override;
 	virtual void SetPrimaryObject(UObject* InObject) override;
-	virtual void SnapshotObject( UObject* InObject ) override;
+	virtual void SnapshotObject( UObject* InObject, TArrayView<const UProperty*> Properties ) override;
 
 	/** BeginOperation should be called when a transaction or undo/redo starts */
 	virtual void BeginOperation() override;
diff --git a/Engine/Source/Editor/UnrealEd/Private/EditorTransaction.cpp b/Engine/Source/Editor/UnrealEd/Private/EditorTransaction.cpp
--- a/Engine/Source/Editor/UnrealEd/Private/EditorTransaction.cpp
+++ b/Engine/Source/Editor/UnrealEd/Private/EditorTransaction.cpp
@@ -305,7 +305,7 @@ void FTransaction::FObjectRecord::Finalize( FTransaction* Owner, TSharedPtr<ITra
 			// have been tracking delta-changes between snapshots and this finalization will need to account for those changes too
 			if (bSnapshot)
 			{
-				Diff(Owner, SerializedObjectSnapshot, CurrentSerializedObject, DeltaChange);
+				Diff(Owner, SerializedObjectSnapshot, CurrentSerializedObject, DeltaChange, /*bFullDiff*/false);
 			}
 
 			SerializedObjectFlip.Swap(CurrentSerializedObject);
@@ -317,7 +317,7 @@ void FTransaction::FObjectRecord::Finalize( FTransaction* Owner, TSharedPtr<ITra
 	}
 }
 
-void FTransaction::FObjectRecord::Snapshot( FTransaction* Owner )
+void FTransaction::FObjectRecord::Snapshot( FTransaction* Owner, TArrayView<const UProperty*> Properties )
 {
 	if (Array)
 	{
@@ -338,14 +338,14 @@ void FTransaction::FObjectRecord::Snapshot( FTransaction* Owner )
 		FSerializedObject CurrentSerializedObject;
 		{
 			CurrentSerializedObject.SetObject(CurrentObject);
-			FWriter Writer(CurrentSerializedObject, bWantsBinarySerialization);
-			SerializeObject(Writer);
+			FWriter Writer(CurrentSerializedObject, bWantsBinarySerialization, Properties);
+			CurrentObject->SerializeScriptProperties(Writer);
 		}
 
 		// Diff against the correct serialized data depending on whether we already had a snapshot
 		const FSerializedObject& InitialSerializedObject = bSnapshot ? SerializedObjectSnapshot : SerializedObject;
 		FTransactionObjectDeltaChange SnapshotDeltaChange;
-		Diff(Owner, InitialSerializedObject, CurrentSerializedObject, SnapshotDeltaChange);
+		Diff(Owner, InitialSerializedObject, CurrentSerializedObject, SnapshotDeltaChange, /*bFullDiff*/false);
 
 		// Update the snapshot data for next time
 		bSnapshot = true;
@@ -361,7 +361,7 @@ void FTransaction::FObjectRecord::Snapshot( FTransaction* Owner )
 	}
 }
 
-void FTransaction::FObjectRecord::Diff( const FTransaction* Owner, const FSerializedObject& OldSerializedObject, const FSerializedObject& NewSerializedObject, FTransactionObjectDeltaChange& OutDeltaChange )
+void FTransaction::FObjectRecord::Diff( const FTransaction* Owner, const FSerializedObject& OldSerializedObject, const FSerializedObject& NewSerializedObject, FTransactionObjectDeltaChange& OutDeltaChange, const bool bFullDiff )
 {
 	auto AreObjectPointersIdentical = [&OldSerializedObject, &NewSerializedObject](const FName InPropertyName)
 	{
@@ -405,18 +405,21 @@ void FTransaction::FObjectRecord::Diff( const FTransaction* Owner, const FSerial
 		return bAreNamesIdentical;
 	};
 
-	OutDeltaChange.bHasNameChange |= OldSerializedObject.ObjectName != NewSerializedObject.ObjectName;
-	OutDeltaChange.bHasOuterChange |= OldSerializedObject.ObjectOuterPathName != NewSerializedObject.ObjectOuterPathName;
-	OutDeltaChange.bHasPendingKillChange |= OldSerializedObject.bIsPendingKill != NewSerializedObject.bIsPendingKill;
-
-	if (!AreObjectPointersIdentical(NAME_None))
+	if (bFullDiff)
 	{
-		OutDeltaChange.bHasNonPropertyChanges = true;
-	}
+		OutDeltaChange.bHasNameChange |= OldSerializedObject.ObjectName != NewSerializedObject.ObjectName;
+		OutDeltaChange.bHasOuterChange |= OldSerializedObject.ObjectOuterPathName != NewSerializedObject.ObjectOuterPathName;
+		OutDeltaChange.bHasPendingKillChange |= OldSerializedObject.bIsPendingKill != NewSerializedObject.bIsPendingKill;
 
-	if (!AreNamesIdentical(NAME_None))
-	{
-		OutDeltaChange.bHasNonPropertyChanges = true;
+		if (!AreObjectPointersIdentical(NAME_None))
+		{
+			OutDeltaChange.bHasNonPropertyChanges = true;
+		}
+
+		if (!AreNamesIdentical(NAME_None))
+		{
+			OutDeltaChange.bHasNonPropertyChanges = true;
+		}
 	}
 
 	if (OldSerializedObject.SerializedProperties.Num() > 0 || NewSerializedObject.SerializedProperties.Num() > 0)
@@ -431,8 +434,11 @@ void FTransaction::FObjectRecord::Diff( const FTransaction* Owner, const FSerial
 			const FSerializedProperty* OldSerializedProperty = OldSerializedObject.SerializedProperties.Find(NewNamePropertyPair.Key);
 			if (!OldSerializedProperty)
 			{
-				// Missing property, assume that the property changed
-				OutDeltaChange.ChangedProperties.AddUnique(NewNamePropertyPair.Key);
+				if (bFullDiff)
+				{
+					// Missing property, assume that the property changed
+					OutDeltaChange.ChangedProperties.AddUnique(NewNamePropertyPair.Key);
+				}
 				continue;
 			}
 
@@ -468,49 +474,55 @@ void FTransaction::FObjectRecord::Diff( const FTransaction* Owner, const FSerial
 			const FSerializedProperty* NewSerializedProperty = NewSerializedObject.SerializedProperties.Find(OldNamePropertyPair.Key);
 			if (!NewSerializedProperty)
 			{
-				// Missing property, assume that the property changed
-				OutDeltaChange.ChangedProperties.AddUnique(OldNamePropertyPair.Key);
+				if (bFullDiff)
+				{
+					// Missing property, assume that the property changed
+					OutDeltaChange.ChangedProperties.AddUnique(OldNamePropertyPair.Key);
+				}
 				continue;
 			}
 		}
 
-		// Compare the data before the property block to see if something else in the object has changed
-		if (!OutDeltaChange.bHasNonPropertyChanges)
+		if (bFullDiff)
 		{
-			const int32 OldHeaderSize = FMath::Min(StartOfOldPropertyBlock, OldSerializedObject.Data.Num());
-			const int32 CurrentHeaderSize = FMath::Min(StartOfNewPropertyBlock, NewSerializedObject.Data.Num());
-
-			bool bIsHeaderIdentical = OldHeaderSize == CurrentHeaderSize;
-			if (bIsHeaderIdentical && CurrentHeaderSize > 0)
-			{
-				bIsHeaderIdentical = FMemory::Memcmp(&OldSerializedObject.Data[0], &NewSerializedObject.Data[0], CurrentHeaderSize) == 0;
-			}
-
-			if (!bIsHeaderIdentical)
+			// Compare the data before the property block to see if something else in the object has changed
+			if (!OutDeltaChange.bHasNonPropertyChanges)
 			{
-				OutDeltaChange.bHasNonPropertyChanges = true;
-			}
-		}
+				const int32 OldHeaderSize = FMath::Min(StartOfOldPropertyBlock, OldSerializedObject.Data.Num());
+				const int32 CurrentHeaderSize = FMath::Min(StartOfNewPropertyBlock, NewSerializedObject.Data.Num());
 
-		// Compare the data after the property block to see if something else in the object has changed
-		if (!OutDeltaChange.bHasNonPropertyChanges)
-		{
-			const int32 OldFooterSize = OldSerializedObject.Data.Num() - FMath::Max(EndOfOldPropertyBlock, 0);
-			const int32 CurrentFooterSize = NewSerializedObject.Data.Num() - FMath::Max(EndOfNewPropertyBlock, 0);
+				bool bIsHeaderIdentical = OldHeaderSize == CurrentHeaderSize;
+				if (bIsHeaderIdentical && CurrentHeaderSize > 0)
+				{
+					bIsHeaderIdentical = FMemory::Memcmp(&OldSerializedObject.Data[0], &NewSerializedObject.Data[0], CurrentHeaderSize) == 0;
+				}
 
-			bool bIsFooterIdentical = OldFooterSize == CurrentFooterSize;
-			if (bIsFooterIdentical && CurrentFooterSize > 0)
-			{
-				bIsFooterIdentical = FMemory::Memcmp(&OldSerializedObject.Data[EndOfOldPropertyBlock], &NewSerializedObject.Data[EndOfNewPropertyBlock], CurrentFooterSize) == 0;
+				if (!bIsHeaderIdentical)
+				{
+					OutDeltaChange.bHasNonPropertyChanges = true;
+				}
 			}
 
-			if (!bIsFooterIdentical)
+			// Compare the data after the property block to see if something else in the object has changed
+			if (!OutDeltaChange.bHasNonPropertyChanges)
 			{
-				OutDeltaChange.bHasNonPropertyChanges = true;
+				const int32 OldFooterSize = OldSerializedObject.Data.Num() - FMath::Max(EndOfOldPropertyBlock, 0);
+				const int32 CurrentFooterSize = NewSerializedObject.Data.Num() - FMath::Max(EndOfNewPropertyBlock, 0);
+
+				bool bIsFooterIdentical = OldFooterSize == CurrentFooterSize;
+				if (bIsFooterIdentical && CurrentFooterSize > 0)
+				{
+					bIsFooterIdentical = FMemory::Memcmp(&OldSerializedObject.Data[EndOfOldPropertyBlock], &NewSerializedObject.Data[EndOfNewPropertyBlock], CurrentFooterSize) == 0;
+				}
+
+				if (!bIsFooterIdentical)
+				{
+					OutDeltaChange.bHasNonPropertyChanges = true;
+				}
 			}
 		}
 	}
-	else
+	else if (bFullDiff)
 	{
 		// No properties, so just compare the whole blob
 		bool bIsBlobIdentical = OldSerializedObject.Data.Num() == NewSerializedObject.Data.Num();
@@ -828,7 +840,7 @@ void FTransaction::SetPrimaryObject(UObject* InObject)
 	}
 }
 
-void FTransaction::SnapshotObject( UObject* InObject )
+void FTransaction::SnapshotObject( UObject* InObject, TArrayView<const UProperty*> Properties )
 {
 	if (InObject && ObjectMap.Contains(InObject))
 	{
@@ -839,7 +851,7 @@ void FTransaction::SnapshotObject( UObject* InObject )
 
 		if (FoundObjectRecord)
 		{
-			FoundObjectRecord->Snapshot(this);
+			FoundObjectRecord->Snapshot(this, Properties);
 		}
 	}
 }
diff --git a/Engine/Source/Runtime/Core/Public/Misc/ITransaction.h b/Engine/Source/Runtime/Core/Public/Misc/ITransaction.h
--- a/Engine/Source/Runtime/Core/Public/Misc/ITransaction.h
+++ b/Engine/Source/Runtime/Core/Public/Misc/ITransaction.h
@@ -364,7 +364,8 @@ public:
 	/**
 	 * Snapshots a UObject within the transaction.
 	 *
-	 * @param Object The object to snapshot.
+	 * @param Object	The object to snapshot.
+	 * @param Property	The optional list of properties that have potentially changed on the object (to avoid snapshotting the entire object).
 	 */
-	virtual void SnapshotObject( UObject* Object ) = 0;
+	virtual void SnapshotObject( UObject* Object, TArrayView<const UProperty*> Properties ) = 0;
 };
diff --git a/Engine/Source/Runtime/CoreUObject/Private/UObject/Obj.cpp b/Engine/Source/Runtime/CoreUObject/Private/UObject/Obj.cpp
--- a/Engine/Source/Runtime/CoreUObject/Private/UObject/Obj.cpp
+++ b/Engine/Source/Runtime/CoreUObject/Private/UObject/Obj.cpp
@@ -346,7 +346,8 @@ void UObject::PostEditChangeProperty(FPropertyChangedEvent& PropertyChangedEvent
 	// This allows listeners to be notified of intermediate changes of state
 	if (PropertyChangedEvent.ChangeType == EPropertyChangeType::Interactive)
 	{
-		SnapshotTransactionBuffer(this);
+		const UProperty* ChangedProperty = PropertyChangedEvent.MemberProperty;
+		SnapshotTransactionBuffer(this, TArrayView<const UProperty*>(&ChangedProperty, 1));
 	}
 }
 
diff --git a/Engine/Source/Runtime/CoreUObject/Private/UObject/UObjectGlobals.cpp b/Engine/Source/Runtime/CoreUObject/Private/UObject/UObjectGlobals.cpp
--- a/Engine/Source/Runtime/CoreUObject/Private/UObject/UObjectGlobals.cpp
+++ b/Engine/Source/Runtime/CoreUObject/Private/UObject/UObjectGlobals.cpp
@@ -2167,6 +2167,11 @@ bool SaveToTransactionBuffer(UObject* Object, bool bMarkDirty)
 }
 
 void SnapshotTransactionBuffer(UObject* Object)
+{
+	SnapshotTransactionBuffer(Object, TArrayView<const UProperty*>());
+}
+
+void SnapshotTransactionBuffer(UObject* Object, TArrayView<const UProperty*> Properties)
 {
 	// Script packages should not end up in the transaction buffer.
 	// PIE objects should go through however. Additionally, in order
@@ -2176,7 +2181,7 @@ void SnapshotTransactionBuffer(UObject* Object)
 
 	if (GUndo && bIsTransactional && bIsNotScriptPackage)
 	{
-		GUndo->SnapshotObject(Object);
+		GUndo->SnapshotObject(Object, Properties);
 	}
 }
 
diff --git a/Engine/Source/Runtime/CoreUObject/Public/UObject/UObjectGlobals.h b/Engine/Source/Runtime/CoreUObject/Public/UObject/UObjectGlobals.h
--- a/Engine/Source/Runtime/CoreUObject/Public/UObject/UObjectGlobals.h
+++ b/Engine/Source/Runtime/CoreUObject/Public/UObject/UObjectGlobals.h
@@ -11,6 +11,7 @@
 #include "UObject/ObjectMacros.h"
 #include "Misc/OutputDeviceRedirector.h"
 #include "UObject/PrimaryAssetId.h"
+#include "Containers/ArrayView.h"
 #include "Templates/Function.h"
 #include "Templates/IsArrayOrRefOfType.h"
 #include "Serialization/ArchiveUObject.h"
@@ -629,8 +630,10 @@ COREUOBJECT_API bool SaveToTransactionBuffer(UObject* Object, bool bMarkDirty);
  *  b) The object has changed since it started transacting.
  *
  * @param	Object		object to snapshot.
+ * @param	Properties	optional list of properties that have potentially changed on the object (to avoid snapshotting the entire object).
  */
 COREUOBJECT_API void SnapshotTransactionBuffer(UObject* Object);
+COREUOBJECT_API void SnapshotTransactionBuffer(UObject* Object, TArrayView<const UProperty*> Properties);
 
 /**
  * Check for StaticAllocateObject error; only for use with the editor, make or other commandlets.
