Fix reimport match for restoring sections data
#rnx

diff --git a/Engine/Source/Editor/UnrealEd/Private/Factories/SkeletalMeshImport.cpp b/Engine/Source/Editor/UnrealEd/Private/Factories/SkeletalMeshImport.cpp
--- a/Engine/Source/Editor/UnrealEd/Private/Factories/SkeletalMeshImport.cpp
+++ b/Engine/Source/Editor/UnrealEd/Private/Factories/SkeletalMeshImport.cpp
@@ -312,9 +312,10 @@ ExistingSkelMeshData* SaveExistingSkelMeshData(USkeletalMesh* ExistingSkelMesh,
 			int32 GenerateUpTo = BackupLODModel->Sections[SectionIndex].GenerateUpToLodIndex;
 			bool bDisabled = BackupLODModel->Sections[SectionIndex].bDisabled;
 			bool bBoneChunkedSection = BackupLODModel->Sections[SectionIndex].ChunkedParentSectionIndex != INDEX_NONE;
-			if (ExistingMeshDataPtr->ExistingImportMaterialOriginalNameData.IsValidIndex(SectionMaterialIndex))
+			//Only save the original section, do not save the chunked sections
+			if (!bBoneChunkedSection && ExistingMeshDataPtr->ExistingImportMaterialOriginalNameData.IsValidIndex(SectionMaterialIndex))
 			{
-				ExistingMeshDataPtr->ExistingImportMeshLodSectionMaterialData[LodIndex].Add(ExistingMeshLodSectionData(ExistingMeshDataPtr->ExistingImportMaterialOriginalNameData[SectionMaterialIndex], SectionCastShadow, SectionRecomputeTangents, GenerateUpTo, bDisabled, bBoneChunkedSection));
+				ExistingMeshDataPtr->ExistingImportMeshLodSectionMaterialData[LodIndex].Add(ExistingMeshLodSectionData(ExistingMeshDataPtr->ExistingImportMaterialOriginalNameData[SectionMaterialIndex], SectionCastShadow, SectionRecomputeTangents, GenerateUpTo, bDisabled));
 			}
 		}
 	}
@@ -1006,20 +1007,19 @@ void RestoreExistingSkelMeshData(ExistingSkelMeshData* MeshData, USkeletalMesh*
 						NewSkelMeshLodModel.Sections[SectionIndex].bRecomputeTangent = MeshData->ExistingImportMeshLodSectionMaterialData[SafeReimportLODIndex][RemapMaterial[NewMeshSectionMaterialIndex]].bRecomputeTangents;
 						NewSkelMeshLodModel.Sections[SectionIndex].GenerateUpToLodIndex = MeshData->ExistingImportMeshLodSectionMaterialData[SafeReimportLODIndex][RemapMaterial[NewMeshSectionMaterialIndex]].GenerateUpTo;
 						NewSkelMeshLodModel.Sections[SectionIndex].bDisabled = MeshData->ExistingImportMeshLodSectionMaterialData[SafeReimportLODIndex][RemapMaterial[NewMeshSectionMaterialIndex]].bDisabled;
-						bool bBoneChunkedSection = MeshData->ExistingImportMeshLodSectionMaterialData[SafeReimportLODIndex][RemapMaterial[NewMeshSectionMaterialIndex]].bBoneChunkedSection;
+						bool bBoneChunkedSection = NewSkelMeshLodModel.Sections[SectionIndex].ChunkedParentSectionIndex >= 0;
+						int32 ParentOriginalSectionIndex = NewSkelMeshLodModel.Sections[SectionIndex].OriginalDataSectionIndex;
 						if (!bBoneChunkedSection)
 						{
 							//Set the new Parent Index
 							ParentSectionIndex = SectionIndex;
-							FSkelMeshSourceSectionUserData& UserSectionData = NewSkelMeshLodModel.UserSectionsData.FindOrAdd(ParentSectionIndex);
+							FSkelMeshSourceSectionUserData& UserSectionData = NewSkelMeshLodModel.UserSectionsData.FindOrAdd(ParentOriginalSectionIndex);
 							UserSectionData.bDisabled = NewSkelMeshLodModel.Sections[SectionIndex].bDisabled;
 							UserSectionData.bCastShadow = NewSkelMeshLodModel.Sections[SectionIndex].bCastShadow;
 							UserSectionData.bRecomputeTangent = NewSkelMeshLodModel.Sections[SectionIndex].bRecomputeTangent;
 							UserSectionData.GenerateUpToLodIndex = NewSkelMeshLodModel.Sections[SectionIndex].GenerateUpToLodIndex;
 							//The cloth will be rebind later after the reimport is done
 						}
-						NewSkelMeshLodModel.Sections[SectionIndex].ChunkedParentSectionIndex = bBoneChunkedSection ? ParentSectionIndex : INDEX_NONE;
-						NewSkelMeshLodModel.Sections[SectionIndex].OriginalDataSectionIndex = ParentSectionIndex;
 					}
 				}
 			}
@@ -1031,9 +1031,10 @@ void RestoreExistingSkelMeshData(ExistingSkelMeshData* MeshData, USkeletalMesh*
 			}
 
 			FName CurrentSectionImportedMaterialName = SkeletalMesh->Materials[NewSkelMeshLodModel.Sections[SectionIndex].MaterialIndex].ImportedMaterialSlotName;
-			for (int32 ExistSectionIndex = 0; ExistSectionIndex < MeshData->ExistingImportMeshLodSectionMaterialData[SafeReimportLODIndex].Num(); ++ExistSectionIndex)
+			for (int32 ExistSectionIndex = 0; ExistSectionIndex < MeshData->LastImportMeshLodSectionMaterialData[SafeReimportLODIndex].Num(); ++ExistSectionIndex)
 			{
-				if (!MeshData->LastImportMeshLodSectionMaterialData[SafeReimportLODIndex].IsValidIndex(ExistSectionIndex) || !MeshData->ExistingImportMeshLodSectionMaterialData[SafeReimportLODIndex].IsValidIndex(ExistSectionIndex))
+				int32 OriginalSectionIndex = NewSkelMeshLodModel.Sections[ExistSectionIndex].OriginalDataSectionIndex;
+				if (!MeshData->LastImportMeshLodSectionMaterialData[SafeReimportLODIndex].IsValidIndex(ExistSectionIndex) || !MeshData->ExistingImportMeshLodSectionMaterialData[SafeReimportLODIndex].IsValidIndex(OriginalSectionIndex))
 				{
 					continue;
 				}
@@ -1045,7 +1046,7 @@ void RestoreExistingSkelMeshData(ExistingSkelMeshData* MeshData, USkeletalMesh*
 				}
 
 				//Get the current skelmesh section slot import name
-				FName ExistMeshSectionSlotName = MeshData->ExistingImportMeshLodSectionMaterialData[SafeReimportLODIndex][ExistSectionIndex].ImportedMaterialSlotName;
+				FName ExistMeshSectionSlotName = MeshData->ExistingImportMeshLodSectionMaterialData[SafeReimportLODIndex][OriginalSectionIndex].ImportedMaterialSlotName;
 				if (ExistMeshSectionSlotName != OriginalImportMeshSectionSlotName)
 				{
 					//The last import slot name match the New import slot name, but the Exist slot name is different then the last import slot name.
diff --git a/Engine/Source/Editor/UnrealEd/Public/SkelImport.h b/Engine/Source/Editor/UnrealEd/Public/SkelImport.h
--- a/Engine/Source/Editor/UnrealEd/Public/SkelImport.h
+++ b/Engine/Source/Editor/UnrealEd/Public/SkelImport.h
@@ -22,20 +22,18 @@ class FSkeletalMeshLODModel;
 
 struct ExistingMeshLodSectionData
 {
-	ExistingMeshLodSectionData(FName InImportedMaterialSlotName, bool InbCastShadow, bool InbRecomputeTangents, int32 InGenerateUpTo, bool InbDisabled, bool InbBoneChunkedSection)
+	ExistingMeshLodSectionData(FName InImportedMaterialSlotName, bool InbCastShadow, bool InbRecomputeTangents, int32 InGenerateUpTo, bool InbDisabled)
 	: ImportedMaterialSlotName(InImportedMaterialSlotName)
 	, bCastShadow(InbCastShadow)
 	, bRecomputeTangents(InbRecomputeTangents)
 	, GenerateUpTo(InGenerateUpTo)
 	, bDisabled(InbDisabled)
-	, bBoneChunkedSection(InbBoneChunkedSection)
 	{}
 	FName ImportedMaterialSlotName;
 	bool bCastShadow;
 	bool bRecomputeTangents;
 	int32 GenerateUpTo;
 	bool bDisabled;
-	bool bBoneChunkedSection;
 };
 
 struct ExistingSkelMeshData
