Fix volumetric lightmap not being removed from scene before ReloadPackages()

(ushell-p4-cherrypick of 9567158 by [REDACTED_PERSON])

diff --git a/Engine/Source/Editor/UnrealEd/Private/PackageTools.cpp b/Engine/Source/Editor/UnrealEd/Private/PackageTools.cpp
--- a/Engine/Source/Editor/UnrealEd/Private/PackageTools.cpp
+++ b/Engine/Source/Editor/UnrealEd/Private/PackageTools.cpp
@@ -622,6 +622,7 @@ UPackageTools::UPackageTools(const FObjectInitializer& ObjectInitializer)
 		// Check to see if we need to reload the current world.
 		FName WorldNameToReload;
 		TMap<FName, const UMapBuildDataRegistry*> LevelsToMapBuildData;
+		bool bReloadingLightingScenario = false; 
 		TArray<ULevelStreaming*> RemovedStreamingLevels;
 		if (UWorld* CurrentWorldPtr = CurrentWorld.Get())
 		{
@@ -673,6 +674,8 @@ UPackageTools::UPackageTools(const FObjectInitializer& ObjectInitializer)
 			// Cache the current map build data for the levels of the current world so we can see if they change due to a reload (we can skip this if reloading the current world).
 			else
 			{
+				bReloadingLightingScenario = CurrentWorldPtr->GetActiveLightingScenario() && CurrentWorldPtr->GetActiveLightingScenario()->MapBuildData && PackagesToReload.Contains(CurrentWorldPtr->GetActiveLightingScenario()->MapBuildData->GetOutermost());
+
 				const TArray<ULevel*>& Levels = CurrentWorldPtr->GetLevels();
 				for (int32 i = Levels.Num() - 1; i >= 0; --i)
 				{
@@ -692,7 +695,13 @@ UPackageTools::UPackageTools(const FObjectInitializer& ObjectInitializer)
 					}
 					else
 					{
-						LevelsToMapBuildData.Add(Level->GetFName(), Level->MapBuildData);
+						if (Level->MapBuildData && PackagesToReload.Contains(Level->MapBuildData->GetOutermost()) || bReloadingLightingScenario)
+						{
+							// Remove any VLM here so FPrecomputedVolumetricLightmapData::RemoveFromSceneData() has the necessary resources to perform GPU brick unplugging before destruction
+							Level->ReleaseRenderingResources();
+
+							LevelsToMapBuildData.Add(Level->GetFName(), Level->MapBuildData);
+						}
 					}
 				}
 			}
@@ -795,9 +804,8 @@ UPackageTools::UPackageTools(const FObjectInitializer& ObjectInitializer)
 				for (int32 LevelIndex = 0; LevelIndex < CurrentWorldPtr->GetNumLevels(); ++LevelIndex)
 				{
 					ULevel* Level = CurrentWorldPtr->GetLevel(LevelIndex);
-					const UMapBuildDataRegistry* OldMapBuildData = LevelsToMapBuildData.FindRef(Level->GetFName());
 
-					if (OldMapBuildData && OldMapBuildData != Level->MapBuildData)
+					if (LevelsToMapBuildData.Contains(Level->GetFName()) || bReloadingLightingScenario)
 					{
 						Level->ReleaseRenderingResources();
 						Level->InitializeRenderingResources();
