Added for the atmospheric light shortcut tweaks into the undo/redo stack.
+ some renaming

(ushell-p4-cherrypick of 9614694 by [REDACTED_PERSON])

diff --git a/Engine/Source/Editor/UnrealEd/Private/LevelEditorViewport.cpp b/Engine/Source/Editor/UnrealEd/Private/LevelEditorViewport.cpp
--- a/Engine/Source/Editor/UnrealEd/Private/LevelEditorViewport.cpp
+++ b/Engine/Source/Editor/UnrealEd/Private/LevelEditorViewport.cpp
@@ -203,6 +203,31 @@ static FVector4 AttemptToSnapLocationToOriginPlane( const FViewportCursorLocatio
 	return Location;
 }
 
+/** Helper function to get an atmosphere light from an index*/
+static UDirectionalLightComponent* GetAtmosphericLight(const uint8 DesiredLightIndex, UWorld* ViewportWorld)
+{
+	UDirectionalLightComponent* SelectedAtmosphericLight = nullptr;
+	float SelectedLightLuminance = 0.0f;
+	for (TObjectIterator<UDirectionalLightComponent> ComponentIt; ComponentIt; ++ComponentIt)
+	{
+		if (ComponentIt->GetWorld() == ViewportWorld)
+		{
+			UDirectionalLightComponent* AtmosphericLight = *ComponentIt;
+
+			if (!AtmosphericLight->IsUsedAsAtmosphereSunLight() || AtmosphericLight->GetAtmosphereSunLightIndex() != DesiredLightIndex || !AtmosphericLight->GetVisibleFlag())
+				continue;
+
+			float LightLuminance = AtmosphericLight->GetColoredLightBrightness().ComputeLuminance();
+			if (!SelectedAtmosphericLight ||					// Set it if null
+				SelectedLightLuminance < LightLuminance)		// Or choose the brightest atmospheric light
+			{
+				SelectedAtmosphericLight = AtmosphericLight;
+			}
+		}
+	}
+	return SelectedAtmosphericLight;
+}
+
 namespace LevelEditorViewportClientHelper
 {
 	UProperty* GetEditTransformProperty(FWidget::EWidgetMode WidgetMode)
@@ -1568,7 +1593,7 @@ FTrackingTransaction::~FTrackingTransaction()
 	USelection::SelectionChangedEvent.RemoveAll(this);
 }
 
-void FTrackingTransaction::Begin(const FText& Description)
+void FTrackingTransaction::Begin(const FText& Description, AActor* AdditionalActor)
 {
 	End();
 	
@@ -1578,11 +1603,8 @@ void FTrackingTransaction::Begin(const FText& Description)
 
 	TSet<AGroupActor*> GroupActors;
 
-	// Modify selected actors to record their state at the start of the transaction
-	for (FSelectionIterator It(GEditor->GetSelectedActorIterator()); It; ++It)
+	auto ProcessActorModify = [&](AActor* Actor)
 	{
-		AActor* Actor = CastChecked<AActor>(*It);
-
 		// Some tracking transactions, such as a duplication operation into a sublevel do not modify the selected Actors
 		// Store the initial package dirty state so we can restore if necessary
 		UPackage* Package = Actor->GetOutermost();
@@ -1602,6 +1624,18 @@ void FTrackingTransaction::Begin(const FText& Description)
 				GroupActors.Add(ActorLockedRootGroup);
 			}
 		}
+	};
+
+	// Modify selected actors to record their state at the start of the transaction
+	for (FSelectionIterator It(GEditor->GetSelectedActorIterator()); It; ++It)
+	{
+		AActor* Actor = CastChecked<AActor>(*It);
+		ProcessActorModify(Actor);
+	}
+	// And the additional optional actor provided
+	if (AdditionalActor)
+	{
+		ProcessActorModify(AdditionalActor);
 	}
 
 	// Modify unique group actors
@@ -2293,7 +2327,7 @@ void FLevelEditorViewportClient::Tick(float DeltaTime)
 
 	UpdateViewForLockedActor(DeltaTime);
 
-	UserIsControllingSunLightTimer = FMath::Max(UserIsControllingSunLightTimer - DeltaTime, 0.0f);
+	UserIsControllingAtmosphericLightTimer = FMath::Max(UserIsControllingAtmosphericLightTimer - DeltaTime, 0.0f);
 }
 
 void FLevelEditorViewportClient::UpdateViewForLockedActor(float DeltaTime)
@@ -2758,21 +2792,38 @@ bool FLevelEditorViewportClient::InputKey(FViewport* InViewport, int32 Controlle
 		return true;
 	}
 
+	UWorld* ViewportWorld = GetWorld();
+	auto ProcessAtmosphericLightShortcut = [&](const uint8 LightIndex, bool& bCurrentUserControl)
+	{
+		UDirectionalLightComponent* SelectedSunLight = GetAtmosphericLight(LightIndex, ViewportWorld);
+		if (SelectedSunLight)
+		{
+			if (!bCurrentUserControl)
+			{
+				FText TrackingDescription = FText::Format(LOCTEXT("RotatationShortcut", "Rotate Atmosphere Light {0}"), LightIndex);
+				TrackingTransaction.Begin(TrackingDescription, SelectedSunLight->GetOwner());
+			}
+			bCurrentUserControl = true;
+			UserIsControllingAtmosphericLightTimer = 3.0f; // Keep the widget open for a few seconds even when not tweaking the sun light
+		}
+	};
 	bool bCtrlLPressed = InputState.IsCtrlButtonPressed() && Key == EKeys::L;
 	if (bCtrlLPressed && InputState.IsShiftButtonPressed())
 	{
-		bUserIsControllingSunLight1 = true;
-		UserIsControllingSunLightTimer = 3.0f; // Keep the widget open for a few seconds even when not tweaking the sun light
+		ProcessAtmosphericLightShortcut(1, bUserIsControllingAtmosphericLight1);
 		return true;
 	}
 	if (bCtrlLPressed)
 	{
-		bUserIsControllingSunLight0 = true;
-		UserIsControllingSunLightTimer = 3.0f; // Keep the widget open for a few seconds even when not tweaking the sun light
+		ProcessAtmosphericLightShortcut(0, bUserIsControllingAtmosphericLight0);
 		return true;
 	}
-	bUserIsControllingSunLight0 = false;
-	bUserIsControllingSunLight1 = false;
+	if (bUserIsControllingAtmosphericLight0 || bUserIsControllingAtmosphericLight1)
+	{
+		TrackingTransaction.End();
+	}
+	bUserIsControllingAtmosphericLight0 = false;
+	bUserIsControllingAtmosphericLight1 = false;
 
 	bool bHandled = FEditorViewportClient::InputKey(InViewport,ControllerId,Key,Event,AmountDepressed,bGamepad);
 
@@ -4115,51 +4166,32 @@ EMouseCursor::Type FLevelEditorViewportClient::GetCursor(FViewport* InViewport,i
 
 void FLevelEditorViewportClient::MouseMove(FViewport* InViewport, int32 x, int32 y)
 {
-	if (bUserIsControllingSunLight0 || bUserIsControllingSunLight1)
+	if (bUserIsControllingAtmosphericLight0 || bUserIsControllingAtmosphericLight1)
 	{
 		UWorld* ViewportWorld = GetWorld();
 
-		UDirectionalLightComponent* SelectedSunLight = nullptr;
-		float SelectedLightLuminance = 0.0f;
-		const uint8 DesiredLightIndex = bUserIsControllingSunLight0 ? 0 : 1;
-		for (TObjectIterator<UDirectionalLightComponent> ComponentIt; ComponentIt; ++ComponentIt)
-		{
-			if (ComponentIt->GetWorld() == ViewportWorld)
-			{
-				UDirectionalLightComponent* SunLight = *ComponentIt;
-
-				if (!SunLight->IsUsedAsAtmosphereSunLight() || SunLight->GetAtmosphereSunLightIndex()!= DesiredLightIndex || !SunLight->GetVisibleFlag())
-					continue;
+		const uint8 DesiredLightIndex = bUserIsControllingAtmosphericLight0 ? 0 : 1;
+		UDirectionalLightComponent* SelectedAtmosphericLight = GetAtmosphericLight(DesiredLightIndex, ViewportWorld);
 
-				float LightLuminance = SunLight->GetColoredLightBrightness().ComputeLuminance();
-				if (!SelectedSunLight ||							// Set it if null
-					SelectedLightLuminance < LightLuminance)		// Or choose the brightest sun light
-				{
-					SelectedSunLight = SunLight;
-				}
-			}
-		}
-
-		if (SelectedSunLight)
+		if (SelectedAtmosphericLight)
 		{
 			int32 mouseDeltaX = x - CachedLastMouseX;
 			int32 mouseDeltaY = y - CachedLastMouseY;
 
-			FTransform ComponentTransform = SelectedSunLight->GetComponentTransform();
-			FQuat SunRotation = ComponentTransform.GetRotation();
+			FTransform ComponentTransform = SelectedAtmosphericLight->GetComponentTransform();
+			FQuat LightRotation = ComponentTransform.GetRotation();
 			// Rotate around up axis (yaw)
 			FVector UpVector = FVector(0, 0, 1);
-			SunRotation = FQuat(UpVector, float(mouseDeltaX)*0.01f) * SunRotation;
-			// Sun Zenith rotation (pitch)
-			FVector PitchRotationAxis = FVector::CrossProduct(SunRotation.GetForwardVector(), UpVector);
+			LightRotation = FQuat(UpVector, float(mouseDeltaX)*0.01f) * LightRotation;
+			// Light Zenith rotation (pitch)
+			FVector PitchRotationAxis = FVector::CrossProduct(LightRotation.GetForwardVector(), UpVector);
 			PitchRotationAxis.Normalize();
-			SunRotation = FQuat(PitchRotationAxis, float(mouseDeltaY)*0.01f) * SunRotation;
+			LightRotation = FQuat(PitchRotationAxis, float(mouseDeltaY)*0.01f) * LightRotation;
 
-			ComponentTransform.SetRotation(SunRotation);
-			SelectedSunLight->SetWorldTransform(ComponentTransform);
+			ComponentTransform.SetRotation(LightRotation);
+			SelectedAtmosphericLight->SetWorldTransform(ComponentTransform);
 
-			// Only manipulate a single light, the first one.
-			UserControlledSunLightMatrix = ComponentTransform;
+			UserControlledAtmosphericLightMatrix = ComponentTransform;
 		}
 	}
 
@@ -4500,9 +4532,9 @@ void FLevelEditorViewportClient::Draw(const FSceneView* View,FPrimitiveDrawInter
 		}
 	}
 
-	if (UserIsControllingSunLightTimer > 0.0f)
+	if (UserIsControllingAtmosphericLightTimer > 0.0f)
 	{
-		// Draw a gizmo helping to figure out where is the sun when moving it using a shortcut.
+		// Draw a gizmo helping to figure out where is the light when moving it using a shortcut.
 		FQuat ViewRotation = FQuat(GetViewRotation());
 		FVector ViewPosition = GetViewLocation();
 		const float GizmoDistance = 50.0f;
@@ -4533,17 +4565,17 @@ void FLevelEditorViewportClient::Draw(const FSceneView* View,FPrimitiveDrawInter
 		DrawArc(PDI, Base, Z, Y, -90.0f, 90.0f, GizmoRadius, 32, FLinearColor(1.0f, 0.2f, 0.2f), SDPG_World);
 		DrawArc(PDI, Base, Z, X, -90.0f, 90.0f, GizmoRadius, 32, FLinearColor(0.2f, 1.0f, 0.2f), SDPG_World);
 
-		// Draw the sun incoming light direction. The arrow is offset outward to help depth perception when it intersects with other gizmo elements.
-		const FLinearColor ArrowColor = -UserControlledSunLightMatrix.GetRotation().GetForwardVector() * 0.5f + 0.5f;
-		const FVector ArrowOrigin = Base - UserControlledSunLightMatrix.GetRotation().GetForwardVector()*GizmoRadius*1.25;
-		const FQuatRotationTranslationMatrix ArrowToWorld(UserControlledSunLightMatrix.GetRotation(), ArrowOrigin);
+		// Draw the light incoming light direction. The arrow is offset outward to help depth perception when it intersects with other gizmo elements.
+		const FLinearColor ArrowColor = -UserControlledAtmosphericLightMatrix.GetRotation().GetForwardVector() * 0.5f + 0.5f;
+		const FVector ArrowOrigin = Base - UserControlledAtmosphericLightMatrix.GetRotation().GetForwardVector()*GizmoRadius*1.25;
+		const FQuatRotationTranslationMatrix ArrowToWorld(UserControlledAtmosphericLightMatrix.GetRotation(), ArrowOrigin);
 		DrawDirectionalArrow(PDI, ArrowToWorld, ArrowColor, GizmoRadius, 0.3f, SDPG_World, ThicknessBold);
 
 		// Now draw x, y and z axis to help getting a sense of depth when look at the vectors on screen.
-		FVector SunArrowTip = -UserControlledSunLightMatrix.GetRotation().GetForwardVector()*GizmoRadius;
-		FVector P0 = Base + SunArrowTip * FVector(1.0f, 0.0f, 0.0f);
-		FVector P1 = Base + SunArrowTip * FVector(1.0f, 1.0f, 0.0f);
-		FVector P2 = Base + SunArrowTip * FVector(1.0f, 1.0f, 1.0f);
+		FVector LightArrowTip = -UserControlledAtmosphericLightMatrix.GetRotation().GetForwardVector()*GizmoRadius;
+		FVector P0 = Base + LightArrowTip * FVector(1.0f, 0.0f, 0.0f);
+		FVector P1 = Base + LightArrowTip * FVector(1.0f, 1.0f, 0.0f);
+		FVector P2 = Base + LightArrowTip * FVector(1.0f, 1.0f, 1.0f);
 		PDI->DrawLine(Base, P0, FLinearColor(1.0f, 0.0f, 0.0f), SDPG_World, ThicknessLight);
 		PDI->DrawLine(P0, P1, FLinearColor(0.0f, 1.0f, 0.0f), SDPG_World, ThicknessLight);
 		PDI->DrawLine(P1, P2, FLinearColor(0.0f, 0.0f, 1.0f), SDPG_World, ThicknessLight);
diff --git a/Engine/Source/Editor/UnrealEd/Public/EditorViewportClient.h b/Engine/Source/Editor/UnrealEd/Public/EditorViewportClient.h
--- a/Engine/Source/Editor/UnrealEd/Public/EditorViewportClient.h
+++ b/Engine/Source/Editor/UnrealEd/Public/EditorViewportClient.h
@@ -1545,10 +1545,10 @@ protected:
 	int32 CachedLastMouseY = 0;
 
 	/** True is the use is controling the light via a shorcut*/
-	bool bUserIsControllingSunLight0 = false;
-	bool bUserIsControllingSunLight1 = false;
-	float UserIsControllingSunLightTimer = 0.0f;
-	FTransform UserControlledSunLightMatrix;
+	bool bUserIsControllingAtmosphericLight0 = false;
+	bool bUserIsControllingAtmosphericLight1 = false;
+	float UserIsControllingAtmosphericLightTimer = 0.0f;
+	FTransform UserControlledAtmosphericLightMatrix;
 
 
 	// -1, -1 if not set
diff --git a/Engine/Source/Editor/UnrealEd/Public/LevelEditorViewport.h b/Engine/Source/Editor/UnrealEd/Public/LevelEditorViewport.h
--- a/Engine/Source/Editor/UnrealEd/Public/LevelEditorViewport.h
+++ b/Engine/Source/Editor/UnrealEd/Public/LevelEditorViewport.h
@@ -83,7 +83,7 @@ struct UNREALED_API FTrackingTransaction
 	/**
 	 * Initiates a transaction.
 	 */
-	void Begin(const FText& Description);
+	void Begin(const FText& Description, AActor* AdditionalActor = nullptr);
 
 	void End();
 
