Fix clang compiler warning/error:
/media/epic/Storage/Builds/Release-4.24/Engine/Source/Editor/UnrealEd/Private/PackageTools.cpp:698:31: error: '&&' within '||' [-Werror,-Wlogical-op-parentheses]

diff --git a/Engine/Source/Editor/UnrealEd/Private/PackageTools.cpp b/Engine/Source/Editor/UnrealEd/Private/PackageTools.cpp
--- a/Engine/Source/Editor/UnrealEd/Private/PackageTools.cpp
+++ b/Engine/Source/Editor/UnrealEd/Private/PackageTools.cpp
@@ -695,7 +695,7 @@ UPackageTools::UPackageTools(const FObjectInitializer& ObjectInitializer)
 					}
 					else
 					{
-						if (Level->MapBuildData && PackagesToReload.Contains(Level->MapBuildData->GetOutermost()) || bReloadingLightingScenario)
+						if ((Level->MapBuildData && PackagesToReload.Contains(Level->MapBuildData->GetOutermost())) || bReloadingLightingScenario)
 						{
 							// Remove any VLM here so FPrecomputedVolumetricLightmapData::RemoveFromSceneData() has the necessary resources to perform GPU brick unplugging before destruction
 							Level->ReleaseRenderingResources();
