Prevent doing nested undo. It can happen if there is a popup dialog during the broadcast of the undo event.

#rnx

diff --git a/Engine/Source/Editor/UnrealEd/Private/EditorTransaction.cpp b/Engine/Source/Editor/UnrealEd/Private/EditorTransaction.cpp
--- a/Engine/Source/Editor/UnrealEd/Private/EditorTransaction.cpp
+++ b/Engine/Source/Editor/UnrealEd/Private/EditorTransaction.cpp
@@ -1306,7 +1306,7 @@ void UTransBuffer::Cancel( int32 StartIndex /*=0*/ )
 bool UTransBuffer::CanUndo( FText* Text )
 {
 	CheckState();
-	if (ActiveCount)
+	if (ActiveCount || CurrentTransaction != nullptr )
 	{
 		if (Text)
 		{
@@ -1343,7 +1343,7 @@ bool UTransBuffer::CanUndo( FText* Text )
 bool UTransBuffer::CanRedo( FText* Text )
 {
 	CheckState();
-	if( ActiveCount )
+	if( ActiveCount || CurrentTransaction != nullptr )
 	{
 		if( Text )
 		{
