Copied CL# 9693447

Bug fixed: Renaming ULayers did not update their internal UActors nor the ViewHiddenLayers that contained them if it was called with a reference to the Layer->LayerName, e.g. RenameLayer(Layer->LayerName, NewLayerName). This created lots of side effects with UActors and HiddenLayers, that after the layer rename pointed to no-longer-existing Layer names.

diff --git a/Engine/Source/Editor/UnrealEd/Private/Layers/LayersSubsystem.cpp b/Engine/Source/Editor/UnrealEd/Private/Layers/LayersSubsystem.cpp
--- a/Engine/Source/Editor/UnrealEd/Private/Layers/LayersSubsystem.cpp
+++ b/Engine/Source/Editor/UnrealEd/Private/Layers/LayersSubsystem.cpp
@@ -1303,6 +1303,7 @@ bool ULayersSubsystem::RenameLayer( const FName& OriginalLayerName, const FName&
 	}
 
 	Layer->Modify();
+	const FName OriginalLayerNameCopy = OriginalLayerName; // Otherwise, bug if RenameLayer(Layer->LayerName, NewLayerName) after the next LayerName rename
 	Layer->LayerName = NewLayerName;
 	Layer->ActorStats.Empty();
 	// Iterate over all actors, swapping layers.
@@ -1314,7 +1315,7 @@ bool ULayersSubsystem::RenameLayer( const FName& OriginalLayerName, const FName&
 			continue;
 		}
 
-		if ( ULayersSubsystem::RemoveActorFromLayer( Actor, OriginalLayerName ) )
+		if (ULayersSubsystem::RemoveActorFromLayer(Actor, OriginalLayerNameCopy))
 		{
 			// No need to mark the actor as modified these functions take care of that
 			AddActorToLayer( Actor, NewLayerName );
@@ -1324,7 +1325,7 @@ bool ULayersSubsystem::RenameLayer( const FName& OriginalLayerName, const FName&
 	// update all views's hidden layers if they had this one
 	for (FLevelEditorViewportClient* ViewportClient : GEditor->GetLevelViewportClients())
 	{
-		if ( ViewportClient->ViewHiddenLayers.Remove( OriginalLayerName ) > 0 )
+		if (ViewportClient->ViewHiddenLayers.Remove(OriginalLayerNameCopy) > 0)
 		{
 			ViewportClient->ViewHiddenLayers.AddUnique( NewLayerName );
 			ViewportClient->Invalidate();
