Fix crash when re-importing a skeletalmesh and there is a animation viewport viewing it
#rnx

diff --git a/Engine/Source/Editor/Persona/Private/AnimationEditorViewportClient.cpp b/Engine/Source/Editor/Persona/Private/AnimationEditorViewportClient.cpp
--- a/Engine/Source/Editor/Persona/Private/AnimationEditorViewportClient.cpp
+++ b/Engine/Source/Editor/Persona/Private/AnimationEditorViewportClient.cpp
@@ -1248,7 +1248,7 @@ void FAnimationViewportClient::DrawMeshBonesBakedAnimation(UDebugSkelMeshCompone
 
 void FAnimationViewportClient::DrawMeshBones(UDebugSkelMeshComponent * MeshComponent, FPrimitiveDrawInterface* PDI) const
 {
-	if ( MeshComponent && MeshComponent->SkeletalMesh )
+	if ( MeshComponent && MeshComponent->SkeletalMesh && MeshComponent->GetNumDrawTransform() > 0)
 	{
 		TArray<FTransform> WorldTransforms;
 		WorldTransforms.AddUninitialized(MeshComponent->GetNumDrawTransform());
diff --git a/Engine/Source/Editor/UnrealEd/Private/Factories/EditorFactories.cpp b/Engine/Source/Editor/UnrealEd/Private/Factories/EditorFactories.cpp
--- a/Engine/Source/Editor/UnrealEd/Private/Factories/EditorFactories.cpp
+++ b/Engine/Source/Editor/UnrealEd/Private/Factories/EditorFactories.cpp
@@ -212,6 +212,8 @@
 
 #include "InstancedFoliageActor.h"
 
+#include "Animation/DebugSkelMeshComponent.h"
+
 #if PLATFORM_WINDOWS
 	// Needed for DDS support.
 	#include "Windows/WindowsHWrapper.h"
@@ -6282,77 +6284,80 @@ EReimportResult::Type UReimportFbxSkeletalMeshFactory::Reimport( UObject* Obj, i
 	UE_LOG(LogEditorFactories, Log, TEXT("Performing atomic reimport of [%s]"), *Filename);
 	CurrentFilename = Filename;
 
-	if( !bOperationCanceled )
+	if (!bOperationCanceled)
 	{
-		FScopedSkeletalMeshPostEditChange ScopedPostEditChange(SkeletalMesh);
-
-		ImportOptions->bCanShowDialog = !IsUnattended;
-
-		if (ImportOptions->bImportAsSkeletalSkinning)
-		{
-			ImportOptions->bImportMaterials = false;
-			ImportOptions->bImportTextures = false;
-			ImportOptions->bImportLOD = false;
-			ImportOptions->bImportSkeletalMeshLODs = false;
-			ImportOptions->bImportAnimations = false;
-			ImportOptions->bImportMorph = false;
-			ImportOptions->VertexColorImportOption = EVertexColorImportOption::Ignore;
-		}
-		else if (ImportOptions->bImportAsSkeletalGeometry)
+		FScopedSuspendAlternateSkinWeightPreview ScopedSuspendAlternateSkinWeightPreview(SkeletalMesh);
 		{
-			ImportOptions->bImportAnimations = false;
-			ImportOptions->bUpdateSkeletonReferencePose = false;
-		}
+			FScopedSkeletalMeshPostEditChange ScopedPostEditChange(SkeletalMesh);
 
-		//Save all skinweight profile infos (need a copy, because they will be removed)
-		const TArray<FSkinWeightProfileInfo> ExistingSkinWeightProfileInfos = SkeletalMesh->GetSkinWeightProfiles();
+			ImportOptions->bCanShowDialog = !IsUnattended;
 
-		if ( FFbxImporter->ImportFromFile( *Filename, FPaths::GetExtension( Filename ), true ) )
-		{
-			if ( FFbxImporter->ReimportSkeletalMesh(SkeletalMesh, ImportData) )
+			if (ImportOptions->bImportAsSkeletalSkinning)
 			{
-				UE_LOG(LogEditorFactories, Log, TEXT("-- imported successfully") );
+				ImportOptions->bImportMaterials = false;
+				ImportOptions->bImportTextures = false;
+				ImportOptions->bImportLOD = false;
+				ImportOptions->bImportSkeletalMeshLODs = false;
+				ImportOptions->bImportAnimations = false;
+				ImportOptions->bImportMorph = false;
+				ImportOptions->VertexColorImportOption = EVertexColorImportOption::Ignore;
+			}
+			else if (ImportOptions->bImportAsSkeletalGeometry)
+			{
+				ImportOptions->bImportAnimations = false;
+				ImportOptions->bUpdateSkeletonReferencePose = false;
+			}
 
-				// Try to find the outer package so we can dirty it up
-				if (SkeletalMesh->GetOuter())
+			//Save all skinweight profile infos (need a copy, because they will be removed)
+			const TArray<FSkinWeightProfileInfo> ExistingSkinWeightProfileInfos = SkeletalMesh->GetSkinWeightProfiles();
+
+			if (FFbxImporter->ImportFromFile(*Filename, FPaths::GetExtension(Filename), true))
+			{
+				if (FFbxImporter->ReimportSkeletalMesh(SkeletalMesh, ImportData))
 				{
-					SkeletalMesh->GetOuter()->MarkPackageDirty();
+					UE_LOG(LogEditorFactories, Log, TEXT("-- imported successfully"));
+
+					// Try to find the outer package so we can dirty it up
+					if (SkeletalMesh->GetOuter())
+					{
+						SkeletalMesh->GetOuter()->MarkPackageDirty();
+					}
+					else
+					{
+						SkeletalMesh->MarkPackageDirty();
+					}
+
+					bSuccess = true;
 				}
 				else
 				{
-					SkeletalMesh->MarkPackageDirty();
+					UE_LOG(LogEditorFactories, Warning, TEXT("-- import failed"));
 				}
-
-				bSuccess = true;
 			}
 			else
 			{
-				UE_LOG(LogEditorFactories, Warning, TEXT("-- import failed") );
+				UE_LOG(LogEditorFactories, Warning, TEXT("-- import failed"));
 			}
-		}
-		else
-		{
-			UE_LOG(LogEditorFactories, Warning, TEXT("-- import failed") );
-		}
-		FFbxImporter->ReleaseScene(); 
+			FFbxImporter->ReleaseScene();
 
-		CleanUp();
+			CleanUp();
 
-		if (bSuccess && ExistingSkinWeightProfileInfos.Num() > 0)
-		{
-			//Restore skin weight profile infos, then reimport affected LODs
-			TArray<FSkinWeightProfileInfo>&SkinWeightsProfile = SkeletalMesh->GetSkinWeightProfiles();
-			SkinWeightsProfile = ExistingSkinWeightProfileInfos;
-			FSkinWeightsUtilities::ReimportAlternateSkinWeight(SkeletalMesh, 0);
-		}
+			if (bSuccess && ExistingSkinWeightProfileInfos.Num() > 0)
+			{
+				//Restore skin weight profile infos, then reimport affected LODs
+				TArray<FSkinWeightProfileInfo>&SkinWeightsProfile = SkeletalMesh->GetSkinWeightProfiles();
+				SkinWeightsProfile = ExistingSkinWeightProfileInfos;
+				FSkinWeightsUtilities::ReimportAlternateSkinWeight(SkeletalMesh, 0);
+			}
 
-		// Reimporting can have dangerous effects if the mesh is still in the transaction buffer.  Reset the transaction buffer if this is the case
-		if(!IsRunningCommandlet() && GEditor->IsObjectInTransactionBuffer( SkeletalMesh ) )
-		{
-			GEditor->ResetTransaction( LOCTEXT("ReimportSkeletalMeshTransactionReset", "Reimporting a skeletal mesh which was in the undo buffer") );
-		}
+			// Reimporting can have dangerous effects if the mesh is still in the transaction buffer.  Reset the transaction buffer if this is the case
+			if (!IsRunningCommandlet() && GEditor->IsObjectInTransactionBuffer(SkeletalMesh))
+			{
+				GEditor->ResetTransaction(LOCTEXT("ReimportSkeletalMeshTransactionReset", "Reimporting a skeletal mesh which was in the undo buffer"));
+			}
 
-		return bSuccess ? EReimportResult::Succeeded : EReimportResult::Failed;
+			return bSuccess ? EReimportResult::Succeeded : EReimportResult::Failed;
+		}
 	}
 	else
 	{
