Copied cl 9486552

PR #6266: Fix SkeletalMeshEdit TimeComparisonThreshold Implicit Cast Error (Contributed by [REDACTED_PERSON])

diff --git a/Engine/Source/Editor/UnrealEd/Private/SkeletalMeshEdit.cpp b/Engine/Source/Editor/UnrealEd/Private/SkeletalMeshEdit.cpp
--- a/Engine/Source/Editor/UnrealEd/Private/SkeletalMeshEdit.cpp
+++ b/Engine/Source/Editor/UnrealEd/Private/SkeletalMeshEdit.cpp
@@ -1550,7 +1550,7 @@ bool UnFbx::FFbxImporter::ImportAnimation(USkeleton* Skeleton, UAnimSequence * D
 		TimeIncrement.SetSecondDouble(1.0 / ((double)(ResampleRate)));
 
 		//Add a threshold when we compare if we have reach the end of the animation
-		const FbxTime TimeComparisonThreshold = (KINDA_SMALL_NUMBER*FBXSDK_TC_SECOND);
+		const FbxTime TimeComparisonThreshold = (static_cast<FbxLongLong>(KINDA_SMALL_NUMBER) * FBXSDK_TC_SECOND);
 		
 		
 		for(int32 SourceTrackIdx = 0; SourceTrackIdx < FbxRawBoneNames.Num(); ++SourceTrackIdx)
