Make sure we create the unique name from the same outer the object will be spawned into.

diff --git a/Engine/Source/Editor/UnrealEd/Private/EditorWorldExtension.cpp b/Engine/Source/Editor/UnrealEd/Private/EditorWorldExtension.cpp
--- a/Engine/Source/Editor/UnrealEd/Private/EditorWorldExtension.cpp
+++ b/Engine/Source/Editor/UnrealEd/Private/EditorWorldExtension.cpp
@@ -66,9 +66,10 @@ AActor* UEditorWorldExtension::SpawnTransientSceneActor(TSubclassOf<AActor> Acto
 	const bool bWasWorldPackageDirty = World->GetOutermost()->IsDirty();
 
 	FActorSpawnParameters ActorSpawnParameters;
-	ActorSpawnParameters.Name = MakeUniqueObjectName(World, ActorClass, *ActorName);
+	ActorSpawnParameters.Name = MakeUniqueObjectName(World->GetCurrentLevel(), ActorClass, *ActorName);
 	ActorSpawnParameters.SpawnCollisionHandlingOverride = ESpawnActorCollisionHandlingMethod::AlwaysSpawn;
 	ActorSpawnParameters.ObjectFlags = InObjectFlags;
+	ActorSpawnParameters.OverrideLevel = World->GetCurrentLevel();
 
 	check(ActorClass != nullptr);
 	AActor* NewActor = World->SpawnActor< AActor >(ActorClass, ActorSpawnParameters);
