Fix crash when saving package named 'None'

diff --git a/Engine/Source/Editor/UnrealEd/Private/UnrealEdSrv.cpp b/Engine/Source/Editor/UnrealEd/Private/UnrealEdSrv.cpp
--- a/Engine/Source/Editor/UnrealEd/Private/UnrealEdSrv.cpp
+++ b/Engine/Source/Editor/UnrealEd/Private/UnrealEdSrv.cpp
@@ -279,6 +279,11 @@ UPackage* UUnrealEdEngine::GeneratePackageThumbnailsIfRequired( const TCHAR* Str
 		FString TempFname;
 		if( FParse::Value( Str, TEXT( "FILE=" ), TempFname ) && ParseObject<UPackage>( Str, TEXT( "Package=" ), Pkg, NULL ) )
 		{
+			if (Pkg == nullptr)
+			{
+				return nullptr;
+			}
+
 			// Update any thumbnails for objects in this package that were modified or generate
 			// new thumbnails for objects that don't have any
 
