Editor: Various fixes relating to deleting assets with editors open for subobjects of those assets

This change comprises 3 changes:
  1. Added handling of blueprints contained within assets to BlueprintActionDatabase. When an asset containing a blueprint was to be deleted, it's inner BP would not get added to the PendingDelete array resulting in the BP actions showing up as references during the delete.
  2. Changed UAssetEditorSubsystem::CloseAllEditorsForAsset to close all editors that relate to the specified asset (including editors for subobjects of the asset)
  3. Reverted change made in 6836393 as this is now handled more generically by CloseAllEditorsForAsset
  4. Ensure that LevelSequence SequenceDirector blueprints do not have RF_Standalone as this cases them to be missed by UObject::TagSubobjects, and subsequently considered an external reference when deleting the object

diff --git a/Engine/Source/Editor/BlueprintGraph/Private/BlueprintActionDatabase.cpp b/Engine/Source/Editor/BlueprintGraph/Private/BlueprintActionDatabase.cpp
--- a/Engine/Source/Editor/BlueprintGraph/Private/BlueprintActionDatabase.cpp
+++ b/Engine/Source/Editor/BlueprintGraph/Private/BlueprintActionDatabase.cpp
@@ -37,6 +37,7 @@
 #include "K2Node_VariableGet.h"
 #include "K2Node_VariableSet.h"
 #include "Kismet2/BlueprintEditorUtils.h"
+#include "BlueprintAssetHandler.h"
 #include "EdGraphNode_Comment.h"
 #include "Animation/AnimInstance.h"
 #include "Editor.h"
@@ -974,7 +975,17 @@ static void BlueprintActionDatabaseImpl::OnAssetsPendingDelete(TArray<UObject*>
 			// in case they choose not to delete the object, we need to add 
 			// these back in to the database, so we track them here
 			PendingDelete.Add(DeletingObject);
-		}		
+		}
+
+		// If this asset contains a blueprint, ensure that it's actions are also marked for pending delete
+		else if (const IBlueprintAssetHandler* Handler = FBlueprintAssetHandler::Get().FindHandler(DeletingObject->GetClass()))
+		{
+			UBlueprint* Blueprint = Handler->RetrieveBlueprint(DeletingObject);
+			if (Blueprint && ActionDatabase.ClearAssetActions(Blueprint))
+			{
+				PendingDelete.Add(Blueprint);
+			}
+		}
 	}
 }
 
diff --git a/Engine/Source/Editor/UnrealEd/Private/ObjectTools.cpp b/Engine/Source/Editor/UnrealEd/Private/ObjectTools.cpp
--- a/Engine/Source/Editor/UnrealEd/Private/ObjectTools.cpp
+++ b/Engine/Source/Editor/UnrealEd/Private/ObjectTools.cpp
@@ -2449,7 +2449,7 @@ namespace ObjectTools
 		{
 			if (Object->IsAsset())
 			{
-				const TArray<IAssetEditorInstance*> ObjectEditors = GEditor->GetEditorSubsystem<UAssetEditorSubsystem>()->FindEditorsForAsset(Object);
+				TArray<IAssetEditorInstance*> ObjectEditors = GEditor->GetEditorSubsystem<UAssetEditorSubsystem>()->FindEditorsForAssetAndSubObjects(Object);
 				for (IAssetEditorInstance* ObjectEditorInstance : ObjectEditors)
 				{
 					if (!ObjectEditorInstance->CloseWindow())
diff --git a/Engine/Source/Editor/UnrealEd/Private/Subsystems/AssetEditorSubsystem.cpp b/Engine/Source/Editor/UnrealEd/Private/Subsystems/AssetEditorSubsystem.cpp
--- a/Engine/Source/Editor/UnrealEd/Private/Subsystems/AssetEditorSubsystem.cpp
+++ b/Engine/Source/Editor/UnrealEd/Private/Subsystems/AssetEditorSubsystem.cpp
@@ -101,10 +101,24 @@ TArray<IAssetEditorInstance*> UAssetEditorSubsystem::FindEditorsForAsset(UObject
 	return AssetEditors;
 }
 
+TArray<IAssetEditorInstance*> UAssetEditorSubsystem::FindEditorsForAssetAndSubObjects(UObject* Asset)
+{
+	TArray<IAssetEditorInstance*> EditorInstances;
+
+	for (const TPair<UObject*, IAssetEditorInstance*>& Pair : OpenedAssets)
+	{
+		if (Pair.Key == Asset || Pair.Key->IsIn(Asset))
+		{		
+			EditorInstances.Add(Pair.Value);
+		}
+	}
+
+	return EditorInstances;
+}
 
 int32 UAssetEditorSubsystem::CloseAllEditorsForAsset(UObject* Asset)
 {
-	TArray<IAssetEditorInstance*> EditorInstances = FindEditorsForAsset(Asset);
+	TArray<IAssetEditorInstance*> EditorInstances = FindEditorsForAssetAndSubObjects(Asset);
 
 	for (IAssetEditorInstance* EditorInstance : EditorInstances)
 	{
diff --git a/Engine/Source/Editor/UnrealEd/Private/Toolkits/AssetEditorManager.cpp b/Engine/Source/Editor/UnrealEd/Private/Toolkits/AssetEditorManager.cpp
--- a/Engine/Source/Editor/UnrealEd/Private/Toolkits/AssetEditorManager.cpp
+++ b/Engine/Source/Editor/UnrealEd/Private/Toolkits/AssetEditorManager.cpp
@@ -779,14 +779,14 @@ void FAssetEditorManager::HandlePackageReloaded(const EPackageReloadPhase InPack
 	if (InPackageReloadPhase == EPackageReloadPhase::PrePackageFixup)
 	{
 		/** Call close for all old assets even if not open, so global callback will go off */
-		TArray<UObject*> ObjectsToClose;
+		TArray<UObject*> AssetsToClose;
 		const TMap<UObject*, UObject*>& RepointedMap = InPackageReloadedEvent->GetRepointedObjects();
 
 		for (const TPair<UObject*, UObject*> RepointPair : RepointedMap)
 		{
 			if (RepointPair.Key->IsAsset())
 			{
-				ObjectsToClose.Add(RepointPair.Key);
+				AssetsToClose.Add(RepointPair.Key);
 			}
 		}
 
@@ -800,13 +800,11 @@ void FAssetEditorManager::HandlePackageReloaded(const EPackageReloadPhase InPack
 				{
 					PendingAssetsToOpen.AddUnique(NewAsset);
 				}
-
-				ObjectsToClose.AddUnique(AssetEditorPair.Key);
 			}
 		}
 
 		int32 NumAssetEditorsClosed = 0;
-		for (UObject* OldAsset : ObjectsToClose)
+		for (UObject* OldAsset : AssetsToClose)
 		{
 			NumAssetEditorsClosed += CloseAllEditorsForAsset(OldAsset);
 		}
diff --git a/Engine/Source/Editor/UnrealEd/Public/Subsystems/AssetEditorSubsystem.h b/Engine/Source/Editor/UnrealEd/Public/Subsystems/AssetEditorSubsystem.h
--- a/Engine/Source/Editor/UnrealEd/Public/Subsystems/AssetEditorSubsystem.h
+++ b/Engine/Source/Editor/UnrealEd/Public/Subsystems/AssetEditorSubsystem.h
@@ -87,6 +87,9 @@ public:
 	/** Returns all editors currently opened for the specified asset */
 	TArray<IAssetEditorInstance*> FindEditorsForAsset(UObject* Asset);
 
+	/** Returns all editors currently opened for the specified asset or any of its subobjects */
+	TArray<IAssetEditorInstance*> FindEditorsForAssetAndSubObjects(UObject* Asset);
+
 	/** Close all active editors for the supplied asset and return the number of asset editors that were closed */
 	UFUNCTION(BlueprintCallable, Category = "Editor Scripting | Asset Tools")
 	int32 CloseAllEditorsForAsset(UObject* Asset);
diff --git a/Engine/Source/Runtime/LevelSequence/Private/LevelSequence.cpp b/Engine/Source/Runtime/LevelSequence/Private/LevelSequence.cpp
--- a/Engine/Source/Runtime/LevelSequence/Private/LevelSequence.cpp
+++ b/Engine/Source/Runtime/LevelSequence/Private/LevelSequence.cpp
@@ -220,6 +220,8 @@ void ULevelSequence::PostLoad()
 
 	if (DirectorBlueprint)
 	{
+		DirectorBlueprint->ClearFlags(RF_Standalone);
+
 		// Remove the binding for the director blueprint recompilation and re-add it to be sure there is only one entry in the list
 		DirectorBlueprint->OnCompiled().RemoveAll(this);
 		DirectorBlueprint->OnCompiled().AddUObject(this, &ULevelSequence::OnDirectorRecompiled);
