Integrating CL# 9712548 from Dev-Enterprise
Fixing asset reimport not working after cancelling "reimport with new file"

diff --git a/Engine/Source/Editor/UnrealEd/Private/Editor.cpp b/Engine/Source/Editor/UnrealEd/Private/Editor.cpp
--- a/Engine/Source/Editor/UnrealEd/Private/Editor.cpp
+++ b/Engine/Source/Editor/UnrealEd/Private/Editor.cpp
@@ -550,7 +550,7 @@ void FReimportManager::ValidateAllSourceFileAndReimport(TArray<UObject*> &ToImpo
 				{
 					TArray<FString> SourceFilenames;
 					this->GetNewReimportPath(Asset, SourceFilenames, FileIndex);
-					if (SourceFilenames.Num() == 0)
+					if (SourceFilenames.Num() == 0 || SourceFilenames[0].IsEmpty())
 					{
 						continue;
 					}
