Moving a folder with a numeric name (eg. "1", "1234") to another level in the hierarchy no longer causes an infinite loop.

This was caused by a search that looked for the last non-numeric character in a string, but failed to account for the case that ALL characters are numeric.

Source: 9545694 (//UE4/Dev-Editor)

diff --git a/Engine/Source/Editor/UnrealEd/Private/EditorActorFolders.cpp b/Engine/Source/Editor/UnrealEd/Private/EditorActorFolders.cpp
--- a/Engine/Source/Editor/UnrealEd/Private/EditorActorFolders.cpp
+++ b/Engine/Source/Editor/UnrealEd/Private/EditorActorFolders.cpp
@@ -306,23 +306,30 @@ FName FActorFolders::GetFolderName(UWorld& InWorld, FName ParentPath, FName InLe
 	// This is potentially very slow but necessary to find a unique name
 	const auto& ExistingFolders = GetFolderPropertiesForWorld(InWorld);
 
-	// Trim any numeric suffix
-	uint32 SuffixLen = 0;
-	for (; InLeafName.ToString().Right(SuffixLen + 1).IsNumeric(); SuffixLen++) {}
+	const FString LeafNameString = InLeafName.ToString();
+
+	// Find the last non-numeric character
+	int32 LastDigit = LeafNameString.FindLastCharByPredicate([](TCHAR Ch) { return !FChar::IsDigit(Ch); });
+	uint32 SuffixLen = (LeafNameString.Len() - LastDigit) - 1;
+
+	if (LastDigit == INDEX_NONE)
+	{
+		// Name is entirely numeric, eg. "123", so no suffix exists
+		SuffixLen = 0;
+	}
 
+	// Trim any numeric suffix
 	uint32 Suffix = 1;
 	FString LeafNameRoot;
-	bool bHasSuffix = true;
 	if (SuffixLen > 0)
 	{
-		LeafNameRoot = InLeafName.ToString().LeftChop(SuffixLen);
-		FString LeafSuffix = InLeafName.ToString().RightChop(InLeafName.ToString().Len() - SuffixLen);
+		LeafNameRoot = LeafNameString.LeftChop(SuffixLen);
+		FString LeafSuffix = LeafNameString.RightChop(LeafNameString.Len() - SuffixLen);
 		Suffix = LeafSuffix.IsNumeric() ? FCString::Atoi(*LeafSuffix) : 1;
 	}
 	else
 	{
-		LeafNameRoot = InLeafName.ToString();
-		bHasSuffix = false;
+		LeafNameRoot = LeafNameString;
 	}
 
 	// Create a valid base name for this folder
@@ -330,7 +337,7 @@ FName FActorFolders::GetFolderName(UWorld& InWorld, FName ParentPath, FName InLe
 	NumberFormat.SetUseGrouping(false);
 	NumberFormat.SetMinimumIntegralDigits(SuffixLen);
 
-	FText LeafName = FText::Format(LOCTEXT("FolderNamePattern", "{0}{1}"), FText::FromString(LeafNameRoot), bHasSuffix ? FText::AsNumber(Suffix++, &NumberFormat) : FText::GetEmpty());
+	FText LeafName = FText::Format(LOCTEXT("FolderNamePattern", "{0}{1}"), FText::FromString(LeafNameRoot), SuffixLen > 0 ? FText::AsNumber(Suffix++, &NumberFormat) : FText::GetEmpty());
 
 	FString ParentFolderPath = ParentPath.IsNone() ? TEXT("") : ParentPath.ToString();
 	if (!ParentFolderPath.IsEmpty())
