Editor: Fixed reallocation of auto-reimport monitors causinginvalid delegate bindings

FContentDirectoryMonitor is a self-referential type since it contains an FFileCache which is also self-referential. This means that reallocation of the monitors can result in invalid delegates being bound for the file monitors. This has been addressed by heap-allocating the monitors (so they are never moved or copied themselves), along with some other trivial changes to bring this code more up-to-date with modern UE4 coding standards.
This is largely in response to github pr 5906 but is a more thorough solution.

#5906

diff --git a/Engine/Source/Editor/UnrealEd/Private/AutoReimport/AutoReimportManager.cpp b/Engine/Source/Editor/UnrealEd/Private/AutoReimport/AutoReimportManager.cpp
--- a/Engine/Source/Editor/UnrealEd/Private/AutoReimport/AutoReimportManager.cpp
+++ b/Engine/Source/Editor/UnrealEd/Private/AutoReimport/AutoReimportManager.cpp
@@ -36,27 +36,31 @@
 #include "AutoReimport/AssetSourceFilenameCache.h"
 
 #define LOCTEXT_NAMESPACE "AutoReimportManager"
-#define yield if (TimeLimit.Exceeded()) return
 
 enum class EStateMachineNode { CallOnce, CallMany };
 
+/** Enum and value to specify the current state of our processing */
+enum class ECurrentState : uint8
+{
+	Idle, Paused, Aborting, PromptUser, Initializing, ProcessAdditions, ProcessModifications, ProcessDeletions, SavePackages
+};
+
 /** A simple state machine class that calls functions mapped on enum values. If any function returns a new enum type, it moves onto that function */
-template<typename TState>
 struct FStateMachine
 {
-	typedef TFunction<TOptional<TState>(const DirectoryWatcher::FTimeLimit&)> FunctionType;
+	typedef TFunction<TOptional<ECurrentState>(const DirectoryWatcher::FTimeLimit&)> FunctionType;
 
 	/** Constructor that specifies the initial state of the machine */
-	FStateMachine(TState InitialState) : CurrentState(InitialState){}
+	FStateMachine(ECurrentState InitialState) : CurrentState(InitialState){}
 
 	/** Add an enum->function mapping for this state machine */
-	void Add(TState State, EStateMachineNode NodeType, FunctionType&& Function)
+	void Add(ECurrentState State, EStateMachineNode NodeType, FunctionType&& Function)
 	{
 		Nodes.Add(State, FStateMachineNode(NodeType, MoveTemp(Function)));
 	}
 
 	/** Set a new state for this machine */
-	void SetState(TState NewState)
+	void SetState(ECurrentState NewState)
 	{
 		CurrentState = NewState;
 	}
@@ -66,8 +70,8 @@ struct FStateMachine
 	{
 		while (!TimeLimit.Exceeded())
 		{
-			const auto& State = Nodes[CurrentState];
-			TOptional<TState> NewState = State.Endpoint(TimeLimit);
+			const FStateMachineNode& State = Nodes[CurrentState];
+			TOptional<ECurrentState> NewState = State.Endpoint(TimeLimit);
 			if (NewState.IsSet())
 			{
 				CurrentState = NewState.GetValue();
@@ -82,7 +86,7 @@ struct FStateMachine
 private:
 
 	/** The current state of this machine */
-	TState CurrentState;
+	ECurrentState CurrentState;
 
 private:
 
@@ -98,18 +102,10 @@ private:
 	};
 
 	/** A map of enum value -> callback information */
-	TMap<TState, FStateMachineNode> Nodes;
+	TMap<ECurrentState, FStateMachineNode> Nodes;
 };
 
-/** Enum and value to specify the current state of our processing */
-enum class ECurrentState
-{
-	Idle, Paused, Aborting, PromptUser, Initializing, ProcessAdditions, ProcessModifications, ProcessDeletions, SavePackages
-};
-uint32 GetTypeHash(ECurrentState State)
-{
-	return (int32)State;
-}
+
 
 /* Deals with auto reimporting of objects when the objects file on disk is modified*/
 class FAutoReimportManager : public FTickableEditorObject, public FGCObject, public TSharedFromThis<FAutoReimportManager>
@@ -175,7 +171,7 @@ private:
 private:
 
 	/** A state machine holding information about the current state of the manager */
-	FStateMachine<ECurrentState> StateMachine;
+	FStateMachine StateMachine;
 
 private:
 
@@ -221,7 +217,7 @@ private:
 	TSharedPtr<FReimportFeedbackContext> FeedbackContextOverride;
 
 	/** Array of objects that detect changes to directories */
-	TArray<FContentDirectoryMonitor> DirectoryMonitors;
+	TArray<TUniquePtr<FContentDirectoryMonitor>> DirectoryMonitors;
 
 	/** A list of packages to save when we've added a bunch of assets */
 	TArray<UPackage*> PackagesToSave;
@@ -261,7 +257,7 @@ FAutoReimportManager::FAutoReimportManager()
 	, bGuardAssetChanges(false)
 	, PausedState(ECurrentState::Idle)
 {
-	auto* Settings = GetMutableDefault<UEditorLoadingSavingSettings>();
+	UEditorLoadingSavingSettings* Settings = GetMutableDefault<UEditorLoadingSavingSettings>();
 	Settings->OnSettingChanged().AddRaw(this, &FAutoReimportManager::HandleLoadingSavingSettingChanged);
 
 	FPackageName::OnContentPathMounted().AddRaw(this, &FAutoReimportManager::OnContentPathChanged);
@@ -305,64 +301,64 @@ FAutoReimportManager::~FAutoReimportManager()
 TArray<FPathAndMountPoint> FAutoReimportManager::GetMonitoredDirectories() const
 {
 	TArray<FPathAndMountPoint> Dirs;
-	for (const auto& Monitor : DirectoryMonitors)
+	for (const TUniquePtr<FContentDirectoryMonitor>& Monitor : DirectoryMonitors)
 	{
-		Dirs.Emplace(Monitor.GetDirectory(), Monitor.GetMountPoint());
+		Dirs.Emplace(Monitor->GetDirectory(), Monitor->GetMountPoint());
 	}
 	return Dirs;
 }
 
 void FAutoReimportManager::IgnoreNewFile(const FString& Filename)
 {
-	for (auto& Monitor : DirectoryMonitors)
+	for (const TUniquePtr<FContentDirectoryMonitor>& Monitor : DirectoryMonitors)
 	{
-		if (Filename.StartsWith(Monitor.GetDirectory()))
+		if (Filename.StartsWith(Monitor->GetDirectory()))
 		{
-			Monitor.IgnoreNewFile(Filename);
+			Monitor->IgnoreNewFile(Filename);
 		}
 	}
 }
 
 void FAutoReimportManager::IgnoreFileModification(const FString& Filename)
 {
-	for (auto& Monitor : DirectoryMonitors)
+	for (const TUniquePtr<FContentDirectoryMonitor>& Monitor : DirectoryMonitors)
 	{
-		if (Filename.StartsWith(Monitor.GetDirectory()))
+		if (Filename.StartsWith(Monitor->GetDirectory()))
 		{
-			Monitor.IgnoreFileModification(Filename);
+			Monitor->IgnoreFileModification(Filename);
 		}
 	}
 }
 
 void FAutoReimportManager::IgnoreMovedFile(const FString& SrcFilename, const FString& DstFilename)
 {
-	for (auto& Monitor : DirectoryMonitors)
+	for (const TUniquePtr<FContentDirectoryMonitor>& Monitor : DirectoryMonitors)
 	{
-		const bool bSrcInFolder = SrcFilename.StartsWith(Monitor.GetDirectory());
-		const bool bDstInFolder = DstFilename.StartsWith(Monitor.GetDirectory());
+		const bool bSrcInFolder = SrcFilename.StartsWith(Monitor->GetDirectory());
+		const bool bDstInFolder = DstFilename.StartsWith(Monitor->GetDirectory());
 
 		if (bSrcInFolder && bDstInFolder)
 		{
-			Monitor.IgnoreMovedFile(SrcFilename, DstFilename);
+			Monitor->IgnoreMovedFile(SrcFilename, DstFilename);
 		}
 		else if (bSrcInFolder)
 		{
-			Monitor.IgnoreDeletedFile(SrcFilename);
+			Monitor->IgnoreDeletedFile(SrcFilename);
 		}
 		else if (bDstInFolder)
 		{
-			Monitor.IgnoreNewFile(DstFilename);
+			Monitor->IgnoreNewFile(DstFilename);
 		}
 	}
 }
 
 void FAutoReimportManager::IgnoreDeletedFile(const FString& Filename)
 {
-	for (auto& Monitor : DirectoryMonitors)
+	for (const TUniquePtr<FContentDirectoryMonitor>& Monitor : DirectoryMonitors)
 	{
-		if (Filename.StartsWith(Monitor.GetDirectory()))
+		if (Filename.StartsWith(Monitor->GetDirectory()))
 		{
-			Monitor.IgnoreDeletedFile(Filename);
+			Monitor->IgnoreDeletedFile(Filename);
 		}
 	}
 }
@@ -376,7 +372,7 @@ void FAutoReimportManager::Destroy()
 		AssetRegistryModule->Get().OnInMemoryAssetDeleted().RemoveAll(this);
 	}
 	
-	if (auto* Settings = GetMutableDefault<UEditorLoadingSavingSettings>())
+	if (UEditorLoadingSavingSettings* Settings = GetMutableDefault<UEditorLoadingSavingSettings>())
 	{
 		Settings->OnSettingChanged().RemoveAll(this);
 	}
@@ -473,18 +469,18 @@ void FAutoReimportManager::OnAssetRenamed(const FAssetData& AssetData, const FSt
 
 int32 FAutoReimportManager::GetNumUnprocessedChanges() const
 {
-	return Utils::Reduce(DirectoryMonitors, [](const FContentDirectoryMonitor& Monitor, int32 Total){
-		return Total + Monitor.GetNumUnprocessedChanges();
+	return Utils::Reduce(DirectoryMonitors, [](const TUniquePtr<FContentDirectoryMonitor>& Monitor, int32 Total){
+		return Total + Monitor->GetNumUnprocessedChanges();
 	}, 0);
 }
 
 void FAutoReimportManager::PopupateMessageLogWithPendingChanges(FMessageLog& MessageLog) const
 {
-	for (const FContentDirectoryMonitor& Monitor : DirectoryMonitors)
+	for (const TUniquePtr<FContentDirectoryMonitor>& Monitor : DirectoryMonitors)
 	{
-		const FString& BasePath = Monitor.GetDirectory();
+		const FString& BasePath = Monitor->GetDirectory();
 
-		Monitor.IterateUnprocessedChanges([&](const DirectoryWatcher::FUpdateCacheTransaction& Transaction, const FDateTime& TimeOfChange){
+		Monitor->IterateUnprocessedChanges([&MessageLog, &BasePath](const DirectoryWatcher::FUpdateCacheTransaction& Transaction, const FDateTime& TimeOfChange){
 			FText ThisMessage;
 			FString FullFilename = BasePath / Transaction.Filename.Get();
 			switch(Transaction.Action)
@@ -587,9 +583,9 @@ TOptional<ECurrentState> FAutoReimportManager::PromptUser()
 	}
 
 	// Keep ticking the monitors
-	for (auto& Monitor : DirectoryMonitors)
+	for (const TUniquePtr<FContentDirectoryMonitor>& Monitor : DirectoryMonitors)
 	{
-		Monitor.Tick();
+		Monitor->Tick();
 	}
 
 	if (ConfirmNotification.IsValid())
@@ -632,9 +628,9 @@ TOptional<ECurrentState> FAutoReimportManager::PromptUser()
 TOptional<ECurrentState> FAutoReimportManager::InitializeOperation()
 {
 	int32 TotalWork = 0;
-	for (auto& Monitor : DirectoryMonitors)
+	for (const TUniquePtr<FContentDirectoryMonitor>& Monitor : DirectoryMonitors)
 	{
-		TotalWork += Monitor.StartProcessing();
+		TotalWork += Monitor->StartProcessing();
 	}
 
 	if (TotalWork > 0)
@@ -657,7 +653,7 @@ TOptional<ECurrentState> FAutoReimportManager::InitializeOperation()
 
 TOptional<ECurrentState> FAutoReimportManager::ProcessAdditions(const DirectoryWatcher::FTimeLimit& TimeLimit)
 {
-	auto NewState = HandlePauseAbort(ECurrentState::ProcessAdditions);
+	TOptional<ECurrentState> NewState = HandlePauseAbort(ECurrentState::ProcessAdditions);
 	if (NewState.IsSet())
 	{
 		return NewState;
@@ -685,24 +681,26 @@ TOptional<ECurrentState> FAutoReimportManager::ProcessAdditions(const DirectoryW
 				FactoryExtensions.Reset();
 				Factory->GetSupportedFileExtensions(FactoryExtensions);
 
-				for (const auto& Ext : FactoryExtensions)
+				for (const FString& Ext : FactoryExtensions)
 				{
-					auto& Array = Factories.FindOrAdd(Ext);
-					Array.Add(Factory);
+					Factories.FindOrAdd(Ext).Add(Factory);
 				}
 			}
 		}
 	}
 
-	for (auto& Pair : Factories)
+	for (TTuple<FString, TArray<UFactory*>>& Pair : Factories)
 	{
 		Pair.Value.Sort([](const UFactory& A, const UFactory& B) { return A.ImportPriority > B.ImportPriority; });
 	}
 
-	for (auto& Monitor : DirectoryMonitors)
+	for (const TUniquePtr<FContentDirectoryMonitor>& Monitor : DirectoryMonitors)
 	{
-		Monitor.ProcessAdditions(TimeLimit, PackagesToSave, Factories, *FeedbackContextOverride);
-		yield TOptional<ECurrentState>();
+		Monitor->ProcessAdditions(TimeLimit, PackagesToSave, Factories, *FeedbackContextOverride);
+		if (TimeLimit.Exceeded())
+		{
+			return TOptional<ECurrentState>();
+		}
 	}
 
 	return ECurrentState::ProcessModifications;
@@ -710,7 +708,7 @@ TOptional<ECurrentState> FAutoReimportManager::ProcessAdditions(const DirectoryW
 
 TOptional<ECurrentState> FAutoReimportManager::ProcessModifications(const DirectoryWatcher::FTimeLimit& TimeLimit)
 {
-	auto NewState = HandlePauseAbort(ECurrentState::ProcessModifications);
+	TOptional<ECurrentState> NewState = HandlePauseAbort(ECurrentState::ProcessModifications);
 	if (NewState.IsSet())
 	{
 		return NewState;
@@ -720,10 +718,14 @@ TOptional<ECurrentState> FAutoReimportManager::ProcessModifications(const Direct
 	TGuardValue<FFeedbackContext*> ScopedContextOverride(GWarn, FeedbackContextOverride.Get());
 	TGuardValue<bool> ScopedAssetChangesGuard(bGuardAssetChanges, true);
 
-	for (auto& Monitor : DirectoryMonitors)
+	for (const TUniquePtr<FContentDirectoryMonitor>& Monitor : DirectoryMonitors)
 	{
-		Monitor.ProcessModifications(TimeLimit, PackagesToSave, *FeedbackContextOverride);
-		yield TOptional<ECurrentState>();
+		Monitor->ProcessModifications(TimeLimit, PackagesToSave, *FeedbackContextOverride);
+
+		if (TimeLimit.Exceeded())
+		{
+			return TOptional<ECurrentState>();
+		}
 	}
 
 	return ECurrentState::ProcessDeletions;
@@ -731,7 +733,7 @@ TOptional<ECurrentState> FAutoReimportManager::ProcessModifications(const Direct
 
 TOptional<ECurrentState> FAutoReimportManager::ProcessDeletions()
 {
-	auto NewState = HandlePauseAbort(ECurrentState::ProcessDeletions);
+	TOptional<ECurrentState> NewState = HandlePauseAbort(ECurrentState::ProcessDeletions);
 	if (NewState.IsSet())
 	{
 		return NewState;
@@ -741,16 +743,16 @@ TOptional<ECurrentState> FAutoReimportManager::ProcessDeletions()
 
 	TArray<FAssetData> AssetsToDelete;
 
-	for (auto& Monitor : DirectoryMonitors)
+	for (const TUniquePtr<FContentDirectoryMonitor>& Monitor : DirectoryMonitors)
 	{
-		Monitor.ExtractAssetsToDelete(AssetsToDelete);
+		Monitor->ExtractAssetsToDelete(AssetsToDelete);
 	}
 
 	FeedbackContextOverride->MainTask->EnterProgressFrame(AssetsToDelete.Num());
 
 	if (AssetsToDelete.Num() > 0)
 	{
-		for (const auto& AssetData : AssetsToDelete)
+		for (const FAssetData& AssetData : AssetsToDelete)
 		{
 			FeedbackContextOverride->AddMessage(EMessageSeverity::Info, FText::Format(LOCTEXT("Success_DeletedAsset", "Attempting to delete {0} (its source file has been removed)."), FText::FromName(AssetData.AssetName)));
 		}
@@ -811,9 +813,9 @@ TOptional<ECurrentState> FAutoReimportManager::Paused()
 
 TOptional<ECurrentState> FAutoReimportManager::Abort()
 {
-	for (auto& Monitor : DirectoryMonitors)
+	for (const TUniquePtr<FContentDirectoryMonitor>& Monitor : DirectoryMonitors)
 	{
-		Monitor.Abort();
+		Monitor->Abort();
 	}
 
 	PackagesToSave.Empty();
@@ -827,7 +829,7 @@ TOptional<ECurrentState> FAutoReimportManager::Idle()
 	// Check whether we need to reset the monitors or not
 	if (ResetMonitorsTimeout.Exceeded())
 	{
-		const auto* Settings = GetDefault<UEditorLoadingSavingSettings>();
+		const UEditorLoadingSavingSettings* Settings = GetDefault<UEditorLoadingSavingSettings>();
 		if (Settings->bMonitorContentDirectories)
 		{
 			DirectoryMonitors.Empty();
@@ -836,9 +838,9 @@ TOptional<ECurrentState> FAutoReimportManager::Idle()
 		else
 		{
 			// Destroy all the existing monitors, including their file caches
-			for (auto& Monitor : DirectoryMonitors)
+			for (const TUniquePtr<FContentDirectoryMonitor>& Monitor : DirectoryMonitors)
 			{
-				Monitor.Destroy();
+				Monitor->Destroy();
 			}
 			DirectoryMonitors.Empty();
 		}
@@ -847,9 +849,9 @@ TOptional<ECurrentState> FAutoReimportManager::Idle()
 		return TOptional<ECurrentState>();
 	}
 
-	for (auto& Monitor : DirectoryMonitors)
+	for (const TUniquePtr<FContentDirectoryMonitor>& Monitor : DirectoryMonitors)
 	{
-		Monitor.Tick();
+		Monitor->Tick();
 	}
 
 	const IAssetRegistry& AssetRegistry = FModuleManager::LoadModuleChecked<FAssetRegistryModule>("AssetRegistry").Get();
@@ -903,7 +905,7 @@ void FAutoReimportManager::HandleLoadingSavingSettingChanged(FName PropertyName)
 
 void FAutoReimportManager::OnContentPathChanged(const FString& InAssetPath, const FString& FileSystemPath)
 {
-	const auto* Settings = GetDefault<UEditorLoadingSavingSettings>();
+	const UEditorLoadingSavingSettings* Settings = GetDefault<UEditorLoadingSavingSettings>();
 	if (Settings->bMonitorContentDirectories)
 	{
 		DirectoryMonitors.Empty();
@@ -921,8 +923,8 @@ void FAutoReimportManager::SetUpDirectoryMonitors()
 	};
 
 	TArray<FParsedSettings> FinalArray;
-	auto SupportedExtensions = GetAllFactoryExtensions();
-	for (const auto& Setting : GetDefault<UEditorLoadingSavingSettings>()->AutoReimportDirectorySettings)
+	FString SupportedExtensions = GetAllFactoryExtensions();
+	for (const FAutoReimportDirectoryConfig& Setting : GetDefault<UEditorLoadingSavingSettings>()->AutoReimportDirectorySettings)
 	{
 		FParsedSettings NewMapping;
 		NewMapping.SourceDirectory = Setting.SourceDirectory;
@@ -935,7 +937,7 @@ void FAutoReimportManager::SetUpDirectoryMonitors()
 
 		// Only include extensions that match a factory
 		NewMapping.Rules.SetApplicableExtensions(SupportedExtensions);
-		for (const auto& WildcardConfig : Setting.Wildcards)
+		for (const FAutoReimportWildcard& WildcardConfig : Setting.Wildcards)
 		{
 			NewMapping.Rules.AddWildcardRule(WildcardConfig.Wildcard, WildcardConfig.bInclude);
 		}
@@ -945,7 +947,7 @@ void FAutoReimportManager::SetUpDirectoryMonitors()
 
 	for (int32 Index = 0; Index < FinalArray.Num(); ++Index)
 	{
-		const auto& Mapping = FinalArray[Index];
+		const FParsedSettings& Mapping = FinalArray[Index];
 
 		// We only create a directory monitor if there are no other's watching parent directories of this one
 		for (int32 OtherIndex = Index + 1; OtherIndex < FinalArray.Num(); ++OtherIndex)
@@ -957,7 +959,7 @@ void FAutoReimportManager::SetUpDirectoryMonitors()
 			}
 		}
 
-		DirectoryMonitors.Emplace(Mapping.SourceDirectory, Mapping.Rules, Mapping.MountPoint);
+		DirectoryMonitors.Emplace(MakeUnique<FContentDirectoryMonitor>(Mapping.SourceDirectory, Mapping.Rules, Mapping.MountPoint));
 
 	next:
 		continue;
@@ -1052,6 +1054,5 @@ void UAutoReimportManager::BeginDestroy()
 }
 
 
-#undef CONDITIONAL_YIELD
 #undef LOCTEXT_NAMESPACE
 
