Fix for component visualizer caching; does not attempt to access a property outside array bounds. This may happen if a BP generated component has been removed from one tick to the next.  This is to fix an editor crash which was occurring.

diff --git a/Engine/Source/Editor/UnrealEd/Private/ComponentVisualizer.cpp b/Engine/Source/Editor/UnrealEd/Private/ComponentVisualizer.cpp
--- a/Engine/Source/Editor/UnrealEd/Private/ComponentVisualizer.cpp
+++ b/Engine/Source/Editor/UnrealEd/Private/ComponentVisualizer.cpp
@@ -99,8 +99,15 @@ UActorComponent* FComponentPropertyPath::GetComponent() const
 					if (UObjectProperty* InnerProp = Cast<UObjectProperty>(ArrayProp->Inner))
 					{
 						FScriptArrayHelper ArrayHelper(ArrayProp, ArrayProp->ContainerPtrToValuePtr<void>(Actor));
-						UObject* Object = InnerProp->GetObjectPropertyValue(ArrayHelper.GetRawPtr(PropertyIndex));
-						Result = Cast<UActorComponent>(Object);
+						if (ArrayHelper.IsValidIndex(PropertyIndex))
+						{
+							UObject* Object = InnerProp->GetObjectPropertyValue(ArrayHelper.GetRawPtr(PropertyIndex));
+							Result = Cast<UActorComponent>(Object);
+						}
+						else
+						{
+							UE_LOG(LogTemp, Log, TEXT("HONO"));
+						}
 					}
 				}
 			}
