Removed extraneous logging.

diff --git a/Engine/Source/Editor/UnrealEd/Private/ComponentVisualizer.cpp b/Engine/Source/Editor/UnrealEd/Private/ComponentVisualizer.cpp
--- a/Engine/Source/Editor/UnrealEd/Private/ComponentVisualizer.cpp
+++ b/Engine/Source/Editor/UnrealEd/Private/ComponentVisualizer.cpp
@@ -104,10 +104,6 @@ UActorComponent* FComponentPropertyPath::GetComponent() const
 							UObject* Object = InnerProp->GetObjectPropertyValue(ArrayHelper.GetRawPtr(PropertyIndex));
 							Result = Cast<UActorComponent>(Object);
 						}
-						else
-						{
-							UE_LOG(LogTemp, Log, TEXT("HONO"));
-						}
 					}
 				}
 			}
