Removed seemingly deprecated hook from the editor to reset the undo buffer when gc-ing steaming levels.
PIE objects can't enter the undo buffer and other hooks already reset the transaction buffer when world/level are removed in the editor

diff --git a/Engine/Source/Editor/UnrealEd/Classes/Editor/EditorEngine.h b/Engine/Source/Editor/UnrealEd/Classes/Editor/EditorEngine.h
--- a/Engine/Source/Editor/UnrealEd/Classes/Editor/EditorEngine.h
+++ b/Engine/Source/Editor/UnrealEd/Classes/Editor/EditorEngine.h
@@ -2862,9 +2862,6 @@ private:
 	/** Delegate callback for when a streaming level is removed from world. */
 	void OnLevelRemovedFromWorld(ULevel* InLevel, UWorld* InWorld);
 
-	/** Delegate callback for when a streamed out levels going to be removed by GC. */
-	void OnGCStreamedOutLevels();
-
 	/** Puts the currently loaded project file at the top of the recents list and trims and files that fell off the list */
 	void UpdateRecentlyLoadedProjectFiles();
 
diff --git a/Engine/Source/Editor/UnrealEd/Private/EditorEngine.cpp b/Engine/Source/Editor/UnrealEd/Private/EditorEngine.cpp
--- a/Engine/Source/Editor/UnrealEd/Private/EditorEngine.cpp
+++ b/Engine/Source/Editor/UnrealEd/Private/EditorEngine.cpp
@@ -899,7 +899,6 @@ void UEditorEngine::Init(IEngineLoop* InEngineLoop)
 	FCoreUObjectDelegates::OnAssetLoaded.AddUObject(this, &UEditorEngine::OnAssetLoaded);
 	FWorldDelegates::LevelAddedToWorld.AddUObject(this, &UEditorEngine::OnLevelAddedToWorld);
 	FWorldDelegates::LevelRemovedFromWorld.AddUObject(this, &UEditorEngine::OnLevelRemovedFromWorld);
-	FLevelStreamingGCHelper::OnGCStreamedOutLevels.AddUObject(this, &UEditorEngine::OnGCStreamedOutLevels);
 
 	FAssetRegistryModule& AssetRegistryModule = FModuleManager::LoadModuleChecked<FAssetRegistryModule>("AssetRegistry");
 	AssetRegistryModule.Get().OnInMemoryAssetCreated().AddUObject(this, &UEditorEngine::OnAssetCreated);
@@ -1230,7 +1229,6 @@ void UEditorEngine::FinishDestroy()
 		FCoreUObjectDelegates::OnAssetLoaded.RemoveAll(this);
 		FWorldDelegates::LevelAddedToWorld.RemoveAll(this);
 		FWorldDelegates::LevelRemovedFromWorld.RemoveAll(this);
-		FLevelStreamingGCHelper::OnGCStreamedOutLevels.RemoveAll(this);
 		GetMutableDefault<UEditorStyleSettings>()->OnSettingChanged().RemoveAll(this);
 
 		FAssetRegistryModule* AssetRegistryModule = FModuleManager::GetModulePtr<FAssetRegistryModule>("AssetRegistry");
@@ -6663,30 +6661,19 @@ void UEditorEngine::OnLevelRemovedFromWorld(ULevel* InLevel, UWorld* InWorld)
 			}
 		}
 	}
-	else
+	// UEngine::LoadMap broadcast this event with InLevel==NULL, before cleaning up the world during travel in Multiplayer
+	else if (InWorld->IsPlayInEditor() && Trans)
 	{
-		// UEngine::LoadMap broadcast this event with InLevel==NULL, before cleaning up the world
-		if (InWorld->IsPlayInEditor())
-		{
-			// Each additional instance of PIE in a multiplayer game will add another barrier, so if the event is triggered then this is the case and we need to lift it
-			// Otherwise there will be an imbalance between barriers set and barriers removed and we won't be able to undo when we return.
-			if (Trans)
-			{
-				Trans->RemoveUndoBarrier();
-			}
-		}
-		else
-		{	
-			// If we're in editor mode, reset transactions buffer, to ensure that there are no references to a world which is about to be destroyed
-			ResetTransaction(NSLOCTEXT("UnrealEd", "LoadMapTransReset", "Loading a New Map"));
-		}
+		// Each additional instance of PIE in a multiplayer game will add another barrier, so if the event is triggered then this is the case and we need to lift it
+		// Otherwise there will be an imbalance between barriers set and barriers removed and we won't be able to undo when we return.
+		Trans->RemoveUndoBarrier();
 	}
-}
 
-void UEditorEngine::OnGCStreamedOutLevels()
-{
-	// Reset transaction buffer because it may hold references to streamed out levels
-	ResetTransaction( NSLOCTEXT("UnrealEd", "GCStreamedOutLevelsTransReset", "GC Streaming Levels") );
+	if (!InWorld->IsPlayInEditor())
+	{
+		// If we're in editor mode, reset transactions buffer, to ensure that there are no references to a world which is about to be destroyed
+		ResetTransaction(NSLOCTEXT("UnrealEd", "LevelRemovedFromWorldEditorCallback", "Level removed from world"));
+	}
 }
 
 void UEditorEngine::UpdateRecentlyLoadedProjectFiles()
