- Prevented the editor engine from resetting the package dirty flag when the package did not save successfully - like the file was read only.

diff --git a/Engine/Source/Editor/UnrealEd/Private/EditorEngine.cpp b/Engine/Source/Editor/UnrealEd/Private/EditorEngine.cpp
--- a/Engine/Source/Editor/UnrealEd/Private/EditorEngine.cpp
+++ b/Engine/Source/Editor/UnrealEd/Private/EditorEngine.cpp
@@ -4481,8 +4481,11 @@ FSavePackageResultStruct UEditorEngine::Save( UPackage* InOuter, UObject* InBase
 				CleanupPhysicsSceneThatWasInitializedForSave(World, bForceInitializedWorld);
 			}
 
-			// Rerunning construction scripts may have made it dirty again
-			InOuter->SetDirtyFlag(false);
+			if (Result == ESavePackageResult::Success) // Package saved successfully?
+			{
+				// Rerunning construction scripts may have made it dirty again
+				InOuter->SetDirtyFlag(false);
+			}
 		}
 	}
 
