Properly apply NewlyCreated package flag when creating a new Untitled world from a template
Use NewlyCreated package flag to skip costly disk IOs in cookonthefly and thumbnail loading for package that have never been saved
   - 3m49s -> 2m19s when importing DATASET-0006a

From dev-enterprise[REDACTED_HANDLE]

diff --git a/Engine/Source/Editor/UnrealEd/Private/AssetThumbnail.cpp b/Engine/Source/Editor/UnrealEd/Private/AssetThumbnail.cpp
--- a/Engine/Source/Editor/UnrealEd/Private/AssetThumbnail.cpp
+++ b/Engine/Source/Editor/UnrealEd/Private/AssetThumbnail.cpp
@@ -1445,7 +1445,8 @@ void FAssetThumbnailPool::DirtyThumbnailForObject(UObject* ObjectBeingModified)
 		// regenerated on demand later. (Before being displayed in the browser, or package saves, etc.)
 		FObjectThumbnail* Thumbnail = ThumbnailTools::GetThumbnailForObject(ObjectBeingModified);
 
-		if (Thumbnail == NULL && !IsGarbageCollecting())
+		// Don't try loading thumbnails for package that have never been saved
+		if (Thumbnail == NULL && !IsGarbageCollecting() && !ObjectBeingModified->GetOutermost()->HasAnyPackageFlags(PKG_NewlyCreated))
 		{
 			// If we don't yet have a thumbnail map, load one from disk if possible
 			// Don't attempt to do this while garbage collecting since loading or finding objects during GC is illegal
diff --git a/Engine/Source/Editor/UnrealEd/Private/CookOnTheFlyServer.cpp b/Engine/Source/Editor/UnrealEd/Private/CookOnTheFlyServer.cpp
--- a/Engine/Source/Editor/UnrealEd/Private/CookOnTheFlyServer.cpp
+++ b/Engine/Source/Editor/UnrealEd/Private/CookOnTheFlyServer.cpp
@@ -3341,7 +3341,8 @@ void UCookOnTheFlyServer::MarkPackageDirtyForCooker( UPackage *Package )
 	{
 		return;
 	}
-	if (Package->HasAnyPackageFlags(PKG_PlayInEditor | PKG_ContainsScript | PKG_CompiledIn) == true && !GetClass()->HasAnyClassFlags(CLASS_DefaultConfig | CLASS_Config))
+
+	if (Package->HasAnyPackageFlags(PKG_PlayInEditor | PKG_ContainsScript | PKG_InMemoryOnly) == true && !GetClass()->HasAnyClassFlags(CLASS_DefaultConfig | CLASS_Config))
 	{
 		return;
 	}
@@ -3358,6 +3359,8 @@ void UCookOnTheFlyServer::MarkPackageDirtyForCooker( UPackage *Package )
 
 	if ( !bIsSavingPackage )
 	{
+		TRACE_CPUPROFILER_EVENT_SCOPE(MarkPackageDirtyForCooker);
+
 		// could have just cooked a file which we might need to write
 		UPackage::WaitForAsyncFileWrites();
 
diff --git a/Engine/Source/Editor/UnrealEd/Private/EditorServer.cpp b/Engine/Source/Editor/UnrealEd/Private/EditorServer.cpp
--- a/Engine/Source/Editor/UnrealEd/Private/EditorServer.cpp
+++ b/Engine/Source/Editor/UnrealEd/Private/EditorServer.cpp
@@ -2533,6 +2533,7 @@ bool UEditorEngine::Map_Load(const TCHAR* Str, FOutputDevice& Ar)
 					const FName WorldPackageFName = WorldPackage->GetFName();
 					UWorld::WorldTypePreLoadMap.FindOrAdd(WorldPackageFName) = EWorldType::Editor;
 					WorldPackage = LoadPackage( WorldPackage, *LongTempFname, LoadFlags );
+					WorldPackage->SetPackageFlags(PKG_NewlyCreated);
 					UWorld::WorldTypePreLoadMap.Remove(WorldPackageFName);
 				}
 				else
