Add support for preventing AutoSave from an active EdMode. Use to prevent Modeling Tools from being auto-aborted by AutoSave.
- Add FEditorModeTools::CanAutoSave(), checks if any active FEdMode/UEdMode wants to prevent AutoSave from occurring.
- Call in FPackageAutoSaver::CanAutoSave().
- Add FEdMode/UEdMode::CanAutoSave(), default to true
- ModelingToolsEditorMode override returns false if an InteractiveTool is active
#rnx

diff --git a/Engine/Plugins/Experimental/ModelingToolsEditorMode/Source/ModelingToolsEditorMode/Private/ModelingToolsEditorMode.cpp b/Engine/Plugins/Experimental/ModelingToolsEditorMode/Source/ModelingToolsEditorMode/Private/ModelingToolsEditorMode.cpp
--- a/Engine/Plugins/Experimental/ModelingToolsEditorMode/Source/ModelingToolsEditorMode/Private/ModelingToolsEditorMode.cpp
+++ b/Engine/Plugins/Experimental/ModelingToolsEditorMode/Source/ModelingToolsEditorMode/Private/ModelingToolsEditorMode.cpp
@@ -86,6 +86,12 @@ bool FModelingToolsEditorMode::ProcessEditDelete()
 }
 
 
+bool FModelingToolsEditorMode::CanAutoSave() const
+{
+	// prevent autosave if any tool is active
+	return ToolsContext->ToolManager->HasAnyActiveTool() == false;
+}
+
 bool FModelingToolsEditorMode::AllowWidgetMove()
 { 
 	return false; 
diff --git a/Engine/Plugins/Experimental/ModelingToolsEditorMode/Source/ModelingToolsEditorMode/Public/ModelingToolsEditorMode.h b/Engine/Plugins/Experimental/ModelingToolsEditorMode/Source/ModelingToolsEditorMode/Public/ModelingToolsEditorMode.h
--- a/Engine/Plugins/Experimental/ModelingToolsEditorMode/Source/ModelingToolsEditorMode/Public/ModelingToolsEditorMode.h
+++ b/Engine/Plugins/Experimental/ModelingToolsEditorMode/Source/ModelingToolsEditorMode/Public/ModelingToolsEditorMode.h
@@ -42,6 +42,7 @@ public:
 
 	virtual bool ProcessEditDelete();
 
+	virtual bool CanAutoSave() const override;
 
 	/*
 	 * focus events
diff --git a/Engine/Source/Editor/UnrealEd/Private/EditorModeManager.cpp b/Engine/Source/Editor/UnrealEd/Private/EditorModeManager.cpp
--- a/Engine/Source/Editor/UnrealEd/Private/EditorModeManager.cpp
+++ b/Engine/Source/Editor/UnrealEd/Private/EditorModeManager.cpp
@@ -2048,3 +2048,25 @@ bool FEditorModeTools::CanCycleWidgetMode() const
 	}
 	return false;
 }
+
+
+bool FEditorModeTools::CanAutoSave() const
+{
+	for (auto& Mode : ActiveModes)
+	{
+		if (Mode->CanAutoSave() == false)
+		{
+			return false;
+		}
+	}
+
+	for (UEdMode* Mode : ActiveScriptableModes)
+	{
+		if (Mode->CanAutoSave() == false)
+		{
+			return false;
+		}
+	}
+
+	return true;
+}
diff --git a/Engine/Source/Editor/UnrealEd/Private/PackageAutoSaver.cpp b/Engine/Source/Editor/UnrealEd/Private/PackageAutoSaver.cpp
--- a/Engine/Source/Editor/UnrealEd/Private/PackageAutoSaver.cpp
+++ b/Engine/Source/Editor/UnrealEd/Private/PackageAutoSaver.cpp
@@ -450,7 +450,10 @@ bool FPackageAutoSaver::CanAutoSave() const
 		}
 	}
 
-	return (bAutosaveEnabled && !bSlowTask && !bInterpEditMode && !bPlayWorldValid && !bAnyMenusVisible && !bAutomationTesting && !bIsInteracting && !GIsDemoMode && bHasGameOrProjectLoaded && !bAreShadersCompiling && !bIsVREditorActive && !bIsSequencerPlaying);
+	// query any active editor modes and allow them to prevent autosave
+	const bool bActiveModesAllowAutoSave = GLevelEditorModeTools().CanAutoSave();
+
+	return (bAutosaveEnabled && !bSlowTask && !bInterpEditMode && !bPlayWorldValid && !bAnyMenusVisible && !bAutomationTesting && !bIsInteracting && !GIsDemoMode && bHasGameOrProjectLoaded && !bAreShadersCompiling && !bIsVREditorActive && !bIsSequencerPlaying && bActiveModesAllowAutoSave);
 }
 
 bool FPackageAutoSaver::DoPackagesNeedAutoSave() const
diff --git a/Engine/Source/Editor/UnrealEd/Public/EdMode.h b/Engine/Source/Editor/UnrealEd/Public/EdMode.h
--- a/Engine/Source/Editor/UnrealEd/Public/EdMode.h
+++ b/Engine/Source/Editor/UnrealEd/Public/EdMode.h
@@ -183,6 +183,12 @@ public:
 
 	virtual void PostUndo() {}
 
+	/** 
+	 * Check to see if this EdMode wants to disallow AutoSave
+	 * @return true if AutoSave can be applied right now
+	 */
+	virtual bool CanAutoSave() const { return true; }
+
 	/**
 	 * Lets each mode/tool handle box selection in its own way.
 	 *
diff --git a/Engine/Source/Editor/UnrealEd/Public/EditorModeManager.h b/Engine/Source/Editor/UnrealEd/Public/EditorModeManager.h
--- a/Engine/Source/Editor/UnrealEd/Public/EditorModeManager.h
+++ b/Engine/Source/Editor/UnrealEd/Public/EditorModeManager.h
@@ -542,6 +542,9 @@ PRAGMA_ENABLE_DEPRECATION_WARNINGS
 	void UpdateInternalData();
 	bool IsOnlyVisibleActiveMode(FEditorModeID InMode) const;
 
+	/** returns true if all active EdModes are OK with an AutoSave happening now  */
+	bool CanAutoSave() const;
+
 protected:
 	/** 
 	 * Delegate handlers
diff --git a/Engine/Source/Editor/UnrealEd/Public/Tools/UEdMode.h b/Engine/Source/Editor/UnrealEd/Public/Tools/UEdMode.h
--- a/Engine/Source/Editor/UnrealEd/Public/Tools/UEdMode.h
+++ b/Engine/Source/Editor/UnrealEd/Public/Tools/UEdMode.h
@@ -169,6 +169,12 @@ public:
 
 	virtual void PostUndo() {}
 
+	/**
+	 * Check to see if this UEdMode wants to disallow AutoSave
+	 * @return true if AutoSave can be applied right now
+	 */
+	virtual bool CanAutoSave() const { return true; }
+
 	virtual void SelectNone();
 	virtual void SelectionChanged() {}
 
