Added null check to helper method of ObjectTools::ForceDeleteObjects to deal with fact that UObject::RetrieveReferencer may return a referencer without a valid UObject referencer.

diff --git a/Engine/Source/Editor/UnrealEd/Private/ObjectTools.cpp b/Engine/Source/Editor/UnrealEd/Private/ObjectTools.cpp
--- a/Engine/Source/Editor/UnrealEd/Private/ObjectTools.cpp
+++ b/Engine/Source/Editor/UnrealEd/Private/ObjectTools.cpp
@@ -2401,13 +2401,16 @@ namespace ObjectTools
 		TArray<FReferencerInformation> ExternalReferencers;
 		Object->RetrieveReferencers(nullptr /* internal refs */, &ExternalReferencers);
 
-		for (const FReferencerInformation& Referencer : ExternalReferencers)
+		for (const FReferencerInformation& RefInfo : ExternalReferencers)
 		{
-			bool bAlreadyIn = false;
-			ReferencingObjects.Add(Referencer.Referencer, &bAlreadyIn);
-			if (!bAlreadyIn)
+			if (RefInfo.Referencer != nullptr)
 			{
-				RecursiveRetrieveReferencers(Referencer.Referencer, ReferencingObjects);
+				bool bAlreadyIn = false;
+				ReferencingObjects.Add(RefInfo.Referencer, &bAlreadyIn);
+				if (!bAlreadyIn)
+				{
+					RecursiveRetrieveReferencers(RefInfo.Referencer, ReferencingObjects);
+				}
 			}
 		}
 	}
