Added additional null check to ObjectTools::ForceDeleteObjects.

diff --git a/Engine/Source/Editor/UnrealEd/Private/ObjectTools.cpp b/Engine/Source/Editor/UnrealEd/Private/ObjectTools.cpp
--- a/Engine/Source/Editor/UnrealEd/Private/ObjectTools.cpp
+++ b/Engine/Source/Editor/UnrealEd/Private/ObjectTools.cpp
@@ -2450,7 +2450,7 @@ namespace ObjectTools
 		bool bClosedAllEditors = true;
 		for (UObject* Object : ReferencingObjects)
 		{
-			if (Object->IsAsset())
+			if (Object != nullptr && Object->IsAsset())
 			{
 				TArray<IAssetEditorInstance*> ObjectEditors = GEditor->GetEditorSubsystem<UAssetEditorSubsystem>()->FindEditorsForAssetAndSubObjects(Object);
 				for (IAssetEditorInstance* ObjectEditorInstance : ObjectEditors)
