Moved BSP and all the related editor modes to a plugin so it can be disabled in some situations

diff --git a/Engine/Plugins/Editor/GeometryMode/Source/BspMode/Private/BspMode.cpp b/Engine/Plugins/Editor/GeometryMode/Source/BspMode/Private/BspMode.cpp
new file mode 100644
--- /dev/null
+++ b/Engine/Plugins/Editor/GeometryMode/Source/BspMode/Private/BspMode.cpp
@@ -0,0 +1,4 @@
+// Copyright 1998-2019 Epic Games, Inc. All Rights Reserved.
+
+#include "BspMode.h"
+
diff --git a/Engine/Plugins/Editor/GeometryMode/Source/BspMode/Private/SBspPalette.cpp b/Engine/Plugins/Editor/GeometryMode/Source/BspMode/Private/SBspPalette.cpp
new file mode 100644
--- /dev/null
+++ b/Engine/Plugins/Editor/GeometryMode/Source/BspMode/Private/SBspPalette.cpp
@@ -0,0 +1,264 @@
+// Copyright 1998-2019 Epic Games, Inc. All Rights Reserved.
+
+#include "SBspPalette.h"
+#include "Modules/ModuleManager.h"
+#include "Widgets/SBoxPanel.h"
+#include "SlateOptMacros.h"
+#include "Styling/CoreStyle.h"
+#include "Widgets/Images/SImage.h"
+#include "Styling/SlateTypes.h"
+#include "Widgets/Text/STextBlock.h"
+#include "Widgets/Layout/SBox.h"
+#include "Widgets/Layout/SScrollBorder.h"
+#include "Widgets/Layout/SBorder.h"
+#include "Widgets/Views/STableViewBase.h"
+#include "Widgets/Views/STableRow.h"
+#include "Widgets/Views/SListView.h"
+#include "Widgets/Input/SCheckBox.h"
+#include "EditorStyleSet.h"
+#include "Engine/Brush.h"
+#include "BspModeStyle.h"
+#include "DragAndDrop/BrushBuilderDragDropOp.h"
+#include "EditorClassUtils.h"
+
+#define LOCTEXT_NAMESPACE "BspPalette"
+
+/** The list view mode of the asset view */
+class SBspBuilderListView : public SListView<TSharedPtr<FBspBuilderType>>
+{
+public:
+	virtual bool SupportsKeyboardFocus() const override { return true; }
+	virtual FReply OnKeyDown( const FGeometry& InGeometry, const FKeyEvent& InKeyEvent ) override
+	{
+		return FReply::Unhandled();
+	}
+};
+
+BEGIN_SLATE_FUNCTION_BUILD_OPTIMIZATION
+void SBspPalette::Construct( const FArguments& InArgs )
+{
+	bIsAdditive = true;
+
+	FBspModeModule& BspModeModule = FModuleManager::GetModuleChecked<FBspModeModule>("BspMode");
+
+	TSharedRef<SBspBuilderListView> ListViewWidget = 
+		SNew(SBspBuilderListView)
+		.SelectionMode(ESelectionMode::Single)
+		.ListItemsSource(&BspModeModule.GetBspBuilderTypes())
+		.OnGenerateRow(this, &SBspPalette::MakeListViewWidget)
+		.OnSelectionChanged(this, &SBspPalette::OnSelectionChanged)
+		.ItemHeight(35);
+
+	ChildSlot
+	[
+		SNew(SVerticalBox)
+
+		+ SVerticalBox::Slot()
+		.FillHeight(1.0f)
+		[
+			SNew(SScrollBorder, ListViewWidget)
+			[
+				ListViewWidget
+			]
+		]
+
+		+ SVerticalBox::Slot()
+		.AutoHeight()
+		[
+			SNew( SHorizontalBox )
+
+			+ SHorizontalBox::Slot()
+			.AutoWidth()
+			.Padding( 3.0f )
+			.HAlign(HAlign_Right)
+			.VAlign(VAlign_Center)
+			[
+				SNew(SCheckBox)
+				.Style(FEditorStyle::Get(), "Toolbar.RadioButton")
+				.ToolTipText(LOCTEXT("BspModeAdditiveTooltip", "Place brushes in additive mode."))
+				.OnCheckStateChanged(this, &SBspPalette::OnAdditiveModeButtonClicked)
+				.IsChecked(this, &SBspPalette::IsAdditiveModeChecked)
+				.ForegroundColor(FCoreStyle::Get().GetSlateColor("Foreground"))
+				.Content()
+				[
+					SNew(SHorizontalBox)
+					+SHorizontalBox::Slot()
+					.AutoWidth()
+					.HAlign(HAlign_Left)
+					.VAlign(VAlign_Center)
+					[
+						SNew(SImage)
+						.Image(this, &SBspPalette::GetAdditiveModeImage)
+					]
+					+SHorizontalBox::Slot()
+					.AutoWidth()
+					.HAlign(HAlign_Center)
+					.VAlign(VAlign_Center)
+					.Padding(2.0f)
+					[
+						SNew(STextBlock)
+						.Text(LOCTEXT("BspModeAdd", "Add"))
+						.Font(FCoreStyle::Get().GetFontStyle("Toolbar.Label.Font"))
+						.ShadowOffset(FVector2D::UnitVector)
+					]
+				]
+			]
+
+			+ SHorizontalBox::Slot()
+			.Padding( 3.0f )
+			.AutoWidth()
+			.HAlign(HAlign_Right)
+			.VAlign(VAlign_Center)
+			[
+				SNew(SCheckBox)
+				.Style(FEditorStyle::Get(), "Toolbar.RadioButton")
+				.ToolTipText(LOCTEXT("BspModeSubtractiveTooltip", "Place brushes in subtractive mode."))
+				.OnCheckStateChanged(this, &SBspPalette::OnSubtractiveModeButtonClicked)
+				.IsChecked(this, &SBspPalette::IsSubtractiveModeChecked)
+				.ForegroundColor(FCoreStyle::Get().GetSlateColor("Foreground"))
+				.Content()
+				[
+					SNew(SHorizontalBox)
+					+SHorizontalBox::Slot()
+					.AutoWidth()
+					.HAlign(HAlign_Left)
+					.VAlign(VAlign_Center)
+					[
+						SNew(SImage)
+						.Image(this, &SBspPalette::GetSubtractiveModeImage)
+					]
+					+SHorizontalBox::Slot()
+					.AutoWidth()
+					.HAlign(HAlign_Left)
+					.VAlign(VAlign_Center)
+					.Padding(2.0f)
+					[
+						SNew(STextBlock)
+						.Text(LOCTEXT("BspModeSubtract", "Subtract"))
+						.Font(FCoreStyle::Get().GetFontStyle("Toolbar.Label.Font"))
+						.ShadowOffset(FVector2D::UnitVector)
+					]
+				]
+			]
+		]
+	];
+}
+END_SLATE_FUNCTION_BUILD_OPTIMIZATION
+
+TSharedRef<ITableRow> SBspPalette::MakeListViewWidget(TSharedPtr<FBspBuilderType> BspBuilder, const TSharedRef<STableViewBase>& OwnerTable)
+{
+	check(BspBuilder.IsValid());
+	check(BspBuilder->BuilderClass.IsValid());
+
+	TSharedRef< STableRow<TSharedPtr<FBspBuilderType>> > TableRowWidget = 
+		SNew( STableRow<TSharedPtr<FBspBuilderType>>, OwnerTable )
+		.Style( FBspModeStyle::Get(), "BspMode.TableRow" )
+		.OnDragDetected( this, &SBspPalette::OnDraggingListViewWidget );
+
+	TSharedRef<SWidget> Content = 
+		SNew(SBorder)
+		.BorderImage(FCoreStyle::Get().GetBrush("NoBorder"))
+		.Padding(0)
+		.ToolTip(FEditorClassUtils::GetTooltip(ABrush::StaticClass(), BspBuilder->ToolTipText))
+		.Cursor( EMouseCursor::GrabHand )
+		[
+			SNew(SHorizontalBox)
+
+			// Icon
+			+SHorizontalBox::Slot()
+			.AutoWidth()
+			[
+				SNew(SBorder)
+				.Padding(4.0f)
+				.BorderImage(FBspModeStyle::Get().GetBrush("BspMode.ThumbnailShadow"))
+				[
+					SNew(SBox)
+					.WidthOverride(35.0f)
+					.HeightOverride(35.0f)
+					[
+						SNew(SBorder)
+						.BorderImage(FBspModeStyle::Get().GetBrush("BspMode.ThumbnailBackground"))
+						.HAlign(HAlign_Center)
+						.VAlign(VAlign_Center)
+						[
+							SNew(SImage)
+							.Image(BspBuilder->Icon)
+						]
+					]
+				]
+			]
+
+			+SHorizontalBox::Slot()
+			.HAlign(HAlign_Left)
+			.VAlign(VAlign_Center)
+			.FillWidth(1.f)
+			[
+				SNew(STextBlock)
+				.TextStyle(FBspModeStyle::Get(), "BspMode.ThumbnailText")
+				.Text(BspBuilder->Text)
+			]
+		];
+
+	TableRowWidget->SetContent(Content);
+
+	return TableRowWidget;
+}
+
+void SBspPalette::OnSelectionChanged(TSharedPtr<FBspBuilderType> BspBuilder, ESelectInfo::Type SelectionType)
+{
+	if(BspBuilder.IsValid())
+	{
+		ActiveBrushBuilder = GEditor->FindBrushBuilder(BspBuilder->BuilderClass.Get());
+	}
+}
+
+FReply SBspPalette::OnDraggingListViewWidget(const FGeometry& MyGeometry, const FPointerEvent& MouseEvent)
+{
+	if ( MouseEvent.IsMouseButtonDown( EKeys::LeftMouseButton ) )
+	{
+		if (ActiveBrushBuilder.IsValid())
+		{
+			FBspModeModule& BspModeModule = FModuleManager::GetModuleChecked<FBspModeModule>("BspMode");
+			TSharedPtr<FBspBuilderType> BspBuilder = BspModeModule.FindBspBuilderType(ActiveBrushBuilder.Get()->GetClass());
+			if(BspBuilder.IsValid())
+			{
+				// We have an active brush builder, start a drag-drop
+				return FReply::Handled().BeginDragDrop(FBrushBuilderDragDropOp::New(ActiveBrushBuilder, BspBuilder->Icon, bIsAdditive));
+			}
+		}
+	}
+
+	return FReply::Unhandled();
+}
+
+void SBspPalette::OnAdditiveModeButtonClicked(ECheckBoxState CheckType)
+{
+	bIsAdditive = CheckType == ECheckBoxState::Checked;
+}
+
+void SBspPalette::OnSubtractiveModeButtonClicked(ECheckBoxState CheckType)
+{
+	bIsAdditive = CheckType != ECheckBoxState::Checked;
+}
+
+ECheckBoxState SBspPalette::IsAdditiveModeChecked() const
+{
+	return bIsAdditive ? ECheckBoxState::Checked : ECheckBoxState::Unchecked;
+}
+
+ECheckBoxState SBspPalette::IsSubtractiveModeChecked() const
+{
+	return !bIsAdditive ? ECheckBoxState::Checked : ECheckBoxState::Unchecked;
+}
+
+const FSlateBrush* SBspPalette::GetAdditiveModeImage() const
+{
+	return FBspModeStyle::Get().GetBrush("BspMode.CSGAdd.Small");
+}
+
+const FSlateBrush* SBspPalette::GetSubtractiveModeImage() const
+{
+	return FBspModeStyle::Get().GetBrush("BspMode.CSGSubtract.Small");
+}
+
+#undef LOCTEXT_NAMESPACE
diff --git a/Engine/Plugins/Editor/GeometryMode/Source/BspMode/Private/SBspPalette.h b/Engine/Plugins/Editor/GeometryMode/Source/BspMode/Private/SBspPalette.h
new file mode 100644
--- /dev/null
+++ b/Engine/Plugins/Editor/GeometryMode/Source/BspMode/Private/SBspPalette.h
@@ -0,0 +1,65 @@
+// Copyright 1998-2019 Epic Games, Inc. All Rights Reserved.
+
+#pragma once
+
+#include "CoreMinimal.h"
+#include "BspModeModule.h"
+#include "Widgets/DeclarativeSyntaxSupport.h"
+#include "Input/Reply.h"
+#include "Widgets/SCompoundWidget.h"
+#include "Misc/NotifyHook.h"
+
+class IDetailsView;
+class ITableRow;
+class STableViewBase;
+class UBrushBuilder;
+struct FSlateBrush;
+enum class ECheckBoxState : uint8;
+
+class SBspPalette : public SCompoundWidget, public FNotifyHook
+{
+public:
+	SLATE_BEGIN_ARGS( SBspPalette ){}
+	SLATE_END_ARGS();
+
+	void Construct( const FArguments& InArgs );
+
+private:
+
+	/** Make a widget for the list view display */
+	TSharedRef<ITableRow> MakeListViewWidget(TSharedPtr<struct FBspBuilderType> BspBuilder, const TSharedRef<STableViewBase>& OwnerTable);
+
+	/** Delegate for when the list view selection changes */
+	void OnSelectionChanged(TSharedPtr<FBspBuilderType> BspBuilder, ESelectInfo::Type SelectionType);
+
+	/** Begin dragging a list widget */
+	FReply OnDraggingListViewWidget(const FGeometry& MyGeometry, const FPointerEvent& MouseEvent);
+
+	/** Select additive mode */
+	void OnAdditiveModeButtonClicked(ECheckBoxState CheckType);
+
+	/** Select subtractive mode */
+	void OnSubtractiveModeButtonClicked(ECheckBoxState CheckType);
+
+	/** @return the check state of the additive mode checkbox */
+	ECheckBoxState IsAdditiveModeChecked() const;
+
+	/** @return the check state of the subtractive mode checkbox */
+	ECheckBoxState IsSubtractiveModeChecked() const;
+
+	/** Get the image for additive mode */
+	const FSlateBrush* GetAdditiveModeImage() const;
+
+	/** Get the image for subtractive mode */
+	const FSlateBrush* GetSubtractiveModeImage() const;
+
+private:
+	/** Property view for brush options */
+	TSharedPtr<class IDetailsView> BrushOptionView;
+
+	/** Brush builder currently active */
+	TWeakObjectPtr<UBrushBuilder> ActiveBrushBuilder;
+
+	/** Additive or subtractive mode */
+	bool bIsAdditive;
+};
diff --git a/Engine/Plugins/Editor/GeometryMode/Source/BspMode/Public/BspModeStyle.h b/Engine/Plugins/Editor/GeometryMode/Source/BspMode/Public/BspModeStyle.h
new file mode 100644
--- /dev/null
+++ b/Engine/Plugins/Editor/GeometryMode/Source/BspMode/Public/BspModeStyle.h
@@ -0,0 +1,25 @@
+// Copyright 1998-2019 Epic Games, Inc. All Rights Reserved.
+#pragma once
+
+#include "CoreMinimal.h"
+#include "Styling/SlateStyle.h"
+
+/**
+ * BSP mode slate style
+ */
+class BSPMODE_API FBspModeStyle
+{
+public:
+
+	static void Initialize();
+	static void Shutdown();
+
+	static const ISlateStyle& Get();
+
+	static const FName& GetStyleSetName();
+
+private:
+
+	/** Singleton instances of this style. */
+	static TSharedPtr< class FSlateStyleSet > StyleSet;
+};
diff --git a/Engine/Plugins/Editor/GeometryMode/Source/BspMode/Public/IBspModeModule.h b/Engine/Plugins/Editor/GeometryMode/Source/BspMode/Public/IBspModeModule.h
new file mode 100644
--- /dev/null
+++ b/Engine/Plugins/Editor/GeometryMode/Source/BspMode/Public/IBspModeModule.h
@@ -0,0 +1,15 @@
+// Copyright 1998-2019 Epic Games, Inc. All Rights Reserved.
+#pragma once
+
+#include "CoreMinimal.h"
+#include "Modules/ModuleInterface.h"
+
+struct FSlateBrush;
+
+class IBspModeModule : public IModuleInterface
+{
+public:
+	virtual void RegisterBspBuilderType( class UClass* InBuilderClass, const FText& InBuilderName, const FText& InBuilderTooltip, const FSlateBrush* InBuilderIcon ) = 0;
+	virtual void UnregisterBspBuilderType( class UClass* InBuilderClass ) = 0;
+};
+
diff --git a/Engine/Plugins/Editor/GeometryMode/Source/GeometryMode/Classes/GeomModifier_Create.h b/Engine/Plugins/Editor/GeometryMode/Source/GeometryMode/Classes/GeomModifier_Create.h
new file mode 100644
--- /dev/null
+++ b/Engine/Plugins/Editor/GeometryMode/Source/GeometryMode/Classes/GeomModifier_Create.h
@@ -0,0 +1,29 @@
+// Copyright 1998-2019 Epic Games, Inc. All Rights Reserved.
+
+
+/**
+ * Creates selected objects.
+ */
+
+#pragma once
+
+#include "CoreMinimal.h"
+#include "UObject/ObjectMacros.h"
+#include "GeomModifier_Edit.h"
+#include "GeomModifier_Create.generated.h"
+
+UCLASS()
+class UGeomModifier_Create : public UGeomModifier_Edit
+{
+	GENERATED_UCLASS_BODY()
+
+
+	//~ Begin UGeomModifier Interface
+	virtual bool Supports() override;
+protected:
+	virtual bool OnApply() override;
+	//~ End UGeomModifier Interface
+};
+
+
+
diff --git a/Engine/Plugins/Editor/GeometryMode/Source/GeometryMode/Classes/GeomModifier_Delete.h b/Engine/Plugins/Editor/GeometryMode/Source/GeometryMode/Classes/GeomModifier_Delete.h
new file mode 100644
--- /dev/null
+++ b/Engine/Plugins/Editor/GeometryMode/Source/GeometryMode/Classes/GeomModifier_Delete.h
@@ -0,0 +1,29 @@
+// Copyright 1998-2019 Epic Games, Inc. All Rights Reserved.
+
+
+/**
+ * Deletes selected objects.
+ */
+
+#pragma once
+
+#include "CoreMinimal.h"
+#include "UObject/ObjectMacros.h"
+#include "GeomModifier_Edit.h"
+#include "GeomModifier_Delete.generated.h"
+
+UCLASS()
+class UGeomModifier_Delete : public UGeomModifier_Edit
+{
+	GENERATED_UCLASS_BODY()
+
+
+	//~ Begin UGeomModifier Interface
+	virtual bool Supports() override;
+protected:
+	virtual bool OnApply() override;
+	//~ End UGeomModifier Interface
+};
+
+
+
diff --git a/Engine/Plugins/Editor/GeometryMode/Source/GeometryMode/Classes/GeomModifier_Flip.h b/Engine/Plugins/Editor/GeometryMode/Source/GeometryMode/Classes/GeomModifier_Flip.h
new file mode 100644
--- /dev/null
+++ b/Engine/Plugins/Editor/GeometryMode/Source/GeometryMode/Classes/GeomModifier_Flip.h
@@ -0,0 +1,29 @@
+// Copyright 1998-2019 Epic Games, Inc. All Rights Reserved.
+
+
+/**
+ * Flips selected objects.
+ */
+
+#pragma once
+
+#include "CoreMinimal.h"
+#include "UObject/ObjectMacros.h"
+#include "GeomModifier_Edit.h"
+#include "GeomModifier_Flip.generated.h"
+
+UCLASS()
+class UGeomModifier_Flip : public UGeomModifier_Edit
+{
+	GENERATED_UCLASS_BODY()
+
+
+	//~ Begin UGeomModifier Interface
+	virtual bool Supports() override;
+protected:
+	virtual bool OnApply() override;
+	//~ End UGeomModifier Interface
+};
+
+
+
diff --git a/Engine/Plugins/Editor/GeometryMode/Source/GeometryMode/Classes/GeomModifier_Split.h b/Engine/Plugins/Editor/GeometryMode/Source/GeometryMode/Classes/GeomModifier_Split.h
new file mode 100644
--- /dev/null
+++ b/Engine/Plugins/Editor/GeometryMode/Source/GeometryMode/Classes/GeomModifier_Split.h
@@ -0,0 +1,29 @@
+// Copyright 1998-2019 Epic Games, Inc. All Rights Reserved.
+
+
+/**
+ * Splits selected objects.
+ */
+
+#pragma once
+
+#include "CoreMinimal.h"
+#include "UObject/ObjectMacros.h"
+#include "GeomModifier_Edit.h"
+#include "GeomModifier_Split.generated.h"
+
+UCLASS()
+class UGeomModifier_Split : public UGeomModifier_Edit
+{
+	GENERATED_UCLASS_BODY()
+
+
+	//~ Begin UGeomModifier Interface
+	virtual bool Supports() override;
+protected:
+	virtual bool OnApply() override;
+	//~ End UGeomModifier Interface
+};
+
+
+
diff --git a/Engine/Plugins/Editor/GeometryMode/Source/GeometryMode/Classes/GeomModifier_Triangulate.h b/Engine/Plugins/Editor/GeometryMode/Source/GeometryMode/Classes/GeomModifier_Triangulate.h
new file mode 100644
--- /dev/null
+++ b/Engine/Plugins/Editor/GeometryMode/Source/GeometryMode/Classes/GeomModifier_Triangulate.h
@@ -0,0 +1,29 @@
+// Copyright 1998-2019 Epic Games, Inc. All Rights Reserved.
+
+
+/**
+ * Triangulates selected objects.
+ */
+
+#pragma once
+
+#include "CoreMinimal.h"
+#include "UObject/ObjectMacros.h"
+#include "GeomModifier_Edit.h"
+#include "GeomModifier_Triangulate.generated.h"
+
+UCLASS()
+class UGeomModifier_Triangulate : public UGeomModifier_Edit
+{
+	GENERATED_UCLASS_BODY()
+
+
+	//~ Begin UGeomModifier Interface
+	virtual bool Supports() override;
+protected:
+	virtual bool OnApply() override;
+	//~ End UGeomModifier Interface
+};
+
+
+
diff --git a/Engine/Plugins/Editor/GeometryMode/Source/GeometryMode/Classes/GeomModifier_Turn.h b/Engine/Plugins/Editor/GeometryMode/Source/GeometryMode/Classes/GeomModifier_Turn.h
new file mode 100644
--- /dev/null
+++ b/Engine/Plugins/Editor/GeometryMode/Source/GeometryMode/Classes/GeomModifier_Turn.h
@@ -0,0 +1,29 @@
+// Copyright 1998-2019 Epic Games, Inc. All Rights Reserved.
+
+
+/**
+ * Turns selected objects.
+ */
+
+#pragma once
+
+#include "CoreMinimal.h"
+#include "UObject/ObjectMacros.h"
+#include "GeomModifier_Edit.h"
+#include "GeomModifier_Turn.generated.h"
+
+UCLASS()
+class UGeomModifier_Turn : public UGeomModifier_Edit
+{
+	GENERATED_UCLASS_BODY()
+
+
+	//~ Begin UGeomModifier Interface
+	virtual bool Supports() override;
+protected:
+	virtual bool OnApply() override;
+	//~ End UGeomModifier Interface
+};
+
+
+
diff --git a/Engine/Plugins/Editor/GeometryMode/Source/GeometryMode/Classes/GeomModifier_Weld.h b/Engine/Plugins/Editor/GeometryMode/Source/GeometryMode/Classes/GeomModifier_Weld.h
new file mode 100644
--- /dev/null
+++ b/Engine/Plugins/Editor/GeometryMode/Source/GeometryMode/Classes/GeomModifier_Weld.h
@@ -0,0 +1,28 @@
+// Copyright 1998-2019 Epic Games, Inc. All Rights Reserved.
+
+/**
+ * Welds selected objects.
+ */
+
+#pragma once
+
+#include "CoreMinimal.h"
+#include "UObject/ObjectMacros.h"
+#include "GeomModifier_Edit.h"
+#include "GeomModifier_Weld.generated.h"
+
+UCLASS()
+class UGeomModifier_Weld : public UGeomModifier_Edit
+{
+	GENERATED_UCLASS_BODY()
+
+
+	//~ Begin UGeomModifier Interface
+	virtual bool Supports() override;
+protected:
+	virtual bool OnApply() override;
+	//~ End UGeomModifier Interface
+};
+
+
+
diff --git a/Engine/Plugins/Editor/GeometryMode/Source/GeometryMode/Private/BrushEditingSubsystemImpl.h b/Engine/Plugins/Editor/GeometryMode/Source/GeometryMode/Private/BrushEditingSubsystemImpl.h
new file mode 100644
--- /dev/null
+++ b/Engine/Plugins/Editor/GeometryMode/Source/GeometryMode/Private/BrushEditingSubsystemImpl.h
@@ -0,0 +1,41 @@
+// Copyright 1998-2019 Epic Games, Inc. All Rights Reserved.
+
+#pragma once
+
+#include "Subsystems/BrushEditingSubsystem.h"
+
+#include "BrushEditingSubsystemImpl.generated.h"
+
+DECLARE_LOG_CATEGORY_EXTERN(LogBrushEditing, Log, All);
+
+struct HGeomPolyProxy;
+struct HGeomEdgeProxy;
+struct HGeomVertexProxy;
+
+UCLASS()
+class UBrushEditingSubsystemImpl : public UBrushEditingSubsystem
+{
+	GENERATED_BODY()
+
+public:
+	UBrushEditingSubsystemImpl();
+
+	virtual void Initialize(FSubsystemCollectionBase& Collection) override;
+
+	virtual bool ProcessClickOnBrushGeometry(FLevelEditorViewportClient* ViewportClient, HHitProxy* InHitProxy, const FViewportClick& Click) override;
+
+	virtual void UpdateGeometryFromSelectedBrushes() override;
+	virtual void UpdateGeometryFromBrush(ABrush* Brush) override;
+
+	virtual bool IsGeometryEditorModeActive() const override;
+
+	virtual void DeselectAllEditingGeometry() override;
+
+	virtual bool HandleActorDelete() override;
+private:
+
+	bool ProcessClickOnGeomPoly(FLevelEditorViewportClient* ViewportClient, HGeomPolyProxy* GeomHitProxy, const FViewportClick& Click);
+	bool ProcessClickOnGeomEdge(FLevelEditorViewportClient* ViewportClient, HGeomEdgeProxy* GeomHitProxy, const FViewportClick& Click);
+	bool ProcessClickOnGeomVertex(FLevelEditorViewportClient* ViewportClient, HGeomVertexProxy* GeomHitProxy, const FViewportClick& Click);
+
+};
diff --git a/Engine/Plugins/Editor/GeometryMode/Source/GeometryMode/Private/GeometryMode.cpp b/Engine/Plugins/Editor/GeometryMode/Source/GeometryMode/Private/GeometryMode.cpp
new file mode 100644
--- /dev/null
+++ b/Engine/Plugins/Editor/GeometryMode/Source/GeometryMode/Private/GeometryMode.cpp
@@ -0,0 +1,331 @@
+// Copyright 1998-2019 Epic Games, Inc. All Rights Reserved.
+
+#include "GeometryMode.h"
+#include "Widgets/Text/STextBlock.h"
+#include "Modules/ModuleManager.h"
+#include "Widgets/Layout/SBorder.h"
+#include "Widgets/Layout/SSeparator.h"
+#include "Widgets/Layout/SGridPanel.h"
+#include "Widgets/Layout/SUniformGridPanel.h"
+#include "Widgets/Input/SButton.h"
+#include "Widgets/SToolTip.h"
+#include "Widgets/Layout/SScrollBox.h"
+#include "Widgets/Input/SCheckBox.h"
+#include "EditorStyleSet.h"
+#include "EditorModeManager.h"
+#include "EditorModes.h"
+#include "GeometryEdMode.h"
+#include "PropertyEditorModule.h"
+#include "IDetailsView.h"
+#include "GeomModifier.h"
+
+#define LOCTEXT_NAMESPACE "GeometryMode"
+
+void SGeometryModeControls::SelectionChanged()
+{
+	// If the currently selected modifier is being disabled, change the selection to Edit
+	for (int32 Idx = 0; Idx < ModifierControls.Num(); ++Idx)
+	{
+		if (ModifierControls[Idx]->IsChecked() && !GetGeometryModeTool()->GetModifier(Idx)->Supports())
+		{
+			if (GetGeometryModeTool()->GetNumModifiers() > 0)
+			{
+				GetGeometryModeTool()->SetCurrentModifier(GetGeometryModeTool()->GetModifier(0));
+			}
+		}
+	}
+}
+
+void SGeometryModeControls::Construct(const FArguments& InArgs)
+{
+	FModeTool_GeometryModify* GeometryModeTool = GetGeometryModeTool();
+	
+	if (GetGeometryModeTool()->GetNumModifiers() > 0)
+	{
+		GetGeometryModeTool()->SetCurrentModifier(GetGeometryModeTool()->GetModifier(0));
+	}
+
+	CreateLayout();
+}
+
+void SGeometryModeControls::OnModifierStateChanged(ECheckBoxState NewCheckedState, UGeomModifier* Modifier)
+{
+	if (NewCheckedState == ECheckBoxState::Checked)
+	{
+		GetGeometryModeTool()->SetCurrentModifier(Modifier);
+
+		TArray<UObject*> PropertyObjects;
+		PropertyObjects.Add(GetGeometryModeTool()->GetCurrentModifier());
+		PropertiesControl->SetObjects(PropertyObjects);
+	}
+}
+
+ECheckBoxState SGeometryModeControls::IsModifierChecked(UGeomModifier* Modifier) const
+{
+	return (GetGeometryModeTool()->GetCurrentModifier() == Modifier)
+		? ECheckBoxState::Checked
+		: ECheckBoxState::Unchecked;
+}
+
+bool SGeometryModeControls::IsModifierEnabled(UGeomModifier* Modifier) const
+{
+	return Modifier->Supports();
+}
+
+EVisibility SGeometryModeControls::IsPropertiesVisible() const
+{
+	if ((GetGeometryModeTool()->GetNumModifiers() > 0) && (GetGeometryModeTool()->GetCurrentModifier() != GetGeometryModeTool()->GetModifier(0)))
+	{
+		return EVisibility::Visible;
+	}
+	else
+	{
+		return EVisibility::Collapsed;
+	}
+}
+
+FReply SGeometryModeControls::OnApplyClicked()
+{
+	check(GLevelEditorModeTools().IsModeActive(FGeometryEditingModes::EM_Geometry));
+
+	GetGeometryModeTool()->GetCurrentModifier()->Apply();
+
+	return FReply::Handled();
+}
+
+FReply SGeometryModeControls::OnModifierClicked(UGeomModifier* Modifier)
+{
+	check(GLevelEditorModeTools().IsModeActive(FGeometryEditingModes::EM_Geometry));
+
+	Modifier->Apply();
+
+	return FReply::Handled();
+}
+
+void SGeometryModeControls::CreateLayout()
+{
+	this->ChildSlot
+	[
+		SNew(SScrollBox)
+		+SScrollBox::Slot()
+		.Padding(0.0f)
+		[
+			SNew(SBorder)
+			.BorderImage(FEditorStyle::GetBrush("ToolPanel.GroupBorder"))
+			[
+				SNew(SVerticalBox)
+				+SVerticalBox::Slot()
+				.AutoHeight()
+				.HAlign(HAlign_Center)
+				[
+					CreateTopModifierButtons()
+				]
+				+SVerticalBox::Slot()
+				.AutoHeight()
+				.Padding(3.0f)
+				[
+					SNew(SSeparator)
+					.Orientation(Orient_Horizontal)
+				]	
+				+SVerticalBox::Slot()
+				.AutoHeight()
+				[
+					CreateModifierProperties()
+				]	
+				+SVerticalBox::Slot()
+				.AutoHeight()
+				.Padding(3.0f)
+				[
+					SNew(SSeparator)
+					.Orientation(Orient_Horizontal)
+					.Visibility(this, &SGeometryModeControls::IsPropertiesVisible)
+				]
+				+SVerticalBox::Slot()
+				.AutoHeight()
+				.HAlign(HAlign_Center)
+				[
+					CreateBottomModifierButtons()
+				]
+			]
+		]
+	];
+}
+
+TSharedRef<SVerticalBox> SGeometryModeControls::CreateTopModifierButtons()
+{
+	FModeTool_GeometryModify* GeometryModeTool = GetGeometryModeTool();
+	TSharedPtr<SVerticalBox> Vbox;
+	const TSharedRef<SGridPanel> RadioButtonPanel = SNew(SGridPanel);
+
+	// Loop through all geometry modifiers and create radio buttons for ones with the bPushButton set to false
+	int32 CurrentModifierButtonCount = 0;
+	for (FModeTool_GeometryModify::TModifierIterator Itor(GeometryModeTool->ModifierIterator()); Itor; ++Itor)
+	{
+		UGeomModifier* Modifier = *Itor;
+		if (!Modifier->bPushButton)
+		{
+			RadioButtonPanel->AddSlot(CurrentModifierButtonCount%2, CurrentModifierButtonCount/2)
+			.Padding( FMargin(20.0f, 5.0f) )
+			[
+				CreateSingleModifierRadioButton(Modifier)
+			];
+
+			++CurrentModifierButtonCount;
+		}
+	}
+
+	// Add the Apply button
+	SAssignNew(Vbox, SVerticalBox)
+	+SVerticalBox::Slot()
+	.AutoHeight()
+	[
+		RadioButtonPanel
+	]
+	+SVerticalBox::Slot()
+	.AutoHeight()
+	.VAlign(VAlign_Center)
+	.HAlign(HAlign_Center)
+	[
+		SNew(SButton)
+		.Text(LOCTEXT("SGeometryModeDialog_Apply", "Apply"))
+		.OnClicked(this, &SGeometryModeControls::OnApplyClicked)
+	];
+
+	return Vbox.ToSharedRef();
+}
+
+TSharedRef<SUniformGridPanel> SGeometryModeControls::CreateBottomModifierButtons()
+{
+	FModeTool_GeometryModify* GeometryModeTool = GetGeometryModeTool();
+	TSharedRef<SUniformGridPanel> ButtonGrid = SNew(SUniformGridPanel).SlotPadding(5.0f);
+
+	// The IDs of the buttons created in this function need to sequentially follow the IDs of the modifier radio buttons
+	// So, this loop simply counts the number of radio buttons so we can use that as an offset
+	int32 CurrentModifierButtonCount = 0;
+	for (FModeTool_GeometryModify::TModifierConstIterator Itor(GeometryModeTool->ModifierConstIterator()); Itor; ++Itor)
+	{
+		const UGeomModifier* Modifier = *Itor;
+		if (!Modifier->bPushButton)
+		{
+			++CurrentModifierButtonCount;
+		}
+	}
+
+	// Loop through all geometry modifiers and create buttons for ones with the bPushButton set to true
+	int32 PushButtonId = 0;
+	for (FModeTool_GeometryModify::TModifierIterator Itor(GeometryModeTool->ModifierIterator()); Itor; ++Itor)
+	{
+		UGeomModifier* Modifier = *Itor;
+		if (Modifier->bPushButton)
+		{
+			ButtonGrid->AddSlot(PushButtonId % 2, PushButtonId / 2)
+			[
+				CreateSingleModifierButton(Modifier)
+			];
+
+			++CurrentModifierButtonCount;
+			++PushButtonId;
+		}
+	}
+
+	return ButtonGrid;
+}
+
+TSharedRef<IDetailsView> SGeometryModeControls::CreateModifierProperties()
+{
+	FDetailsViewArgs Args;
+	Args.bHideSelectionTip = true;
+	Args.bAllowSearch = false;
+
+	FPropertyEditorModule& PropertyModule = FModuleManager::LoadModuleChecked<FPropertyEditorModule>("PropertyEditor");
+	PropertiesControl = PropertyModule.CreateDetailView(Args);
+	PropertiesControl->SetVisibility(TAttribute<EVisibility>(this, &SGeometryModeControls::IsPropertiesVisible));
+
+	return PropertiesControl.ToSharedRef();
+}
+
+TSharedRef<SCheckBox> SGeometryModeControls::CreateSingleModifierRadioButton(UGeomModifier* Modifier)
+{
+	TSharedRef<SCheckBox> CheckBox =
+	SNew(SCheckBox)
+	.Style(FEditorStyle::Get(), "RadioButton")
+	.Padding(FMargin(4.0f, 0.0f, 0.0f, 0.0f))
+	.IsChecked(this, &SGeometryModeControls::IsModifierChecked, Modifier)
+	.IsEnabled(this, &SGeometryModeControls::IsModifierEnabled, Modifier)
+	.OnCheckStateChanged(this, &SGeometryModeControls::OnModifierStateChanged, Modifier)
+	.ToolTip(SNew(SToolTip).Text(Modifier->GetModifierTooltip()))
+	[
+		SNew(STextBlock).Text( Modifier->GetModifierDescription() )
+	];
+
+	ModifierControls.Add(CheckBox);
+
+	return CheckBox;
+}
+
+TSharedRef<SButton> SGeometryModeControls::CreateSingleModifierButton(UGeomModifier* Modifier)
+{
+	TSharedRef<SButton> Widget =
+	SNew(SButton)
+	.Text( Modifier->GetModifierDescription() )
+	.ToolTip(SNew(SToolTip).Text(Modifier->GetModifierTooltip()))
+	.HAlign(HAlign_Center)
+	.IsEnabled(this, &SGeometryModeControls::IsModifierEnabled, Modifier)
+	.OnClicked(this, &SGeometryModeControls::OnModifierClicked, Modifier);
+
+	return Widget;
+}
+
+FModeTool_GeometryModify* SGeometryModeControls::GetGeometryModeTool() const
+{
+	FEdModeGeometry* Mode = (FEdModeGeometry*)GLevelEditorModeTools().GetActiveMode(FGeometryEditingModes::EM_Geometry );
+	FModeTool* Tool = Mode? Mode->GetCurrentTool(): NULL;
+
+	check(Tool);
+
+	return (FModeTool_GeometryModify*)Tool;
+}
+
+void FGeometryMode::RegisterTabSpawners(const TSharedRef<class FTabManager>& TabManager)
+{
+
+}
+
+void FGeometryMode::UnregisterTabSpawners(const TSharedRef<class FTabManager>& TabManager)
+{
+
+}
+
+void FGeometryMode::Init(const TSharedPtr< class IToolkitHost >& InitToolkitHost)
+{
+	GeomWidget = SNew(SGeometryModeControls);
+
+	FModeToolkit::Init(InitToolkitHost);
+}
+
+FName FGeometryMode::GetToolkitFName() const
+{
+	return FName("GeometryMode");
+}
+
+FText FGeometryMode::GetBaseToolkitName() const
+{
+	return LOCTEXT( "ToolkitName", "Geometry Mode" );
+}
+
+class FEdMode* FGeometryMode::GetEditorMode() const
+{
+	return (FEdModeGeometry*)GLevelEditorModeTools().GetActiveMode(FGeometryEditingModes::EM_Geometry);
+}
+
+void FGeometryMode::SelectionChanged()
+{
+	GeomWidget->SelectionChanged();
+}
+
+TSharedPtr<SWidget> FGeometryMode::GetInlineContent() const
+{
+	return SNew(SGeometryModeControls);
+}
+
+#undef LOCTEXT_NAMESPACE
diff --git a/Engine/Plugins/Editor/GeometryMode/Source/GeometryMode/Private/GeometryMode.h b/Engine/Plugins/Editor/GeometryMode/Source/GeometryMode/Private/GeometryMode.h
new file mode 100644
--- /dev/null
+++ b/Engine/Plugins/Editor/GeometryMode/Source/GeometryMode/Private/GeometryMode.h
@@ -0,0 +1,117 @@
+// Copyright 1998-2019 Epic Games, Inc. All Rights Reserved.
+
+#pragma once
+
+#include "CoreMinimal.h"
+#include "Layout/Visibility.h"
+#include "Input/Reply.h"
+#include "Widgets/DeclarativeSyntaxSupport.h"
+#include "Widgets/SWidget.h"
+#include "Widgets/SCompoundWidget.h"
+#include "Widgets/SBoxPanel.h"
+#include "Misc/NotifyHook.h"
+#include "Toolkits/IToolkitHost.h"
+#include "Toolkits/BaseToolkit.h"
+
+class IDetailsView;
+class SButton;
+class SCheckBox;
+class SUniformGridPanel;
+class UGeomModifier;
+
+/** Geometry Mode widget for controls */
+class SGeometryModeControls : public SCompoundWidget, public FNotifyHook
+{
+public:
+	SLATE_BEGIN_ARGS(SGeometryModeControls) {}
+	SLATE_END_ARGS()
+
+public:
+
+	void SelectionChanged();
+
+	/** SCompoundWidget functions */
+	void Construct(const FArguments& InArgs);
+
+protected:
+	/** Called when a new modifier mode is selected */
+	void OnModifierStateChanged(ECheckBoxState NewCheckedState, UGeomModifier* Modifier);
+
+	/** Returns the state of a modifier radio button */
+	ECheckBoxState IsModifierChecked(UGeomModifier* Modifier) const;
+
+	/** Returns the enabled state of a modifier button */
+	bool IsModifierEnabled(UGeomModifier* Modifier) const;
+
+	/** Returns the visibility state of the properties control */
+	EVisibility IsPropertiesVisible() const;
+
+	/** Called when the Apply button is clicked */
+	FReply OnApplyClicked();
+
+	/** Called when a modifier button is clicked */
+	FReply OnModifierClicked(UGeomModifier* Modifier);
+
+private:
+	/** Creates the geometry mode controls */
+	void CreateLayout();
+
+	/** Creates controls for the modifiers section */
+	TSharedRef<SVerticalBox> CreateTopModifierButtons();
+	
+	/** Creates controls for the actions section */
+	TSharedRef<SUniformGridPanel> CreateBottomModifierButtons();
+
+	/** Creates controls for the modifier properties section */
+	TSharedRef<class IDetailsView> CreateModifierProperties();
+
+	/** Creates a modifier radio button */
+	TSharedRef<SCheckBox> CreateSingleModifierRadioButton(UGeomModifier* Modifier);
+
+	/** Creates an action button */
+	TSharedRef<SButton> CreateSingleModifierButton(UGeomModifier* Modifier);
+
+	/** Returns a reference to the geometry mode tool */
+	class FModeTool_GeometryModify* GetGeometryModeTool() const;
+
+	void MakeBuilderBrush( UClass* BrushBuilderClass );
+
+	void OnAddVolume( UClass* VolumeClass );
+
+private:
+	/** Pointer to the parent window, so we know to destroy it when done */
+	TWeakPtr<SWindow> ParentWindow;
+
+	/** A list of the checkbox modifier controls */
+	TArray< TSharedPtr<SCheckBox> > ModifierControls;
+
+	/** The properties control */
+	TSharedPtr<class IDetailsView> PropertiesControl;
+};
+
+
+/**
+ * Mode Toolkit for the Geometry Tools
+ */
+class FGeometryMode : public FModeToolkit
+{
+public:
+	virtual void RegisterTabSpawners(const TSharedRef<class FTabManager>& TabManager) override;
+	virtual void UnregisterTabSpawners(const TSharedRef<class FTabManager>& TabManager) override;
+
+	/** Initializes the geometry mode toolkit */
+	virtual void Init(const TSharedPtr< class IToolkitHost >& InitToolkitHost) override;
+
+	/** IToolkit interface */
+	virtual FName GetToolkitFName() const override;
+	virtual FText GetBaseToolkitName() const override;
+	virtual class FEdMode* GetEditorMode() const override;
+	virtual TSharedPtr<class SWidget> GetInlineContent() const override;
+
+	/** Method called when the selection */
+	virtual void SelectionChanged();
+
+private:
+	/** Geometry tools widget */
+	TSharedPtr<class SGeometryModeControls> GeomWidget;
+};
diff --git a/Engine/Plugins/Editor/GeometryMode/Source/GeometryMode/Private/GeometryModeModule.cpp b/Engine/Plugins/Editor/GeometryMode/Source/GeometryMode/Private/GeometryModeModule.cpp
new file mode 100644
--- /dev/null
+++ b/Engine/Plugins/Editor/GeometryMode/Source/GeometryMode/Private/GeometryModeModule.cpp
@@ -0,0 +1,31 @@
+// Copyright 1998-2019 Epic Games, Inc. All Rights Reserved.
+
+#include "GeometryModeModule.h"
+#include "Modules/ModuleManager.h"
+#include "EditorModeManager.h"
+#include "Editor.h"
+
+
+FEditorModeID FGeometryEditingModes::EM_Geometry = FEditorModeID(TEXT("EM_Geometry"));
+FEditorModeID FGeometryEditingModes::EM_Bsp = FEditorModeID(TEXT("EM_Bsp"));
+FEditorModeID FGeometryEditingModes::EM_TextureAlign = FEditorModeID(TEXT("EM_TextureAlign"));
+
+void FGeometryModeModule::StartupModule()
+{
+	FEditorModeRegistry::Get().RegisterMode<FEdModeGeometry>(
+		FGeometryEditingModes::EM_Geometry,
+		NSLOCTEXT("EditorModes", "GeometryMode", "Geometry Editing"),
+		FSlateIcon(FEditorStyle::GetStyleSetName(), "LevelEditor.BspMode", "LevelEditor.BspMode.Small"),
+		true, 500
+	);
+}
+
+
+void FGeometryModeModule::ShutdownModule()
+{
+	FEditorModeRegistry::Get().UnregisterMode(FGeometryEditingModes::EM_Geometry);
+}
+
+IMPLEMENT_MODULE(FGeometryModeModule, GeometryMode);
+
+#undef LOCTEXT_NAMESPACE
diff --git a/Engine/Plugins/Editor/GeometryMode/Source/GeometryMode/Public/GeometryModeModule.h b/Engine/Plugins/Editor/GeometryMode/Source/GeometryMode/Public/GeometryModeModule.h
new file mode 100644
--- /dev/null
+++ b/Engine/Plugins/Editor/GeometryMode/Source/GeometryMode/Public/GeometryModeModule.h
@@ -0,0 +1,26 @@
+// Copyright 1998-2019 Epic Games, Inc. All Rights Reserved.
+
+#pragma once
+
+#include "CoreMinimal.h"
+#include "Modules/ModuleInterface.h"
+
+typedef FName FEditorModeID;
+
+struct GEOMETRYMODE_API FGeometryEditingModes
+{
+	static FEditorModeID EM_Geometry;
+	static FEditorModeID EM_Bsp;
+	static FEditorModeID EM_TextureAlign;
+};
+/**
+ * Geometry mode module
+ */
+class FGeometryModeModule : public IModuleInterface
+{
+public:
+	// IModuleInterface
+	virtual void StartupModule() override;
+	virtual void ShutdownModule() override;
+	// End of IModuleInterface
+};
diff --git a/Engine/Plugins/Editor/GeometryMode/Source/TextureAlignMode/Private/TextureAlignMode.cpp b/Engine/Plugins/Editor/GeometryMode/Source/TextureAlignMode/Private/TextureAlignMode.cpp
new file mode 100644
--- /dev/null
+++ b/Engine/Plugins/Editor/GeometryMode/Source/TextureAlignMode/Private/TextureAlignMode.cpp
@@ -0,0 +1,35 @@
+// Copyright 1998-2019 Epic Games, Inc. All Rights Reserved.
+
+#include "TextureAlignMode.h"
+#include "EditorModeManager.h"
+#include "EditorModes.h"
+#include "TextureAlignEdMode.h"
+
+#define LOCTEXT_NAMESPACE "TextureAlignMode"
+
+void FTextureAlignMode::RegisterTabSpawners(const TSharedRef<class FTabManager>& TabManager)
+{
+
+}
+
+void FTextureAlignMode::UnregisterTabSpawners(const TSharedRef<class FTabManager>& TabManager)
+{
+
+}
+
+FName FTextureAlignMode::GetToolkitFName() const
+{
+	return FName("TextureAlignMode");
+}
+
+FText FTextureAlignMode::GetBaseToolkitName() const
+{
+	return LOCTEXT("ToolkitName", "Texture Align Mode");
+}
+
+class FEdMode* FTextureAlignMode::GetEditorMode() const
+{
+	return (FEdModeTexture*)GLevelEditorModeTools().GetActiveMode(FGeometryEditingModes::EM_TextureAlign);
+}
+
+#undef LOCTEXT_NAMESPACE
diff --git a/Engine/Plugins/Editor/GeometryMode/Source/TextureAlignMode/Public/TextureAlignMode.h b/Engine/Plugins/Editor/GeometryMode/Source/TextureAlignMode/Public/TextureAlignMode.h
new file mode 100644
--- /dev/null
+++ b/Engine/Plugins/Editor/GeometryMode/Source/TextureAlignMode/Public/TextureAlignMode.h
@@ -0,0 +1,24 @@
+// Copyright 1998-2019 Epic Games, Inc. All Rights Reserved.
+
+#pragma once
+
+#include "CoreMinimal.h"
+#include "Toolkits/BaseToolkit.h"
+
+/**
+ * Mode Toolkit for the Texture Align Tools
+ */
+class FTextureAlignMode : public FModeToolkit
+{
+public:
+	virtual void RegisterTabSpawners(const TSharedRef<class FTabManager>& TabManager) override;
+	virtual void UnregisterTabSpawners(const TSharedRef<class FTabManager>& TabManager) override;
+
+	/** IToolkit interface */
+	virtual FName GetToolkitFName() const override;
+	virtual FText GetBaseToolkitName() const override;
+	virtual class FEdMode* GetEditorMode() const override;
+
+private:
+
+};
diff --git a/Engine/Source/Editor/PlacementMode/Private/PlacementMode.cpp b/Engine/Source/Editor/PlacementMode/Private/PlacementMode.cpp
--- a/Engine/Source/Editor/PlacementMode/Private/PlacementMode.cpp
+++ b/Engine/Source/Editor/PlacementMode/Private/PlacementMode.cpp
@@ -449,8 +449,6 @@ bool FPlacementMode::UsesPropertyWidgets() const
 bool FPlacementMode::IsCompatibleWith(FEditorModeID OtherModeID) const
 {
 	return
-		OtherModeID == FBuiltinEditorModes::EM_Bsp			||
-		OtherModeID == FBuiltinEditorModes::EM_Geometry		||
 		OtherModeID == FBuiltinEditorModes::EM_InterpEdit	||
 		OtherModeID == FBuiltinEditorModes::EM_MeshPaint	||
 		OtherModeID == FBuiltinEditorModes::EM_Foliage		||
diff --git a/Engine/Source/Editor/UnrealEd/Private/DragTool_BoxSelect.cpp b/Engine/Source/Editor/UnrealEd/Private/DragTool_BoxSelect.cpp
--- a/Engine/Source/Editor/UnrealEd/Private/DragTool_BoxSelect.cpp
+++ b/Engine/Source/Editor/UnrealEd/Private/DragTool_BoxSelect.cpp
@@ -15,6 +15,7 @@
 #include "ScopedTransaction.h"
 #include "Engine/LevelStreaming.h"
 #include "CanvasTypes.h"
+#include "Subsystems/BrushEditingSubsystem.h"
 
 ///////////////////////////////////////////////////////////////////////////////////////////////////////////////////////
 //
@@ -130,7 +131,8 @@ void FDragTool_ActorBoxSelect::AddDelta( const FVector& InDelta )
 */
 void FDragTool_ActorBoxSelect::EndDrag()
 {
-	const bool bGeometryMode = ModeTools->IsModeActive(FBuiltinEditorModes::EM_Geometry);
+	UBrushEditingSubsystem* BrushSubsystem = GEditor->GetEditorSubsystem<UBrushEditingSubsystem>();
+	const bool bGeometryMode = BrushSubsystem ? BrushSubsystem->IsGeometryEditorModeActive() : false;
 	
 	FScopedTransaction Transaction( NSLOCTEXT("ActorFrustumSelect", "MarqueeSelectTransation", "Marquee Select" ) );
 
@@ -291,7 +293,8 @@ bool FDragTool_ActorBoxSelect::IntersectsBox( AActor& InActor, const FBox& InBox
 {
 	bool bActorHitByBox = false;
 
-	const bool bGeometryMode = ModeTools->IsModeActive(FBuiltinEditorModes::EM_Geometry);
+	UBrushEditingSubsystem* BrushSubsystem = GEditor->GetEditorSubsystem<UBrushEditingSubsystem>();
+	const bool bGeometryMode = BrushSubsystem ? BrushSubsystem->IsGeometryEditorModeActive() : false;
 
 	// Check for special cases (like certain show flags that might hide an actor)
 	bool bActorIsHiddenByShowFlags = false;
diff --git a/Engine/Source/Editor/UnrealEd/Private/DragTool_FrustumSelect.cpp b/Engine/Source/Editor/UnrealEd/Private/DragTool_FrustumSelect.cpp
--- a/Engine/Source/Editor/UnrealEd/Private/DragTool_FrustumSelect.cpp
+++ b/Engine/Source/Editor/UnrealEd/Private/DragTool_FrustumSelect.cpp
@@ -16,6 +16,7 @@
 #include "ScopedTransaction.h"
 #include "HModel.h"
 #include "CanvasTypes.h"
+#include "Subsystems/BrushEditingSubsystem.h"
 
 ///////////////////////////////////////////////////////////////////////////////
 //
@@ -52,7 +53,8 @@ void FDragTool_ActorFrustumSelect::StartDrag(FEditorViewportClient* InViewportCl
 
 void FDragTool_ActorFrustumSelect::EndDrag()
 {
-	const bool bGeometryMode = ModeTools->IsModeActive(FBuiltinEditorModes::EM_Geometry);
+	UBrushEditingSubsystem* BrushSubsystem = GEditor->GetEditorSubsystem<UBrushEditingSubsystem>();
+	const bool bGeometryMode = BrushSubsystem ? BrushSubsystem->IsGeometryEditorModeActive() : false;
 
 	FSceneViewFamilyContext ViewFamily(FSceneViewFamily::ConstructionValues(LevelViewportClient->Viewport, LevelViewportClient->GetScene(), LevelViewportClient->EngineShowFlags ));
 	FSceneView* SceneView = LevelViewportClient->CalcSceneView(&ViewFamily);
@@ -190,7 +192,8 @@ bool FDragTool_ActorFrustumSelect::IntersectsFrustum( AActor& InActor, const FCo
 {	
 	bool bActorHitByBox = false;
 
-	const bool bGeometryMode = ModeTools->IsModeActive(FBuiltinEditorModes::EM_Geometry);
+	UBrushEditingSubsystem* BrushSubsystem = GEditor->GetEditorSubsystem<UBrushEditingSubsystem>();
+	const bool bGeometryMode = BrushSubsystem ? BrushSubsystem->IsGeometryEditorModeActive() : false;
 
 	// Check for special cases (like certain show flags that might hide an actor)
 	bool bActorIsHiddenByShowFlags = false;
diff --git a/Engine/Source/Editor/UnrealEd/Private/EditorActor.cpp b/Engine/Source/Editor/UnrealEd/Private/EditorActor.cpp
--- a/Engine/Source/Editor/UnrealEd/Private/EditorActor.cpp
+++ b/Engine/Source/Editor/UnrealEd/Private/EditorActor.cpp
@@ -49,9 +49,8 @@
 #include "EditorLevelUtils.h"
 #include "Kismet2/BlueprintEditorUtils.h"
 #include "LevelEditorViewport.h"
+#include "Layers/ILayers.h"
 #include "Layers/LayersSubsystem.h"
-#include "Editor/GeometryMode/Public/GeometryEdMode.h"
-#include "Editor/GeometryMode/Public/EditorGeometry.h"
 #include "ActorEditorUtils.h"
 #include "Particles/ParticleSystemComponent.h"
 #include "UnrealExporter.h"
@@ -65,6 +64,8 @@
 #include "AssetToolsModule.h"
 #include "AssetSelection.h"
 #include "Framework/Application/SlateApplication.h"
+#include "EdMode.h"
+#include "Subsystems/BrushEditingSubsystem.h"
 #include "Misc/ScopedSlowTask.h"
 #include "Subsystems/AssetEditorSubsystem.h"
 
@@ -2593,24 +2594,9 @@ void UUnrealEdEngine::edactAlignVertices()
 						PolyIdx = -1;
 					}
 
-					// Determine if we are in geometry edit mode.
-					if ( GLevelEditorModeTools().IsModeActive(FBuiltinEditorModes::EM_Geometry) )
+					if (UBrushEditingSubsystem* BrushSubsystem = GEditor->GetEditorSubsystem<UBrushEditingSubsystem>())
 					{
-						// If we are in geometry mode, go through the list of geometry objects
-						// and find our current brush and update its source data as it might have changed 
-						// in RecomputePoly
-						FEdModeGeometry* GeomMode = (FEdModeGeometry*)GLevelEditorModeTools().GetActiveMode(FBuiltinEditorModes::EM_Geometry);
-						FEdModeGeometry::TGeomObjectIterator GeomModeIt = GeomMode->GeomObjectItor();
-						for( ; GeomModeIt; ++GeomModeIt )
-						{
-							FGeomObjectPtr Object = *GeomModeIt;
-							if( Object->GetActualBrush() == Brush )
-							{
-								// We found our current brush, update the geometry object's data
-								Object->GetFromSource();
-								break;
-							}
-						}
+						BrushSubsystem->UpdateGeometryFromBrush(Brush);
 					}
 				}
 			}
diff --git a/Engine/Source/Editor/UnrealEd/Private/EditorModeInterpolation.cpp b/Engine/Source/Editor/UnrealEd/Private/EditorModeInterpolation.cpp
--- a/Engine/Source/Editor/UnrealEd/Private/EditorModeInterpolation.cpp
+++ b/Engine/Source/Editor/UnrealEd/Private/EditorModeInterpolation.cpp
@@ -191,10 +191,8 @@ void FEdModeInterpEdit::ActorSelectionChangeNotify()
 bool FEdModeInterpEdit::IsCompatibleWith(FEditorModeID OtherModeID) const
 {
 	return
-		OtherModeID == FBuiltinEditorModes::EM_Placement	||
-		OtherModeID == FBuiltinEditorModes::EM_MeshPaint	||
-		OtherModeID == FBuiltinEditorModes::EM_Geometry		||
-		OtherModeID == FBuiltinEditorModes::EM_Bsp;
+		OtherModeID == FBuiltinEditorModes::EM_Placement ||
+		OtherModeID == FBuiltinEditorModes::EM_MeshPaint;
 }
 
 //////////////////////////////////////////////////////////////////////////
diff --git a/Engine/Source/Editor/UnrealEd/Private/EditorModeManager.cpp b/Engine/Source/Editor/UnrealEd/Private/EditorModeManager.cpp
--- a/Engine/Source/Editor/UnrealEd/Private/EditorModeManager.cpp
+++ b/Engine/Source/Editor/UnrealEd/Private/EditorModeManager.cpp
@@ -30,6 +30,7 @@
 #include "Framework/Commands/UICommandList.h"
 #include "Framework/MultiBox/MultiBoxBuilder.h"
 #include "Toolkits/BaseToolkit.h"
+#include "Subsystems/BrushEditingSubsystem.h"
 #include "Tools/UEdMode.h"
 #include "Widgets/Images/SImage.h"
 
@@ -1403,8 +1404,13 @@ bool FEditorModeTools::ShouldDrawBrushWireframe( AActor* InActor ) const
 /** true if brush vertices should be drawn */
 bool FEditorModeTools::ShouldDrawBrushVertices() const
 {
-	// Currently only geometry mode being active prevents vertices from being drawn.
-	return !IsModeActive( FBuiltinEditorModes::EM_Geometry );
+	if(UBrushEditingSubsystem* BrushSubsystem = GEditor->GetEditorSubsystem<UBrushEditingSubsystem>())
+	{
+		// Currently only geometry mode being active prevents vertices from being drawn.
+		return !BrushSubsystem->IsGeometryEditorModeActive();
+	}
+
+	return true;
 }
 
 /** Ticks all active modes */
diff --git a/Engine/Source/Editor/UnrealEd/Private/EditorModeRegistry.cpp b/Engine/Source/Editor/UnrealEd/Private/EditorModeRegistry.cpp
--- a/Engine/Source/Editor/UnrealEd/Private/EditorModeRegistry.cpp
+++ b/Engine/Source/Editor/UnrealEd/Private/EditorModeRegistry.cpp
@@ -9,12 +9,9 @@
 
 #include "Editor/PlacementMode/Public/IPlacementModeModule.h"
 #include "Editor/LandscapeEditor/Public/LandscapeEditorModule.h"
-#include "Editor/BspMode/Public/IBspModeModule.h"
 #include "Editor/MeshPaint/Public/MeshPaintModule.h"
-#include "Editor/GeometryMode/Public/GeometryEdMode.h"
 #include "Editor/ActorPickerMode/Public/ActorPickerMode.h"
 #include "Editor/SceneDepthPickerMode/Public/SceneDepthPickerMode.h"
-#include "Editor/TextureAlignMode/Public/TextureAlignEdMode.h"
 #include "Editor/FoliageEdit/Public/FoliageEditModule.h"
 #include "Editor/VirtualTexturingEditor/Public/VirtualTexturingEditorModule.h"
 #include "Tools/UEdMode.h"
@@ -59,9 +56,6 @@ void FEditorModeRegistry::Initialize()
 	// Load editor mode modules that will automatically register their editor modes, and clean themselves up on unload.
 	//@TODO: ROCKET: These are probably good plugin candidates, that shouldn't have to be force-loaded here but discovery loaded somehow
 	FModuleManager::LoadModuleChecked<IPlacementModeModule>(TEXT("PlacementMode"));
-	FModuleManager::LoadModuleChecked<IBspModeModule>(TEXT("BspMode"));
-	FModuleManager::LoadModuleChecked<FTextureAlignModeModule>(TEXT("TextureAlignMode"));
-	FModuleManager::LoadModuleChecked<FGeometryModeModule>(TEXT("GeometryMode"));
 	FModuleManager::LoadModuleChecked<FActorPickerModeModule>(TEXT("ActorPickerMode"));
 	FModuleManager::LoadModuleChecked<FSceneDepthPickerModeModule>(TEXT("SceneDepthPickerMode"));
 	FModuleManager::LoadModuleChecked<IMeshPaintModule>(TEXT("MeshPaintMode"));
diff --git a/Engine/Source/Editor/UnrealEd/Private/EditorModes.cpp b/Engine/Source/Editor/UnrealEd/Private/EditorModes.cpp
--- a/Engine/Source/Editor/UnrealEd/Private/EditorModes.cpp
+++ b/Engine/Source/Editor/UnrealEd/Private/EditorModes.cpp
@@ -10,10 +10,7 @@ DEFINE_LOG_CATEGORY(LogEditorModes);
 const FEditorModeID FBuiltinEditorModes::EM_None = NAME_None;
 const FEditorModeID FBuiltinEditorModes::EM_Default(TEXT("EM_Default"));
 const FEditorModeID FBuiltinEditorModes::EM_Placement(TEXT("PLACEMENT"));
-const FEditorModeID FBuiltinEditorModes::EM_Bsp(TEXT("BSP"));
-const FEditorModeID FBuiltinEditorModes::EM_Geometry(TEXT("EM_Geometry"));
 const FEditorModeID FBuiltinEditorModes::EM_InterpEdit(TEXT("EM_InterpEdit"));
-const FEditorModeID FBuiltinEditorModes::EM_Texture(TEXT("EM_Texture"));
 const FEditorModeID FBuiltinEditorModes::EM_MeshPaint(TEXT("EM_MeshPaint"));
 const FEditorModeID FBuiltinEditorModes::EM_Landscape(TEXT("EM_Landscape"));
 const FEditorModeID FBuiltinEditorModes::EM_Foliage(TEXT("EM_Foliage"));
diff --git a/Engine/Source/Editor/UnrealEd/Private/EditorSelectUtils.cpp b/Engine/Source/Editor/UnrealEd/Private/EditorSelectUtils.cpp
--- a/Engine/Source/Editor/UnrealEd/Private/EditorSelectUtils.cpp
+++ b/Engine/Source/Editor/UnrealEd/Private/EditorSelectUtils.cpp
@@ -26,6 +26,7 @@
 #include "SnappingUtils.h"
 #include "Logging/MessageLog.h"
 #include "ActorGroupingUtils.h"
+#include "Subsystems/BrushEditingSubsystem.h"
 
 #define LOCTEXT_NAMESPACE "EditorSelectUtils"
 
@@ -354,9 +355,10 @@ void UUnrealEdEngine::UpdatePivotLocationForSelection( bool bOnChange )
 	}
 	else if( SingleActor != NULL ) 
 	{		
+		UBrushEditingSubsystem* BrushSubsystem = GEditor->GetEditorSubsystem<UBrushEditingSubsystem>();
+
 		// For geometry mode use current pivot location as it's set to selected face, not actor
-		FEditorModeTools& Tools = GLevelEditorModeTools();
-		if( Tools.IsModeActive(FBuiltinEditorModes::EM_Geometry) == false || bOnChange == true )
+		if(BrushSubsystem->IsGeometryEditorModeActive() == false || bOnChange == true )
 		{
 			// Set pivot point to the actor's location, accounting for any set pivot offset
 			FVector PivotPoint = SingleActor->GetTransform().TransformPosition(SingleActor->GetPivotOffset());
diff --git a/Engine/Source/Editor/UnrealEd/Private/EditorServer.cpp b/Engine/Source/Editor/UnrealEd/Private/EditorServer.cpp
--- a/Engine/Source/Editor/UnrealEd/Private/EditorServer.cpp
+++ b/Engine/Source/Editor/UnrealEd/Private/EditorServer.cpp
@@ -110,8 +110,6 @@
 #include "PackageTools.h"
 #include "LevelEditor.h"
 #include "Kismet2/BlueprintEditorUtils.h"
-#include "Editor/GeometryMode/Public/GeometryEdMode.h"
-#include "Editor/GeometryMode/Public/EditorGeometry.h"
 #include "LandscapeProxy.h"
 #include "Lightmass/PrecomputedVisibilityOverrideVolume.h"
 #include "Animation/AnimSet.h"
@@ -159,6 +157,8 @@
 #include "HAL/PlatformApplicationMisc.h"
 #include "AssetExportTask.h"
 #include "EditorBuildUtils.h"
+#include "Subsystems/BrushEditingSubsystem.h"
+#include "EdMode.h"
 
 DEFINE_LOG_CATEGORY_STATIC(LogEditorServer, Log, All);
 
@@ -1939,11 +1939,11 @@ void UEditorEngine::RebuildAlteredBSP()
 
 void UEditorEngine::BSPIntersectionHelper(UWorld* InWorld, ECsgOper Operation)
 {
-	FEdModeGeometry* Mode = GLevelEditorModeTools().GetActiveModeTyped<FEdModeGeometry>(FBuiltinEditorModes::EM_Geometry);
-	if (Mode)
+	if (UBrushEditingSubsystem* BrushSubsystem = GetEditorSubsystem<UBrushEditingSubsystem>())
 	{
-		Mode->GeometrySelectNone(true, true);
+		BrushSubsystem->DeselectAllEditingGeometry();
 	}
+
 	ABrush* DefaultBrush = InWorld->GetDefaultBrush();
 	if (DefaultBrush != NULL)
 	{
@@ -6120,8 +6120,14 @@ bool UEditorEngine::HandleSelectCommand( const TCHAR* Str, FOutputDevice& Ar, UW
 
 bool UEditorEngine::HandleDeleteCommand( const TCHAR* Str, FOutputDevice& Ar, UWorld* InWorld )
 {
-	// If geometry mode is active, give it a chance to handle this command.  If it does not, use the default handler
-	if( !GLevelEditorModeTools().IsModeActive( FBuiltinEditorModes::EM_Geometry ) || !( (FEdModeGeometry*)GLevelEditorModeTools().GetActiveMode(FBuiltinEditorModes::EM_Geometry) )->ExecDelete() )
+	bool bHandled = false;
+	UBrushEditingSubsystem* BrushSubsystem = GetEditorSubsystem<UBrushEditingSubsystem>();
+	if (BrushSubsystem)
+	{
+		bHandled = BrushSubsystem->HandleActorDelete();
+	}
+
+	if(!bHandled)
 	{
 		return Exec( InWorld, TEXT("ACTOR DELETE") );
 	}
diff --git a/Engine/Source/Editor/UnrealEd/Private/LevelEditorViewport.cpp b/Engine/Source/Editor/UnrealEd/Private/LevelEditorViewport.cpp
--- a/Engine/Source/Editor/UnrealEd/Private/LevelEditorViewport.cpp
+++ b/Engine/Source/Editor/UnrealEd/Private/LevelEditorViewport.cpp
@@ -53,7 +53,6 @@
 #include "AssetRegistryModule.h"
 #include "IPlacementModeModule.h"
 #include "Engine/Polys.h"
-#include "Editor/GeometryMode/Public/EditorGeometry.h"
 #include "ActorEditorUtils.h"
 #include "ObjectTools.h"
 #include "PackageTools.h"
@@ -77,6 +76,7 @@
 #include "ActorGroupingUtils.h"
 #include "EditorWorldExtension.h"
 #include "VREditorMode.h"
+#include "Subsystems/BrushEditingSubsystem.h"
 #include "Engine/VolumeTexture.h"
 #include "Materials/MaterialExpressionDivide.h"
 #include "Materials/MaterialExpressionSubtract.h"
@@ -2118,18 +2118,19 @@ void FLevelEditorViewportClient::LostFocus(FViewport* InViewport)
 //
 void FLevelEditorViewportClient::ProcessClick(FSceneView& View, HHitProxy* HitProxy, FKey Key, EInputEvent Event, uint32 HitX, uint32 HitY)
 {
+	UBrushEditingSubsystem* BrushSubsystem = GEditor->GetEditorSubsystem<UBrushEditingSubsystem>();
 
 	const FViewportClick Click(&View,this,Key,Event,HitX,HitY);
 	if (Click.GetKey() == EKeys::MiddleMouseButton && !Click.IsAltDown() && !Click.IsShiftDown())
 	{
-		ClickHandlers::ClickViewport(this, Click);
+		LevelViewportClickHandlers::ClickViewport(this, Click);
 		return;
 	}
 	if (!ModeTools->HandleClick(this, HitProxy,Click))
 	{
 		if (HitProxy == NULL)
 		{
-			ClickHandlers::ClickBackdrop(this,Click);
+			LevelViewportClickHandlers::ClickBackdrop(this,Click);
 		}
 		else if (HitProxy->IsA(HWidgetAxis::StaticGetType()))
 		{
@@ -2199,11 +2200,11 @@ void FLevelEditorViewportClient::ProcessClick(FSceneView& View, HHitProxy* HitPr
 
 			if (bSelectComponent)
 			{
-				ClickHandlers::ClickComponent(this, ActorHitProxy, Click);
+				LevelViewportClickHandlers::ClickComponent(this, ActorHitProxy, Click);
 			}
 			else
 			{
-				ClickHandlers::ClickActor(this, ConsideredActor, Click, true);
+				LevelViewportClickHandlers::ClickActor(this, ConsideredActor, Click, true);
 			}
 
 			// We clicked an actor, allow the pivot to reposition itself.
@@ -2211,56 +2212,19 @@ void FLevelEditorViewportClient::ProcessClick(FSceneView& View, HHitProxy* HitPr
 		}
 		else if (HitProxy->IsA(HInstancedStaticMeshInstance::StaticGetType()))
 		{
-			ClickHandlers::ClickActor(this, ((HInstancedStaticMeshInstance*)HitProxy)->Component->GetOwner(), Click, true);
+			LevelViewportClickHandlers::ClickActor(this, ((HInstancedStaticMeshInstance*)HitProxy)->Component->GetOwner(), Click, true);
 		}
 		else if (HitProxy->IsA(HBSPBrushVert::StaticGetType()) && ((HBSPBrushVert*)HitProxy)->Brush.IsValid())
 		{
-			ClickHandlers::ClickBrushVertex(this,((HBSPBrushVert*)HitProxy)->Brush.Get(),((HBSPBrushVert*)HitProxy)->Vertex,Click);
+			LevelViewportClickHandlers::ClickBrushVertex(this,((HBSPBrushVert*)HitProxy)->Brush.Get(),((HBSPBrushVert*)HitProxy)->Vertex,Click);
 		}
 		else if (HitProxy->IsA(HStaticMeshVert::StaticGetType()))
 		{
-			ClickHandlers::ClickStaticMeshVertex(this,((HStaticMeshVert*)HitProxy)->Actor,((HStaticMeshVert*)HitProxy)->Vertex,Click);
+			LevelViewportClickHandlers::ClickStaticMeshVertex(this,((HStaticMeshVert*)HitProxy)->Actor,((HStaticMeshVert*)HitProxy)->Vertex,Click);
 		}
-		else if (HitProxy->IsA(HGeomPolyProxy::StaticGetType()))
+		else if (BrushSubsystem && BrushSubsystem->ProcessClickOnBrushGeometry(this, HitProxy, Click))
 		{
-			HGeomPolyProxy* GeomHitProxy = (HGeomPolyProxy*)HitProxy;
-
-			if( GeomHitProxy->GetGeomObject() )
-			{
-				FHitResult CheckResult(ForceInit);
-				FCollisionQueryParams BoxParams(SCENE_QUERY_STAT(ProcessClickTrace), false, GeomHitProxy->GetGeomObject()->ActualBrush);
-				bool bHit = GWorld->SweepSingleByObjectType(CheckResult, Click.GetOrigin(), Click.GetOrigin() + Click.GetDirection() * HALF_WORLD_MAX, FQuat::Identity, FCollisionObjectQueryParams(ECC_WorldStatic), FCollisionShape::MakeBox(FVector(1.f)), BoxParams);
-
-				if(bHit)
-				{
-					GEditor->UnsnappedClickLocation = CheckResult.Location;
-					GEditor->ClickLocation = CheckResult.Location;
-					GEditor->ClickPlane = FPlane(CheckResult.Location, CheckResult.Normal);
-				}
-
-				if(!ClickHandlers::ClickActor(this, GeomHitProxy->GetGeomObject()->ActualBrush, Click, false))
-				{
-					ClickHandlers::ClickGeomPoly(this, GeomHitProxy, Click);
-				}
-
-				Invalidate(true, true);
-			}
-		}
-		else if (HitProxy->IsA(HGeomEdgeProxy::StaticGetType()))
-		{
-			HGeomEdgeProxy* GeomHitProxy = (HGeomEdgeProxy*)HitProxy;
-
-			if( GeomHitProxy->GetGeomObject() !=nullptr )
-			{
-				if(!ClickHandlers::ClickGeomEdge(this, GeomHitProxy, Click))
-				{
-					ClickHandlers::ClickActor(this, GeomHitProxy->GetGeomObject()->ActualBrush, Click, true);
-				}
-			}
-		}
-		else if (HitProxy->IsA(HGeomVertexProxy::StaticGetType()))
-		{
-			ClickHandlers::ClickGeomVertex(this,(HGeomVertexProxy*)HitProxy,Click);
+			// Handled by the brush subsystem
 		}
 		else if (HitProxy->IsA(HModel::StaticGetType()))
 		{
@@ -2273,12 +2237,12 @@ void FLevelEditorViewportClient::ProcessClick(FSceneView& View, HHitProxy* HitPr
 			uint32 SurfaceIndex = INDEX_NONE;
 			if(ModelHit->ResolveSurface(SceneView,HitX,HitY,SurfaceIndex))
 			{
-				ClickHandlers::ClickSurface(this,ModelHit->GetModel(),SurfaceIndex,Click);
+				LevelViewportClickHandlers::ClickSurface(this,ModelHit->GetModel(),SurfaceIndex,Click);
 			}
 		}
 		else if (HitProxy->IsA(HLevelSocketProxy::StaticGetType()))
 		{
-			ClickHandlers::ClickLevelSocket(this, HitProxy, Click);
+			LevelViewportClickHandlers::ClickLevelSocket(this, HitProxy, Click);
 		}
 	}
 }
@@ -5200,10 +5164,3 @@ bool FLevelEditorViewportClient::GetPivotForOrbit(FVector& Pivot) const
 }
 
 #undef LOCTEXT_NAMESPACE
-
-// Doxygen cannot parse these correctly since the declarations are made in Editor, not UnrealEd
-#if !UE_BUILD_DOCS
-IMPLEMENT_HIT_PROXY(HGeomPolyProxy,HHitProxy);
-IMPLEMENT_HIT_PROXY(HGeomEdgeProxy,HHitProxy);
-IMPLEMENT_HIT_PROXY(HGeomVertexProxy,HHitProxy);
-#endif
diff --git a/Engine/Source/Editor/UnrealEd/Private/LevelViewportClickHandlers.cpp b/Engine/Source/Editor/UnrealEd/Private/LevelViewportClickHandlers.cpp
--- a/Engine/Source/Editor/UnrealEd/Private/LevelViewportClickHandlers.cpp
+++ b/Engine/Source/Editor/UnrealEd/Private/LevelViewportClickHandlers.cpp
@@ -24,17 +24,16 @@
 #include "ScopedTransaction.h"
 #include "ILevelEditor.h"
 #include "SnappingUtils.h"
-#include "Editor/GeometryMode/Public/EditorGeometry.h"
-#include "Editor/GeometryMode/Public/GeometryEdMode.h"
 #include "Logging/MessageLog.h"
 #include "ActorEditorUtils.h"
 #include "Editor/ActorPositioning.h"
 #include "StaticLightingSystem/StaticLightingPrivate.h"
 #include "LightMap.h"
+#include "Subsystems/BrushEditingSubsystem.h"
 
 #define LOCTEXT_NAMESPACE "ClickHandlers"
 
-namespace ClickHandlers
+namespace LevelViewportClickHandlers
 {
 	static void PrivateSummonContextMenu( FLevelEditorViewportClient* ViewportClient )
 	{
@@ -468,286 +467,6 @@ namespace ClickHandlers
 			ViewportClient->Invalidate( true, true );
 		}
 	}
-
-	bool ClickGeomPoly(FLevelEditorViewportClient* ViewportClient, HGeomPolyProxy* InHitProxy, const FViewportClick& Click)
-	{
-		//something is wrong with the hitproxy relating to this click - create a debug log to help identify what
-		if( InHitProxy == NULL )
-		{
-			UE_LOG(LogEditorViewport, Warning, TEXT("Invalid hitproxy"));
-			return false;
-		}
-
-		if( !InHitProxy->GeomObjectWeakPtr.IsValid() )
-		{
-			return false;
-		}
-
-		// Pivot snapping
-		if( Click.GetKey() == EKeys::MiddleMouseButton && Click.IsAltDown() )
-		{
-			GEditor->SetPivot( GEditor->ClickLocation, true, false, true );
-
-			return true;
-		}
-		else if( Click.GetKey() == EKeys::LeftMouseButton && Click.IsControlDown() && Click.IsShiftDown() && !Click.IsAltDown() )
-		{
-			GEditor->SelectActor( InHitProxy->GetGeomObject()->GetActualBrush(), false, true );
-		}
-		else if( Click.GetKey() == EKeys::LeftMouseButton )
-		{
-			// This should only happen in geometry mode
-			FEdMode* Mode = GLevelEditorModeTools().GetActiveMode( FBuiltinEditorModes::EM_Geometry );
-			if( Mode )
-			{
-				if( ( InHitProxy->GetGeomObject() != NULL ) && (InHitProxy->GetGeomObject()->PolyPool.IsValidIndex( InHitProxy->PolyIndex ) == true ) )					
-				{
-					Mode->GetCurrentTool()->StartTrans();
-
-					if( !Click.IsControlDown() )
-					{
-						Mode->GetCurrentTool()->SelectNone();
-					}
-
-					FGeomPoly& gp = InHitProxy->GetGeomObject()->PolyPool[ InHitProxy->PolyIndex ];
-					gp.Select( Click.IsControlDown() ? !gp.IsSelected() : true );
-
-					Mode->SelectionChanged();
-
-					Mode->GetCurrentTool()->EndTrans();
-					ViewportClient->Invalidate( true, false );
-				}
-				else
-				{
-					//try to get the name of the object also
-					FString name = TEXT("UNKNOWN");
-					if( InHitProxy->GetGeomObject()->GetActualBrush() != NULL )
-					{
-						name = InHitProxy->GetGeomObject()->GetActualBrush()->GetName();
-					}
-					UE_LOG(LogEditorViewport, Warning, TEXT("Invalid PolyIndex %d on %s" ) ,InHitProxy->PolyIndex , *name );
-				}
-			}
-		}
-
-		return false;
-	}
-
-	/**
-	 * Utility method used by ClickGeomEdge and ClickGeomVertex.  Returns true if the projections of
-	 * the vectors onto the specified viewport plane are equal within the given tolerance.
-	 */
-	bool OrthoEqual(ELevelViewportType ViewportType, const FVector& Vec0, const FVector& Vec1, float Tolerance=0.1f)
-	{
-		bool bResult = false;
-		switch( ViewportType )
-		{
-		case LVT_OrthoXY:
-		case LVT_OrthoNegativeXY:
-			bResult = FMath::Abs(Vec0.X - Vec1.X) < Tolerance && FMath::Abs(Vec0.Y - Vec1.Y) < Tolerance;
-			break;
-		case LVT_OrthoXZ:
-		case LVT_OrthoNegativeXZ:
-			bResult = FMath::Abs(Vec0.X - Vec1.X) < Tolerance && FMath::Abs(Vec0.Z - Vec1.Z) < Tolerance;
-			break;
-		case LVT_OrthoYZ:
-		case LVT_OrthoNegativeYZ:
-			bResult = FMath::Abs(Vec0.Y - Vec1.Y) < Tolerance && FMath::Abs(Vec0.Z - Vec1.Z) < Tolerance;
-			break;
-		default:
-			check( 0 );
-			break;
-		}
-		return bResult;
-	}
-
-	bool ClickGeomEdge(FLevelEditorViewportClient* ViewportClient, HGeomEdgeProxy* InHitProxy, const FViewportClick& Click)
-	{
-		if( !InHitProxy->GetGeomObject() )
-		{
-			return false;
-		}
-
-		// Pivot snapping
-		if( Click.GetKey() == EKeys::MiddleMouseButton && Click.IsAltDown() )
-		{
-			GEditor->SetPivot( GEditor->ClickLocation, true, false, true );
-
-			return true;
-		}
-		else if( Click.GetKey() == EKeys::LeftMouseButton && Click.IsControlDown() && Click.IsShiftDown() && !Click.IsAltDown() )
-		{
-			GEditor->SelectActor( InHitProxy->GetGeomObject()->GetActualBrush(), false, true );
-
-			return true;
-		}
-		else if( Click.GetKey() == EKeys::LeftMouseButton )
-		{
-			FEdMode* Mode = GLevelEditorModeTools().GetActiveMode( FBuiltinEditorModes::EM_Geometry );
-
-			if( Mode )
-			{ 
-				Mode->GetCurrentTool()->StartTrans();
-
-				const bool bControlDown = Click.IsControlDown();
-				if( !bControlDown )
-				{
-					Mode->GetCurrentTool()->SelectNone();
-				}
-
-				FGeomEdge& HitEdge = InHitProxy->GetGeomObject()->EdgePool[ InHitProxy->EdgeIndex ];
-				HitEdge.Select( bControlDown ? !HitEdge.IsSelected() : true );
-
-				if( ViewportClient->IsOrtho() )
-				{
-					// Select all edges in the brush that match the projected mid point of the original edge.
-					for( int32 EdgeIndex = 0 ; EdgeIndex < InHitProxy->GetGeomObject()->EdgePool.Num() ; ++EdgeIndex )
-					{
-						if ( EdgeIndex != InHitProxy->EdgeIndex )
-						{
-							FGeomEdge& GeomEdge = InHitProxy->GetGeomObject()->EdgePool[ EdgeIndex ];
-							if ( OrthoEqual( ViewportClient->ViewportType, GeomEdge.GetMid(), HitEdge.GetMid() ) )
-							{
-								GeomEdge.Select( bControlDown ? !GeomEdge.IsSelected() : true );
-							}
-						}
-					}
-				}
-
-				Mode->SelectionChanged();
-
-				Mode->GetCurrentTool()->EndTrans();
-				ViewportClient->Invalidate( true, true );
-				return true;
-			}
-
-			return false;
-
-
-		}
-
-		return false;
-	}
-
-	bool ClickGeomVertex(FLevelEditorViewportClient* ViewportClient,HGeomVertexProxy* InHitProxy,const FViewportClick& Click)
-	{
-		if(InHitProxy->GetGeomObject() == nullptr)
-		{
-			return false;
-		}
-
-		if( !GLevelEditorModeTools().IsModeActive( FBuiltinEditorModes::EM_Geometry ) )
-		{
-			return false;
-		}
-
-		FEdModeGeometry* Mode = static_cast<FEdModeGeometry*>( GLevelEditorModeTools().GetActiveMode(FBuiltinEditorModes::EM_Geometry) );
-
-		// Note: The expected behavior is that right clicking on a vertex will snap the vertex that was
-		// right-clicked on to the nearest grid point, then move all SELECTED verts by the appropriate
-		// delta.  So, we need to handle the right mouse button click BEFORE we change the selection set below.
-
-		if( Click.GetKey() == EKeys::RightMouseButton )
-		{
-
-			if (InHitProxy->VertexIndex < 0 || InHitProxy->VertexIndex >= InHitProxy->GetGeomObject()->VertexPool.Num() )
-			{
-				UE_LOG(LogEditorViewport, Warning, TEXT("Invalid InHitProxy->VertexIndex" ) );		
-				return false;
-			}
-
-			FModeTool_GeometryModify* Tool = static_cast<FModeTool_GeometryModify*>( Mode->GetCurrentTool() );
-			Tool->StartTrans();
-
-			// Compute out far to move to get back on the grid.
-			const FVector WorldLoc = InHitProxy->GetGeomObject()->GetActualBrush()->ActorToWorld().TransformPosition( InHitProxy->GetGeomObject()->VertexPool[ InHitProxy->VertexIndex ] );
-
-			FVector SnappedLoc = WorldLoc;
-			FSnappingUtils::SnapPointToGrid( SnappedLoc, FVector(GEditor->GetGridSize()) );
-
-			const FVector Delta = SnappedLoc - WorldLoc;
-			GEditor->SetPivot( SnappedLoc, false, false );
-
-			for( int32 VertexIndex = 0 ; VertexIndex < InHitProxy->GetGeomObject()->VertexPool.Num() ; ++VertexIndex )
-			{
-				FGeomVertex& GeomVertex = InHitProxy->GetGeomObject()->VertexPool[VertexIndex];
-				if( GeomVertex.IsSelected() )
-				{
-					GeomVertex += Delta;
-				}
-			}
-
-			Tool->EndTrans();
-			InHitProxy->GetGeomObject()->SendToSource();
-			ViewportClient->Invalidate( true, true );
-
-			// HACK: The Bsp update has to occur after SendToSource() updates the vert pool, putting it outside
-			// of the mode tool's transaction, therefore, the Bsp update requires a transaction of its own
-			{
-				FScopedTransaction Transaction( NSLOCTEXT("UnrealEd", "GeoModeVertexSnap", "Vertex Snap") );
-
-				// Update Bsp
-				GEditor->RebuildAlteredBSP();
-			}
-		}
-
-		if( Click.GetKey() == EKeys::MiddleMouseButton && Click.IsAltDown() )
-		{
-			// Pivot snapping
-
-			GEditor->SetPivot( GEditor->ClickLocation, true, false, true );
-
-			return true;
-		}
-		else if( Click.GetKey() == EKeys::LeftMouseButton && Click.IsControlDown() && Click.IsShiftDown() && !Click.IsAltDown() )
-		{
-			GEditor->SelectActor( InHitProxy->GetGeomObject()->GetActualBrush(), false, true );
-		}
-		else if( Click.GetKey() == EKeys::LeftMouseButton )
-		{
-			Mode->GetCurrentTool()->StartTrans();
-
-			// Disable Ctrl+clicking for selection if selecting with RMB.
-			const bool bControlDown = Click.IsControlDown();
-			if( !bControlDown )
-			{
-				Mode->GetCurrentTool()->SelectNone();
-			}
-
-			FGeomVertex& HitVertex = InHitProxy->GetGeomObject()->VertexPool[ InHitProxy->VertexIndex ];
-			bool bSelect = bControlDown ? !HitVertex.IsSelected() : true;
-
-			HitVertex.Select( bSelect );
-
-			if( ViewportClient->IsOrtho() )
-			{
-				// Select all vertices that project to the same location.
-				for( int32 VertexIndex = 0 ; VertexIndex < InHitProxy->GetGeomObject()->VertexPool.Num() ; ++VertexIndex )
-				{
-					if ( VertexIndex != InHitProxy->VertexIndex )
-					{
-						FGeomVertex& GeomVertex = InHitProxy->GetGeomObject()->VertexPool[VertexIndex];
-						if ( OrthoEqual( ViewportClient->ViewportType, GeomVertex, HitVertex ) )
-						{
-							GeomVertex.Select( bSelect );
-						}
-					}
-				}
-			}
-
-			Mode->SelectionChanged();
-
-			Mode->GetCurrentTool()->EndTrans();
-
-			ViewportClient->Invalidate( true, true );
-
-			return true;
-		}
-
-		return false;
-	}
-
-
 	static FBspSurf GSaveSurf;
 
 	void ClickSurface(FLevelEditorViewportClient* ViewportClient,UModel* Model,int32 iSurf,const FViewportClick& Click)
diff --git a/Engine/Source/Editor/UnrealEd/Private/Subsystems/BrushEditingSubsystem.cpp b/Engine/Source/Editor/UnrealEd/Private/Subsystems/BrushEditingSubsystem.cpp
new file mode 100644
--- /dev/null
+++ b/Engine/Source/Editor/UnrealEd/Private/Subsystems/BrushEditingSubsystem.cpp
@@ -0,0 +1,11 @@
+// Copyright 1998-2019 Epic Games, Inc. All Rights Reserved.
+
+#include "Subsystems/BrushEditingSubsystem.h"
+#include "Editor.h"
+
+
+UBrushEditingSubsystem::UBrushEditingSubsystem()
+	: UEditorSubsystem()
+{
+
+}
diff --git a/Engine/Source/Editor/UnrealEd/Private/UnrealEdSrv.cpp b/Engine/Source/Editor/UnrealEd/Private/UnrealEdSrv.cpp
--- a/Engine/Source/Editor/UnrealEd/Private/UnrealEdSrv.cpp
+++ b/Engine/Source/Editor/UnrealEd/Private/UnrealEdSrv.cpp
@@ -78,7 +78,6 @@
 #include "AssetToolsModule.h"
 #include "IContentBrowserSingleton.h"
 #include "ContentBrowserModule.h"
-#include "Editor/GeometryMode/Public/GeometryEdMode.h"
 #include "AssetRegistryModule.h"
 #include "Matinee/MatineeActor.h"
 #include "MatineeExporter.h"
@@ -103,6 +102,8 @@
 	#include "Windows/WindowsHWrapper.h"
 #endif
 #include "ActorGroupingUtils.h"
+#include "EdMode.h"
+#include "Subsystems/BrushEditingSubsystem.h"
 
 DEFINE_LOG_CATEGORY_STATIC(LogUnrealEdSrv, Log, All);
 
@@ -1774,11 +1775,9 @@ static void MirrorActors(const FVector& MirrorScale)
 		LevelDirtyCallback.Request();
 	}
 
-	if ( GLevelEditorModeTools().IsModeActive( FBuiltinEditorModes::EM_Geometry ) )
+	if (UBrushEditingSubsystem* BrushSubsystem = GEditor->GetEditorSubsystem<UBrushEditingSubsystem>())
 	{
-		// If we are in geometry mode, make sure to update the mode with new source data for selected brushes
-		FEdModeGeometry* Mode = (FEdModeGeometry*)GLevelEditorModeTools().GetActiveMode( FBuiltinEditorModes::EM_Geometry );
-		Mode->GetFromSource();
+		BrushSubsystem->UpdateGeometryFromSelectedBrushes();
 	}
 
 	GEditor->RedrawLevelEditingViewports();
@@ -3003,18 +3002,11 @@ bool UUnrealEdEngine::Exec_Mode( const TCHAR* Str, FOutputDevice& Ar )
 	//
 	FEditorModeID EditorMode = FBuiltinEditorModes::EM_None;
 
-	if 		(FParse::Command(&Str,TEXT("CAMERAMOVE")))		{ EditorMode = FBuiltinEditorModes::EM_Default;		}
-	else if	(FParse::Command(&Str,TEXT("GEOMETRY"))) 		{ EditorMode = FBuiltinEditorModes::EM_Geometry;	}
-	else if	(FParse::Command(&Str,TEXT("TEXTURE"))) 		{ EditorMode = FBuiltinEditorModes::EM_Texture;		}
-	else if (FParse::Command(&Str,TEXT("MESHPAINT")))		{ EditorMode = FBuiltinEditorModes::EM_MeshPaint;	}
-	else if (FParse::Command(&Str,TEXT("LANDSCAPE")))		{ EditorMode = FBuiltinEditorModes::EM_Landscape;	}
-	else if (FParse::Command(&Str,TEXT("FOLIAGE")))			{ EditorMode = FBuiltinEditorModes::EM_Foliage;		}
+	FString CommandToken = FParse::Token(Str, false);
+	FEdMode* FoundMode = GLevelEditorModeTools().GetActiveMode(FName(*CommandToken));
 
-	if ( EditorMode == FBuiltinEditorModes::EM_None )
+	if (FoundMode != NULL)
 	{
-		FString CommandToken = FParse::Token(Str, false);
-		FEdMode* FoundMode = GLevelEditorModeTools().GetActiveMode( FName( *CommandToken ) );
-
 		if ( FoundMode != NULL )
 		{
 			EditorMode = FName( *CommandToken );
diff --git a/Engine/Source/Editor/UnrealEd/Public/EditorModes.h b/Engine/Source/Editor/UnrealEd/Public/EditorModes.h
--- a/Engine/Source/Editor/UnrealEd/Public/EditorModes.h
+++ b/Engine/Source/Editor/UnrealEd/Public/EditorModes.h
@@ -28,18 +28,9 @@ public:
 	/** Placement mode */
 	const static FEditorModeID EM_Placement;
 
-	/** Bsp mode */
-	const static FEditorModeID EM_Bsp;
-
-	/** Geometry editing mode. */
-	const static FEditorModeID EM_Geometry;
-
 	/** Interpolation editing. */
 	const static FEditorModeID EM_InterpEdit;
 
-	/** Texture alignment via the widget. */
-	const static FEditorModeID EM_Texture;
-
 	/** Mesh paint tool */
 	const static FEditorModeID EM_MeshPaint;
 
diff --git a/Engine/Source/Editor/UnrealEd/Public/LevelViewportClickHandlers.h b/Engine/Source/Editor/UnrealEd/Public/LevelViewportClickHandlers.h
new file mode 100644
--- /dev/null
+++ b/Engine/Source/Editor/UnrealEd/Public/LevelViewportClickHandlers.h
@@ -0,0 +1,63 @@
+// Copyright 1998-2019 Epic Games, Inc. All Rights Reserved.
+
+#pragma once
+
+#include "CoreMinimal.h"
+#include "UObject/UObjectGlobals.h"
+#include "HitProxies.h"
+
+class AActor;
+class ABrush;
+class FLevelEditorViewportClient;
+class UModel;
+class USceneComponent;
+struct FViewportClick;
+struct HActor;
+struct HGeomEdgeProxy;
+struct HGeomPolyProxy;
+struct HGeomVertexProxy;
+
+/**
+ * A hit proxy class for sockets in the main editor viewports.
+ */
+struct HLevelSocketProxy : public HHitProxy
+{
+	DECLARE_HIT_PROXY();
+
+	AActor* Actor;
+	USceneComponent* SceneComponent;
+	FName SocketName;
+
+	HLevelSocketProxy(AActor* InActor, USceneComponent* InSceneComponent, FName InSocketName)
+		:	HHitProxy(HPP_UI)
+		,   Actor( InActor )
+		,   SceneComponent( InSceneComponent )
+		,	SocketName( InSocketName )
+	{}
+	virtual void AddReferencedObjects( FReferenceCollector& Collector ) override
+	{
+		Collector.AddReferencedObject( Actor );
+		Collector.AddReferencedObject( SceneComponent );
+	}
+};
+
+namespace LevelViewportClickHandlers
+{
+	bool ClickViewport(FLevelEditorViewportClient* ViewportClient, const FViewportClick& Click);
+
+	bool UNREALED_API ClickActor(FLevelEditorViewportClient* ViewportClient,AActor* Actor,const FViewportClick& Click,bool bAllowSelectionChange);
+
+	bool ClickComponent(FLevelEditorViewportClient* ViewportClient, HActor* ActorHitProxy, const FViewportClick& Click);
+
+	void ClickBrushVertex(FLevelEditorViewportClient* ViewportClient,ABrush* InBrush,FVector* InVertex,const FViewportClick& Click);
+
+	void ClickStaticMeshVertex(FLevelEditorViewportClient* ViewportClient,AActor* InActor,FVector& InVertex,const FViewportClick& Click);
+	
+	void ClickSurface(FLevelEditorViewportClient* ViewportClient, UModel* Model, int32 iSurf, const FViewportClick& Click);
+
+	void ClickBackdrop(FLevelEditorViewportClient* ViewportClient,const FViewportClick& Click);
+
+	void ClickLevelSocket(FLevelEditorViewportClient* ViewportClient, HHitProxy* HitProxy, const FViewportClick& Click);
+};
+
+
diff --git a/Engine/Source/Editor/UnrealEd/Public/Subsystems/BrushEditingSubsystem.h b/Engine/Source/Editor/UnrealEd/Public/Subsystems/BrushEditingSubsystem.h
new file mode 100644
--- /dev/null
+++ b/Engine/Source/Editor/UnrealEd/Public/Subsystems/BrushEditingSubsystem.h
@@ -0,0 +1,36 @@
+// Copyright 1998-2019 Epic Games, Inc. All Rights Reserved.
+
+#pragma once
+
+#include "EditorSubsystem.h"
+
+#include "BrushEditingSubsystem.generated.h"
+
+class FLevelEditorViewportClient;
+class FEditorModeTools;
+class HHitProxy;
+class ABrush;
+struct FViewportClick;
+
+UCLASS(abstract)
+class UNREALED_API UBrushEditingSubsystem : public UEditorSubsystem
+{
+	GENERATED_BODY()
+
+public:
+	UBrushEditingSubsystem();
+
+	virtual bool ProcessClickOnBrushGeometry(FLevelEditorViewportClient* ViewportClient, HHitProxy* InHitProxy, const FViewportClick& Click) { return false; }
+
+
+	virtual void UpdateGeometryFromSelectedBrushes() {}
+	virtual void UpdateGeometryFromBrush(ABrush* Brush) {}
+
+	virtual bool IsGeometryEditorModeActive() const { return false; }
+
+	virtual void DeselectAllEditingGeometry() {};
+
+	virtual bool HandleActorDelete() { return false; };
+private:
+
+};
diff --git a/Engine/Source/Editor/UnrealEd/UnrealEd.Build.cs b/Engine/Source/Editor/UnrealEd/UnrealEd.Build.cs
--- a/Engine/Source/Editor/UnrealEd/UnrealEd.Build.cs
+++ b/Engine/Source/Editor/UnrealEd/UnrealEd.Build.cs
@@ -76,7 +76,6 @@ public class UnrealEd : ModuleRules
 		PublicDependencyModuleNames.AddRange(
 			new string[]
 			{
-				"BspMode",
 				"Core",
 				"CoreUObject",
 				"ApplicationCore",
@@ -220,8 +219,6 @@ public class UnrealEd : ModuleRules
 				"SessionFrontend",
 				"Sequencer",
 				"StringTableEditor",
-				"GeometryMode",
-				"TextureAlignMode",
 				"FoliageEdit",
 				"ImageWrapper",
 				"Blutility",
diff --git a/Engine/Source/Runtime/NavigationSystem/Private/NavigationSystem.cpp b/Engine/Source/Runtime/NavigationSystem/Private/NavigationSystem.cpp
--- a/Engine/Source/Runtime/NavigationSystem/Private/NavigationSystem.cpp
+++ b/Engine/Source/Runtime/NavigationSystem/Private/NavigationSystem.cpp
@@ -32,8 +32,6 @@
 #if WITH_EDITOR
 #include "EditorModeManager.h"
 #include "EditorModes.h"
-#include "Editor/GeometryMode/Public/GeometryEdMode.h"
-#include "Editor/GeometryMode/Public/EditorGeometry.h"
 #include "Editor/LevelEditor/Public/LevelEditor.h"
 #endif
 
@@ -422,14 +420,6 @@ UNavigationSystemV1::UNavigationSystemV1(const FObjectInitializer& ObjectInitial
 		const FTransform RecastToUnrealTransfrom(Recast2UnrealMatrix());
 		SetCoordTransform(ENavigationCoordSystem::Navigation, ENavigationCoordSystem::Unreal, RecastToUnrealTransfrom);
 	}
-
-#if WITH_EDITOR
-	if (GIsEditor && HasAnyFlags(RF_ClassDefaultObject) == false)
-	{
-		FEditorDelegates::EditorModeIDEnter.AddUObject(this, &UNavigationSystemV1::OnEditorModeIDChanged, true);
-		FEditorDelegates::EditorModeIDExit.AddUObject(this, &UNavigationSystemV1::OnEditorModeIDChanged, false);
-	}
-#endif // WITH_EDITOR
 }
 
 PRAGMA_DISABLE_DEPRECATION_WARNINGS
@@ -437,14 +427,6 @@ UNavigationSystemV1::~UNavigationSystemV1()
 {
 	CleanUp(FNavigationSystem::ECleanupMode::CleanupUnsafe);
 
-#if WITH_EDITOR
-	if (GIsEditor)
-	{
-		FEditorDelegates::EditorModeIDEnter.RemoveAll(this);
-		FEditorDelegates::EditorModeIDExit.RemoveAll(this);
-	}
-#endif // WITH_EDITOR
-
 #if !UE_BUILD_SHIPPING
 	FCoreDelegates::OnGetOnScreenMessages.RemoveAll(this);
 #endif // !UE_BUILD_SHIPPING
@@ -2769,47 +2751,6 @@ void UNavigationSystemV1::UpdateLevelCollision(ULevel* InLevel)
 		OnLevelAddedToWorld(InLevel, World);
 	}
 }
-
-void UNavigationSystemV1::OnEditorModeChanged(FEdMode* Mode, bool IsEntering)
-{
-	if (Mode == NULL)
-	{
-		return;
-	}
-
-	if (IsEntering == false && Mode->GetID() == FBuiltinEditorModes::EM_Geometry)
-	{
-		// check if any of modified brushes belongs to an ANavMeshBoundsVolume
-		FEdModeGeometry* GeometryMode = (FEdModeGeometry*)Mode;
-		for (auto GeomObjectIt = GeometryMode->GeomObjectItor(); GeomObjectIt; GeomObjectIt++)
-		{
-			ANavMeshBoundsVolume* Volume = Cast<ANavMeshBoundsVolume>((*GeomObjectIt)->GetActualBrush());
-			if (Volume)
-			{
-				OnNavigationBoundsUpdated(Volume);
-			}
-		}
-	}
-}
-
-void UNavigationSystemV1::OnEditorModeIDChanged(const FEditorModeID& ModeID, bool IsEntering)
-{
-	if (IsEntering == false && ModeID == FBuiltinEditorModes::EM_Geometry)
-	{
-		// check if any of modified brushes belongs to an ANavMeshBoundsVolume
-		FEdMode* Mode = GLevelEditorModeTools().GetActiveMode(ModeID);
-		FEdModeGeometry* GeometryMode = (FEdModeGeometry*)Mode;
-		for (auto GeomObjectIt = GeometryMode->GeomObjectItor(); GeomObjectIt; GeomObjectIt++)
-		{
-			ANavMeshBoundsVolume* Volume = Cast<ANavMeshBoundsVolume>((*GeomObjectIt)->GetActualBrush());
-			if (Volume)
-			{
-				OnNavigationBoundsUpdated(Volume);
-			}
-		}
-	}
-}
-
 #endif
 
 void UNavigationSystemV1::OnNavigationBoundsUpdated(ANavMeshBoundsVolume* NavVolume)
diff --git a/Engine/Source/Runtime/NavigationSystem/Public/NavigationSystem.h b/Engine/Source/Runtime/NavigationSystem/Public/NavigationSystem.h
--- a/Engine/Source/Runtime/NavigationSystem/Public/NavigationSystem.h
+++ b/Engine/Source/Runtime/NavigationSystem/Public/NavigationSystem.h
@@ -800,9 +800,6 @@ public:
 
 	void UpdateLevelCollision(ULevel* InLevel);
 
-	UE_DEPRECATED(4.24, "Use OnEditorModeIDChanged instead")
-	virtual void OnEditorModeChanged(FEdMode* Mode, bool IsEntering);
-	virtual void OnEditorModeIDChanged(const FEditorModeID& ModeID, bool IsEntering);
 #endif // WITH_EDITOR
 
 	FORCEINLINE bool IsSetUpForLazyGeometryExporting() const { return bGenerateNavigationOnlyAroundNavigationInvokers; }
