Fix Static Analysis

diff --git a/Engine/Source/Editor/UnrealEd/Private/UnrealEdSrv.cpp b/Engine/Source/Editor/UnrealEd/Private/UnrealEdSrv.cpp
--- a/Engine/Source/Editor/UnrealEd/Private/UnrealEdSrv.cpp
+++ b/Engine/Source/Editor/UnrealEd/Private/UnrealEdSrv.cpp
@@ -3007,10 +3007,7 @@ bool UUnrealEdEngine::Exec_Mode( const TCHAR* Str, FOutputDevice& Ar )
 
 	if (FoundMode != NULL)
 	{
-		if ( FoundMode != NULL )
-		{
-			EditorMode = FName( *CommandToken );
-		}
+		EditorMode = FName( *CommandToken );
 	}
 
 	if( EditorMode != FBuiltinEditorModes::EM_None )
