Fixed non uobject reference potentially lingering on world preventing unloading when remove a level from the level list or when deleting it.

diff --git a/Engine/Source/Editor/UnrealEd/Private/EditorLevelUtils.cpp b/Engine/Source/Editor/UnrealEd/Private/EditorLevelUtils.cpp
--- a/Engine/Source/Editor/UnrealEd/Private/EditorLevelUtils.cpp
+++ b/Engine/Source/Editor/UnrealEd/Private/EditorLevelUtils.cpp
@@ -807,9 +807,18 @@ bool UEditorLevelUtils::EditorDestroyLevel(ULevel* InLevel)
 {
 	UWorld* World = InLevel->OwningWorld;
 
-	InLevel->GetOuter()->MarkPendingKill();
+	UObject* Outer = InLevel->GetOuter();
+
+	// Call cleanup on the outer world of the level so external hooks can be properly released, so that unloading the package isn't prevented.
+	UWorld* OuterWorld = Cast<UWorld>(Outer);
+	if (OuterWorld && OuterWorld != World)
+	{
+		OuterWorld->CleanupWorld();
+	}
+
+	Outer->MarkPendingKill();
 	InLevel->MarkPendingKill();
-	InLevel->GetOuter()->ClearFlags(RF_Public | RF_Standalone);
+	Outer->ClearFlags(RF_Public | RF_Standalone);
 
 	UPackage* Package = InLevel->GetOutermost();
 	// We want to unconditionally destroy the level, so clear the dirty flag here so it can be unloaded successfully
diff --git a/Engine/Source/Editor/UnrealEd/Private/ObjectTools.cpp b/Engine/Source/Editor/UnrealEd/Private/ObjectTools.cpp
--- a/Engine/Source/Editor/UnrealEd/Private/ObjectTools.cpp
+++ b/Engine/Source/Editor/UnrealEd/Private/ObjectTools.cpp
@@ -2344,6 +2344,13 @@ namespace ObjectTools
 			{
 				MorphTarget->BaseSkelMesh->UnregisterMorphTarget(MorphTarget);
 			}
+
+			// @todo Hack for 4.23.2 since public headers can't be touched.
+			// Worlds get hooked on by a lot of external non-uobject system through GCObject, call World cleanup to fire delegates to tell them to unhook and release reference
+			if (UWorld* World = Cast<UWorld>(ObjectToDelete))
+			{
+				World->CleanupWorld();
+			}
 		}
 
 		if ( bPerformReferenceCheck )
