Force deleting an asset no longer crashes if it's being referenced from a view that is also getting closed.

diff --git a/Engine/Source/Editor/UnrealEd/Private/ObjectTools.cpp b/Engine/Source/Editor/UnrealEd/Private/ObjectTools.cpp
--- a/Engine/Source/Editor/UnrealEd/Private/ObjectTools.cpp
+++ b/Engine/Source/Editor/UnrealEd/Private/ObjectTools.cpp
@@ -2403,7 +2403,7 @@ namespace ObjectTools
 		return true;
 	}
 
-	static void RecursiveRetrieveReferencers(UObject* Object, TSet<UObject*>& ReferencingObjects)
+	static void RecursiveRetrieveReferencers(UObject* Object, TSet<FWeakObjectPtr>& ReferencingObjects)
 	{
 		TArray<FReferencerInformation> ExternalReferencers;
 		Object->RetrieveReferencers(nullptr /* internal refs */, &ExternalReferencers);
@@ -2445,7 +2445,7 @@ namespace ObjectTools
 		}
 
 		// Recursively find all references to objects being deleted
-		TSet<UObject*> ReferencingObjects;
+		TSet<FWeakObjectPtr> ReferencingObjects;
 		for (UObject* ToDelete : InObjectsToDelete)
 		{
 			ReferencingObjects.Add(ToDelete);
@@ -2455,8 +2455,9 @@ namespace ObjectTools
 
 		// Attempt to close all editors referencing any of the deleted objects
 		bool bClosedAllEditors = true;
-		for (UObject* Object : ReferencingObjects)
+		for (const FWeakObjectPtr& ObjectPtr : ReferencingObjects)
 		{
+			UObject* Object = ObjectPtr.Get();
 			if (Object != nullptr && Object->IsAsset())
 			{
 				TArray<IAssetEditorInstance*> ObjectEditors = GEditor->GetEditorSubsystem<UAssetEditorSubsystem>()->FindEditorsForAssetAndSubObjects(Object);
