UE4R - Use ES3.1 for mobile preview

diff --git a/Engine/Source/Editor/UnrealEd/Private/Kismet2/DebuggerCommands.cpp b/Engine/Source/Editor/UnrealEd/Private/Kismet2/DebuggerCommands.cpp
--- a/Engine/Source/Editor/UnrealEd/Private/Kismet2/DebuggerCommands.cpp
+++ b/Engine/Source/Editor/UnrealEd/Private/Kismet2/DebuggerCommands.cpp
@@ -331,7 +331,7 @@ void FPlayWorldCommands::RegisterCommands()
 	UI_COMMAND( PlayInViewport, "Selected Viewport", "Play this level in the active level editor viewport", EUserInterfaceActionType::Check, FInputChord() );
 	UI_COMMAND( PlayInEditorFloating, "New Editor Window (PIE)", "Play this level in a new window", EUserInterfaceActionType::Check, FInputChord() );
 	UI_COMMAND( PlayInVR, "VR Preview", "Play this level in VR", EUserInterfaceActionType::Check, FInputChord() );
-	UI_COMMAND( PlayInMobilePreview, "Mobile Preview ES2 (PIE)", "Play this level as a mobile device preview in ES2 mode (runs in its own process)", EUserInterfaceActionType::Check, FInputChord());
+	UI_COMMAND( PlayInMobilePreview, "Mobile Preview ES3.1 (PIE)", "Play this level as a mobile device preview in ES3.1 mode (runs in its own process)", EUserInterfaceActionType::Check, FInputChord());
 	UI_COMMAND( PlayInVulkanPreview, "Vulkan Mobile Preview (PIE)", "Play this level using mobile Vulkan rendering (runs in its own process)", EUserInterfaceActionType::Check, FInputChord() );
 	UI_COMMAND( PlayInNewProcess, "Standalone Game", "Play this level in a new window that runs in its own process", EUserInterfaceActionType::Check, FInputChord() );
 	UI_COMMAND( PlayInCameraLocation, "Current Camera Location", "Spawn the player at the current camera location", EUserInterfaceActionType::RadioButton, FInputChord() );
diff --git a/Engine/Source/Editor/UnrealEd/Private/PlayLevel.cpp b/Engine/Source/Editor/UnrealEd/Private/PlayLevel.cpp
--- a/Engine/Source/Editor/UnrealEd/Private/PlayLevel.cpp
+++ b/Engine/Source/Editor/UnrealEd/Private/PlayLevel.cpp
@@ -1382,7 +1382,7 @@ void UEditorEngine::PlayStandaloneLocalPc(FString MapNameOverride, FIntPoint* Wi
 		}
 		else
 		{
-			AdditionalParameters += TEXT(" -featureleveles2");
+			AdditionalParameters += TEXT(" -featureleveles31");
 		}
 
 		if (IsOpenGLPlatform(GShaderPlatformForFeatureLevel[GMaxRHIFeatureLevel]))
