Fixed non uobject reference potentially lingering on world preventing unloading when remove a level from the level list or when deleting it.

Added a last resort ReloadPackages if we still can't delete an object after a force replace references

diff --git a/Engine/Source/Editor/UnrealEd/Private/ObjectTools.cpp b/Engine/Source/Editor/UnrealEd/Private/ObjectTools.cpp
--- a/Engine/Source/Editor/UnrealEd/Private/ObjectTools.cpp
+++ b/Engine/Source/Editor/UnrealEd/Private/ObjectTools.cpp
@@ -2345,7 +2345,7 @@ namespace ObjectTools
 				MorphTarget->BaseSkelMesh->UnregisterMorphTarget(MorphTarget);
 			}
 
-			// @todo Hack for 4.23.2 since public headers can't be touched.
+			// @todo FH: Temporary Hack for world to clean up references until `ForceReplaceReferences` can be made consistent with `IsReferenced`
 			// Worlds get hooked on by a lot of external non-uobject system through GCObject, call World cleanup to fire delegates to tell them to unhook and release reference
 			if (UWorld* World = Cast<UWorld>(ObjectToDelete))
 			{
@@ -2663,6 +2663,7 @@ namespace ObjectTools
 			ReloadEditorWorldForReferenceReplacementIfNecessary(ObjectsToDelete);
 		}
 
+		TArray<UPackage*> PackagesToReload;
 		{
 			int32 ReplaceableObjectsNum = 0;
 			{
@@ -2763,10 +2764,9 @@ namespace ObjectTools
 			FAssetToolsModule& AssetToolsModule = FModuleManager::LoadModuleChecked<FAssetToolsModule>(TEXT("AssetTools"));
 
 			int32 Count = 0;
-			for(TWeakObjectPtr<UObject>& Object : ObjectsToDelete)
+			for(auto It = ObjectsToDelete.CreateIterator(); It; ++It)
 			{
-				UObject* CurObject = Object.Get();
-
+				UObject* CurObject = It->Get();
 				if ( !ensure(CurObject != NULL) )
 				{
 					continue;
@@ -2781,6 +2781,13 @@ namespace ObjectTools
 						++NumDeletedObjects;
 					}
 				}
+				// if the delete fails at this point, it means the object won't be able to be purged and might be left in a weird state, as a last resort queue its package for reload
+				else
+				{
+					UE_LOG(LogObjectTools, Warning, TEXT("ForceDeleteObject failed to delete %s, queuing its package for reloading"), *CurObject->GetName());
+					PackagesToReload.AddUnique(CurObject->GetOutermost());
+					It.RemoveCurrent();
+				}
 
 				GWarn->StatusUpdate(Count, ReplaceableObjectsNum, NSLOCTEXT("UnrealEd", "ConsolidateAssetsUpdate_DeletingObjects", "Deleting Assets..."));
 				++Count;
@@ -2803,6 +2810,12 @@ namespace ObjectTools
 		}
 		ObjectsToDelete.Empty();
 
+		// Reload packages of objects we failed to clean as a last resort since they might be left in an unstable state due to the force replace references
+		FText ErrorMessage;
+		ensureMsgf(PackagesToReload.Num() == 0, TEXT("Failed to unload all packages during ForceDeleteObject"));
+		bool bSuccess = PackageTools::ReloadPackages(PackagesToReload, ErrorMessage, UPackageTools::EReloadPackagesInteractionMode::AssumePositive);
+		ensure(bSuccess);
+
 		GWarn->EndSlowTask();
 
 		if (GUnrealEd)
