Fix deperecation warning about PackageTools

diff --git a/Engine/Source/Editor/UnrealEd/Private/ObjectTools.cpp b/Engine/Source/Editor/UnrealEd/Private/ObjectTools.cpp
--- a/Engine/Source/Editor/UnrealEd/Private/ObjectTools.cpp
+++ b/Engine/Source/Editor/UnrealEd/Private/ObjectTools.cpp
@@ -2813,7 +2813,7 @@ namespace ObjectTools
 		// Reload packages of objects we failed to clean as a last resort since they might be left in an unstable state due to the force replace references
 		FText ErrorMessage;
 		ensureMsgf(PackagesToReload.Num() == 0, TEXT("Failed to unload all packages during ForceDeleteObject"));
-		bool bSuccess = PackageTools::ReloadPackages(PackagesToReload, ErrorMessage, UPackageTools::EReloadPackagesInteractionMode::AssumePositive);
+		bool bSuccess = UPackageTools::ReloadPackages(PackagesToReload, ErrorMessage, UPackageTools::EReloadPackagesInteractionMode::AssumePositive);
 		ensure(bSuccess);
 
 		GWarn->EndSlowTask();
