Fixed a component visualizer regression in which components without a UPROPERTY reference in their parent actor were no longer able to be visualized.

diff --git a/Engine/Source/Editor/UnrealEd/Private/ComponentVisualizer.cpp b/Engine/Source/Editor/UnrealEd/Private/ComponentVisualizer.cpp
--- a/Engine/Source/Editor/UnrealEd/Private/ComponentVisualizer.cpp
+++ b/Engine/Source/Editor/UnrealEd/Private/ComponentVisualizer.cpp
@@ -7,7 +7,7 @@
 IMPLEMENT_HIT_PROXY(HComponentVisProxy, HHitProxy);
 
 
-static TTuple<FName, int32> GetActorPropertyNameAndIndexForComponent(const AActor* Actor, const UActorComponent* Component)
+static FPropertyNameAndIndex GetActorPropertyNameAndIndexForComponent(const AActor* Actor, const UActorComponent* Component)
 {
 	if (Actor != nullptr && Component != nullptr)
 	{
@@ -23,7 +23,7 @@ static TTuple<FName, int32> GetActorPropertyNameAndIndexForComponent(const AActo
 				if (Object == Component)
 				{
 					// It does! Return this name
-					return MakeTuple(ObjectProp->GetFName(), Index);
+					return FPropertyNameAndIndex(ObjectProp->GetFName(), Index);
 				}
 			}
 		}
@@ -40,21 +40,24 @@ static TTuple<FName, int32> GetActorPropertyNameAndIndexForComponent(const AActo
 					UObject* Object = InnerProp->GetObjectPropertyValue(ArrayHelper.GetRawPtr(Index));
 					if (Object == Component)
 					{
-						return MakeTuple(ArrayProp->GetFName(), Index);
+						return FPropertyNameAndIndex(ArrayProp->GetFName(), Index);
 					}
 				}
 			}
 		}
 	}
 
-	return MakeTuple<FName, int32>(NAME_None, INDEX_NONE);
+	return FPropertyNameAndIndex();
 }
 
 
 void FComponentPropertyPath::Set(const UActorComponent* Component)
 {
+	// Determine which property on the component's actor owner references the component.
 	AActor* Actor = Component->GetOwner();
 
+	// If such a property were found, build a chain of such properties, recursing up actor parent components,
+	// until we reach the top
 	UChildActorComponent* ParentComponent = Actor->GetParentComponent();
 	if (ParentComponent)
 	{
@@ -64,28 +67,58 @@ void FComponentPropertyPath::Set(const UActorComponent* Component)
 	{
 		// If there are no further parents, store this one (the outermost)
 		ParentOwningActor = Actor;
+
+		// We have successfully arrived at the top of the actor/component tree, and have a valid property chain.
+		// Hence we don't need to cache the current component ptr as a last resort.
+		LastResortComponentPtr = nullptr;
+	}
+
+	// If a last resort component ptr has been set, no need to build the property chain
+	if (LastResortComponentPtr.IsValid())
+	{
+		return;
 	}
 
-	// Add the next property to the chain after the recursion, so they are added outermost-first
-	PropertyChain.Add(GetActorPropertyNameAndIndexForComponent(Actor, Component));
+	FPropertyNameAndIndex PropertyNameAndIndex = GetActorPropertyNameAndIndexForComponent(Actor, Component);
+
+	if (PropertyNameAndIndex.IsValid())
+	{
+		// If we found a property, add it to the chain after the recursion, so they are added outermost-first.
+		PropertyChain.Add(PropertyNameAndIndex);
+	}
+	else
+	{
+		// If no such property were found, we set a "last resort" weak ptr to the component itself and get rid of the property chain.
+		// This is not preferable as we can't recuperate the component if its address changes (e.g. on hot reload or BP reconstruction).
+		// However it is valid to have an owned component without a UPROPERTY reference, so we need to handle this case.
+		LastResortComponentPtr = const_cast<UActorComponent*>(Component);
+		PropertyChain.Empty();
+	}
 }
 
 
 UActorComponent* FComponentPropertyPath::GetComponent() const
 {
+	// If there's a valid "last resort" component ptr, use this. 
+	if (LastResortComponentPtr.IsValid())
+	{
+		return LastResortComponentPtr.Get();
+	}
+
 	UActorComponent* Result = nullptr;
 	const AActor* Actor = ParentOwningActor.Get();
 	if (Actor)
 	{
 		int32 Level = 0;
-		for (TPair<FName, int32> PropertyNameAndIndex : PropertyChain)
+		for (const FPropertyNameAndIndex& PropertyNameAndIndex : PropertyChain)
 		{
-			FName PropertyName = PropertyNameAndIndex.Get<0>();
-			int32 PropertyIndex = PropertyNameAndIndex.Get<1>();
 			Result = nullptr;
 
-			if (PropertyName != NAME_None && PropertyIndex != INDEX_NONE)
+			if (PropertyNameAndIndex.IsValid())
 			{
+				FName PropertyName = PropertyNameAndIndex.Name;
+				int32 PropertyIndex = PropertyNameAndIndex.Index;
+
 				UClass* ActorClass = Actor->GetClass();
 				UProperty* Prop = FindField<UProperty>(ActorClass, PropertyName);
 
@@ -113,7 +146,7 @@ UActorComponent* FComponentPropertyPath::GetComponent() const
 			if (Level < PropertyChain.Num())
 			{
 				UChildActorComponent* ChildActorComponent = Cast<UChildActorComponent>(Result);
-				if (Result == nullptr)
+				if (ChildActorComponent == nullptr)
 				{
 					break;
 				}
@@ -129,17 +162,20 @@ UActorComponent* FComponentPropertyPath::GetComponent() const
 
 bool FComponentPropertyPath::IsValid() const
 {
+	// If there's a valid "last resort" component, this will always be valid.
+	if (LastResortComponentPtr.IsValid())
+	{
+		return true;
+	}
+
 	if (!ParentOwningActor.IsValid())
 	{
 		return false;
 	}
 
-	for (TPair<FName, int32> PropertyNameAndIndex : PropertyChain)
+	for (const FPropertyNameAndIndex& PropertyNameAndIndex : PropertyChain)
 	{
-		FName PropertyName = PropertyNameAndIndex.Get<0>();
-		int32 PropertyIndex = PropertyNameAndIndex.Get<1>();
-
-		if (PropertyName == NAME_None || PropertyIndex == INDEX_NONE)
+		if (!PropertyNameAndIndex.IsValid())
 		{
 			return false;
 		}
@@ -149,12 +185,11 @@ bool FComponentPropertyPath::IsValid() const
 }
 
 
-FComponentVisualizer::FPropertyNameAndIndex FComponentVisualizer::GetComponentPropertyName(const UActorComponent* Component)
+FPropertyNameAndIndex FComponentVisualizer::GetComponentPropertyName(const UActorComponent* Component)
 {
 	if (Component)
 	{
-		TTuple<FName, int32> Result = GetActorPropertyNameAndIndexForComponent(Component->GetOwner(), Component);
-		return FPropertyNameAndIndex(Result.Get<0>(), Result.Get<1>());
+		return GetActorPropertyNameAndIndexForComponent(Component->GetOwner(), const_cast<UActorComponent*>(Component));
 	}
 
 	// Didn't find actor property referencing this component
diff --git a/Engine/Source/Editor/UnrealEd/Public/ComponentVisualizer.h b/Engine/Source/Editor/UnrealEd/Public/ComponentVisualizer.h
--- a/Engine/Source/Editor/UnrealEd/Public/ComponentVisualizer.h
+++ b/Engine/Source/Editor/UnrealEd/Public/ComponentVisualizer.h
@@ -35,6 +35,31 @@ struct HComponentVisProxy : public HHitProxy
 };
 
 
+struct FPropertyNameAndIndex
+{
+	FPropertyNameAndIndex()
+		: Name(NAME_None)
+		, Index(INDEX_NONE)
+	{}
+
+	explicit FPropertyNameAndIndex(FName InName, int32 InIndex = 0)
+		: Name(InName)
+		, Index(InIndex)
+	{}
+
+	bool IsValid() const { return Name != NAME_None && Index != INDEX_NONE; }
+
+	void Clear()
+	{
+		Name = NAME_None;
+		Index = INDEX_NONE;
+	}
+
+	FName Name;
+	int32 Index;
+};
+
+
 /**
  * Describes a chain of properties from the parent actor of a given component, to the component itself.
  */
@@ -49,6 +74,7 @@ public:
 	void Reset()
 	{
 		ParentOwningActor = nullptr;
+		LastResortComponentPtr = nullptr;
 		PropertyChain.Reset();
 	}
 
@@ -67,7 +93,8 @@ private:
 	void Set(const UActorComponent* Component);
 
 	TWeakObjectPtr<AActor> ParentOwningActor;
-	TArray<TTuple<FName, int32>> PropertyChain;
+	TWeakObjectPtr<UActorComponent> LastResortComponentPtr;
+	TArray<FPropertyNameAndIndex> PropertyChain;
 };
 
 
@@ -112,37 +139,15 @@ public:
 	/** */
 	virtual bool IsVisualizingArchetype() const { return false; }
 
-	struct FPropertyNameAndIndex
-	{
-		FPropertyNameAndIndex()
-			: Name(NAME_None)
-			, Index(INDEX_NONE)
-		{}
-
-		explicit FPropertyNameAndIndex(FName InName, int32 InIndex = 0)
-			: Name(InName)
-			, Index(InIndex)
-		{}
-
-		bool IsValid() const { return Name != NAME_None && Index != INDEX_NONE; }
-
-		void Clear()
-		{
-			Name = NAME_None;
-			Index = INDEX_NONE;
-		}
-
-		FName Name;
-		int32 Index;
-	};
-
+	// So deprecated code expecting this as an inner class still works
+	using FPropertyNameAndIndex = ::FPropertyNameAndIndex;
 
 	/** Find the name of the property that points to this component */
-	UE_DEPRECATED(4.25, "Please use the FComponentPropertyPath class to build property name paths for components.")
+	UE_DEPRECATED(4.24, "Please use the FComponentPropertyPath class to build property name paths for components.")
 	static FPropertyNameAndIndex GetComponentPropertyName(const UActorComponent* Component);
 
 	/** Get a component pointer from the property name */
-	UE_DEPRECATED(4.25, "Please use the FComponentPropertyPath::GetComponent() to retrieve a component pointer from a property name path.")
+	UE_DEPRECATED(4.24, "Please use the FComponentPropertyPath::GetComponent() to retrieve a component pointer from a property name path.")
 	static UActorComponent* GetComponentFromPropertyName(const AActor* CompOwner, const FPropertyNameAndIndex& Property);
 
 	/** Notify that a component property has been modified */
