Fix ensure when in ForceDeleteObjects when reloading packages

diff --git a/Engine/Source/Editor/UnrealEd/Private/ObjectTools.cpp b/Engine/Source/Editor/UnrealEd/Private/ObjectTools.cpp
--- a/Engine/Source/Editor/UnrealEd/Private/ObjectTools.cpp
+++ b/Engine/Source/Editor/UnrealEd/Private/ObjectTools.cpp
@@ -2811,10 +2811,13 @@ namespace ObjectTools
 		ObjectsToDelete.Empty();
 
 		// Reload packages of objects we failed to clean as a last resort since they might be left in an unstable state due to the force replace references
-		FText ErrorMessage;
-		ensureMsgf(PackagesToReload.Num() == 0, TEXT("Failed to unload all packages during ForceDeleteObject"));
-		bool bSuccess = UPackageTools::ReloadPackages(PackagesToReload, ErrorMessage, UPackageTools::EReloadPackagesInteractionMode::AssumePositive);
-		ensure(bSuccess);
+		ensureMsgf(PackagesToReload.Num() == 0, TEXT("Failed to unload all packages during ForceDeleteObjects"));
+		if (PackagesToReload.Num() > 0)
+		{
+			FText ErrorMessage;
+			bool bSuccess = UPackageTools::ReloadPackages(PackagesToReload, ErrorMessage, UPackageTools::EReloadPackagesInteractionMode::AssumePositive);
+			ensureMsgf(bSuccess, TEXT("Failed to reload package as a last resort in ForceDeleteObjects"));
+		}
 
 		GWarn->EndSlowTask();
 
