Fixed atmospheric light shortcut control to work when real-time mode is disabled.

diff --git a/Engine/Source/Editor/UnrealEd/Private/LevelEditorViewport.cpp b/Engine/Source/Editor/UnrealEd/Private/LevelEditorViewport.cpp
--- a/Engine/Source/Editor/UnrealEd/Private/LevelEditorViewport.cpp
+++ b/Engine/Source/Editor/UnrealEd/Private/LevelEditorViewport.cpp
@@ -2766,6 +2766,7 @@ bool FLevelEditorViewportClient::InputKey(FViewport* InViewport, int32 Controlle
 			{
 				FText TrackingDescription = FText::Format(LOCTEXT("RotatationShortcut", "Rotate Atmosphere Light {0}"), LightIndex);
 				TrackingTransaction.Begin(TrackingDescription, SelectedSunLight->GetOwner());
+				SetRealtime(true, true); // The first time, save that setting for RestoreRealtime
 			}
 			bCurrentUserControl = true;
 			UserIsControllingAtmosphericLightTimer = 3.0f; // Keep the widget open for a few seconds even when not tweaking the sun light
@@ -2784,9 +2785,10 @@ bool FLevelEditorViewportClient::InputKey(FViewport* InViewport, int32 Controlle
 	}
 	if (bUserIsControllingAtmosphericLight0 || bUserIsControllingAtmosphericLight1)
 	{
-		TrackingTransaction.End();
+		TrackingTransaction.End();					// End undo/redo translation
+		RestoreRealtime(true);						// Restore previous real-time state
 	}
-	bUserIsControllingAtmosphericLight0 = false;
+	bUserIsControllingAtmosphericLight0 = false;	// Disable all atmospheric light controls
 	bUserIsControllingAtmosphericLight1 = false;
 
 	bool bHandled = FEditorViewportClient::InputKey(InViewport,ControllerId,Key,Event,AmountDepressed,bGamepad);
