Error in Rebuild HLOD: DefaultFortStreamingData.ini - file(s) not opened on this client

UResavePackagesCommandlet::CheckoutFile() would fail if a non binary & non exclusive checkout file was checked out by someone else.
In this case, DefaultFortStreamingData.ini is neither a binary file nor an exclusive checkout file (+l), there's no reason the checkout should fail.

diff --git a/Engine/Source/Editor/UnrealEd/Private/Commandlets/ContentCommandlets.cpp b/Engine/Source/Editor/UnrealEd/Private/Commandlets/ContentCommandlets.cpp
--- a/Engine/Source/Editor/UnrealEd/Private/Commandlets/ContentCommandlets.cpp
+++ b/Engine/Source/Editor/UnrealEd/Private/Commandlets/ContentCommandlets.cpp
@@ -1271,21 +1271,28 @@ bool UResavePackagesCommandlet::CheckoutFile(const FString& Filename, bool bAddF
 		return true;
 	}
 
-	//FString FullPath = FPaths::ConvertRelativePathToFull(StreamingLevelPackageFilename);
 	ISourceControlProvider& SourceControlProvider = ISourceControlModule::Get().GetProvider();
 	FSourceControlStatePtr SourceControlState = SourceControlProvider.GetState(*Filename, EStateCacheUsage::ForceUpdate);
 	if (SourceControlState.IsValid())
 	{
-		FString CurrentlyCheckedOutUser;
-		if (SourceControlState->IsCheckedOutOther(&CurrentlyCheckedOutUser))
+		// Already checked out/added this file
+		if (SourceControlState->IsCheckedOut() || SourceControlState->IsAdded())
 		{
-			if (!bIgnoreAlreadyCheckedOut)
-			{
-				UE_LOG(LogContentCommandlet, Error, TEXT("[REPORT] %s level is already checked out by someone else (%s), can not submit!"), *Filename, *CurrentlyCheckedOutUser);
-			}
-			else
+			return true;
+		}
+		else if (!SourceControlState->IsSourceControlled())
+		{
+			if ( bAddFile )
 			{
-				UE_LOG(LogContentCommandlet, Warning, TEXT("[REPORT] %s level is already checked out by someone else (%s), can not submit!"), *Filename, *CurrentlyCheckedOutUser);
+				if (SourceControlProvider.Execute(ISourceControlOperation::Create<FMarkForAdd>(), *Filename) == ECommandResult::Succeeded)
+				{
+					UE_LOG(LogContentCommandlet, Display, TEXT("[REPORT] %s successfully added"), *Filename);
+					return true;
+				}
+				else
+				{
+					UE_LOG(LogContentCommandlet, Error, TEXT("[REPORT] %s could not be added!"), *Filename);
+				}
 			}
 		}
 		else if (!SourceControlState->IsCurrent())
@@ -1299,28 +1306,27 @@ bool UResavePackagesCommandlet::CheckoutFile(const FString& Filename, bool bAddF
 				UE_LOG(LogContentCommandlet, Warning, TEXT("[REPORT] %s is not synced to head, can not submit"), *Filename);
 			}
 		}
-		else if ( SourceControlState->IsSourceControlled() == false )
+		else if (!SourceControlState->CanCheckout())
 		{
-			if ( bAddFile )
+			FString CurrentlyCheckedOutUser;
+			if (SourceControlState->IsCheckedOutOther(&CurrentlyCheckedOutUser))
 			{
-				if (SourceControlProvider.Execute(ISourceControlOperation::Create<FMarkForAdd>(), *Filename) == ECommandResult::Succeeded)
+				if (!bIgnoreAlreadyCheckedOut)
 				{
-					UE_LOG(LogContentCommandlet, Display, TEXT("[REPORT] %s successfully added"), *Filename);
-					return true;
+					UE_LOG(LogContentCommandlet, Error, TEXT("[REPORT] %s level is already checked out by someone else (%s), can not submit!"), *Filename, *CurrentlyCheckedOutUser);
 				}
 				else
 				{
-					UE_LOG(LogContentCommandlet, Error, TEXT("[REPORT] %s could not be added!"), *Filename);
+					UE_LOG(LogContentCommandlet, Warning, TEXT("[REPORT] %s level is already checked out by someone else (%s), can not submit!"), *Filename, *CurrentlyCheckedOutUser);
 				}
 			}
+			else
+			{
+				UE_LOG(LogContentCommandlet, Error, TEXT("[REPORT] Unable to checkout %s, can not submit"), *Filename);
+			}
 		}
 		else 
 		{
-			// already checked out this file
-			if (SourceControlState->IsCheckedOut() || SourceControlState->IsAdded() )
-			{
-				return true;
-			}
 			if (SourceControlProvider.Execute(ISourceControlOperation::Create<FCheckOut>(), *Filename) == ECommandResult::Succeeded)
 			{
 				UE_LOG(LogContentCommandlet, Display, TEXT("[REPORT] %s Checked out successfully"), *Filename);
