Make sure we always have a sparse class data struct, if one is specified in the class, when serializing.
When creating a new asset that has sparse class data initialize the sparse class data and mark it as serializable.

diff --git a/Engine/Source/Editor/UnrealEd/Private/Kismet2/Kismet2.cpp b/Engine/Source/Editor/UnrealEd/Private/Kismet2/Kismet2.cpp
--- a/Engine/Source/Editor/UnrealEd/Private/Kismet2/Kismet2.cpp
+++ b/Engine/Source/Editor/UnrealEd/Private/Kismet2/Kismet2.cpp
@@ -578,6 +578,14 @@ UBlueprint* FKismetEditorUtilities::CreateBlueprint(UClass* ParentClass, UObject
 		}
 	}
 	
+	// Create the sparse class data and set the flag saying it is safe to serialize it
+	if (NewBP)
+	{
+		UBlueprintGeneratedClass* BlueprintGeneratedClass = CastChecked<UBlueprintGeneratedClass>(NewBP->GeneratedClass);
+		void* SparseDataPtr = BlueprintGeneratedClass->GetOrCreateSparseClassData();
+		BlueprintGeneratedClass->bIsSparseClassDataSerializable = SparseDataPtr != nullptr;
+	}
+
 
 	// Report blueprint creation to analytics
 	if (FEngineAnalytics::IsAvailable())
diff --git a/Engine/Source/Runtime/CoreUObject/Private/UObject/Class.cpp b/Engine/Source/Runtime/CoreUObject/Private/UObject/Class.cpp
--- a/Engine/Source/Runtime/CoreUObject/Private/UObject/Class.cpp
+++ b/Engine/Source/Runtime/CoreUObject/Private/UObject/Class.cpp
@@ -3985,10 +3985,8 @@ void UClass::SerializeSparseClassData(UObject* Object, FStructuredArchive::FSlot
 	// tell the archive that it's allowed to load data for transient properties
 	FArchive& UnderlyingArchive = Slot.GetUnderlyingArchive();
 
-	if (UnderlyingArchive.IsSaving())
-	{
-		GetOrCreateSparseClassData();
-	}
+	// make sure we always have sparse class a sparse class data struct to read from/write to
+	GetOrCreateSparseClassData();
 
 	if (((UnderlyingArchive.IsLoading() || UnderlyingArchive.IsSaving()) && !UnderlyingArchive.WantBinaryPropertySerialization()))
 	{
