Added additional information to the asserts in FindBlueprintForNodeChecked and FindBlueprintForGraphChecked

diff --git a/Engine/Source/Editor/UnrealEd/Private/Kismet2/BlueprintEditorUtils.cpp b/Engine/Source/Editor/UnrealEd/Private/Kismet2/BlueprintEditorUtils.cpp
--- a/Engine/Source/Editor/UnrealEd/Private/Kismet2/BlueprintEditorUtils.cpp
+++ b/Engine/Source/Editor/UnrealEd/Private/Kismet2/BlueprintEditorUtils.cpp
@@ -2098,7 +2098,9 @@ UBlueprint* FBlueprintEditorUtils::FindBlueprintForNode(const UEdGraphNode* Node
 // Helper function to get the blueprint that ultimately owns a node.  Cannot fail.
 UBlueprint* FBlueprintEditorUtils::FindBlueprintForNodeChecked(const UEdGraphNode* Node)
 {
-	return FindBlueprintForGraphChecked(Node->GetGraph());
+	UBlueprint* Result = FindBlueprintForNode(Node);
+	checkf(Result, TEXT("FBlueprintEditorUtils::FindBlueprintForNodeChecked(%s) failed to find a Blueprint."), *GetPathNameSafe(Node));
+	return Result;
 }
 
 
@@ -2120,7 +2122,7 @@ UBlueprint* FBlueprintEditorUtils::FindBlueprintForGraph(const UEdGraph* Graph)
 UBlueprint* FBlueprintEditorUtils::FindBlueprintForGraphChecked(const UEdGraph* Graph)
 {
 	UBlueprint* Result = FBlueprintEditorUtils::FindBlueprintForGraph(Graph);
-	check(Result);
+	checkf(Result, TEXT("FBlueprintEditorUtils::FindBlueprintForGraphChecked(%s) failed to find a Blueprint."), *GetPathNameSafe(Graph));
 	return Result;
 }
 
