Fix static analysis warning as NewBP is assumed to be valid repeatedly throughout the function already
#rnx

diff --git a/Engine/Source/Editor/UnrealEd/Private/Kismet2/Kismet2.cpp b/Engine/Source/Editor/UnrealEd/Private/Kismet2/Kismet2.cpp
--- a/Engine/Source/Editor/UnrealEd/Private/Kismet2/Kismet2.cpp
+++ b/Engine/Source/Editor/UnrealEd/Private/Kismet2/Kismet2.cpp
@@ -579,13 +579,9 @@ UBlueprint* FKismetEditorUtilities::CreateBlueprint(UClass* ParentClass, UObject
 	}
 	
 	// Create the sparse class data and set the flag saying it is safe to serialize it
-	if (NewBP)
-	{
-		UBlueprintGeneratedClass* BlueprintGeneratedClass = CastChecked<UBlueprintGeneratedClass>(NewBP->GeneratedClass);
-		void* SparseDataPtr = BlueprintGeneratedClass->GetOrCreateSparseClassData();
-		BlueprintGeneratedClass->bIsSparseClassDataSerializable = SparseDataPtr != nullptr;
-	}
-
+	UBlueprintGeneratedClass* BlueprintGeneratedClass = CastChecked<UBlueprintGeneratedClass>(NewBP->GeneratedClass);
+	void* SparseDataPtr = BlueprintGeneratedClass->GetOrCreateSparseClassData();
+	BlueprintGeneratedClass->bIsSparseClassDataSerializable = SparseDataPtr != nullptr;
 
 	// Report blueprint creation to analytics
 	if (FEngineAnalytics::IsAvailable())
