Adding logging to get more information about a material crash on save.

diff --git a/Engine/Source/Editor/UnrealEd/Private/Toolkits/AssetEditorToolkit.cpp b/Engine/Source/Editor/UnrealEd/Private/Toolkits/AssetEditorToolkit.cpp
--- a/Engine/Source/Editor/UnrealEd/Private/Toolkits/AssetEditorToolkit.cpp
+++ b/Engine/Source/Editor/UnrealEd/Private/Toolkits/AssetEditorToolkit.cpp
@@ -530,9 +530,9 @@ void FAssetEditorToolkit::SaveAsset_Execute()
 
 	TArray<UPackage*> PackagesToSave;
 
-	for (auto Object : ObjectsToSave)
+	for (UObject* Object : ObjectsToSave)
 	{
-		check((Object != nullptr) && Object->IsAsset());
+		checkf(((Object != nullptr) && Object->IsAsset()), TEXT("Invalid object to save: %s"), (Object != nullptr) ? *Object->GetFullName() : TEXT("Null Object"));
 		PackagesToSave.Add(Object->GetOutermost());
 	}
 
