Changing undo behavior to also handle material function changes and more strictly update the parent.

diff --git a/Engine/Source/Editor/UnrealEd/Private/PreviewMaterial.cpp b/Engine/Source/Editor/UnrealEd/Private/PreviewMaterial.cpp
--- a/Engine/Source/Editor/UnrealEd/Private/PreviewMaterial.cpp
+++ b/Engine/Source/Editor/UnrealEd/Private/PreviewMaterial.cpp
@@ -1679,11 +1679,22 @@ void UMaterialEditorInstanceConstant::CopyBasePropertiesFromParent()
 #if WITH_EDITOR
 void UMaterialEditorInstanceConstant::PostEditUndo()
 {
-	FMaterialUpdateContext UpdateContext;
+	if (bIsFunctionPreviewMaterial)
+	{
+		bIsFunctionInstanceDirty = true;
+		ApplySourceFunctionChanges();
+	}
+	else
+	{
+		FMaterialUpdateContext Context;
 
-	UpdateSourceInstanceParent();
+		UpdateSourceInstanceParent();
 
-	UpdateContext.AddMaterialInstance(SourceInstance);
+		Context.AddMaterialInstance(SourceInstance);
+
+		// Fully update static parameters before recreating render state for all components
+		SetSourceInstance(SourceInstance);
+	}
 
 	Super::PostEditUndo();
 }
