Added support for command button on MAC for atmospheric light shortcut.

diff --git a/Engine/Source/Editor/UnrealEd/Private/LevelEditorViewport.cpp b/Engine/Source/Editor/UnrealEd/Private/LevelEditorViewport.cpp
--- a/Engine/Source/Editor/UnrealEd/Private/LevelEditorViewport.cpp
+++ b/Engine/Source/Editor/UnrealEd/Private/LevelEditorViewport.cpp
@@ -2772,13 +2772,15 @@ bool FLevelEditorViewportClient::InputKey(FViewport* InViewport, int32 Controlle
 			UserIsControllingAtmosphericLightTimer = 3.0f; // Keep the widget open for a few seconds even when not tweaking the sun light
 		}
 	};
-	bool bCtrlLPressed = InputState.IsCtrlButtonPressed() && Key == EKeys::L;
-	if (bCtrlLPressed && InputState.IsShiftButtonPressed())
+
+
+	bool bCmdCtrlLPressed = (InputState.IsCommandButtonPressed() || InputState.IsCtrlButtonPressed()) && Key == EKeys::L;
+	if (bCmdCtrlLPressed && InputState.IsShiftButtonPressed())
 	{
 		ProcessAtmosphericLightShortcut(1, bUserIsControllingAtmosphericLight1);
 		return true;
 	}
-	if (bCtrlLPressed)
+	if (bCmdCtrlLPressed)
 	{
 		ProcessAtmosphericLightShortcut(0, bUserIsControllingAtmosphericLight0);
 		return true;
diff --git a/Engine/Source/Editor/UnrealEd/Public/EditorViewportClient.h b/Engine/Source/Editor/UnrealEd/Public/EditorViewportClient.h
--- a/Engine/Source/Editor/UnrealEd/Public/EditorViewportClient.h
+++ b/Engine/Source/Editor/UnrealEd/Public/EditorViewportClient.h
@@ -95,6 +95,7 @@ public:
 	bool IsCtrlButtonEvent() const { return (Key == EKeys::LeftControl || Key == EKeys::RightControl); }
 	bool IsShiftButtonEvent() const { return (Key == EKeys::LeftShift || Key == EKeys::RightShift); }
 	bool IsAltButtonEvent() const { return (Key == EKeys::LeftAlt || Key == EKeys::RightAlt); }
+	bool IsCommandButtonEvent() const { return (Key == EKeys::LeftCommand || Key == EKeys::RightCommand); }
 
 	bool IsLeftMouseButtonPressed() const { return IsButtonPressed( EKeys::LeftMouseButton ); }
 	bool IsMiddleMouseButtonPressed() const { return IsButtonPressed( EKeys::MiddleMouseButton ); }
@@ -108,7 +109,9 @@ public:
 	bool IsAltButtonPressed() const { return !( IsAltButtonEvent() && InputEvent == IE_Released ) && ( IsButtonPressed( EKeys::LeftAlt ) || IsButtonPressed( EKeys::RightAlt ) ); }
 	bool IsShiftButtonPressed() const { return !( IsShiftButtonEvent() && InputEvent == IE_Released ) && ( IsButtonPressed( EKeys::LeftShift ) || IsButtonPressed( EKeys::RightShift ) ); }
 	bool IsCtrlButtonPressed() const { return !( IsCtrlButtonEvent() && InputEvent == IE_Released ) && ( IsButtonPressed( EKeys::LeftControl ) || IsButtonPressed( EKeys::RightControl ) ); }
+	bool IsCommandButtonPressed() const { return !(IsCommandButtonEvent() && InputEvent == IE_Released) &&  (IsButtonPressed( EKeys::LeftCommand ) || IsButtonPressed( EKeys::RightCommand ) ); }
 	bool IsSpaceBarPressed() const { return IsButtonPressed( EKeys::SpaceBar ); }
+
 private:
 	/** Viewport the event was sent to */
 	FViewport* Viewport;
