Made a command arg param "SKIPMAPCHECK" to skip the "check for map errors" pass after any editor build step.

#rnx

diff --git a/Engine/Source/Editor/UnrealEd/Private/EditorBuildUtils.cpp b/Engine/Source/Editor/UnrealEd/Private/EditorBuildUtils.cpp
--- a/Engine/Source/Editor/UnrealEd/Private/EditorBuildUtils.cpp
+++ b/Engine/Source/Editor/UnrealEd/Private/EditorBuildUtils.cpp
@@ -315,7 +315,7 @@ bool FEditorBuildUtils::EditorBuild( UWorld* InWorld, FName Id, const bool bAllo
 		BuildProgressWidget.Pin()->SetBuildType(BuildType);
 	}
 
-	bool bShouldMapCheck = true;
+	bool bShouldMapCheck = !FParse::Param(FCommandLine::Get(), TEXT("SkipMapCheck"));
 	if (Id == FBuildOptions::BuildGeometry)
 	{
 		// We can't set the busy cursor for all windows, because lighting
