Fix Reset to fbx reimport feature
#rnx

diff --git a/Engine/Source/Editor/UnrealEd/Private/Factories/SkeletalMeshImport.cpp b/Engine/Source/Editor/UnrealEd/Private/Factories/SkeletalMeshImport.cpp
--- a/Engine/Source/Editor/UnrealEd/Private/Factories/SkeletalMeshImport.cpp
+++ b/Engine/Source/Editor/UnrealEd/Private/Factories/SkeletalMeshImport.cpp
@@ -746,6 +746,10 @@ void RestoreExistingSkelMeshData(ExistingSkelMeshData* MeshData, USkeletalMesh*
 			//Copy the re ordered materials (this ensure the material array do not change when we re-import)
 			SkeletalMesh->Materials = MaterialOrdered;
 		}
+		else
+		{
+			bMaterialReset = true;
+		}
 	}
 
 	SkeletalMesh->LODSettings = MeshData->ExistingLODSettings;
@@ -964,21 +968,26 @@ void RestoreExistingSkelMeshData(ExistingSkelMeshData* MeshData, USkeletalMesh*
 		if (SafeReimportLODIndex < 1 && SkeletalMesh->GetLODInfoArray().IsValidIndex(SafeReimportLODIndex))
 		{
 			FSkeletalMeshLODInfo& BaseLODInfo = SkeletalMesh->GetLODInfoArray()[SafeReimportLODIndex];
-			//Restore the Base MaterialMap
-			if (SkeletalMesh->GetImportedModel() && SkeletalMesh->GetImportedModel()->LODModels.IsValidIndex(SafeReimportLODIndex))
+			if(bMaterialReset)
+			{
+				//If we reset the material array there is no point keeping the user changes
+				BaseLODInfo.LODMaterialMap.Empty();
+			}
+			else if (SkeletalMesh->GetImportedModel() && SkeletalMesh->GetImportedModel()->LODModels.IsValidIndex(SafeReimportLODIndex))
 			{
+				//Restore the Base MaterialMap
 				FSkeletalMeshLODModel& LODModel = SkeletalMesh->GetImportedModel()->LODModels[SafeReimportLODIndex];
 				for (int32 SectionIndex = 0; SectionIndex < LODModel.Sections.Num(); ++SectionIndex)
 				{
 					int32 MaterialIndex = LODModel.Sections[SectionIndex].MaterialIndex;
-					if (MeshData->ExistingBaseLODInfo.LODMaterialMap.IsValidIndex(MaterialIndex))
+					if (MeshData->ExistingBaseLODInfo.LODMaterialMap.IsValidIndex(SectionIndex))
 					{
-						int32 ExistingLODMaterialIndex = MeshData->ExistingBaseLODInfo.LODMaterialMap[MaterialIndex];
-						while (BaseLODInfo.LODMaterialMap.Num() <= MaterialIndex)
+						int32 ExistingLODMaterialIndex = MeshData->ExistingBaseLODInfo.LODMaterialMap[SectionIndex];
+						while (BaseLODInfo.LODMaterialMap.Num() <= SectionIndex)
 						{
 							BaseLODInfo.LODMaterialMap.Add(INDEX_NONE);
 						}
-						BaseLODInfo.LODMaterialMap[MaterialIndex] = ExistingLODMaterialIndex;
+						BaseLODInfo.LODMaterialMap[SectionIndex] = ExistingLODMaterialIndex;
 					}
 				}
 			}
