Blueprint interface conflicts are now reported to the user as toasts rather than silent log messages

diff --git a/Engine/Source/Editor/UnrealEd/Private/Kismet2/BlueprintEditorUtils.cpp b/Engine/Source/Editor/UnrealEd/Private/Kismet2/BlueprintEditorUtils.cpp
--- a/Engine/Source/Editor/UnrealEd/Private/Kismet2/BlueprintEditorUtils.cpp
+++ b/Engine/Source/Editor/UnrealEd/Private/Kismet2/BlueprintEditorUtils.cpp
@@ -5657,14 +5657,48 @@ void FBlueprintEditorUtils::FindNativizationDependencies(UBlueprint* Blueprint,
 	NativizeDependenciesOut.AddUnique(Blueprint->ParentClass);
 }
 
-/** Shared function for posting notification toasts (utilized by the nativization property system) */
-static void PostNativizationWarning(const FText& Message)
+/** Shared function for posting notification toasts */
+static void ShowNotification(const FText& Message, EMessageSeverity::Type Severity)
 {
-	FNotificationInfo Warning(Message);
-	Warning.ExpireDuration = 5.0f;
-	Warning.bFireAndForget = true;
-	Warning.Image = FCoreStyle::Get().GetBrush(TEXT("MessageLog.Warning"));
-	FSlateNotificationManager::Get().AddNotification(Warning);
+	if(FApp::IsUnattended())
+	{
+		switch(Severity)
+		{
+		case EMessageSeverity::CriticalError:
+		case EMessageSeverity::Error:
+			UE_LOG(LogBlueprint, Error, TEXT("%s"), *Message.ToString());
+			break;
+		case EMessageSeverity::PerformanceWarning:
+		case EMessageSeverity::Warning:
+			UE_LOG(LogBlueprint, Warning, TEXT("%s"), *Message.ToString());
+			break;
+		case EMessageSeverity::Info:
+			UE_LOG(LogBlueprint, Log, TEXT("%s"), *Message.ToString());
+			break;
+		}
+	}
+	else
+	{
+		FNotificationInfo Warning(Message);
+		Warning.ExpireDuration = 5.0f;
+		Warning.bFireAndForget = true;
+		switch(Severity)
+		{
+		case EMessageSeverity::CriticalError:
+		case EMessageSeverity::Error:
+			Warning.Image = FCoreStyle::Get().GetBrush(TEXT("MessageLog.Error"));
+			break;
+		case EMessageSeverity::PerformanceWarning:
+		case EMessageSeverity::Warning:
+			Warning.Image = FCoreStyle::Get().GetBrush(TEXT("MessageLog.Warning"));
+			break;
+		case EMessageSeverity::Info:
+			Warning.Image = FCoreStyle::Get().GetBrush(TEXT("MessageLog.Info"));
+			break;
+		}
+	
+		FSlateNotificationManager::Get().AddNotification(Warning);
+	}
 }
 
 bool FBlueprintEditorUtils::PropagateNativizationSetting(UBlueprint* Blueprint)
@@ -5714,7 +5748,7 @@ bool FBlueprintEditorUtils::PropagateNativizationSetting(UBlueprint* Blueprint)
 			bSettingsChanged |= bAddedDependencies;
 			if (bAddedDependencies)
 			{
-				PostNativizationWarning(LOCTEXT("DependenciesSavedForNativization", "Saved extra (required dependency) Blueprints for nativization."));
+				ShowNotification(LOCTEXT("DependenciesSavedForNativization", "Saved extra (required dependency) Blueprints for nativization."), EMessageSeverity::Warning);
 			}
 		}
 		break;
@@ -5824,12 +5858,13 @@ bool FBlueprintEditorUtils::ImplementNewInterface(UBlueprint* Blueprint, const F
 	{
 		if( Blueprint->ImplementedInterfaces[i].Interface == InterfaceClass )
 		{
-			Blueprint->Message_Warn(
+			ShowNotification(
 				FText::Format(
-					LOCTEXT("InterfaceAlreadyImplementedFmt", "ImplementNewInterface: Blueprint '{0}' already implements the interface called '{1}'"),
-					FText::FromString(Blueprint->GetPathName()),
+					LOCTEXT("InterfaceAlreadyImplementedFmt", "Blueprint '{0}' already implements the interface '{1}'"),
+					FText::FromString(Blueprint->GetName()),
 					FText::FromString(InterfaceClassName.ToString())
-				).ToString()
+				),
+				EMessageSeverity::Warning
 			);
 			return false;
 		}
@@ -5855,13 +5890,14 @@ bool FBlueprintEditorUtils::ImplementNewInterface(UBlueprint* Blueprint, const F
 			{
 				bAllFunctionsAdded = false;
 
-				Blueprint->Message_Error(
+				ShowNotification(
 					FText::Format(
-						LOCTEXT("InterfaceFunctionConflictsFmt", "ImplementNewInterface: Blueprint '{0}' has a function or graph which conflicts with the function {1} in the interface called '{2}'"),
-						FText::FromString(Blueprint->GetPathName()),
+						LOCTEXT("InterfaceFunctionConflictsFmt", "Blueprint '{0}' has a function or graph which conflicts with function '{1}' in interface '{2}'"),
+						FText::FromString(Blueprint->GetName()),
 						FText::FromName(FunctionName),
 						FText::FromName(InterfaceClassName)
-					).ToString()
+					),
+					EMessageSeverity::Error
 				);
 				break;
 			}
@@ -5895,10 +5931,11 @@ bool FBlueprintEditorUtils::ImplementNewInterface(UBlueprint* Blueprint, const F
 			if (InterfaceBlueprint && InterfaceBlueprint->NativizationFlag == EBlueprintNativizationFlag::Disabled)
 			{
 				InterfaceBlueprint->NativizationFlag = EBlueprintNativizationFlag::Dependency;
-				PostNativizationWarning(FText::Format(
+				ShowNotification(FText::Format(
 					LOCTEXT("InterfaceFlaggedForNativization", "{0} flagged for nativization (as a required dependency)."),
 					FText::FromName(InterfaceBlueprint->GetFName())
-					)
+					),
+					EMessageSeverity::Warning
 				);
 			}
 		}
