Ensure the "ResavePackages -BuildHLOD -BuildOptions=Clusters" commandlet is deterministic

diff --git a/Engine/Source/Editor/UnrealEd/Private/HierarchicalLOD.cpp b/Engine/Source/Editor/UnrealEd/Private/HierarchicalLOD.cpp
--- a/Engine/Source/Editor/UnrealEd/Private/HierarchicalLOD.cpp
+++ b/Engine/Source/Editor/UnrealEd/Private/HierarchicalLOD.cpp
@@ -14,6 +14,9 @@
 #include "Misc/MapErrors.h"
 #include "GameFramework/WorldSettings.h"
 
+#include "Algo/Transform.h"
+#include "EditorLevelUtils.h"
+
 #if WITH_EDITOR
 #include "Engine/LODActor.h"
 #include "ObjectTools.h"
@@ -117,7 +120,6 @@ void FHierarchicalLODBuilder::PreviewBuild()
 		// Only build clusters for levels that are visible
 		if (LevelIter->bIsVisible)
 		{
-			LevelIter->MarkPackageDirty();
 			BuildClusters(LevelIter, false);
 		}
 		else
@@ -150,22 +152,27 @@ void FHierarchicalLODBuilder::BuildClusters(ULevel* InLevel, const bool bCreateM
 	// I'm using stack mem within this scope of the function
 	// so we need this
 	FMemMark Mark(FMemStack::Get());
-	
-	DeleteLODActors(InLevel);
-	
-	LODLevelLODActors.AddDefaulted(BuildLODLevelSettings.Num());
-	const int32 NumHLODLevels = BuildLODLevelSettings.Num();
 
-	LODLevelLODActors.Empty();
-	LODLevelLODActors.AddZeroed(NumHLODLevels);
+	for (AActor* Actor : InLevel->Actors)
+	{
+		if (ALODActor* LODActor = Cast<ALODActor>(Actor))
+		{
+			OldLODActors.Add(LODActor);
+		}
+	}
+
+	const int32 NumHLODLevels = BuildLODLevelSettings.Num();
+	LODLevelLODActors.AddDefaulted(NumHLODLevels);
 
 	// only build if it's enabled
 	if (InLevel->GetWorldSettings()->bEnableHierarchicalLODSystem && BuildLODLevelSettings.Num() > 0)
 	{
+		CreateTempLODActorLevel(InLevel);
+
 		if (InLevel->GetWorldSettings()->bGenerateSingleClusterForLevel)
 		{
 			GenerateAsSingleCluster(NumHLODLevels, InLevel, bCreateMeshes);
-		}		
+		}
 		else
 		{
 			for (int32 LODId = 0; LODId < NumHLODLevels; ++LODId)
@@ -192,7 +199,7 @@ void FHierarchicalLODBuilder::BuildClusters(ULevel* InLevel, const bool bCreateM
 									{
 										FLODCluster ActorCluster(Actor);
 										PreviousActorCluster += ActorCluster;
-									}									
+									}
 								}
 
 								// Reassess whether or not objects that were excluded from the previous HLOD level should be included in this one
@@ -230,7 +237,7 @@ void FHierarchicalLODBuilder::BuildClusters(ULevel* InLevel, const bool bCreateM
 							}
 						}
 
-						ALODActor* LODActor = PreviousActorCluster.BuildActor(InLevel, LODId, bCreateMeshes);
+						ALODActor* LODActor = CreateLODActor(PreviousActorCluster, InLevel, LODId, bCreateMeshes);
 						LODLevelLODActors[LODId].Add(LODActor);
 
 						ValidStaticMeshActorsInLevel.RemoveAll([PreviousActorCluster](AActor* InActor) { return PreviousActorCluster.Actors.Contains(InActor); });
@@ -277,6 +284,8 @@ void FHierarchicalLODBuilder::BuildClusters(ULevel* InLevel, const bool bCreateM
 				}
 			}
 		}
+
+		ApplyClusteringChanges(InLevel);
 	}
 	else
 	{
@@ -294,6 +303,114 @@ void FHierarchicalLODBuilder::BuildClusters(ULevel* InLevel, const bool bCreateM
 	Clusters.Shrink();
 }
 
+void FHierarchicalLODBuilder::CreateTempLODActorLevel(ULevel* InLevel)
+{
+	TempLevel = NewObject<ULevel>(GetTransientPackage(), TEXT("TempLODActorLevel"));
+	TempLevel->Initialize(FURL(nullptr));
+	check(TempLevel);
+	TempLevel->AddToRoot();
+	TempLevel->OwningWorld = InLevel->GetWorld();
+	TempLevel->Model = NewObject<UModel>(TempLevel);
+	TempLevel->Model->Initialize(nullptr, true);
+	TempLevel->bIsVisible = true;
+
+	TempLevel->SetFlags(RF_Transactional);
+	TempLevel->Model->SetFlags(RF_Transactional);
+}
+
+// Hash the clusters using the LODLevel & subactors pointers
+// Take into account child ALODActors too.
+int32 HashLODActorForClusterComparison(ALODActor* LODActor)
+{
+	uint32 HashValue = 0;
+
+	HashValue = HashCombine(HashValue, LODActor->LODLevel);
+	HashValue = HashCombine(HashValue, LODActor->SubActors.Num());
+
+	TArray<AActor*> Actors = LODActor->SubActors;
+	Actors.Sort();
+
+	TArray<ALODActor*> ChildLODActors;
+
+	for (AActor* Actor : Actors)
+	{
+		if (ALODActor* ChildLODActor = Cast<ALODActor>(Actor))
+		{
+			ChildLODActors.Add(ChildLODActor);
+		}
+		else
+		{
+			HashValue = HashCombine(HashValue, GetTypeHash(Actor));
+		}
+	}
+
+	TArray<int32> ChildLODActorsHashes;
+	Algo::Transform(ChildLODActors, ChildLODActorsHashes, [](ALODActor* ChildLODActor) { return HashLODActorForClusterComparison(ChildLODActor); });
+	ChildLODActorsHashes.Sort();
+
+	for (int32 ChildLODActorsHash : ChildLODActorsHashes)
+	{
+		HashValue = HashCombine(HashValue, ChildLODActorsHash);
+	}
+
+	return HashValue;
+}
+
+void FHierarchicalLODBuilder::ApplyClusteringChanges(ULevel* InLevel)
+{
+	// Compare the LOD actors we spawned against those in the level
+	bool bChanged = OldLODActors.Num() != NewLODActors.Num();
+	if (!bChanged)
+	{
+		TSet<int32> HashedLODActors;
+		Algo::Transform(OldLODActors, HashedLODActors, [](ALODActor* LODActor) { return HashLODActorForClusterComparison(LODActor); });
+
+		for (ALODActor* LODActor : NewLODActors)
+		{
+			int32 Hash = HashLODActorForClusterComparison(LODActor);
+			if (!HashedLODActors.Contains(Hash))
+			{
+				bChanged = true;
+				break;
+			}
+		}
+	}
+
+	// If clusters changed, delete old LOD actors and move the new ones in the proper level
+	if (bChanged)
+	{
+		InLevel->MarkPackageDirty();
+
+		// Delete all 
+		DeleteLODActors(InLevel);
+
+		for (ALODActor* LODActor : NewLODActors)
+		{
+			// Move the LOD actor from the temp level to the proper level
+			LODActor->Rename(nullptr, InLevel);
+
+			// Reinsert actors properly in the LODActor subactors array
+			// Will also setup LODParentPrimitive for each actors primitive components.
+			TArray<AActor*> SubActors = LODActor->SubActors;
+			LODActor->SubActors.Empty();
+			for (AActor* Actor : SubActors)
+			{
+				LODActor->AddSubActor(Actor);
+			}
+		}
+	}
+
+	// Delete the temporary level
+	TempLevel->ClearLevelComponents();
+	InLevel->GetWorld()->RemoveLevel(TempLevel);
+	TempLevel->OwningWorld = NULL;
+	TempLevel->RemoveFromRoot();
+	TempLevel = nullptr;
+
+	OldLODActors.Empty();
+	NewLODActors.Empty();
+}
+
 void FHierarchicalLODBuilder::GenerateAsSingleCluster(const int32 NumHLODLevels, ULevel* InLevel, const bool bCreateMeshes)
 {
 	Clusters.Empty();
@@ -342,7 +459,7 @@ void FHierarchicalLODBuilder::GenerateAsSingleCluster(const int32 NumHLODLevels,
 
 		if (LevelCluster.IsValid())
 		{
-			ALODActor* LODActor = LevelCluster.BuildActor(InLevel, LODId, bCreateMeshes);
+			ALODActor* LODActor = CreateLODActor(LevelCluster, InLevel, LODId, bCreateMeshes);
 			PreviousLevelActor = LODActor;
 		}
 	}
@@ -566,10 +683,16 @@ bool FHierarchicalLODBuilder::ShouldGenerateCluster(AActor* Actor, const bool bP
 		return false;
 	}
 
-	ALODActor* LODActor = Cast<ALODActor>(Actor);
-	if (bPreviewBuild && LODActor)
+	if (ALODActor* LODActor = Cast<ALODActor>(Actor))
 	{
-		if (LODActor->GetStaticMeshComponent()->GetStaticMesh())
+		if (bPreviewBuild && NewLODActors.Contains(LODActor))
+		{
+			if (LODActor->GetStaticMeshComponent()->GetStaticMesh())
+			{
+				return false;
+			}
+		}
+		else if (OldLODActors.Contains(LODActor))
 		{
 			return false;
 		}
@@ -592,16 +715,6 @@ bool FHierarchicalLODBuilder::ShouldGenerateCluster(AActor* Actor, const bool bP
 	{
 		for (UStaticMeshComponent* Component : Components)
 		{			
-			if (Component->GetLODParentPrimitive())
-			{
-				ALODActor* ParentActor = CastChecked<ALODActor>(Component->GetLODParentPrimitive()->GetOwner());
-				
-				if (ParentActor && bPreviewBuild)
-				{
-					return false;
-				}
-			}
-
 			if (Component->bHiddenInGame)
 			{
 				return false;
@@ -926,7 +1039,6 @@ void FHierarchicalLODBuilder::DeleteLODActors(ULevel* InLevel)
 	IHierarchicalLODUtilities* Utilities = Module.GetUtilities();
 
 	// you still have to delete all objects just in case they had it and didn't want it anymore
-	TArray<UObject*> AssetsToDelete;
 	for (int32 ActorId = InLevel->Actors.Num() - 1; ActorId >= 0; --ActorId)
 	{
 		ALODActor* LodActor = Cast<ALODActor>(InLevel->Actors[ActorId]);
@@ -1067,7 +1179,7 @@ void FHierarchicalLODBuilder::MergeClustersAndBuildActors(ULevel* InLevel, const
 
 					if (Cluster.Actors.Num() >= MinNumActors)
 					{
-						ALODActor* LODActor = Cluster.BuildActor(InLevel, LODIdx, bCreateMeshes);
+						ALODActor* LODActor = CreateLODActor(Cluster, InLevel, LODIdx, bCreateMeshes);
 						if (LODActor)
 						{
 							LODLevelLODActors[LODIdx].Add(LODActor);
@@ -1084,9 +1196,79 @@ void FHierarchicalLODBuilder::MergeClustersAndBuildActors(ULevel* InLevel, const
 							RejectedActorsInLevel.RemoveSingleSwap(RemoveActor, false);
 						}
 					}
-				}				
+				}
+			}
+		}
+	}
+}
+
+ALODActor* FHierarchicalLODBuilder::CreateLODActor(const FLODCluster& InCluster, ULevel* InLevel, const int32 LODIdx, const bool bCreateMeshes)
+{
+	ALODActor* NewActor = nullptr;
+
+	if (InLevel && InLevel->GetWorld())
+	{
+		// create asset using Actors
+		const FHierarchicalSimplification& LODSetup = InLevel->GetWorldSettings()->GetHierarchicalLODSetup()[LODIdx];
+
+		// Retrieve draw distance for current and next LOD level
+		const int32 LODCount = InLevel->GetWorldSettings()->GetNumHierarchicalLODLevels();
+
+		// Where generated assets will be stored
+		FHierarchicalLODUtilitiesModule& Module = FModuleManager::LoadModuleChecked<FHierarchicalLODUtilitiesModule>("HierarchicalLODUtilities");
+		IHierarchicalLODUtilities* Utilities = Module.GetUtilities();
+
+		TArray<UStaticMeshComponent*> AllComponents;
+		for (auto& Actor : InCluster.Actors)
+		{
+			TArray<UStaticMeshComponent*> Components;
+
+			if (Actor->IsA<ALODActor>())
+			{
+				Utilities->ExtractStaticMeshComponentsFromLODActor(Actor, Components);
+			}
+			else
+			{
+				Actor->GetComponents<UStaticMeshComponent>(Components);
+			}
+
+			AllComponents.Append(Components);
+		}
+
+		if (AllComponents.Num())
+		{
+			// Create LOD Actor
+			UWorld* LevelWorld = Cast<UWorld>(InLevel->GetOuter());
+			check(LevelWorld);
+
+			FTransform Transform;
+
+			FActorSpawnParameters ActorSpawnParams;
+			ActorSpawnParams.OverrideLevel = TempLevel;
+
+			NewActor = LevelWorld->SpawnActor<ALODActor>(ALODActor::StaticClass(), Transform, ActorSpawnParams);
+			NewLODActors.Add(NewActor);
+			NewActor->LODLevel = LODIdx + 1;
+			NewActor->CachedNumHLODLevels = InLevel->GetWorldSettings()->GetNumHierarchicalLODLevels();
+			NewActor->SetDrawDistance(0.0f);
+
+			// now set as parent
+			for (auto& Actor : InCluster.Actors)
+			{
+				NewActor->SubActors.Add(Actor);
+			}
+
+			if (bCreateMeshes)
+			{
+				UHLODProxy* Proxy = Utilities->CreateOrRetrieveLevelHLODProxy(InLevel, LODIdx);
+				UPackage* AssetsOuter = Proxy->GetOutermost();
+				checkf(AssetsOuter != nullptr, TEXT("Failed to created outer for generated HLOD assets"));
+				Utilities->BuildStaticMeshForLODActor(NewActor, AssetsOuter, LODSetup);
 			}
+			NewActor->PostEditChange();
 		}
 	}
+
+	return NewActor;
 }
 #undef LOCTEXT_NAMESPACE 
diff --git a/Engine/Source/Editor/UnrealEd/Private/LODCluster.cpp b/Engine/Source/Editor/UnrealEd/Private/LODCluster.cpp
--- a/Engine/Source/Editor/UnrealEd/Private/LODCluster.cpp
+++ b/Engine/Source/Editor/UnrealEd/Private/LODCluster.cpp
@@ -258,75 +258,6 @@ void FLODCluster::SubtractCluster(const FLODCluster& Other)
 	}
 }
 
-
-ALODActor* FLODCluster::BuildActor(ULevel* InLevel, const int32 LODIdx, const bool bCreateMeshes)
-{
-	ALODActor* NewActor = nullptr;
-
-	if (InLevel && InLevel->GetWorld())
-	{
-		// create asset using Actors
-		const FHierarchicalSimplification& LODSetup = InLevel->GetWorldSettings()->GetHierarchicalLODSetup()[LODIdx];
-
-		// Retrieve draw distance for current and next LOD level
-		const int32 LODCount = InLevel->GetWorldSettings()->GetNumHierarchicalLODLevels();
-
-		// Where generated assets will be stored
-		FHierarchicalLODUtilitiesModule& Module = FModuleManager::LoadModuleChecked<FHierarchicalLODUtilitiesModule>("HierarchicalLODUtilities");
-		IHierarchicalLODUtilities* Utilities = Module.GetUtilities();
- 	
-		TArray<UStaticMeshComponent*> AllComponents;
-		for (auto& Actor : Actors)
-		{
-			TArray<UStaticMeshComponent*> Components;
-
-			if (Actor->IsA<ALODActor>())
-			{
-				Utilities->ExtractStaticMeshComponentsFromLODActor(Actor, Components);
-			}
-			else
-			{
-				Actor->GetComponents<UStaticMeshComponent>(Components);
-			}
-
-			AllComponents.Append(Components);
-		}
-
-		if (AllComponents.Num())
-		{
-			// Create LOD Actor
-			UWorld* LevelWorld = Cast<UWorld>(InLevel->GetOuter());
-			check(LevelWorld);
-
-			FTransform Transform;
-			NewActor = LevelWorld->SpawnActor<ALODActor>(ALODActor::StaticClass(), Transform);
-			NewActor->LODLevel = LODIdx + 1;
-			NewActor->CachedNumHLODLevels = InLevel->GetWorldSettings()->GetNumHierarchicalLODLevels();
-			NewActor->SetDrawDistance(0.0f);
-
-			// now set as parent
-			for (auto& Actor : Actors)
-			{
-				NewActor->AddSubActor(Actor);
-			}
-
-			if (bCreateMeshes)
-			{
-				UHLODProxy* Proxy = Utilities->CreateOrRetrieveLevelHLODProxy(InLevel, LODIdx);
-				UPackage* AssetsOuter = Proxy->GetOutermost();
-				checkf(AssetsOuter != nullptr, TEXT("Failed to created outer for generated HLOD assets"));
-				Utilities->BuildStaticMeshForLODActor(NewActor, AssetsOuter, LODSetup);
-			}
-			NewActor->PostEditChange();
-		}
-
-		
-	}
-
-	return NewActor;
-}
-
-
 bool FLODCluster::Contains(FLODCluster& Other) const
 {
 	if (IsValid() && Other.IsValid())
diff --git a/Engine/Source/Editor/UnrealEd/Public/HierarchicalLOD.h b/Engine/Source/Editor/UnrealEd/Public/HierarchicalLOD.h
--- a/Engine/Source/Editor/UnrealEd/Public/HierarchicalLOD.h
+++ b/Engine/Source/Editor/UnrealEd/Public/HierarchicalLOD.h
@@ -166,6 +166,19 @@ private:
 	*/
 	bool ShouldGenerateCluster(AActor* Actor, const bool bPreviewBuild, const int32 HLODLevelIndex);
 	
+	/**
+	* if criteria matches, creates new LODActor and replace current Actors with that. We don't need
+	* this clears previous actors and sets to this new actor
+	* this is required when new LOD is created from these actors, this will be replaced
+	* to save memory and to reduce memory increase during this process, we discard previous actors and replace with this actor
+	*
+	* @param InLevel - Level for which currently the HLODs are being build
+	* @param LODIdx - LOD index to build
+	* @param bCreateMeshes - Whether or not the new StaticMeshes should be created
+	* @return ALODActor*
+	*/
+	ALODActor* CreateLODActor(const FLODCluster& InCluster, ULevel* InLevel, const int32 LODIdx, const bool bCreateMeshes);
+
 	/**
 	* Deletes LOD actors from the world	
 	*
@@ -175,6 +188,17 @@ private:
 	*/
 	void DeleteLODActors(ULevel* InLevel);
 
+	/**
+	 * Create a temporary level in which we'll spawn newly created LODActor.
+	 * This is to avoid dirtying the main level when no changes are detected.
+	 */
+	void CreateTempLODActorLevel(ULevel* InLevel);
+
+	/**
+	 * Delete the temporary LODActor level.
+	 */
+	void ApplyClusteringChanges(ULevel* InLevel);
+
 	/** Array of LOD Clusters - this is only valid within scope since mem stack allocator */
 	TArray<FLODCluster, TMemStackAllocator<>> Clusters;
 
@@ -194,4 +218,13 @@ private:
 	TArray<AActor*> ValidStaticMeshActorsInLevel;
 	/** Actors which were rejected from the previous HLOD level(s) */
 	TArray<AActor*> RejectedActorsInLevel;
+
+	/** Temporary LODActor levels */
+	ULevel* TempLevel;
+
+	/** Previous LODActors found in level */
+	TArray<ALODActor*> OldLODActors;
+
+	/** Newly spawned LODActors from cluster(s) rebuilding */
+	TArray<ALODActor*> NewLODActors;
 };
diff --git a/Engine/Source/Editor/UnrealEd/Public/LODCluster.h b/Engine/Source/Editor/UnrealEd/Public/LODCluster.h
--- a/Engine/Source/Editor/UnrealEd/Public/LODCluster.h
+++ b/Engine/Source/Editor/UnrealEd/Public/LODCluster.h
@@ -64,28 +64,7 @@ struct FLODCluster
 	/** Cached cluster cost, FMath::Pow(Bound.W, 3) / FillingFactor */
 	float ClusterCost;
 
-	/**
-	* if criteria matches, build new LODActor and replace current Actors with that. We don't need
-	* this clears previous actors and sets to this new actor
-	* this is required when new LOD is created from these actors, this will be replaced
-	* to save memory and to reduce memory increase during this process, we discard previous actors and replace with this actor
-	*
-	* @param InLevel - Level for which currently the HLODs are being build
-	* @param LODIdx - LOD index to build
-	* @param bCreateMeshes - Whether or not the new StaticMeshes should be created
-	* @return ALODActor*
-	*/
-	ALODActor* BuildActor(ULevel* InLevel, const int32 LODIdx, const bool bCreateMeshes);
-
-	/**
-	* Recursively retrieves StaticMeshComponents from a LODActor and its child LODActors 
-	*
-	* @param Actor - LODActor instance
-	* @param InOutComponents - Will hold the StaticMeshComponents	
-	*/
-	void ExtractStaticMeshComponentsFromLODActor(AActor* Actor, TArray<UStaticMeshComponent*>& InOutComponents);
-
-private:	
+private:
 	/**
 	* Merges this cluster with Other by combining the actor arrays and updating the bounds, filling factor and cluster cost
 	*
