Do not resave _BuildData.assets in the HLOD rebuild commandlet as they aren't dirty

diff --git a/Engine/Source/Editor/UnrealEd/Private/Commandlets/ContentCommandlets.cpp b/Engine/Source/Editor/UnrealEd/Private/Commandlets/ContentCommandlets.cpp
--- a/Engine/Source/Editor/UnrealEd/Private/Commandlets/ContentCommandlets.cpp
+++ b/Engine/Source/Editor/UnrealEd/Private/Commandlets/ContentCommandlets.cpp
@@ -1726,7 +1726,7 @@ void UResavePackagesCommandlet::PerformAdditionalOperations(class UWorld* World,
 				if (InLevel && InLevel->MapBuildData && (bShouldBuildLighting || bShouldBuildHLOD || bShouldBuildReflectionCaptures) )
 				{
 					UPackage* MapBuildDataPackage = InLevel->MapBuildData->GetOutermost();
-					if (MapBuildDataPackage != InLevel->GetOutermost())
+					if (MapBuildDataPackage != InLevel->GetOutermost() && MapBuildDataPackage->IsDirty())
 					{
 						CheckoutAndSavePackage(MapBuildDataPackage, CheckedOutPackagesFilenames);
 					}
