Fix for crash when geometry plugin is disabled.

diff --git a/Engine/Source/Editor/UnrealEd/Private/EditorSelectUtils.cpp b/Engine/Source/Editor/UnrealEd/Private/EditorSelectUtils.cpp
--- a/Engine/Source/Editor/UnrealEd/Private/EditorSelectUtils.cpp
+++ b/Engine/Source/Editor/UnrealEd/Private/EditorSelectUtils.cpp
@@ -354,11 +354,12 @@ void UUnrealEdEngine::UpdatePivotLocationForSelection( bool bOnChange )
 		SetPivot(SingleComponent->GetComponentLocation(), false, true);
 	}
 	else if( SingleActor != NULL ) 
-	{		
+	{
 		UBrushEditingSubsystem* BrushSubsystem = GEditor->GetEditorSubsystem<UBrushEditingSubsystem>();
+		const bool bGeometryMode = BrushSubsystem ? BrushSubsystem->IsGeometryEditorModeActive() : false;
 
 		// For geometry mode use current pivot location as it's set to selected face, not actor
-		if(BrushSubsystem->IsGeometryEditorModeActive() == false || bOnChange == true )
+		if (!bGeometryMode || bOnChange == true)
 		{
 			// Set pivot point to the actor's location, accounting for any set pivot offset
 			FVector PivotPoint = SingleActor->GetTransform().TransformPosition(SingleActor->GetPivotOffset());
