Crash fix for moving engine basicshapes and then trying to place them with tthe placement browser.

Notes:
You nd up in this code with a redirector! so you can't CastChecked()
Ideally this wouldn't be possible, but unfortunately that is a much bigger refactor.
You can chase the redirector here and it causes this to be functional at first, BUT when you restart the editor this no longer works because the initial load chases the redirector then the CanCreateActorFrom() fails all together.

diff --git a/Engine/Source/Editor/UnrealEd/Private/Factories/ActorFactory.cpp b/Engine/Source/Editor/UnrealEd/Private/Factories/ActorFactory.cpp
--- a/Engine/Source/Editor/UnrealEd/Private/Factories/ActorFactory.cpp
+++ b/Engine/Source/Editor/UnrealEd/Private/Factories/ActorFactory.cpp
@@ -438,21 +438,26 @@ void UActorFactoryBasicShape::PostSpawnActor(UObject* Asset, AActor* NewActor)
 {
 	Super::PostSpawnActor(Asset, NewActor);
 
-	// Change properties
-	UStaticMesh* StaticMesh = CastChecked<UStaticMesh>(Asset);
-
-	AStaticMeshActor* StaticMeshActor = CastChecked<AStaticMeshActor>(NewActor);
-	UStaticMeshComponent* StaticMeshComponent = StaticMeshActor->GetStaticMeshComponent();
-
-	if( StaticMeshComponent )
+	// HACK 4.24 crash fix
+	// You CAN end up in this code with a redirector! so you can't CastChecked()
+	// Ideally this wouldn't be possible, but unfortunately that is a much bigger refactor.
+	// You can chase the redirector here and it causes this to be functional at first, BUT when you restart the editor this no longer works because the initial load chases the redirector then the CanCreateActorFrom() fails all together.
+	if (UStaticMesh* StaticMesh = Cast<UStaticMesh>(Asset))
 	{
-		StaticMeshComponent->UnregisterComponent();
+		// Change properties
+		AStaticMeshActor* StaticMeshActor = CastChecked<AStaticMeshActor>(NewActor);
+		UStaticMeshComponent* StaticMeshComponent = StaticMeshActor->GetStaticMeshComponent();
 
-		StaticMeshComponent->SetStaticMesh(StaticMesh);
-		StaticMeshComponent->StaticMeshDerivedDataKey = StaticMesh->RenderData->DerivedDataKey;
-		StaticMeshComponent->SetMaterial(0, LoadObject<UMaterial>(nullptr, TEXT("/Engine/BasicShapes/BasicShapeMaterial.BasicShapeMaterial")));
-		// Init Component
-		StaticMeshComponent->RegisterComponent();
+		if (StaticMeshComponent)
+		{
+			StaticMeshComponent->UnregisterComponent();
+
+			StaticMeshComponent->SetStaticMesh(StaticMesh);
+			StaticMeshComponent->StaticMeshDerivedDataKey = StaticMesh->RenderData->DerivedDataKey;
+			StaticMeshComponent->SetMaterial(0, LoadObject<UMaterial>(nullptr, TEXT("/Engine/BasicShapes/BasicShapeMaterial.BasicShapeMaterial")));
+			// Init Component
+			StaticMeshComponent->RegisterComponent();
+		}
 	}
 }
 
