Include the Windows 10 SDK when installing new versions of VS2017.

diff --git a/Engine/Source/Editor/UnrealEd/Private/SourceCodeNavigation.cpp b/Engine/Source/Editor/UnrealEd/Private/SourceCodeNavigation.cpp
--- a/Engine/Source/Editor/UnrealEd/Private/SourceCodeNavigation.cpp
+++ b/Engine/Source/Editor/UnrealEd/Private/SourceCodeNavigation.cpp
@@ -2120,7 +2120,7 @@ FString FSourceCodeNavigationImpl::GetSuggestedIDEInstallerFileName()
 void FSourceCodeNavigationImpl::LaunchIDEInstaller(const FString& Filepath)
 {
 #if PLATFORM_WINDOWS
-	auto Params = TEXT("--productId \"Microsoft.VisualStudio.Product.Community\" --add \"Microsoft.VisualStudio.Workload.NativeGame\" --add \"Component.Unreal\" --campaign \"EpicGames_UE4\"");
+	auto Params = TEXT("--productId \"Microsoft.VisualStudio.Product.Community\" --add \"Microsoft.VisualStudio.Workload.NativeGame\" --add \"Component.Unreal\" --add \"Microsoft.VisualStudio.Component.Windows10SDK.17763\" --campaign \"EpicGames_UE4\"");
 	FPlatformProcess::ExecElevatedProcess(*Filepath, Params, nullptr);
 #endif
 }
