Add support for export of runtime virtual texture streaming mips to BMP

diff --git a/Engine/Source/Editor/UnrealEd/Classes/Exporters/SoundSurroundExporterWAV.h b/Engine/Source/Editor/UnrealEd/Classes/Exporters/SoundSurroundExporterWAV.h
--- a/Engine/Source/Editor/UnrealEd/Classes/Exporters/SoundSurroundExporterWAV.h
+++ b/Engine/Source/Editor/UnrealEd/Classes/Exporters/SoundSurroundExporterWAV.h
@@ -22,8 +22,8 @@ public:
 	USoundSurroundExporterWAV(const FObjectInitializer& ObjectInitializer = FObjectInitializer());
 
 	//~ Begin UExporter Interface
-	virtual int32 GetFileCount( void ) const override;
-	virtual FString GetUniqueFilename( const TCHAR* Filename, int32 FileIndex ) override;
+	virtual int32 GetFileCount( UObject* Object ) const override;
+	virtual FString GetUniqueFilename( const TCHAR* Filename, int32 FileIndex, int32 FileCount ) override;
 	virtual bool ExportBinary( UObject* Object, const TCHAR* Type, FArchive& Ar, FFeedbackContext* Warn, int32 FileIndex = 0, uint32 PortFlags=0 ) override;
 	virtual bool SupportsObject(UObject* Object) const override;
 	//~ End UExporter Interface
diff --git a/Engine/Source/Editor/UnrealEd/Classes/Exporters/TextureExporterBMP.h b/Engine/Source/Editor/UnrealEd/Classes/Exporters/TextureExporterBMP.h
--- a/Engine/Source/Editor/UnrealEd/Classes/Exporters/TextureExporterBMP.h
+++ b/Engine/Source/Editor/UnrealEd/Classes/Exporters/TextureExporterBMP.h
@@ -17,8 +17,16 @@ class UTextureExporterBMP : public UExporter
 	GENERATED_UCLASS_BODY()
 
 	virtual bool SupportsObject(UObject* Object) const override;
-	virtual bool ExportBinary( UObject* Object, const TCHAR* Type, FArchive& Ar, FFeedbackContext* Warn, int32 FileIndex = 0, uint32 PortFlags=0 ) override;
-};
-
+	virtual int32 GetFileCount(UObject* Object) const override;
+	virtual FString GetUniqueFilename(const TCHAR* Filename, int32 FileIndex, int32 FileCount) override;
+	virtual bool ExportBinary(UObject* Object, const TCHAR* Type, FArchive& Ar, FFeedbackContext* Warn, int32 FileIndex = 0, uint32 PortFlags = 0) override;
 
+protected:
+	UTexture2D* GetExportTexture(UObject* Object) const;
+};
 
+UCLASS()
+class URuntimeVirtualTextureExporterBMP : public UTextureExporterBMP
+{
+	GENERATED_UCLASS_BODY()
+};
diff --git a/Engine/Source/Editor/UnrealEd/Private/EditorExporters.cpp b/Engine/Source/Editor/UnrealEd/Private/EditorExporters.cpp
--- a/Engine/Source/Editor/UnrealEd/Private/EditorExporters.cpp
+++ b/Engine/Source/Editor/UnrealEd/Private/EditorExporters.cpp
@@ -217,12 +217,12 @@ bool USoundSurroundExporterWAV::SupportsObject(UObject* Object) const
 	return bSupportsObject;
 }
 
-int32 USoundSurroundExporterWAV::GetFileCount( void ) const
+int32 USoundSurroundExporterWAV::GetFileCount( UObject* Object ) const
 {
 	return( SPEAKER_Count );
 }
 
-FString USoundSurroundExporterWAV::GetUniqueFilename( const TCHAR* Filename, int32 FileIndex )
+FString USoundSurroundExporterWAV::GetUniqueFilename( const TCHAR* Filename, int32 FileIndex, int32 FileCount )
 {
 	static FString SpeakerLocations[SPEAKER_Count] =
 	{
diff --git a/Engine/Source/Editor/UnrealEd/Private/Factories/EditorFactories.cpp b/Engine/Source/Editor/UnrealEd/Private/Factories/EditorFactories.cpp
--- a/Engine/Source/Editor/UnrealEd/Private/Factories/EditorFactories.cpp
+++ b/Engine/Source/Editor/UnrealEd/Private/Factories/EditorFactories.cpp
@@ -214,6 +214,9 @@
 
 #include "Animation/DebugSkelMeshComponent.h"
 
+#include "VT/RuntimeVirtualTexture.h"
+#include "VT/RuntimeVirtualTextureStreamingProxy.h"
+
 #if PLATFORM_WINDOWS
 	// Needed for DDS support.
 	#include "Windows/WindowsHWrapper.h"
@@ -4741,32 +4744,54 @@ UTextureExporterBMP::UTextureExporterBMP(const FObjectInitializer& ObjectInitial
 
 }
 
+UTexture2D* UTextureExporterBMP::GetExportTexture(UObject* Object) const
+{
+	URuntimeVirtualTexture* RuntimeVirtualTexture = Cast<URuntimeVirtualTexture>(Object);
+	if (RuntimeVirtualTexture != nullptr && RuntimeVirtualTexture->GetStreamingTexture() != nullptr)
+	{
+		return RuntimeVirtualTexture->GetStreamingTexture();
+	}
+	return Cast<UTexture2D>(Object);
+}
+
 bool UTextureExporterBMP::SupportsObject(UObject* Object) const
 {
 	bool bSupportsObject = false;
 	if (Super::SupportsObject(Object))
 	{
-		UTexture2D* Texture = Cast<UTexture2D>(Object);
-
-		if (Texture)
-		{
-			bSupportsObject = Texture->Source.GetFormat() == TSF_BGRA8 || Texture->Source.GetFormat() == TSF_RGBA16 || Texture->Source.GetFormat() == TSF_G8;
-		}
+		UTexture2D* Texture = GetExportTexture(Object);
+		bSupportsObject = Texture != nullptr && (Texture->Source.GetFormat() == TSF_BGRA8 || Texture->Source.GetFormat() == TSF_RGBA16 || Texture->Source.GetFormat() == TSF_G8);
 	}
 	return bSupportsObject;
 }
 
+int32 UTextureExporterBMP::GetFileCount(UObject* Object) const
+{
+	UTexture2D* Texture = GetExportTexture(Object);
+	return (Texture != nullptr) ? Texture->Source.GetNumLayers() : 1;
+}
+
+FString UTextureExporterBMP::GetUniqueFilename(const TCHAR* Filename, int32 FileIndex, int32 FileCount)
+{
+	return (FileCount == 1) ? Filename : FString::Printf(TEXT("%s%d%s"), *FPaths::GetBaseFilename(Filename, false), FileIndex, *FPaths::GetExtension(Filename, true));
+}
+
 bool UTextureExporterBMP::ExportBinary( UObject* Object, const TCHAR* Type, FArchive& Ar, FFeedbackContext* Warn, int32 FileIndex, uint32 PortFlags )
 {
-	UTexture2D* Texture = CastChecked<UTexture2D>( Object );
+	UTexture2D* Texture = GetExportTexture(Object);
+	if(Texture == nullptr)
+	{
+		return false;
+	}
 
-	if( !Texture->Source.IsValid() || ( Texture->Source.GetFormat() != TSF_BGRA8 && Texture->Source.GetFormat() != TSF_RGBA16  && Texture->Source.GetFormat() != TSF_G8) )
+	const int32 LayerIndex = FileIndex;
+	if( !Texture->Source.IsValid() || ( Texture->Source.GetFormat(LayerIndex) != TSF_BGRA8 && Texture->Source.GetFormat(LayerIndex) != TSF_RGBA16 && Texture->Source.GetFormat() != TSF_G8) )
 	{
 		return false;
 	}
 
-	const bool bIsRGBA16 = Texture->Source.GetFormat() == TSF_RGBA16;
-	const bool bIsG8 = Texture->Source.GetFormat() == TSF_G8;
+	const bool bIsRGBA16 = Texture->Source.GetFormat(LayerIndex) == TSF_RGBA16;
+	const bool bIsG8 = Texture->Source.GetFormat(LayerIndex) == TSF_G8;
 	const int32 SourceBytesPerPixel = bIsRGBA16 ? 8 : (bIsG8 ? 1 : 4);
 
 	if (bIsRGBA16)
@@ -4781,7 +4806,7 @@ bool UTextureExporterBMP::ExportBinary( UObject* Object, const TCHAR* Type, FArc
 	int32 SizeX = Texture->Source.GetSizeX();
 	int32 SizeY = Texture->Source.GetSizeY();
 	TArray64<uint8> RawData;
-	Texture->Source.GetMipData(RawData, 0);
+	Texture->Source.GetMipData(RawData, LayerIndex, 0);
 
 	FBitmapFileHeader bmf;
 	FBitmapInfoHeader bmhdr;
@@ -4843,6 +4868,16 @@ bool UTextureExporterBMP::ExportBinary( UObject* Object, const TCHAR* Type, FArc
 	return true;
 }
 
+URuntimeVirtualTextureExporterBMP::URuntimeVirtualTextureExporterBMP(const FObjectInitializer& ObjectInitializer)
+	: Super(ObjectInitializer)
+{
+	SupportedClass = URuntimeVirtualTexture::StaticClass();
+	PreferredFormatIndex = 0;
+	FormatExtension.Add(TEXT("BMP"));
+	FormatDescription.Add(TEXT("Windows Bitmap"));
+
+}
+
 /*------------------------------------------------------------------------------
 	URenderTargetExporterHDR implementation.
 	Exports render targets.
diff --git a/Engine/Source/Runtime/Engine/Classes/Exporters/Exporter.h b/Engine/Source/Runtime/Engine/Classes/Exporters/Exporter.h
--- a/Engine/Source/Runtime/Engine/Classes/Exporters/Exporter.h
+++ b/Engine/Source/Runtime/Engine/Classes/Exporters/Exporter.h
@@ -136,12 +136,12 @@ class ENGINE_VTABLE UExporter : public UObject
 	 * which have upto 8 raw waves stored within them.
 	 * @todo document - this will only ever return 1 is that correct.
 	 */
-	virtual int32 GetFileCount( void ) const { return( 1 ); }
+	virtual int32 GetFileCount( UObject* Object ) const { return( 1 ); }
 
 	/** 
 	 * Differentiates the filename for objects with multiple files to export. Only needs to be overridden if GetFileCount() returns > 1.
 	 */
-	virtual FString GetUniqueFilename( const TCHAR* Filename, int32 FileIndex ) { check( FileIndex == 0 ); return( FString( Filename ) ); }
+	virtual FString GetUniqueFilename( const TCHAR* Filename, int32 FileIndex, int32 FileCount ) { check( FileIndex == 0 && FileCount == 1 ); return( FString( Filename ) ); }
 
 	/**
 	* Return true if the exporter is in batch mode.
diff --git a/Engine/Source/Runtime/Engine/Classes/VT/RuntimeVirtualTexture.h b/Engine/Source/Runtime/Engine/Classes/VT/RuntimeVirtualTexture.h
--- a/Engine/Source/Runtime/Engine/Classes/VT/RuntimeVirtualTexture.h
+++ b/Engine/Source/Runtime/Engine/Classes/VT/RuntimeVirtualTexture.h
@@ -139,6 +139,8 @@ public:
 	FVector4 GetUniformParameter(int32 Index) const;
 
 #if WITH_EDITOR
+	/** Get the streaming texture */
+	URuntimeVirtualTextureStreamingProxy* GetStreamingTexture() const { return StreamingTexture; }
 	/** Get a hash of the current state to use for streaming texture invalidation. */
 	uint32 GetStreamingTextureBuildHash() const;
 	/** Initialize the low mip streaming texture with the passed in size and data. */
diff --git a/Engine/Source/Runtime/Engine/Private/UnrealExporter.cpp b/Engine/Source/Runtime/Engine/Private/UnrealExporter.cpp
--- a/Engine/Source/Runtime/Engine/Private/UnrealExporter.cpp
+++ b/Engine/Source/Runtime/Engine/Private/UnrealExporter.cpp
@@ -353,12 +353,13 @@ bool UExporter::RunAssetExportTask(class UAssetExportTask* Task)
 	}
 	else
 	{
-		for( int32 i = 0; i < Exporter->GetFileCount(); i++ )
+		const int32 FileCount = Exporter->GetFileCount(Task->Object);
+		for( int32 i = 0; i < FileCount; i++ )
 		{
 			FBufferArchive Buffer;
 			if(ExportToArchive(Task->Object, Exporter, Buffer, *Extension, i))
 			{
-				FString UniqueFilename = Exporter->GetUniqueFilename(*Task->Filename, i);
+				FString UniqueFilename = Exporter->GetUniqueFilename(*Task->Filename, i, FileCount);
 
 				if(!Task->bReplaceIdentical)
 				{
