Build fix

diff --git a/Engine/Source/Editor/UnrealEd/Classes/Exporters/TextureExporterBMP.h b/Engine/Source/Editor/UnrealEd/Classes/Exporters/TextureExporterBMP.h
--- a/Engine/Source/Editor/UnrealEd/Classes/Exporters/TextureExporterBMP.h
+++ b/Engine/Source/Editor/UnrealEd/Classes/Exporters/TextureExporterBMP.h
@@ -22,7 +22,7 @@ class UTextureExporterBMP : public UExporter
 	virtual bool ExportBinary(UObject* Object, const TCHAR* Type, FArchive& Ar, FFeedbackContext* Warn, int32 FileIndex = 0, uint32 PortFlags = 0) override;
 
 protected:
-	UTexture2D* GetExportTexture(UObject* Object) const;
+	class UTexture2D* GetExportTexture(UObject* Object) const;
 };
 
 UCLASS()
