This solves the error described in JIRA UE-83453: New Window Left/Top Position changes every time the window is closed on a display using custom display scaling.

diff --git a/Engine/Source/Editor/UnrealEd/Private/PlayLevel.cpp b/Engine/Source/Editor/UnrealEd/Private/PlayLevel.cpp
--- a/Engine/Source/Editor/UnrealEd/Private/PlayLevel.cpp
+++ b/Engine/Source/Editor/UnrealEd/Private/PlayLevel.cpp
@@ -3445,9 +3445,7 @@ UGameInstance* UEditorEngine::CreatePIEGameInstance(int32 InPIEInstance, bool bI
 						static void OnPIEWindowClosed( const TSharedRef< SWindow >& WindowBeingClosed, TWeakPtr< SViewport > PIEViewportWidget, int32 index, bool bRestoreRootWindow )
 						{
 							// Save off the window position
-							FVector2D PIEWindowPos = WindowBeingClosed->GetPositionInScreen();
-							const float DPIScale = FPlatformApplicationMisc::GetDPIScaleFactorAtPoint(PIEWindowPos.X, PIEWindowPos.Y);
-							PIEWindowPos /= DPIScale;
+							const FVector2D PIEWindowPos = WindowBeingClosed->GetPositionInScreen();
 
 							ULevelEditorPlaySettings* LevelEditorPlaySettings = ULevelEditorPlaySettings::StaticClass()->GetDefaultObject<ULevelEditorPlaySettings>();
 
