Kismet Reinstancing code witll now use FReferenceFinder instead of FArchiveHasReferences to improve reinstancing performance

diff --git a/Engine/Source/Editor/UnrealEd/Private/Kismet2/KismetReinstanceUtilities.cpp b/Engine/Source/Editor/UnrealEd/Private/Kismet2/KismetReinstanceUtilities.cpp
--- a/Engine/Source/Editor/UnrealEd/Private/Kismet2/KismetReinstanceUtilities.cpp
+++ b/Engine/Source/Editor/UnrealEd/Private/Kismet2/KismetReinstanceUtilities.cpp
@@ -23,7 +23,7 @@
 #include "Kismet2/BlueprintEditorUtils.h"
 #include "Layers/LayersSubsystem.h"
 #include "Editor.h"
-#include "Serialization/ArchiveHasReferences.h"
+#include "UObject/ReferencerFinder.h"
 
 #include "UObject/UObjectHash.h"
 #include "UObject/UObjectIterator.h"
@@ -100,7 +100,7 @@ struct FReplaceReferenceHelper
 		{
 			BP_SCOPED_COMPILER_EVENT_STAT(EKismetReinstancerStats_FindReferencers);
 
-			Targets = FArchiveHasReferences::GetAllReferencers(SourceObjects, ObjectsThatShouldUseOldStuff);
+			Targets = FReferencerFinder::GetAllReferencers(SourceObjects, ObjectsThatShouldUseOldStuff);
 		}
 
 		{
