Remove DumpBlueprintsInfo commandlet as it was relying on content that was not available to all

diff --git a/Engine/Source/Editor/UnrealEd/Private/Commandlets/DumpBlueprintsInfoCommandlet.cpp b/Engine/Source/Editor/UnrealEd/Private/Commandlets/DumpBlueprintsInfoCommandlet.cpp
--- a/Engine/Source/Editor/UnrealEd/Private/Commandlets/DumpBlueprintsInfoCommandlet.cpp
+++ b/Engine/Source/Editor/UnrealEd/Private/Commandlets/DumpBlueprintsInfoCommandlet.cpp
@@ -1,2427 +1,10 @@
 // Copyright 1998-2019 Epic Games, Inc. All Rights Reserved.
 
 #include "Commandlets/DumpBlueprintsInfoCommandlet.h"
-#include "HAL/PlatformProcess.h"
-#include "HAL/FileManager.h"
-#include "Misc/Paths.h"
-#include "Modules/ModuleManager.h"
-#include "UObject/Class.h"
-#include "UObject/Interface.h"
-#include "UObject/UnrealType.h"
-#include "Misc/PackageName.h"
-#include "Engine/Level.h"
-#include "EdGraph/EdGraphPin.h"
-#include "Components/ActorComponent.h"
-#include "Components/SceneComponent.h"
-#include "GameFramework/Actor.h"
-#include "Engine/Blueprint.h"
-#include "Components/StaticMeshComponent.h"
-#include "ActorFactories/ActorFactory.h"
-#include "ActorFactories/ActorFactoryBlueprint.h"
-#include "Animation/Skeleton.h"
-#include "EdGraph/EdGraph.h"
-#include "AssetData.h"
-#include "EdGraph/EdGraphSchema.h"
-#include "Animation/AnimInstance.h"
-#include "Engine/Engine.h"
-#include "Animation/AnimBlueprint.h"
-#include "Factories/AnimBlueprintFactory.h"
-#include "Factories/BlueprintFactory.h"
-#include "EngineGlobals.h"
-#include "Engine/LevelScriptActor.h"
-#include "Engine/BlueprintGeneratedClass.h"
-#include "Animation/AnimBlueprintGeneratedClass.h"
-#include "Engine/UserDefinedEnum.h"
-#include "Engine/UserDefinedStruct.h"
-#include "UnrealEdGlobals.h"
-#include "Kismet2/KismetEditorUtilities.h"
-#include "K2Node.h"
-#include "AnimationGraph.h"
-#include "AnimGraphNode_StateMachine.h"
-#include "EdGraphSchema_K2.h"
-#include "Engine/LevelScriptBlueprint.h"
-#include "Engine/SCS_Node.h"
-#include "Kismet2/BlueprintEditorUtils.h"
-#include "AnimStateConduitNode.h"
-#include "AnimStateNode.h"
-#include "AssetSelection.h"
-#include "UObject/UObjectHash.h"
-#include "UObject/UObjectIterator.h"
-#include "K2Node_CustomEvent.h"
-#include "ProfilingDebugging/ScopedTimers.h"
-#include "BlueprintActionMenuBuilder.h"
-#include "BlueprintActionFilter.h"
-#include "BlueprintActionDatabase.h"
-#include "BlueprintActionMenuUtils.h"
-#include "K2Node_Composite.h"
-#include "IAssetTools.h"
-#include "AssetToolsModule.h"
-#include "BlueprintNodeSpawner.h"
-#include "AnimationGraphSchema.h"
-#include "AssetRegistryModule.h"
-#include "PackageHelperFunctions.h"
-#include "K2Node_FunctionEntry.h"
-#include "Engine/Selection.h"
+#include "Logging/LogMacros.h"
 
 DEFINE_LOG_CATEGORY_STATIC(LogBlueprintInfoDump, Log, All);
 
-/*******************************************************************************
- * Static Helpers
- ******************************************************************************/
-
-namespace DumpBlueprintInfoUtils
-{
-	static const FString HelpString = TEXT("\n\
-\n\
-DumpBlueprintsInfo commandlet params: \n\
-\n\
-    -class=<Class>      Used to specify the blueprint's parent class, if left  \n\
-                        unset then it will go through EVERY blueprint parent   \n\
-                        class available.                                       \n\
-\n\
-    -multifile          Used to keep file size down, will split each blueprint \n\
-                        into its own file (meaning only one file will be       \n\
-                        created when used with -class).                        \n\
-\n\
-    -palette            Defaults to on, unless other flags were specified.     \n\
-                        Dumps all actions from the blueprint's palette menu    \n\
-                        (constant across all graphs).                          \n\
-\n\
-    -palfilter=<Class>  Simulates picking a class from the blueprint palette's \n\
-                        drop down. Setting -palfilter=all will dump the palette\n\
-                        for every possible class.                              \n\
-\n\
-    -context            Will dump all actions from the context menu, for each  \n\
-                        graph and every pin type within (dumps a lot of info). \n\
-\n\
-    -noPinContext       Use after the -context switch to only dump context info\n\
-                        for each graph (not for each pin type within each      \n\
-                        graph).                                                \n\
-\n\
-    -noActionInfo       Whenever an action is dumped it will omit its sub      \n\
-                        fields (can help cut down on the lines that have to be \n\
-                        compared in a diff).                                   \n\
-\n\
-    -graph=<Filter>     Used to only dump contextual info for specific graph   \n\
-                        types. The <Filter> param can only be certain values:  \n\
-                        function, ubergraph, macro, animation, or statemachine.\n\
-\n\
-    -pin=<PinType>      Used to specify a single pin type to dump contextual   \n\
-                        data for. The <PinType> param can be a POD type like   \n\
-                        'int' or 'bool', or an object class like 'Actor', etc. \n\
-\n\
-    -classpin=<Class>   There is no way with the -pin switch to specify a class\n\
-                        pin, this will override any previous -pin switches and \n\
-                        replace them with the class specified (like 'Actor').  \n\
-\n\
-    -select=<Class>     For graph context dumps, this will simulate an object  \n\
-                        being selected during action menu generation. <Class>  \n\
-                        must be an Actor or ActorComponent sub-class. You can  \n\
-                        also specify -select=all to go through every viable    \n\
-                        Actor and ActorComponent class.                        \n\
-\n\
-    -time               When enabled, will record timings during menu building \n\
-                        (as it has been time sync in the past). This is not    \n\
-                        ideal for diffs though (since times can easily vary).  \n\
-\n\
-    -diff=<DiffPath>    Can be a folder or file path, if a folder then it will \n\
-                        look for a matching filename within that folder to     \n\
-                        compare against. Can help automate looking for changes.\n\
-\n\
-    -diffcmd=<Command>  Diffing will attempt to use the P4 diff-tool, but if   \n\
-                        you wish to use your own tool then specify it here. Use\n\
-                        '{1}' and '{2}' as placeholders for filenames, like so:\n\
-                        -diffcmd=\"AraxisP4Diff.exe {2} {1}\".                 \n\
-\n\
-    -name=<Filename>    Overrides the default filename. Leave off the extention\n\
-                        (this will add .json to the end). When -multifile is   \n\
-                        supplied, the class name will be postfixed to the name.\n\
-\n\
-    -interface=<Class>  Appends the desired interface to blueprints that are   \n\
-                        being dumped. The <Class> param has to match a known   \n\
-                        interface class.                                       \n\
-\n\
-    -loadBP=<Blueprint> Before recording any info, this will attempt to load   \n\
-                        the specified Blueprint. The blueprint name can contain\n\
-                        wildcards (to match multiple blueprints). Can also be  \n\
-                        set to \"all\", to load every non-developer blueprint. \n\
-\n\
-    -dbInfo             Will dump info regarding the blueprint action database \n\
-                        (which is used to consturct blueprint action menus).   \n\
-\n\
-    -help, -h, -?       Display this message and then exit.                    \n\
-\n");
-
-	/** Flags that govern the verbosity of the dump. */
-	enum EDumpFlags
-	{
-		BPDUMP_UnfilteredPalette	= (1<<0), 
-		BPDUMP_FilteredPalette		= (1<<1), 
-		BPDUMP_GraphContextActions	= (1<<2), 
-		BPDUMP_PinContextActions	= (1<<3), 
-		BPDUMP_LogHelp				= (1<<4),
-
-		BPDUMP_FilePerBlueprint		= (1<<5), 
-		BPDUMP_PinTypeIsClass		= (1<<6),
-		BPDUMP_DoNotDumpActionInfo  = (1<<7),
-		BPDUMP_RecordTiming			= (1<<8),
-		BPDUMP_SelectAllObjTypes    = (1<<9), 
-
-		BPDUMP_ActionDatabaseInfo   = (1<<11),
-
-		BPDUMP_PaletteMask = (BPDUMP_UnfilteredPalette | BPDUMP_FilteredPalette),
-		BPDUMP_ContextMask = (BPDUMP_GraphContextActions | BPDUMP_PinContextActions),
-	};
-
-	/**
-	 * A collection of variables that represent the various command switches 
-	 * that users can specify when running the commandlet. See the HelpString 
-	 * variable for a listing of supported switches.
-	 */
-	struct CommandletOptions
-	{
-		CommandletOptions()
-			: BlueprintClass(nullptr)
-			, DumpFlags(DumpBlueprintInfoUtils::BPDUMP_UnfilteredPalette)
-			, PaletteFilter(nullptr)
-			, GraphFilter(GT_MAX)
-			, SelectedObjectType(nullptr)
-			, InterfaceClass(nullptr)
-		{
-		}
-
-		/**
-		* Parses the string command switches into flags, class pointers, and 
-		* booleans that will govern what should be dumped. Logs errors if any 
-		* switch was misused.
-		*/
-		CommandletOptions(TArray<FString> const& Switches);
-
-		UClass*    BlueprintClass;
-		uint32     DumpFlags;
-		UClass*    PaletteFilter;
-		EGraphType GraphFilter;
-		FName      PinType;
-		UClass*    SelectedObjectType;
-		FString    DiffPath;
-		FString    DiffCommand;
-		FString    Filename;
-		UClass*	   InterfaceClass;
-	};
-
-	/** Static instance of the command switches (so we don't have to pass one along the call stack) */
-	static CommandletOptions CommandOptions;
-
-	/** Tracks spawned level actors (so we don't have to create more the we have to). */
-	static TMap<UClass*, AActor*> LevelActors;
-	
-	/** Tracks instantiated blueprints (so we don't have to create more the we have to). */
-	static TMap<UClass*, UBlueprint*> ClassBlueprints;
-
-	/**
-	 * Certain blueprints (like level blueprints) require a level outer, and 
-	 * for certain actions we need a level actor selected. This utility function
-	 * provides an easy way to grab the world (which has a level that we can use 
-	 * for these purposes).
-	 */
-	static UWorld* GetWorld();
-
-	/** 
-	 * Utility function for spawning and selecting an actor of the specified 
-	 * type (utilized when the user uses the -select switch with an actor class).
-	 */
-	static AActor* SpawnLevelActor(UClass* ActorClass, bool bSelect = true);
-
-	/**
-	 * Spawns a transient blueprint of the specified type. Adds all possible 
-	 * graph types (function, macro, etc.), and does some additional setup for
-	 * unique blueprint types (like level and anim blueprints).
-	 */
-	static UBlueprint* MakeTempBlueprint(UClass* BlueprintClass);
-
-	/**
-	 * Adds an instance of the specified component type to the supplied 
-	 * blueprint (invoked when the user has specified the -select command switch
-	 * with a component class).
-	 */
-	static bool AddComponentToBlueprint(UBlueprint* Blueprint, UClass* ComponentClass);
-
-	/**
-	 * Certain nodes add specific graph types that we want to dump info on (like 
-	 * state machine graphs for anim blueprints). The best way to add those 
-	 * graphs is through the natural process of adding those nodes (which this
-	 * method is intended for).
-	 */
-	template<class NodeType>
-	static NodeType* AddNodeToGraph(UEdGraph* Graph);
-
-	/**
-	 * Builds a fully qualified file path for a new dump file. If using the 
-	 * -multifile command switch, then this will create a sub-directory and name
-	 * the file after the class. Generally, dump files are placed in the 
-	 * project's ".../Saved/Commandlets/" directory.
-	 */
-	static FString BuildDumpFilePath(UClass* BlueprintClass);
-
-	/**
-	 * Utility function to convert a tab integer into a string of whitespace.
-	 * Defaults to tab characters, but if bUseSpaces is enabled, then single 
-	 * spaces are used.
-	 */
-	static FString BuildIndentString(uint32 IndentCount, bool bUseSpaces = false);
-
-	/**
-	 * Utility function to convert a graph's EGraphType into a string. Used as 
-	 * an aid when writing graph information to json.
-	 */
-	static FString GetGraphTypeString(UEdGraph const* Graph);
-
-	/**
-	 * Concatenates the action's category with its menu name (to help
-	 * distinguish similarly named actions). Can then be used to sort and 
-	 * uniquely identify actions.
-	 */
-	static FString GetActionKey(FGraphActionListBuilderBase::ActionGroup const& Action);
-
-	/**
-	 * Goes through all of the blueprint skeleton's object properties and pulls
-	 * out the ones that are associated with an UActorComponent (and are visbile
-	 * to the blueprint).
-	 */
-	static void GetComponentProperties(UBlueprint* Blueprint, TArray<UObjectProperty*>& PropertiesOut);
-
-	/**
-	 * Takes the given size (in bytes) and returns a formatted string in either
-	 * bytes, kilobytes, megabytes, or gigabytes.
-	 *
-	 * @return A formatted string, representing the specified size in the best units possible.
-	 */
-	static FString BuildByteSizeString(int32 ByteSize);
-
-	/**
-	 * Dumps stats on the new blueprint menu system (database size, number of 
-	 * enteries, etc.).
-	 *
-	 * @return True if any data was written to the file, otherwise false (if nothing was dumped).
-	 */
-	static bool DumpActionDatabaseInfo(uint32 Indent, FArchive* FileOutWriter);
-
-	/**
-	 * Constructs a temporary blueprint (of the class type specified) and kicks 
-	 * off a dump of all its nested information (palette, graph, contextual 
-	 * actions, etc.). 
-	 */
-	static void DumpInfoForClass(uint32 Indent, UClass* BlueprintClass, FArchive* FileOutWriter);
-	
-	/**
-	 * Assumes that the specified ActionListBuilder is configured with the 
-	 * proper blueprint. Starts by clearing any actions it contained and then
-	 * runs through and re-adds all actions.
-	 *
-	 * @return The amount of time (in seconds) that the menu building took.
-	 */
-	static double GetPaletteMenuActions(FCategorizedGraphActionListBuilder& PaletteBuilder, UBlueprint* Blueprint, UClass* PaletteFilter);
-
-	/**
-	 * Dumps all palette actions listed for the specified blueprint. Determines
-	 * if the user specified any filter class for the palette and adjusts 
-	 * accordingly (can dump multiple palettes if -palfilter=all was specified).
-	 */
-	static void DumpPalette(uint32 Indent, UBlueprint* Blueprint, FArchive* FileOutWriter);
-
-	/**
-	 * Dumps a single instance of the blueprint's palette (using the ClassFilter).
-	 * ClassFilter can be null and the full unfiltered palette will be dumped.
-	 */
-	static void DumpPalette(uint32 Indent, UBlueprint* Blueprint, UClass* ClassFilter, FArchive* FileOutWriter);
-
-	/**
-	 * Generic function utilized by both palette and context-menu dumps. Take a
-	 * GraphActionListBuilder and writes out every action that it has captured.
-	 */
-	static void DumpActionList(uint32 Indent, FGraphActionListBuilderBase& ActionList, FArchive* FileOutWriter);
-
-	/**
-	 * Generic function that dumps information on a single action (like it's 
-	 * name, category, an associated node if it has one, etc.).
-	 */
-	static void DumpActionMenuItem(uint32 Indent, FGraphActionListBuilderBase::ActionGroup const& Action, FArchive* FileOutWriter);
-
-	/**
-	 * Emulates the blueprint's context menu and goes through each of its graphs,
-	 * dumping the context menu(s) for each. Entry point for dumping all nested
-	 * context actions.
-	 */
-	static void DumpContextualActions(uint32 Indent, UBlueprint* Blueprint, FArchive* FileOutWriter);
-
-	/**
-	 * Dumps the context menu actions that can be found when right clicking in a
-	 * bare graph (of the specified type). If pin context dumping is enabled, 
-	 * then this will continue down into that.
-	 */
-	static void DumpGraphContextActions(uint32 Indent, UEdGraph* Graph, FArchive* FileOutWriter);
-
-	/**
-	 * Will go through every pin type that a user can create a pin for and dumps
-	 * the full context menu as if it were dragged from each type.
-	 *
-	 * WARNING: This will dump A LOT of info, as it accounts for every type as
-	 *          an input/output and as an array input/output (use -pin= to reign it in).
-	 */
-	static bool DumpPinContextActions(uint32 Indent, UEdGraph* Graph, FArchive* FileOutWriter);
-
-	/**
-	 * Pin type info comes in a tree format (intended for menus), this recursive 
-	 * function traverses the tree and dumps contextual pin actions for each 
-	 * leaf it encounters (utilized by DumpPinContextActions).
-	 */
-	static bool DumpTypeTreeActions(uint32 Indent, UEdGraph* Graph, TSharedPtr<UEdGraphSchema_K2::FPinTypeTreeInfo> PinTypeInfo, FArchive* FileOutWriter);
-
-	/**
-	 * Generic function that takes a specific FEdGraphPinType and dumps all the 
-	 * actions available to that pin (as if you dragged and spawned a context 
-	 * menu from it).
-	 */
-	static void DumpContextualPinTypeActions(uint32 Indent, UEdGraph* Graph, const FEdGraphPinType& PinType, FArchive* FileOutWriter);
-	
-	/**
-	 * Generic function that takes a contextual GraphActionList and calls down 
-	 * into the DumpActionList(), while dumping context information to go along 
-	 * with it (to give the reader context).
-	 */
-	static void DumpContextActionList(uint32 Indent, FGraphContextMenuBuilder ActionBuilder, UBlueprint* Blueprint, FArchive* FileOutWriter);
-
-	/**
-	 * Assumes that the specified ActionListBuilder is configured with all the 
-	 * proper filter/contextual data. Starts by clearing any actions it 
-	 * contained and then runs through building adding new ones that match the
-	 * context.
-	 *
-	 * @return The amount of time (in seconds) that the menu building took.
-	 */
-	static double GetContextMenuActions(FGraphContextMenuBuilder ActionBuilder, UBlueprint* Blueprint);
-
-	/**
-	 * Looks at the filter/context items set on the specified GraphActionListBuilder
-	 * and writes them out, to provide context with any dumped actions.
-	 */
-	static void DumpContextInfo(uint32 Indent, FGraphContextMenuBuilder ActionBuilder, UBlueprint* Blueprint, FArchive* FileOutWriter);
-
-	/**
-	 * A utility function aimed at determining whether we should open a diff for
-	 * the dump file (with -multifile we only want to open files that differ, 
-	 * since we could be going through every class in the project).
-	 *
-	 * @TODO: Right now this is very platform specific (it uses dos commands). 
-	 *        Maybe it should be moved into a platform generic file like IFileManager?
-	 */
-	static bool DoFilesDiffer(FString const& NewFilePath, FString const& OldFilePath);
-
-	/**
-	 * Looks at the two dump files, and opens the specified diff tool if there 
-	 * are differences detected (defaults to p4merge.exe if a diffcmd wasn't 
-	 * specified by the user).
-	 */
-	static void DiffDumpFiles(FString const& NewFilePath, FString const& OldFilePath, FString const& UserDiffCmd);
-
-	/**
-	 * Takes the user specified class name, and attempts to translate it into 
-	 * a class pointer. ClassName can be a blueprint name; if so, this will make
-	 * sure that blueprint is loaded and return that blueprint's generated class.
-	 * 
-	 * @param  ClassName	A name that the user provided, in attempt to identify a class.
-	 * @return Null if no matching class could be found, otherwise a valid class pointer.
-	 */
-	static UClass* GetUserNamedClass(FString ClassName);
-
-	/**
-	 * Attempts to load a subset of blueprint assets. The AssetName can be left
-	 * blank to force load all blueprints.
-	 * 
-	 * @param  AssetName			The name of the blueprint(s) that you want to load (can have wildcards). If left empty, will load all blueprints.
-	 * @param  bAllowDevBlueprints	If true, blueprint in the developer folder will also be loaded.
-	 * @return The number of blueprints that were loaded.
-	 */
-	static int32 LoadBlueprints(FString AssetName, bool bAllowDevBlueprints);
-}
-
-//------------------------------------------------------------------------------
-DumpBlueprintInfoUtils::CommandletOptions::CommandletOptions(TArray<FString> const& Switches)
-	: BlueprintClass(nullptr)
-	, DumpFlags(BPDUMP_UnfilteredPalette)
-	, PaletteFilter(nullptr)
-	, GraphFilter(GT_MAX)
-	, SelectedObjectType(nullptr)
-	, InterfaceClass(nullptr)	
-{
-	uint32 NewDumpFlags = 0x00;
-	for (FString const& Switch : Switches)
-	{
-		if (Switch.StartsWith("class="))
-		{
-			FString ClassSwitch, ClassName;
-			Switch.Split(TEXT("="), &ClassSwitch, &ClassName);
-			BlueprintClass = GetUserNamedClass(ClassName);
-
-			if (BlueprintClass == nullptr)
-			{
-				UE_LOG(LogBlueprintInfoDump, Error, TEXT("Unrecognized blueprint class '%s', defaulting to 'Actor'"), *ClassName);
-				BlueprintClass = AActor::StaticClass();
-			}
-		}
-		else if (Switch.StartsWith("palfilter="))
-		{
-			FString ClassSwitch, ClassName;
-			Switch.Split(TEXT("="), &ClassSwitch, &ClassName);
-			PaletteFilter = GetUserNamedClass(ClassName);
-
-			NewDumpFlags |= BPDUMP_FilteredPalette;
-			if (PaletteFilter == nullptr)
-			{
-				if (ClassName.Compare("all", ESearchCase::IgnoreCase))
-				{
-					UE_LOG(LogBlueprintInfoDump, Error, TEXT("Unrecognized palette filter '%s', defaulting to unfiltered"), *ClassName);
-					NewDumpFlags &= ~(DumpBlueprintInfoUtils::BPDUMP_FilteredPalette);
-					NewDumpFlags |= DumpBlueprintInfoUtils::BPDUMP_UnfilteredPalette;
-				}
-			}
-		}
-		else if (Switch.StartsWith("graph="))
-		{
-			FString GraphSwitch, FilterName;
-			Switch.Split(TEXT("="), &GraphSwitch, &FilterName);
-
-#define CHECK_GRAPH_SWITCH(GraphType) \
-			if ((GraphFilter == GT_MAX) && !FilterName.Compare(#GraphType, ESearchCase::IgnoreCase)) \
-			{ \
-				GraphFilter = GT_##GraphType; \
-			}
-
-			CHECK_GRAPH_SWITCH(Function)
-			CHECK_GRAPH_SWITCH(Ubergraph)
-			CHECK_GRAPH_SWITCH(Macro)
-			CHECK_GRAPH_SWITCH(Animation)
-			CHECK_GRAPH_SWITCH(StateMachine)
-
-#undef CHECK_GRAPH_SWITCH
-
-			if (GraphFilter == GT_MAX)
-			{
-				UE_LOG(LogBlueprintInfoDump, Error, TEXT("Invalid graph type '%s', can only be one of the following: function, ubergraph, macro, animation, or statemachine"), *FilterName);
-			}
-
-			NewDumpFlags |= BPDUMP_GraphContextActions;
-		}
-		else if (Switch.StartsWith("pin="))
-		{
-			NewDumpFlags &= ~BPDUMP_PinTypeIsClass;
-
-			FString PinSwitch;
-			FString PinTypeStr;
-			Switch.Split(TEXT("="), &PinSwitch, &PinTypeStr);
-			PinType = *PinTypeStr;
-
-			NewDumpFlags |= DumpBlueprintInfoUtils::BPDUMP_PinContextActions;
-			// implies that we want contextual actions
-			NewDumpFlags |= DumpBlueprintInfoUtils::BPDUMP_GraphContextActions;
-		}
-		else if (Switch.StartsWith("classpin="))
-		{
-			NewDumpFlags |= BPDUMP_PinTypeIsClass;
-
-			FString PinSwitch;
-			FString PinTypeStr;
-			Switch.Split(TEXT("="), &PinSwitch, &PinTypeStr);
-			PinType = *PinTypeStr;
-			
-			NewDumpFlags |= BPDUMP_PinContextActions;
-			// implies that we want contextual actions
-			NewDumpFlags |= BPDUMP_GraphContextActions;
-		}
-		else if (Switch.StartsWith("select="))
-		{
-			FString ClassSwitch, ClassName;
-			Switch.Split(TEXT("="), &ClassSwitch, &ClassName);
-			SelectedObjectType = GetUserNamedClass(ClassName);
-
-			if (!ClassName.Compare("all", ESearchCase::IgnoreCase))
-			{
-				NewDumpFlags |= BPDUMP_SelectAllObjTypes;
-			}
-			else if (SelectedObjectType == nullptr)
-			{
-				SelectedObjectType = UStaticMeshComponent::StaticClass();
-				UE_LOG(LogBlueprintInfoDump, Error, TEXT("Unrecognized selection class '%s', defaulting to '%s'"), *ClassName, *SelectedObjectType->GetName());
-			}
-			else if (SelectedObjectType->IsChildOf(UActorComponent::StaticClass()))
-			{
-				if (SelectedObjectType->HasAnyClassFlags(CLASS_Abstract))
-				{
-					SelectedObjectType = UStaticMeshComponent::StaticClass();
-					UE_LOG(LogBlueprintInfoDump, Error, TEXT("Selection component cannot be a abstract ('%s' is), defaulting to '%s'"), *ClassName, *SelectedObjectType->GetName());
-				}
-			}
-			else if (!SelectedObjectType->IsChildOf(AActor::StaticClass()))
-			{
-				SelectedObjectType = AActor::StaticClass();
-				UE_LOG(LogBlueprintInfoDump, Error, TEXT("Selection class must be an actor or component ('%s' is not), defaulting to '%s'"), *ClassName, *SelectedObjectType->GetName());
-			}
-			else if (GetWorld() == nullptr)
-			{
-				UE_LOG(LogBlueprintInfoDump, Error, TEXT("Cannot select a level actor without a valid editor world (clearing the selection)"), *ClassName, *SelectedObjectType->GetName());
-				SelectedObjectType = nullptr;
-			}
-
-			// implies that we want contextual actions
-			NewDumpFlags |= BPDUMP_GraphContextActions;
-		}
-		else if (Switch.StartsWith("diff="))
-		{
-			FString DiffSwitch;
-			Switch.Split(TEXT("="), &DiffSwitch, &DiffPath);
-		}
-		else if (Switch.StartsWith("diffcmd="))
-		{
-			FString DiffSwitch;
-			Switch.Split(TEXT("="), &DiffSwitch, &DiffCommand);
-			
-			DiffCommand = DiffCommand.Replace(TEXT("\""), TEXT(""));
-		}
-		else if (!Switch.Compare("palette", ESearchCase::IgnoreCase))
-		{
-			NewDumpFlags |= BPDUMP_UnfilteredPalette;
-		}
-		else if (!Switch.Compare("context", ESearchCase::IgnoreCase))
-		{
-			NewDumpFlags |= BPDUMP_GraphContextActions;
-			NewDumpFlags |= BPDUMP_PinContextActions;
-		}
-		else if (!Switch.Compare("noPinContext", ESearchCase::IgnoreCase))
-		{
-			NewDumpFlags &= (~BPDUMP_PinContextActions);
-		}
-		else if (!Switch.Compare("h", ESearchCase::IgnoreCase) || 
-		         !Switch.Compare("?", ESearchCase::IgnoreCase) || 
-				 !Switch.Compare("help", ESearchCase::IgnoreCase))
-		{
-			NewDumpFlags |= BPDUMP_LogHelp;
-		}
-		else if (!Switch.Compare("multifile", ESearchCase::IgnoreCase))
-		{
-			NewDumpFlags |= BPDUMP_FilePerBlueprint;
-		}
-		else if (!Switch.Compare("noActionInfo", ESearchCase::IgnoreCase))
-		{
-			NewDumpFlags |= BPDUMP_DoNotDumpActionInfo;
-		}
-		else if (!Switch.Compare("time", ESearchCase::IgnoreCase))
-		{
-			NewDumpFlags |= BPDUMP_RecordTiming;
-		}
-		else if (Switch.StartsWith("name="))
-		{
-			FString NameSwitch;
-			Switch.Split(TEXT("="), &NameSwitch, &Filename);
-		}
-		else if (Switch.StartsWith("interface="))
-		{
-			FString InterfaceSwitch, InterfaceName;
-			Switch.Split(TEXT("="), &InterfaceSwitch, &InterfaceName);
-
-			if (UClass* Class = GetUserNamedClass(InterfaceName))
-			{
-				if (Class->IsChildOf<UInterface>())
-				{
-					InterfaceClass = Class;
-				}
-			}
-			
-			if (InterfaceClass == nullptr)
-			{
-				UE_LOG(LogBlueprintInfoDump, Warning, TEXT("Could not find a matching interface class matching this name: '%s'"), *InterfaceName);
-			}
-		}
-		else if (Switch.StartsWith("loadBP="))
-		{
-			FString LoadSwitch, BlueprintName;
-			Switch.Split(TEXT("="), &LoadSwitch, &BlueprintName);
-
-			if (!BlueprintName.Compare("all", ESearchCase::IgnoreCase))
-			{
-				BlueprintName = TEXT("");
-			}
-
-			bool const bAllowDevBlueprints = BlueprintName.IsEmpty() ? false : true;
-			LoadBlueprints(BlueprintName, bAllowDevBlueprints);
-		}
-		else if (!Switch.Compare("dbInfo", ESearchCase::IgnoreCase))
-		{
-			NewDumpFlags |= BPDUMP_ActionDatabaseInfo;
-		}
-	}
-
-	if (NewDumpFlags != 0)
-	{
-		DumpFlags = NewDumpFlags;
-	}
-	if ((DumpFlags & (BPDUMP_PaletteMask | BPDUMP_ContextMask)) == 0)
-	{
-		DumpFlags |= BPDUMP_UnfilteredPalette;
-	}
-}
-
-//------------------------------------------------------------------------------
-static UWorld* DumpBlueprintInfoUtils::GetWorld()
-{
-	UWorld* World = nullptr;
-	for (FWorldContext const& WorldContext : GEngine->GetWorldContexts())
-	{
-		World = WorldContext.World();
-		if (World != nullptr)
-		{
-			break;
-		}
-	}
-
-	if (World == nullptr)
-	{
-		static UWorld* CommandletWorld = nullptr;
-		if (CommandletWorld == nullptr)
-		{
-			if (GUnrealEd == nullptr)
-			{
-				UE_LOG(LogBlueprintInfoDump, Error, TEXT("Cannot create a temp map to test within, without a valid editor world"));
-			}
-			else
-			{
-				CommandletWorld = GEditor->NewMap();
-			}
-		}
-
-		World = CommandletWorld;
-	}
-
-	return World;
-}
-
-//------------------------------------------------------------------------------
-static AActor* DumpBlueprintInfoUtils::SpawnLevelActor(UClass* ActorClass, bool bSelect)
-{
-	check(ActorClass->IsChildOf(AActor::StaticClass()));
-
-	AActor* SpawnedActor = nullptr;
-	if (AActor** ExistingActorPtr = LevelActors.Find(ActorClass))
-	{
-		SpawnedActor = *ExistingActorPtr;
-	}
-	else if (FKismetEditorUtilities::CanCreateBlueprintOfClass(ActorClass))
-	{
-		UActorFactory* NewFactory = NewObject<UActorFactory>(GetTransientPackage(), UActorFactoryBlueprint::StaticClass());
-		NewFactory->AddToRoot();
-
-		UBlueprint* ActorTemplate = MakeTempBlueprint(ActorClass);
-		SpawnedActor = FActorFactoryAssetProxy::AddActorForAsset(ActorTemplate,
-			/*SelectActor =*/bSelect,
-			RF_Transient, NewFactory, NAME_None);
-		NewFactory->RemoveFromRoot();
-	}
-	// @TODO: What about non-blueprintable actors (brushes, etc.)... the code below crashes
-// 	else
-// 	{
-// 		UWorld* World = GetWorld();
-// 		SpawnedActor = GEditor->AddActor(World->GetCurrentLevel(), ActorClass, FVector::ZeroVector);
-// 	}
-
-	if (SpawnedActor != nullptr)
-	{
-		LevelActors.Add(ActorClass, SpawnedActor);
-	}
-
-	if (bSelect)
-	{
-		USelection* SelectedActors = GEditor->GetSelectedActors();
-		SelectedActors->DeselectAll();
-
-		if (SpawnedActor != nullptr)
-		{
-			SelectedActors->Select(SpawnedActor);
-		}
-	}
-
-	return SpawnedActor;
-}
-
-//------------------------------------------------------------------------------
-static UBlueprint* DumpBlueprintInfoUtils::MakeTempBlueprint(UClass* ParentClass)
-{
-	UBlueprint* MadeBlueprint = nullptr;
-	if (UBlueprint** FoundBlueprint = ClassBlueprints.Find(ParentClass))
-	{
-		MadeBlueprint = *FoundBlueprint;
-	}
-	else if (UBlueprint* ClassBlueprint = Cast<UBlueprint>(ParentClass->ClassGeneratedBy))
-	{
-		ClassBlueprints.Add(ParentClass, ClassBlueprint);
-		MadeBlueprint = ClassBlueprint;
-	}
-	else
-	{
-		UObject* BlueprintOuter = GetTransientPackage();
-
-		bool bIsAnimBlueprint = ParentClass->IsChildOf(UAnimInstance::StaticClass());
-		bool bIsLevelBlueprint = ParentClass->IsChildOf(ALevelScriptActor::StaticClass());
-
-		UClass* BlueprintClass = UBlueprint::StaticClass();
-		UClass* GeneratedClass = UBlueprintGeneratedClass::StaticClass();
-		EBlueprintType BlueprintType = BPTYPE_Normal;
-		UFactory* AssetFactory = nullptr;
-
-		if (bIsAnimBlueprint)
-		{
-			BlueprintClass = UAnimBlueprint::StaticClass();
-			GeneratedClass = UAnimBlueprintGeneratedClass::StaticClass();
-
-			UAnimBlueprintFactory* BlueprintFactory = NewObject<UAnimBlueprintFactory>();
-			BlueprintFactory->ParentClass    = ParentClass;
-			BlueprintFactory->BlueprintType  = BlueprintType;
-			BlueprintFactory->TargetSkeleton = (USkeleton*)StaticLoadObject(USkeleton::StaticClass(), /*Outer =*/nullptr, TEXT("/Engine/NotForLicensees/Automation/QAAutomationtest_Assets/TEST_SkeletalMesh_Skeleton.TEST_SkeletalMesh_Skeleton"));
-
-			AssetFactory = BlueprintFactory;
-		}
-		else if (bIsLevelBlueprint)
-		{
-			UWorld* World = GetWorld();
-			if (World == nullptr)
-			{
-				UE_LOG(LogBlueprintInfoDump, Error, TEXT("Cannot make a proper level blueprint without a valid editor level for its outer."));
-			}
-			else
-			{
-				BlueprintClass = ULevelScriptBlueprint::StaticClass();
-				BlueprintType = BPTYPE_LevelScript;
-				BlueprintOuter = World->GetCurrentLevel();
-			}
-		}
-		else 
-		{
-			UBlueprintFactory* BlueprintFactory = NewObject<UBlueprintFactory>();
-			BlueprintFactory->ParentClass = ParentClass;
-			AssetFactory = BlueprintFactory;
-		}
-		// @TODO: UEditorUtilityBlueprint
-
-		FString const ClassName = ParentClass->GetName();
-		FString const DesiredName = FString::Printf(TEXT("COMMANDLET_TEMP_Blueprint_%s"), *ClassName);
-		FName   const TempBpName = MakeUniqueObjectName(BlueprintOuter, BlueprintClass, FName(*DesiredName));
-
-		check(FKismetEditorUtilities::CanCreateBlueprintOfClass(ParentClass));
-		if (AssetFactory != nullptr)
-		{
-			IAssetTools& AssetTools = FModuleManager::LoadModuleChecked<FAssetToolsModule>("AssetTools").Get();
-			UObject* NewAsset = AssetTools.CreateAsset(TempBpName.ToString(), BlueprintOuter->GetPathName(), BlueprintClass, AssetFactory, FName("DumpBlueprintsInfoCommandlet"));
-			MadeBlueprint = CastChecked<UBlueprint>(NewAsset);
-		}
-		else
-		{
-			MadeBlueprint = FKismetEditorUtilities::CreateBlueprint(ParentClass, BlueprintOuter, TempBpName, BlueprintType, BlueprintClass, GeneratedClass);
-		}
-		
-		// if this is an animation blueprint, then we want anim specific graphs to test as well (if it has an anim graph)...
-		if (bIsAnimBlueprint && (MadeBlueprint->FunctionGraphs.Num() > 0))
-		{
-			UAnimationGraph* AnimGraph = CastChecked<UAnimationGraph>(MadeBlueprint->FunctionGraphs[0]);
-			check(AnimGraph != nullptr);
-
-			// should add a state-machine graph
-			UAnimGraphNode_StateMachine* StateMachineNode = AddNodeToGraph<UAnimGraphNode_StateMachine>(AnimGraph);
-
-			UAnimationStateMachineGraph* StateMachineGraph = StateMachineNode->EditorStateMachineGraph;
-			// should add an conduit graph 
-			UAnimStateConduitNode* ConduitNode = AddNodeToGraph<UAnimStateConduitNode>((UEdGraph*)StateMachineGraph);
-
-			UAnimStateNode* StateNode = AddNodeToGraph<UAnimStateNode>((UEdGraph*)StateMachineGraph);
-			// should create a transition graph
-			StateNode->AutowireNewNode(ConduitNode->GetOutputPin());
-		}
-		else if (bIsLevelBlueprint)
-		{
-			ULevel* Level = CastChecked<ULevel>(BlueprintOuter);
-			Level->LevelScriptBlueprint = Cast<ULevelScriptBlueprint>(MadeBlueprint);
-		}
-
-		// may have been altered in CreateBlueprint()
-		BlueprintType = MadeBlueprint->BlueprintType;
-
-		bool bCanAddFunctions = (BlueprintType != BPTYPE_MacroLibrary); // taken from FBlueprintEditor::NewDocument_IsVisibleForType()
-		if (bCanAddFunctions)
-		{
-			// add a functions graph that isn't the construction script (or an animation graph)
-			FName FuncGraphName = MakeUniqueObjectName(MadeBlueprint, UEdGraph::StaticClass(), FName(TEXT("NewFunction")));
-			UEdGraph* FuncGraph = FBlueprintEditorUtils::CreateNewGraph(MadeBlueprint, FuncGraphName, UEdGraph::StaticClass(), UEdGraphSchema_K2::StaticClass());
-			FBlueprintEditorUtils::AddFunctionGraph<UClass>(MadeBlueprint, FuncGraph, /*bIsUserCreated =*/true, nullptr);
-		}
-
-		bool bCanAddMacros = ((BlueprintType == BPTYPE_MacroLibrary) || (BlueprintType == BPTYPE_Normal) || (BlueprintType == BPTYPE_LevelScript));
-		if (bCanAddMacros)
-		{
-			FName MacroGraphName = MakeUniqueObjectName(MadeBlueprint, UEdGraph::StaticClass(), FName(TEXT("NewMacro")));
-			UEdGraph* MacroGraph = FBlueprintEditorUtils::CreateNewGraph(MadeBlueprint, MacroGraphName, UEdGraph::StaticClass(), UEdGraphSchema_K2::StaticClass());
-			FBlueprintEditorUtils::AddMacroGraph(MadeBlueprint, MacroGraph, /*bIsUserCreated =*/true, nullptr);
-		}
-
-		// if you can add custom events to this blueprint, do so (so show that we 
-		// can call an event on ourselves)
-		if (MadeBlueprint->UbergraphPages.Num() > 0)
-		{
-			UK2Node_CustomEvent* CustomEventNode = AddNodeToGraph<UK2Node_CustomEvent>(MadeBlueprint->UbergraphPages[0]);
-			CustomEventNode->CustomFunctionName = FBlueprintEditorUtils::FindUniqueCustomEventName(MadeBlueprint);
-		}
-
-		UClass* ObjTypeToAdd = CommandOptions.SelectedObjectType;
-		if (ObjTypeToAdd && ObjTypeToAdd->IsChildOf(UActorComponent::StaticClass()))
-		{
-			if (!AddComponentToBlueprint(MadeBlueprint, ObjTypeToAdd))
-			{
-				UE_LOG(LogBlueprintInfoDump, Error, TEXT("Cannot add a '%s' to a '%s' blueprint."), *ObjTypeToAdd->GetName(), *ClassName);
-			}
-		}
-		else if (CommandOptions.DumpFlags & BPDUMP_SelectAllObjTypes)
-		{
-			for (TObjectIterator<UClass> ClassIt; ClassIt; ++ClassIt)
-			{
-				UClass* Class = *ClassIt;
-				if (Class->IsChildOf(UActorComponent::StaticClass()))
-				{
-					AddComponentToBlueprint(MadeBlueprint, Class);
-				}
-			}
-		}
-
-		FBlueprintEditorUtils::MarkBlueprintAsStructurallyModified(MadeBlueprint);
-		MadeBlueprint->AddToRoot(); // to keep the BP from being garbage collected
-		FKismetEditorUtilities::CompileBlueprint(MadeBlueprint);
-		ClassBlueprints.Add(ParentClass, MadeBlueprint);
-	}
-
-	check(MadeBlueprint != nullptr);
-	return MadeBlueprint;
-}
-
-//------------------------------------------------------------------------------
-static bool DumpBlueprintInfoUtils::AddComponentToBlueprint(UBlueprint* Blueprint, UClass* ComponentClass)
-{
-	bool const bCanUserAddComponents  = FBlueprintEditorUtils::DoesSupportComponents(Blueprint);
-	bool const bClassIsActorComponent = ComponentClass->IsChildOf(UActorComponent::StaticClass());
-	bool const bCanBeAddedToBlueprint = !ComponentClass->HasAnyClassFlags(CLASS_Abstract) && ComponentClass->HasMetaData(FBlueprintMetadata::MD_BlueprintSpawnableComponent);
-	bool const bCanMakeComponent = (bCanUserAddComponents && bClassIsActorComponent && bCanBeAddedToBlueprint);
-
-	if (bCanMakeComponent)
-	{
-		USimpleConstructionScript* SCS = Blueprint->SimpleConstructionScript;
-
-		USCS_Node* NewNode = SCS->CreateNode(ComponentClass);
-
-		if (ComponentClass->IsChildOf(USceneComponent::StaticClass()) || (SCS->GetRootNodes().Num() == 0))
-		{
-			SCS->AddNode(NewNode);
-		}
-		else
-		{
-			USCS_Node* RootNode = SCS->GetDefaultSceneRootNode();
-			check(RootNode != nullptr);
-			RootNode->AddChildNode(NewNode);
-		}
-
-		// regenerate the skeleton class
-		FBlueprintEditorUtils::MarkBlueprintAsStructurallyModified(Blueprint);
-	}
-
-	return bCanMakeComponent;
-}
-
-//------------------------------------------------------------------------------
-template<class NodeType>
-static NodeType* DumpBlueprintInfoUtils::AddNodeToGraph(UEdGraph* Graph)
-{
-	check(Graph != nullptr);
-	NodeType* NewNode = NewObject<NodeType>(Graph);
-	Graph->AddNode(NewNode, /*bFromUI =*/true, /*bSelectNewNode =*/false);
-
-	NewNode->CreateNewGuid();
-	NewNode->PostPlacedNewNode();
-	NewNode->AllocateDefaultPins();
-	return NewNode;
-}
-
-//------------------------------------------------------------------------------
-static FString DumpBlueprintInfoUtils::BuildDumpFilePath(UClass* BlueprintClass)
-{
-	FString Pathname = FString::Printf(TEXT("BlueprintsInfoDump_%s"), FPlatformTime::StrTimestamp());
-	Pathname = Pathname.Replace(TEXT(" "), TEXT("_"));
-	Pathname = Pathname.Replace(TEXT("/"), TEXT("-"));
-	Pathname = Pathname.Replace(TEXT(":"), TEXT("."));
-	
-	bool const bSplitBlueprintsByFile = (CommandOptions.DumpFlags & BPDUMP_FilePerBlueprint) != 0;
-	
-	static FString CommandletSaveDir;
-	if (CommandletSaveDir.IsEmpty())
-	{
-		CommandletSaveDir = FPaths::ProjectSavedDir() + TEXT("Commandlets/");
-		CommandletSaveDir = FPaths::ConvertRelativePathToFull(CommandletSaveDir);
-		IFileManager::Get().MakeDirectory(*CommandletSaveDir);
-		
-		if (bSplitBlueprintsByFile)
-		{
-			CommandletSaveDir += Pathname + "/";
-			IFileManager::Get().MakeDirectory(*CommandletSaveDir);
-		}
-	}
-	
-	if (!CommandOptions.Filename.IsEmpty())
-	{
-		Pathname = CommandOptions.Filename;
-	}
-
-	if (bSplitBlueprintsByFile && (BlueprintClass != nullptr))
-	{
-		if (CommandOptions.Filename.IsEmpty())
-		{
-			Pathname = TEXT("BlueprintInfo");
-		}
-		Pathname += "_" + BlueprintClass->GetName() + ".json";
-	}
-	else
-	{
-		Pathname += ".json";
-	}
-
-	return CommandletSaveDir / *Pathname;
-}
-
-//------------------------------------------------------------------------------
-static FString DumpBlueprintInfoUtils::BuildIndentString(uint32 IndentCount, bool bUseSpaces)
-{
-	char RepatingChar = '\t';
-	if (bUseSpaces)
-	{
-		RepatingChar = ' ';
-	}
-
-	FString IndentString;
-	while (IndentCount > 0)
-	{
-		IndentString += RepatingChar;
-		--IndentCount;
-	}
-	return IndentString;
-}
-
-//------------------------------------------------------------------------------
-static FString DumpBlueprintInfoUtils::GetGraphTypeString(UEdGraph const* Graph)
-{
-#define RETURN_GRAPH_TYPE(GraphType) \
-	case GraphType: \
-	{ \
-		return FString(#GraphType); \
-	} \
-	break
-
-	UEdGraphSchema const* GraphSchema = GetDefault<UEdGraphSchema>(Graph->Schema);
-	switch (GraphSchema->GetGraphType(Graph))
-	{
-		RETURN_GRAPH_TYPE(GT_Ubergraph);
-		RETURN_GRAPH_TYPE(GT_Function);
-		RETURN_GRAPH_TYPE(GT_Macro);
-		RETURN_GRAPH_TYPE(GT_Animation);
-		RETURN_GRAPH_TYPE(GT_StateMachine);
-
-		default:
-			break;
-	}
-#undef RETURN_GRAPH_TYPE
-
-	return FString("<UNRECOGNIZED>");
-}
-
-//------------------------------------------------------------------------------
-static FString DumpBlueprintInfoUtils::GetActionKey(FGraphActionListBuilderBase::ActionGroup const& Action)
-{
-	const TArray<FString>& MenuHierarchy = Action.GetCategoryChain();
-
-	FString ActionKey;
-	for (FString const& SubCategory : MenuHierarchy)
-	{
-		ActionKey += SubCategory + TEXT("|");
-	}
-	if (MenuHierarchy.Num() > 0)
-	{
-		ActionKey.RemoveAt(ActionKey.Len() - 1); // remove the trailing '|'
-	}
-
-	TSharedPtr<FEdGraphSchemaAction> MainAction = Action.Actions[0];
-	ActionKey += MainAction->GetMenuDescription().ToString();
-
-	return ActionKey;
-}
-
-//------------------------------------------------------------------------------
-static void DumpBlueprintInfoUtils::GetComponentProperties(UBlueprint* Blueprint, TArray<UObjectProperty*>& PropertiesOut)
-{
-	UClass* BpClass = (Blueprint->SkeletonGeneratedClass != nullptr) ? Blueprint->SkeletonGeneratedClass : Blueprint->ParentClass;
-	if (BpClass->IsChildOf<AActor>())
-	{
-		for (TFieldIterator<UObjectProperty> PropertyIt(BpClass, EFieldIteratorFlags::IncludeSuper); PropertyIt; ++PropertyIt)
-		{
-			// SMyBlueprint filters out component variables in SMyBlueprint::CollectAllActions() using CPF_BlueprintVisible/CPF_Parm flags
-			if (PropertyIt->PropertyClass->IsChildOf(UActorComponent::StaticClass()) &&
-				PropertyIt->HasAnyPropertyFlags(CPF_BlueprintVisible) && !PropertyIt->HasAnyPropertyFlags(CPF_Parm))
-			{
-				PropertiesOut.Add(*PropertyIt);
-			}
-		}
-	}
-}
-
-//------------------------------------------------------------------------------
-static FString DumpBlueprintInfoUtils::BuildByteSizeString(int32 ByteSize)
-{
-	enum 
-	{
-		Bytes,
-		KiloBytes,
-		MegaBytes,
-		GigaBytes,
-
-		ByteUnits_MAX
-	};
-
-	const TCHAR* ByteUnits[ByteUnits_MAX];
-	ByteUnits[Bytes] = TEXT("Bytes");
-	ByteUnits[KiloBytes] = TEXT("KB");
-	ByteUnits[MegaBytes] = TEXT("MB");
-	ByteUnits[GigaBytes] = TEXT("GB");
-
-	int32 UnitsIndex = 0;
-	float ConvertedSize = ByteSize;
-
-	bool bIsIntegral = true;
-
-	float const MetricStepSize = 1024.0f;
-	while ((ConvertedSize > MetricStepSize) && (UnitsIndex < ByteUnits_MAX))
-	{
-		ConvertedSize /= MetricStepSize;
-		++UnitsIndex;
-
-		bIsIntegral = false;
-	}
-
-	if (bIsIntegral)
-	{
-		return FString::Printf(TEXT("%.0f %s"), ConvertedSize, ByteUnits[UnitsIndex]);
-	}
-	else
-	{
-		return FString::Printf(TEXT("%.2f %s"), ConvertedSize, ByteUnits[UnitsIndex]);
-	}
-}
-
-//------------------------------------------------------------------------------
-static bool DumpBlueprintInfoUtils::DumpActionDatabaseInfo(uint32 Indent, FArchive* FileOutWriter)
-{
-	bool bWroteToFile = false;
-
-	uint32 const DbInfoMask = BPDUMP_ActionDatabaseInfo;
-	if ((CommandOptions.DumpFlags & DbInfoMask) == BPDUMP_ActionDatabaseInfo)
-	{
-		UE_LOG(LogBlueprintInfoDump, Display, TEXT("%sDumping Database info..."), *BuildIndentString(Indent, true));
-		//--------------------------------------
-		// Composing Data
-		//--------------------------------------
-		double DatabaseBuildTime = 0.0;
-		{
-			FScopedDurationTimer DatabaseBuildTimer(DatabaseBuildTime);
-			// prime the database so we can record information from it
-			FBlueprintActionDatabase::Get();
-		}
-		FBlueprintActionDatabase& Database = FBlueprintActionDatabase::Get();
-		
-		FBlueprintActionDatabase::FActionRegistry const& ActionRegistry = Database.GetAllActions();
-		int32 EstimatedDatabaseSize = sizeof(Database) + ActionRegistry.GetAllocatedSize();
-		int32 EstimatedSystemSize   = EstimatedDatabaseSize;
-
-		TSet<UBlueprint*> TemplateOuters;
-		int32 DatabaseCount = 0;
-		double TotalPrimingTime = 0.0;
-		int32 TemplateCount = 0;
-		int32 UnknownAssetActions = 0;
-
-		struct FSpawnerInfo
-		{
-			int32  Count;
-			int32  TemplateNodeCount;
-			double TotalPrimingTime;
-
-			FSpawnerInfo() : Count(0), TemplateNodeCount(0), TotalPrimingTime(0.0) {}
-		};
-		TMap<UClass*, FSpawnerInfo> DatabaseBreakdown;
-
-		for (auto const& DbEntry : ActionRegistry)
-		{
-			UObject const* ActionSetKey = DbEntry.Key.ResolveObjectPtr();
-			if (ActionSetKey)
-			{
-				bool const bIsUnknownAssetEntry = ActionSetKey->IsAsset() &&
-					!ActionSetKey->IsA<UBlueprint>() &&
-					!ActionSetKey->IsA<UUserDefinedStruct>() &&
-					!ActionSetKey->IsA<UUserDefinedEnum>();
-
-				for (UBlueprintNodeSpawner* BpAction : DbEntry.Value)
-				{
-					++DatabaseCount;
-					// @TODO: doesn't account for any allocated memory (for delegates, text strings, etc.)
-					EstimatedDatabaseSize += sizeof(*BpAction);
-
-					FSpawnerInfo& SpawnerInfo = DatabaseBreakdown.FindOrAdd(BpAction->GetClass());
-					SpawnerInfo.Count += 1;
-
-					int32 OldPrimingTime = TotalPrimingTime;
-					{
-						FScopedDurationTimer PrimingTimer(TotalPrimingTime);
-						BpAction->Prime();
-					}
-					SpawnerInfo.TotalPrimingTime += (TotalPrimingTime - OldPrimingTime);
-
-					if (UEdGraphNode* TemplateNode = BpAction->GetCachedTemplateNode())
-					{
-						UObject* TemplateOuter = TemplateNode->GetOuter();
-						while ((TemplateOuter != nullptr) && (Cast<UBlueprint>(TemplateOuter) == nullptr))
-						{
-							TemplateOuter = TemplateOuter->GetOuter();
-						}
-						UBlueprint* OuterBlueprint = CastChecked<UBlueprint>(TemplateOuter);
-						TemplateOuters.Add(OuterBlueprint);
-
-						++TemplateCount;
-						SpawnerInfo.TemplateNodeCount += 1;
-					}
-
-					if (bIsUnknownAssetEntry)
-					{
-						++UnknownAssetActions;
-					}
-				}
-			}
-		}
-
-		FString const OriginalIndent = BuildIndentString(Indent);
-		FString const IndentedNewline = "\n" + BuildIndentString(Indent + 1);
-
-		{
-			int32 SpawnerCount = 0;
-			for (TObjectIterator<UBlueprintNodeSpawner> NodeSpawnerIt; NodeSpawnerIt; ++NodeSpawnerIt)
-			{
-				++SpawnerCount;
-				// @TODO: doesn't account for any allocated memory (for delegates, text strings, etc.)
-				EstimatedSystemSize += sizeof(**NodeSpawnerIt);
-			}
-
-			FString DatabaseInfoHeading = FString::Printf(TEXT("%s\"ActionDatabaseInfo\" : {%s\"TotalNodeSpawnerCount\" : %d,"), 
-				*OriginalIndent, *IndentedNewline, SpawnerCount);
-			FileOutWriter->Serialize(TCHAR_TO_ANSI(*DatabaseInfoHeading), DatabaseInfoHeading.Len());
-		}
-
-		//--------------------------------------
-		// Dumping Database Stats
-		//--------------------------------------
-		UE_LOG(LogBlueprintInfoDump, Display, TEXT("%sDumping raw action stats..."), *BuildIndentString(Indent+1, true));
-
-		FString const DoubleIndent = BuildIndentString(Indent + 2);
-		FString const DblIndentedNewline = "\n" + DoubleIndent;
-		FString const SubDictEndingBrace = IndentedNewline + "}";
-
-		FString const DatabaseSizeStr = BuildByteSizeString(EstimatedDatabaseSize);
-		FString const AvgActionSizeStr = BuildByteSizeString(EstimatedDatabaseSize / DatabaseCount);
-
-		FString DatabaseStats = FString::Printf(TEXT("%s\"Database Stats\" : {"), *IndentedNewline);
-		DatabaseStats += FString::Printf(TEXT("%s\"DatabaseBuildTime\"     : %.3f seconds,"), *DblIndentedNewline, DatabaseBuildTime);
-		DatabaseStats += FString::Printf(TEXT("%s\"NodeSpawnerCount\"      : %d,"), *DblIndentedNewline, DatabaseCount);
-		DatabaseStats += FString::Printf(TEXT("%s\"EstimatedDatabaseSize\" : %s,"), *DblIndentedNewline, *DatabaseSizeStr);
-		DatabaseStats += FString::Printf(TEXT("%s\"AvgSizePerEntry\"       : %s"), *DblIndentedNewline, *AvgActionSizeStr);
-		DatabaseStats += SubDictEndingBrace + ",";
-
-		FileOutWriter->Serialize(TCHAR_TO_ANSI(*DatabaseStats), DatabaseStats.Len());
-
-		//--------------------------------------
-		// Dumping Template Cache Stats
-		//--------------------------------------
-		UE_LOG(LogBlueprintInfoDump, Display, TEXT("%sDumping template-cache stats..."), *BuildIndentString(Indent+1, true));
-
-		int32 EstimatedCacheSize = 0;
-		for (UBlueprint* CacheBlueprint : TemplateOuters)
-		{
-			TArray<UObject*> ChildObjs;
-			GetObjectsWithOuter(CacheBlueprint, ChildObjs);
-
-			EstimatedCacheSize += sizeof(*CacheBlueprint);
-			for (UObject* ChildObj : ChildObjs)
-			{
-				// @TODO: doesn't account for any allocated memory (for member TArrays, etc.)
-				EstimatedCacheSize += sizeof(*ChildObj);
-			}
-		}
-		EstimatedSystemSize += EstimatedCacheSize;
-
-		FString const NodeCacheSizeStr = BuildByteSizeString(EstimatedCacheSize);
-		FString const AvgNodeSizeStr = BuildByteSizeString((TemplateCount > 0) ? EstimatedCacheSize / TemplateCount : 0);
-		FString NodeCacheStats = FString::Printf(TEXT("%s\"Template-Cache Stats\" : {"), *IndentedNewline);
-		NodeCacheStats += FString::Printf(TEXT("%s\"TotalPrimingDuration\" : %.3f seconds,"), *DblIndentedNewline, TotalPrimingTime);
-		NodeCacheStats += FString::Printf(TEXT("%s\"CachedNodeCount\"      : %d,"), *DblIndentedNewline, TemplateCount);
-		NodeCacheStats += FString::Printf(TEXT("%s\"EstimatedCacheSize\"   : %s,"), *DblIndentedNewline, *NodeCacheSizeStr);
-		NodeCacheStats += FString::Printf(TEXT("%s\"AvgSizePerEntry\"      : %s"), *DblIndentedNewline, *AvgNodeSizeStr);
-		NodeCacheStats += SubDictEndingBrace + ",";
-
-		FileOutWriter->Serialize(TCHAR_TO_ANSI(*NodeCacheStats), NodeCacheStats.Len());
-
-		//--------------------------------------
-		// Dumping Database Breakdown
-		//--------------------------------------
-		TCHAR const* LineEnding = TEXT(",\n");
-
-		FString BreakdownStats = FString::Printf(TEXT("%s\"Database Breakdown\" : {\n"), *IndentedNewline);
-		for (auto const& SpawnerEntry : DatabaseBreakdown)
-		{
-			int32 SpawnerCount = SpawnerEntry.Value.Count;
-
-			static FString const TripleIndent = BuildIndentString(Indent + 3);
-			BreakdownStats += FString::Printf(TEXT("%s\"%s\" : {\n"), *DoubleIndent, *SpawnerEntry.Key->GetName());
-			BreakdownStats += FString::Printf(TEXT("%s\"Total\"               : %d,\n"), *TripleIndent, SpawnerCount);
-			BreakdownStats += FString::Printf(TEXT("%s\"TemplateNodesPrimed\" : %d,\n"), *TripleIndent, SpawnerEntry.Value.TemplateNodeCount);
-			BreakdownStats += FString::Printf(TEXT("%s\"AvgPrimingDuration\"  : %.03f seconds\n"), *TripleIndent, SpawnerEntry.Value.TotalPrimingTime / SpawnerCount);
-			BreakdownStats += DoubleIndent + "}" + LineEnding;
-		}
-		BreakdownStats.RemoveFromEnd(LineEnding);
-		BreakdownStats += IndentedNewline + "},";
-
-		FileOutWriter->Serialize(TCHAR_TO_ANSI(*BreakdownStats), BreakdownStats.Len());
-
-		//--------------------------------------
-		// Dumping Blueprint Stats
-		//--------------------------------------
-		UE_LOG(LogBlueprintInfoDump, Display, TEXT("%sDumping blueprint related stats..."), *BuildIndentString(Indent+1, true));
-
-		int32 BlueprintsCount = 0;
-		int32 BlueprintsDbCount = 0;
-		int32 TotalBlueprintActionCount = 0;
-		int32 BlueprintFunctionCount = 0;
-		int32 BlueprintVariableCount = 0;
-		int32 BlueprintDelegateCount = 0;
-		int32 BlueprintLocalVarCount = 0;
-
-		for (TObjectIterator<UBlueprint> BlueprintIt; BlueprintIt; ++BlueprintIt)
-		{
-			if (BlueprintIt->IsAsset())
-			{
-				++BlueprintsCount;
-				if (FBlueprintActionDatabase::FActionList const* ActionList = ActionRegistry.Find(*BlueprintIt))
-				{
-					if (ActionList->Num() > 0)
-					{
-						++BlueprintsDbCount;
-					}
-					TotalBlueprintActionCount += ActionList->Num();
-				}
-
-				for (auto GraphIt = BlueprintIt->FunctionGraphs.CreateConstIterator(); GraphIt; ++GraphIt)
-				{
-					UEdGraph* FunctionGraph = (*GraphIt);
-
-					TArray<UK2Node_FunctionEntry*> GraphEntryNodes;
-					FunctionGraph->GetNodesOfClass<UK2Node_FunctionEntry>(GraphEntryNodes);
-
-					for (UK2Node_FunctionEntry* FunctionEntry : GraphEntryNodes)
-					{
-						for (FBPVariableDescription const& LocalVar : FunctionEntry->LocalVariables)
-						{
-							++BlueprintLocalVarCount;
-						}
-					}
-				}
-
-				UClass* BlueprintClass = BlueprintIt->GeneratedClass;
-				if (BlueprintClass == nullptr)
-				{
-					continue;
-				}
-
-				for (TFieldIterator<UFunction> FunctionIt(BlueprintClass, EFieldIteratorFlags::ExcludeSuper); FunctionIt; ++FunctionIt)
-				{
-					++BlueprintFunctionCount;
-				}
-
-				for (TFieldIterator<UProperty> PropertyIt(BlueprintClass, EFieldIteratorFlags::ExcludeSuper); PropertyIt; ++PropertyIt)
-				{
-					bool const bIsDelegate = PropertyIt->IsA(UMulticastDelegateProperty::StaticClass());
-					if (bIsDelegate)
-					{
-						++BlueprintDelegateCount;
-					}
-					else
-					{
-						++BlueprintVariableCount;
-					}
-				}
-			}
-		};
-
-		FString BlueprintStats = FString::Printf(TEXT("%s\"Blueprint Stats\" : {"), *IndentedNewline);
-		BlueprintStats += FString::Printf(TEXT("%s\"BlueprintsLoaded\"         : %d,"), *DblIndentedNewline, BlueprintsCount);
-		BlueprintStats += FString::Printf(TEXT("%s\"BlueprintsWithActions\"    : %d,"), *DblIndentedNewline, BlueprintsDbCount);
-		BlueprintStats += FString::Printf(TEXT("%s\"AvgFunctionsPerBlueprint\" : %d,"), *DblIndentedNewline, FMath::RoundToInt((float)BlueprintFunctionCount / BlueprintsCount));
-		BlueprintStats += FString::Printf(TEXT("%s\"AvgVariablesPerBlueprint\" : %d,"), *DblIndentedNewline, FMath::RoundToInt((float)BlueprintVariableCount / BlueprintsCount));
-		BlueprintStats += FString::Printf(TEXT("%s\"AvgDelegatesPerBlueprint\" : %d,"), *DblIndentedNewline, FMath::RoundToInt((float)BlueprintDelegateCount / BlueprintsCount));
-		BlueprintStats += FString::Printf(TEXT("%s\"AvgLocalVarsPerBlueprint\" : %d,"), *DblIndentedNewline, FMath::RoundToInt((float)BlueprintLocalVarCount / BlueprintsCount));
-		BlueprintStats += FString::Printf(TEXT("%s\"TotalBlueprintActions\"    : %d,"), *DblIndentedNewline, TotalBlueprintActionCount);
-		BlueprintStats += FString::Printf(TEXT("%s\"AvgActionsPerBlueprint\"   : %d"), *DblIndentedNewline, FMath::RoundToInt((float)TotalBlueprintActionCount / BlueprintsDbCount));
-		BlueprintStats += SubDictEndingBrace + ",";
-		
-		FileOutWriter->Serialize(TCHAR_TO_ANSI(*BlueprintStats), BlueprintStats.Len());
-
-		//--------------------------------------
-		// Additional Asset Stats
-		//--------------------------------------
-		UE_LOG(LogBlueprintInfoDump, Display, TEXT("%sDumping additional asset stats..."), *BuildIndentString(Indent+1, true));
-
-		int32 EnumAssetCount = 0;
-		int32 EnumsWithActionsCount = 0;
-		int32 TotalEnumActions = 0;
-
-		for (TObjectIterator<UUserDefinedEnum> EnumIt; EnumIt; ++EnumIt)
-		{
-			if (EnumIt->IsAsset())
-			{
-				++EnumAssetCount;
-				if (FBlueprintActionDatabase::FActionList const* ActionList = ActionRegistry.Find(*EnumIt))
-				{
-					if (ActionList->Num() > 0)
-					{
-						++EnumsWithActionsCount;
-					}
-					TotalEnumActions += ActionList->Num();
-				}
-			}
-			else
-			{
-				// @TODO: bad assumption? all UUserDefinedEnum should be assets
-			}
-		};
-
-		int32 StructAssetCount = 0;
-		int32 StructsWithActionsCount = 0;
-		int32 TotalStructActions = 0;
-
-		for (TObjectIterator<UUserDefinedStruct> StructIt; StructIt; ++StructIt)
-		{
-			if (StructIt->IsAsset())
-			{
-				++StructAssetCount;
-				if (FBlueprintActionDatabase::FActionList const* ActionList = ActionRegistry.Find(*StructIt))
-				{
-					if (ActionList->Num() > 0)
-					{
-						++StructsWithActionsCount;
-					}
-					TotalStructActions += ActionList->Num();
-				}
-			}
-			else
-			{
-				// @TODO: bad assumption? all UUserDefinedStructs should be assets
-			}
-		};
-
-		FString OtherAssetStats = FString::Printf(TEXT("%s\"Other Asset Stats\"  : {"), *IndentedNewline);
-		OtherAssetStats += FString::Printf(TEXT("%s\"EnumAssetsLoaded\"        : %d,"), *DblIndentedNewline, EnumAssetCount);
-		OtherAssetStats += FString::Printf(TEXT("%s\"EnumAssetsWithActions\"   : %d,"), *DblIndentedNewline, EnumsWithActionsCount);
-		OtherAssetStats += FString::Printf(TEXT("%s\"TotalEnumActionCount\"    : %d,"), *DblIndentedNewline, TotalEnumActions);
-		OtherAssetStats += FString::Printf(TEXT("%s\"StructAssetsLoaded\"      : %d,"), *DblIndentedNewline, StructAssetCount);
-		OtherAssetStats += FString::Printf(TEXT("%s\"StructAssetsWithActions\" : %d,"), *DblIndentedNewline, StructsWithActionsCount);
-		OtherAssetStats += FString::Printf(TEXT("%s\"TotalStructActionCount\"  : %d,"), *DblIndentedNewline, TotalStructActions);
-		OtherAssetStats += FString::Printf(TEXT("%s\"OtherAssetActions\"       : %d"), *DblIndentedNewline, UnknownAssetActions);
-		OtherAssetStats += SubDictEndingBrace + ",";
-
-		FileOutWriter->Serialize(TCHAR_TO_ANSI(*OtherAssetStats), OtherAssetStats.Len());
-		
-		//--------------------------------------
-		// Stats Closing
-		//--------------------------------------
-		
-		FString const TotalSystemSizeStr = BuildByteSizeString(EstimatedSystemSize);
-
-		FString DatabaseInfoClosing = FString::Printf(TEXT("%s\"EstimatedSystemSize\" : \"%s\"\n%s}"),
-			*IndentedNewline, *TotalSystemSizeStr, *OriginalIndent);
-		FileOutWriter->Serialize(TCHAR_TO_ANSI(*DatabaseInfoClosing), DatabaseInfoClosing.Len());
-
-		bWroteToFile = true;
-	}
-	return bWroteToFile;
-}
-
-//------------------------------------------------------------------------------
-static void DumpBlueprintInfoUtils::DumpInfoForClass(uint32 Indent, UClass* BlueprintClass, FArchive* FileOutWriter)
-{
-	FString const ClassName = BlueprintClass->GetName();
-	UE_LOG(LogBlueprintInfoDump, Display, TEXT("%sDumping BP class: '%s'..."), *BuildIndentString(Indent, true), *ClassName);
-
-	FString const ClassEntryIndent = BuildIndentString(Indent);
-	FString BeginClassEntry = FString::Printf(TEXT("%s\"%s\" : {"), *ClassEntryIndent, *ClassName);
-
-	FString IndentedNewline = "\n" + BuildIndentString(Indent + 1);
-
-	BeginClassEntry += IndentedNewline + "\"ClassContext\" : \"" + ClassName + "\",\n";
-	FileOutWriter->Serialize(TCHAR_TO_ANSI(*BeginClassEntry), BeginClassEntry.Len());
-
-	UBlueprint* TempBlueprint = MakeTempBlueprint(BlueprintClass);
-	if (CommandOptions.InterfaceClass != nullptr)
-	{
-		UClass* BpClass = (TempBlueprint->SkeletonGeneratedClass != nullptr) ? TempBlueprint->SkeletonGeneratedClass : TempBlueprint->ParentClass;
-		if (!BpClass->ImplementsInterface(CommandOptions.InterfaceClass))
-		{
-			if (!FBlueprintEditorUtils::ImplementNewInterface(TempBlueprint, CommandOptions.InterfaceClass->GetFName()))
-			{
-				UE_LOG(LogBlueprintInfoDump, Warning, TEXT("Failed to add interface (%s), to blueprint class: '%s'"), *CommandOptions.InterfaceClass->GetName(), *ClassName);
-			}
-		}		
-	}
-
-	bool bNeedsClosingComma = false;
-	if (CommandOptions.DumpFlags & BPDUMP_PaletteMask)
-	{
-		DumpPalette(Indent + 1, TempBlueprint, FileOutWriter);
-		bNeedsClosingComma = true;
-	}
-
-	if (CommandOptions.DumpFlags & BPDUMP_ContextMask)
-	{
-		if (bNeedsClosingComma)
-		{
-			FileOutWriter->Serialize(TCHAR_TO_ANSI(TEXT(",\n")), 2);
-		}
-		DumpContextualActions(Indent + 1, TempBlueprint, FileOutWriter);
-		bNeedsClosingComma = true;
-	}
-
-	FString EndClassEntry = "\n" + ClassEntryIndent + "}";
-	FileOutWriter->Serialize(TCHAR_TO_ANSI(*EndClassEntry), EndClassEntry.Len());
-}
-
-//------------------------------------------------------------------------------
-static double DumpBlueprintInfoUtils::GetPaletteMenuActions(FCategorizedGraphActionListBuilder& PaletteBuilder, UBlueprint* Blueprint, UClass* PaletteFilter)
-{
-	PaletteBuilder.Empty();
-	UEdGraphSchema_K2 const* K2Schema = GetDefault<UEdGraphSchema_K2>();
-
-	double MenuBuildDuration = 0.0;
-	
-	FBlueprintActionContext FilterContext;
-	FilterContext.Blueprints.Add(const_cast<UBlueprint*>(Blueprint));
-
-	FBlueprintActionMenuBuilder MenuBuilder(nullptr);
-	{			
-		FScopedDurationTimer DurationTimer(MenuBuildDuration);
-		FBlueprintActionMenuUtils::MakePaletteMenu(FilterContext, PaletteFilter, MenuBuilder);
-		PaletteBuilder.Append(MenuBuilder);
-	}
-
-	return MenuBuildDuration;
-}
-
-//------------------------------------------------------------------------------
-static void DumpBlueprintInfoUtils::DumpPalette(uint32 Indent, UBlueprint* Blueprint, FArchive* FileOutWriter)
-{
-	UClass* PaletteFilter = CommandOptions.PaletteFilter;
-	DumpPalette(Indent, Blueprint, PaletteFilter, FileOutWriter);
-	bool bNeedsEndline = true;
-
-	if ((CommandOptions.DumpFlags & BPDUMP_FilteredPalette) && (PaletteFilter == nullptr))
-	{
-		// anim blueprints don't have a palette, so it is ok to assume this
-		UEdGraphSchema_K2 const* K2Schema = GetDefault<UEdGraphSchema_K2>();
-
-		for (TObjectIterator<UClass> ClassIt; ClassIt; ++ClassIt)
-		{
-			UClass* Class = *ClassIt;
-			if (!K2Schema->ClassHasBlueprintAccessibleMembers(Class))
-			{
-				continue;
-			}
-
-			FileOutWriter->Serialize(TCHAR_TO_ANSI(TEXT(",\n")), 2);
-			DumpPalette(Indent, Blueprint, Class, FileOutWriter);
-		}
-	}
-}
-
-//------------------------------------------------------------------------------
-static void DumpBlueprintInfoUtils::DumpPalette(uint32 Indent, UBlueprint* Blueprint, UClass* ClassFilter, FArchive* FileOutWriter)
-{
-	FString const PaletteEntryIndent = BuildIndentString(Indent);
-	FString BeginPaletteEntry = FString::Printf(TEXT("%s\"Palette"), *PaletteEntryIndent);
-
-	FString FilterClassName("<UNFILTERED>");
-	if (ClassFilter != nullptr)
-	{
-		FilterClassName = ClassFilter->GetName();
-		BeginPaletteEntry += "-" + FilterClassName;
-	}
-	BeginPaletteEntry += "\" : {\n";
-
-	FString const NestedIndent = BuildIndentString(Indent + 1);
-	UE_LOG(LogBlueprintInfoDump, Display, TEXT("%sDumping palette: %s"), *BuildIndentString(Indent, true), *FilterClassName);
-
-	bool bIsAnimBlueprint = (Cast<UAnimBlueprint>(Blueprint) != nullptr);
-	// animation blueprints don't have a palette
-	if (bIsAnimBlueprint)
-	{
-		BeginPaletteEntry += NestedIndent + "\"IsAnimBlueprint\" : true";
-		FileOutWriter->Serialize(TCHAR_TO_ANSI(*BeginPaletteEntry), BeginPaletteEntry.Len());
-	}
-	else
-	{
-		FCategorizedGraphActionListBuilder PaletteBuilder;
-		PaletteBuilder.OwnerOfTemporaries = NewObject<UEdGraph>((UObject*)Blueprint);
-		PaletteBuilder.OwnerOfTemporaries->Schema = UEdGraphSchema_K2::StaticClass();
-		PaletteBuilder.OwnerOfTemporaries->SetFlags(RF_Transient);
-		double MenuBuildDuration = GetPaletteMenuActions(PaletteBuilder, Blueprint, ClassFilter);
-
-		BeginPaletteEntry += NestedIndent + "\"FilterClass\" : \"" + FilterClassName + "\",\n";
-		if (CommandOptions.DumpFlags & BPDUMP_RecordTiming)
-		{
-			BeginPaletteEntry += FString::Printf(TEXT("%s\"MenuBuildTime_Seconds\" : %f,\n"), *NestedIndent, MenuBuildDuration);
-		}
-
-		FileOutWriter->Serialize(TCHAR_TO_ANSI(*BeginPaletteEntry), BeginPaletteEntry.Len());
-		DumpActionList(Indent + 1, PaletteBuilder, FileOutWriter);
-	}
-
-	FString EndPaletteEntry = "\n" + PaletteEntryIndent + "}";
-	FileOutWriter->Serialize(TCHAR_TO_ANSI(*EndPaletteEntry), EndPaletteEntry.Len());
-}
-
-//------------------------------------------------------------------------------
-static void DumpBlueprintInfoUtils::DumpActionList(uint32 Indent, FGraphActionListBuilderBase& ActionList, FArchive* FileOutWriter)
-{
-	TArray< FGraphActionListBuilderBase::ActionGroup const* > SortedActions;
-	for (int32 ActionIndex = 0; ActionIndex < ActionList.GetNumActions(); ++ActionIndex)
-	{
-		FGraphActionListBuilderBase::ActionGroup& Action = ActionList.GetAction(ActionIndex);
-		if (Action.Actions.Num() <= 0)
-		{
-			continue;
-		}
-
-		SortedActions.Add(&Action);
-	}
-
-	FString const ActionListIndent = BuildIndentString(Indent);
-	FString const NestedIndent     = BuildIndentString(Indent + 1);
-
-	FString BeginActionListEntry = FString::Printf(TEXT("%s\"ActionSet\" : {\n%s\"ActionCount\" : %d"), *ActionListIndent, *NestedIndent, SortedActions.Num());
-	BeginActionListEntry += ",\n" + NestedIndent + "\"Actions\" : ";
-
-	bool const bNoActionInfo = (CommandOptions.DumpFlags & BPDUMP_DoNotDumpActionInfo) != 0;
-	if (bNoActionInfo)
-	{
-		BeginActionListEntry += "[";
-	}
-	else
-	{
-		BeginActionListEntry += "{";
-	}
-	FileOutWriter->Serialize(TCHAR_TO_ANSI(*BeginActionListEntry), BeginActionListEntry.Len());
-
-	struct ActionSortFunctor
-	{
-		bool operator()(FGraphActionListBuilderBase::ActionGroup const& LHS, FGraphActionListBuilderBase::ActionGroup const& RHS) const
-		{
-			TSharedPtr<FEdGraphSchemaAction> LHSAction = LHS.Actions[0];
-			TSharedPtr<FEdGraphSchemaAction> RHSAction = RHS.Actions[0];
-			
-			if (LHSAction->GetGrouping() != RHSAction->GetGrouping())
-			{
-				return LHSAction->GetGrouping() > RHSAction->GetGrouping();
-			}
-			
-			FString LhKey = GetActionKey(LHS);
-			FString RhKey = GetActionKey(RHS);
-			return (LhKey.Compare(RhKey) < 0);
-		}
-	};
-	// need to sort so we can easily compare from one generation to the next
-	SortedActions.Sort(ActionSortFunctor());
-
-	FString LineEnding("\n");
-	for (FGraphActionListBuilderBase::ActionGroup const* Action : SortedActions)
-	{
-		FileOutWriter->Serialize(TCHAR_TO_ANSI(*LineEnding), LineEnding.Len());
-		DumpActionMenuItem(Indent + 2, *Action, FileOutWriter);
-		LineEnding = TEXT(",\n");
-	}
-
-	FString EndActionListEntry = "\n" + NestedIndent;
-	if (bNoActionInfo)
-	{
-		EndActionListEntry += "]";
-	}
-	else
-	{
-		EndActionListEntry += "}";
-	}
-	EndActionListEntry += "\n" + ActionListIndent + "}";
-	FileOutWriter->Serialize(TCHAR_TO_ANSI(*EndActionListEntry), EndActionListEntry.Len());
-}
-
-//------------------------------------------------------------------------------
-static void DumpBlueprintInfoUtils::DumpActionMenuItem(uint32 Indent, FGraphActionListBuilderBase::ActionGroup const& Action, FArchive* FileOutWriter)
-{
-	check(Action.Actions.Num() > 0);
-
-	// Get action category info
-	TArray<FString> MenuHierarchy = Action.GetCategoryChain();
-
-	FString ActionCategory = TEXT("");
-
-	bool bHasCategory = (MenuHierarchy.Num() > 0);
-	if (bHasCategory)
-	{
-		for (FString const& SubCategory : MenuHierarchy)
-		{
-			ActionCategory += SubCategory + TEXT("|");
-		}
-	}
-
-	TSharedPtr<FEdGraphSchemaAction> PrimeAction = Action.Actions[0];
-	FString const ActionName = PrimeAction->GetMenuDescription().ToString();
-
-	FString const ActionEntryIndent = BuildIndentString(Indent);
-	FString ActionEntry = ActionEntryIndent + "\"" + ActionCategory + ActionName + "\"";;
-
-	bool const bNoActionInfo = (CommandOptions.DumpFlags & BPDUMP_DoNotDumpActionInfo) != 0;
-	if (!bNoActionInfo)
-	{
-		++Indent;
-		FString const IndentedNewline = "\n" + BuildIndentString(Indent);
-
-		ActionEntry += " : {";
-		ActionEntry += IndentedNewline + "\"ActionType\"  : \"" + PrimeAction->GetTypeId().ToString() + "\",";
-		ActionEntry += IndentedNewline + "\"Name\"        : \"" + ActionName + "\",";
-		ActionEntry += IndentedNewline + "\"Category\"    : \"";
-		if (bHasCategory)
-		{
-			ActionEntry += ActionCategory;
-			ActionEntry.RemoveAt(ActionEntry.Len() - 1); // remove the trailing '|'
-		}
-		ActionEntry += "\","; // end action category data
-
-		FString TooltipStr = PrimeAction->GetTooltipDescription().ToString().Replace(TEXT("\\\""), TEXT("'")).Replace(TEXT("\""), TEXT("'"));
-		FString const TooltipFieldLabel("\"Tooltip\"     : \"");
-		TooltipStr = TooltipStr.Replace(TEXT("\n"), *(IndentedNewline + BuildIndentString(TooltipFieldLabel.Len(), /*bUseSpaces =*/true)));
-
-		ActionEntry += IndentedNewline + TooltipFieldLabel + TooltipStr + "\",";
-		ActionEntry += IndentedNewline + "\"Keywords\"    : \"" + PrimeAction->GetKeywords().ToString() + "\",";
-		ActionEntry += IndentedNewline + "\"SearchTitle\" : \"" + PrimeAction->GetMenuDescription().ToString() + "\",";
-		ActionEntry += IndentedNewline + FString::Printf(TEXT("\"Grouping\"    : %d"), PrimeAction->GetGrouping());
-		
-		// Get action node type info
-		UK2Node const* NodeTemplate = FBlueprintActionMenuUtils::ExtractNodeTemplateFromAction(PrimeAction);
-		if (NodeTemplate != nullptr)
-		{
-			// Build action node type data
-			ActionEntry += "," + IndentedNewline + "\"Node\"        : \"";
-			ActionEntry += NodeTemplate->GetClass()->GetPathName();
-			ActionEntry += "\"";
-		}
-		// Finish action entry
-		ActionEntry += "\n" + ActionEntryIndent + "}";
-	}
-
-	// Write entry to file
-	FileOutWriter->Serialize(TCHAR_TO_ANSI(*ActionEntry), ActionEntry.Len());
-}
-
-//------------------------------------------------------------------------------
-static void DumpBlueprintInfoUtils::DumpContextualActions(uint32 Indent, UBlueprint* Blueprint, FArchive* FileOutWriter)
-{
-	UE_LOG(LogBlueprintInfoDump, Display, TEXT("%sDumping contextual info..."), *BuildIndentString(Indent, true));
-
-	FString const ContextualEntryIndent = BuildIndentString(Indent);
-	FString const BeginContextualEntry = FString::Printf(TEXT("%s\"GraphActions\" : {\n"), *ContextualEntryIndent);
-
-	FileOutWriter->Serialize(TCHAR_TO_ANSI(*BeginContextualEntry), BeginContextualEntry.Len());
-
-	TArray<UEdGraph*> BpGraphs;
-	Blueprint->GetAllGraphs(BpGraphs);
-
-	bool bIsFirstEntry = true;
-	for (UEdGraph* Graph : BpGraphs)
-	{
-		UEdGraphSchema const* GraphSchema = GetDefault<UEdGraphSchema>(Graph->Schema);
-		if ((CommandOptions.GraphFilter != GT_MAX) && (GraphSchema->GetGraphType(Graph) != CommandOptions.GraphFilter))
-		{
-			continue;
-		}
-
-		if (bIsFirstEntry)
-		{
-			FileOutWriter->Serialize(TCHAR_TO_ANSI(TEXT("\n")), 1);
-			bIsFirstEntry = false;
-		}
-		else
-		{
-			FileOutWriter->Serialize(TCHAR_TO_ANSI(TEXT(",\n")), 2);
-		}
-		DumpGraphContextActions(Indent + 1, Graph, FileOutWriter);
-	}
-
-	FString const EndContextualEntry = FString::Printf(TEXT("\n%s}"), *ContextualEntryIndent);
-	FileOutWriter->Serialize(TCHAR_TO_ANSI(*EndContextualEntry), EndContextualEntry.Len());
-}
-
-//------------------------------------------------------------------------------
-static void DumpBlueprintInfoUtils::DumpGraphContextActions(uint32 Indent, UEdGraph* Graph, FArchive* FileOutWriter)
-{
-	UE_LOG(LogBlueprintInfoDump, Display, TEXT("%sDumping graph: '%s'..."), *BuildIndentString(Indent, true), *Graph->GetName());
-
-	FString const GraphEntryIndent = BuildIndentString(Indent);
-	FString BeginGraphEntry = FString::Printf(TEXT("%s\"%s\" : {"), *GraphEntryIndent, *Graph->GetName());
-
-	FString const NestedIndent = BuildIndentString(++Indent);
-	FString const IndentedNewline = "\n" + NestedIndent;
-	BeginGraphEntry += IndentedNewline + "\"GraphType\" : \"" + GetGraphTypeString(Graph) + "\",";
-	BeginGraphEntry += IndentedNewline + "\"GraphName\" : \"" + Graph->GetName() + "\",";
-	BeginGraphEntry += IndentedNewline + "\"GraphContextMenu\" : \n";
-
-	FileOutWriter->Serialize(TCHAR_TO_ANSI(*BeginGraphEntry), BeginGraphEntry.Len());
-
-	UObject* GraphOuter = Graph->GetOuter();
-	UBlueprint* Blueprint = Cast<UBlueprint>(GraphOuter);
-	while ((Blueprint == nullptr) && (GraphOuter != nullptr))
-	{
-		GraphOuter = GraphOuter->GetOuter();
-		Blueprint = Cast<UBlueprint>(GraphOuter);
-	}
-	check(Blueprint != nullptr);
-
-	FGraphContextMenuBuilder ActionBuilder(Graph);
-	UE_LOG(LogBlueprintInfoDump, Display, TEXT("%sDumping graph context actions..."), *BuildIndentString(Indent, true));
-	DumpContextActionList(Indent, ActionBuilder, Blueprint, FileOutWriter);
-
-	if (CommandOptions.DumpFlags & BPDUMP_SelectAllObjTypes)
-	{
-		for (TObjectIterator<UClass> ClassIt; ClassIt; ++ClassIt)
-		{
-			UClass* Class = *ClassIt;
-			
-			bool const bIsTempBpClass = (Class->GetOuterUPackage() == GetTransientPackage());
-			// have to exclude classes that were created explicitly in this commandlet
-			if (bIsTempBpClass)
-			{
-				continue;
-			}
-			
-			if (Class->IsChildOf<AActor>() && !Class->IsChildOf<ALevelScriptActor>())
-			{				
-				if (AActor* LevelActor = SpawnLevelActor(Class, true))
-				{
-					UE_LOG(LogBlueprintInfoDump, Display, TEXT("%sDumping actions with actor selection: '%s'..."), *BuildIndentString(Indent, true), *Class->GetName());
-
-					FString ActorSelectionEntry = "," + IndentedNewline + "\"LevelActorMenu-" + Class->GetName() + "\" : \n";
-					FileOutWriter->Serialize(TCHAR_TO_ANSI(*ActorSelectionEntry), ActorSelectionEntry.Len());
-
-					DumpContextActionList(Indent, ActionBuilder, Blueprint, FileOutWriter);
-				}
-			}
-		}
-	}
-
-	TArray<UObjectProperty*> ComponentProperties;
-	GetComponentProperties(Blueprint, ComponentProperties);
-
-	bool const bOnlyDumpSpecificComponents = (CommandOptions.SelectedObjectType != nullptr) &&
-		CommandOptions.SelectedObjectType->IsChildOf<UActorComponent>();
-	for (UObjectProperty* Component : ComponentProperties)
-	{
-		if (bOnlyDumpSpecificComponents && !Component->PropertyClass->IsChildOf(CommandOptions.SelectedObjectType))
-		{
-			continue;
-		}
-		UE_LOG(LogBlueprintInfoDump, Display, TEXT("%sDumping actions with component selection: '%s'..."), *BuildIndentString(Indent, true), *Component->GetName());
-
-		FString SelectionContextEntry = "," + IndentedNewline + "\"ComponentContextMenu-" + Component->GetName() + "\" : \n";
-		FileOutWriter->Serialize(TCHAR_TO_ANSI(*SelectionContextEntry), SelectionContextEntry.Len());
-
-		ActionBuilder.SelectedObjects.Empty();
-		ActionBuilder.SelectedObjects.Add(Component);
-		DumpContextActionList(Indent, ActionBuilder, Blueprint, FileOutWriter);
-	}
-
-	FString EndGraphEntry;
-	if (CommandOptions.DumpFlags & BPDUMP_PinContextActions)
-	{
-		FString const PinActionsEntry = "," + IndentedNewline + IndentedNewline + "\"PinContextActions\" : [\n";
-		FileOutWriter->Serialize(TCHAR_TO_ANSI(*PinActionsEntry), PinActionsEntry.Len());
-
-		if (DumpPinContextActions(Indent + 1, Graph, FileOutWriter))
-		{
-			EndGraphEntry = IndentedNewline + "]";
-		}
-		else
-		{
-			EndGraphEntry = NestedIndent + "]";
-		}
-	}
-
-	EndGraphEntry += "\n" + GraphEntryIndent + "}";
-	FileOutWriter->Serialize(TCHAR_TO_ANSI(*EndGraphEntry), EndGraphEntry.Len());
-}
-
-//------------------------------------------------------------------------------
-static bool DumpBlueprintInfoUtils::DumpPinContextActions(uint32 Indent, UEdGraph* Graph, FArchive* FileOutWriter)
-{
-	FGraphContextMenuBuilder ContextMenuBuilder(Graph);
-
-	bool bWroteToFile = false;
-	if (!CommandOptions.PinType.IsNone())
-	{
-		FEdGraphPinType PinType;
-		PinType.PinCategory = CommandOptions.PinType;
-
-		bool const bUsePinTypeClass = (CommandOptions.DumpFlags & BPDUMP_PinTypeIsClass) != 0;
-
-		bool bIsValidPinType = true;	
-		if (UClass* TypeClass = FindObject<UClass>(ANY_PACKAGE, *CommandOptions.PinType.ToString()))
-		{
-			bIsValidPinType = UEdGraphSchema_K2::IsAllowableBlueprintVariableType(TypeClass);
-			if (bIsValidPinType)
-			{
-				PinType.PinSubCategoryObject = TypeClass;
-				if (TypeClass->IsChildOf(UInterface::StaticClass()))
-				{
-					PinType.PinCategory = UEdGraphSchema_K2::PC_Interface;
-				}
-				else
-				{
-					PinType.PinCategory = UEdGraphSchema_K2::PC_Object;
-				}
-
-				if (bUsePinTypeClass)
-				{
-					PinType.PinCategory = UEdGraphSchema_K2::PC_Class;
-				}
-			}
-			
-		}
-		else if (UScriptStruct* StructType = FindObject<UScriptStruct>(ANY_PACKAGE, *CommandOptions.PinType.ToString()))
-		{
-			PinType.PinCategory = UEdGraphSchema_K2::PC_Struct;
-			PinType.PinSubCategoryObject = StructType;
-			bIsValidPinType = UEdGraphSchema_K2::IsAllowableBlueprintVariableType(StructType);
-
-			UEdGraphSchema const* Schema = Graph->GetSchema();
-			if (UAnimationGraphSchema const* AnimSchema = Cast<UAnimationGraphSchema>(Schema))
-			{
-				bIsValidPinType |= AnimSchema->IsPosePin(PinType);
-			}
-		}
-		else if (CommandOptions.PinType == TEXT("self"))
-		{
-			PinType.PinCategory = UEdGraphSchema_K2::PC_Object;
-			if (bUsePinTypeClass)
-			{
-				PinType.PinCategory = UEdGraphSchema_K2::PC_Class;
-			}
-			PinType.PinSubCategory = UEdGraphSchema_K2::PSC_Self;
-			bIsValidPinType = true;
-		}
-		else if ((CommandOptions.PinType != UEdGraphSchema_K2::PC_Delegate) ||
-		         (CommandOptions.PinType != UEdGraphSchema_K2::PC_MCDelegate))
-		{
-			// @TODO: PC_Delegate, PC_MCDelegate
-			bIsValidPinType = false;
-		}
-		
-		if (bIsValidPinType)
-		{
-			DumpContextualPinTypeActions(Indent, Graph, PinType, FileOutWriter);
-			FileOutWriter->Serialize(TCHAR_TO_ANSI(TEXT(",\n")), 2);
-
-			PinType.bIsReference = true;
-			DumpContextualPinTypeActions(Indent, Graph, PinType, FileOutWriter);
-			FileOutWriter->Serialize(TCHAR_TO_ANSI(TEXT(",\n")), 2);
-			PinType.bIsReference = false;
-
-			PinType.ContainerType = EPinContainerType::Array;
-			DumpContextualPinTypeActions(Indent, Graph, PinType, FileOutWriter);
-			FileOutWriter->Serialize(TCHAR_TO_ANSI(TEXT(",\n")), 2);
-
-			PinType.ContainerType = EPinContainerType::Set;
-			DumpContextualPinTypeActions(Indent, Graph, PinType, FileOutWriter);
-			FileOutWriter->Serialize(TCHAR_TO_ANSI(TEXT(",\n")), 2);
-
-			PinType.ContainerType = EPinContainerType::Map;
-			DumpContextualPinTypeActions(Indent, Graph, PinType, FileOutWriter);
-
-			bWroteToFile = true;
-		}
-	}
-	else if (Graph->Schema->IsChildOf(UEdGraphSchema_K2::StaticClass()))
-	{
-		TArray< TSharedPtr<UEdGraphSchema_K2::FPinTypeTreeInfo> > TypeTree;
-		UEdGraphSchema_K2 const* GraphSchema = GetDefault<UEdGraphSchema_K2>(Graph->Schema);
-		GraphSchema->GetVariableTypeTree(TypeTree, ETypeTreeFilter::AllowExec | ETypeTreeFilter::AllowWildcard);
-
-		for (TSharedPtr<UEdGraphSchema_K2::FPinTypeTreeInfo> TypeInfo : TypeTree)
-		{
-			bWroteToFile |= DumpTypeTreeActions(Indent, Graph, TypeInfo, FileOutWriter);
-		}
-	}
-	else // state machine graph?
-	{
-		// look in graph for nodes and mine their pin types?
-	}
-
-	return bWroteToFile;
-}
-
-//------------------------------------------------------------------------------
-static bool DumpBlueprintInfoUtils::DumpTypeTreeActions(uint32 Indent, UEdGraph* Graph, TSharedPtr<UEdGraphSchema_K2::FPinTypeTreeInfo> PinTypeInfo, FArchive* FileOutWriter)
-{
-	FString PendingLineEnding;
-
-	bool bWroteToFile = false;
-	if (!PinTypeInfo->bReadOnly)
-	{
-		FEdGraphPinType PinType = PinTypeInfo->GetPinType(/*bForceLoadedSubCategoryObject =*/false);
-		DumpContextualPinTypeActions(Indent, Graph, PinType, FileOutWriter);
-		FileOutWriter->Serialize(TCHAR_TO_ANSI(TEXT(",\n")), 2);
-
-		PinType.bIsReference = true;
-		DumpContextualPinTypeActions(Indent, Graph, PinType, FileOutWriter);
-		FileOutWriter->Serialize(TCHAR_TO_ANSI(TEXT(",\n")), 2);
-		PinType.bIsReference = false;
-
-		PinType.ContainerType = EPinContainerType::Array;
-		DumpContextualPinTypeActions(Indent, Graph, PinType, FileOutWriter);
-		FileOutWriter->Serialize(TCHAR_TO_ANSI(TEXT(",\n")), 2);
-
-		PinType.ContainerType = EPinContainerType::Set;
-		DumpContextualPinTypeActions(Indent, Graph, PinType, FileOutWriter);
-		FileOutWriter->Serialize(TCHAR_TO_ANSI(TEXT(",\n")), 2);
-
-		PinType.ContainerType = EPinContainerType::Map;
-		DumpContextualPinTypeActions(Indent, Graph, PinType, FileOutWriter);
-
-		PendingLineEnding = TEXT(",\n");
-		bWroteToFile = true;
-	}
-
-	for (int32 ChildIndex = 0; ChildIndex < PinTypeInfo->Children.Num(); ++ChildIndex)
-	{
-		TSharedPtr<UEdGraphSchema_K2::FPinTypeTreeInfo> ChildInfo = PinTypeInfo->Children[ChildIndex];
-		if (ChildInfo.IsValid())
-		{
-			if (!PendingLineEnding.IsEmpty())
-			{
-				FileOutWriter->Serialize(TCHAR_TO_ANSI(*PendingLineEnding), PendingLineEnding.Len());
-				PendingLineEnding.Empty();
-			}
-
-			if (DumpTypeTreeActions(Indent, Graph, ChildInfo, FileOutWriter))
-			{
-				bWroteToFile |= true;
-				PendingLineEnding = TEXT(",\n");
-			}
-			// else, opps... we may be left with an invalid PendingLineEnding that was serialized
-		}
-	}
-
-	return bWroteToFile;
-}
-
-//------------------------------------------------------------------------------
-static void DumpBlueprintInfoUtils::DumpContextualPinTypeActions(uint32 Indent, UEdGraph* Graph, const FEdGraphPinType& PinType, FArchive* FileOutWriter)
-{
-	FGraphContextMenuBuilder ContextMenuBuilder(Graph);
-
-	UK2Node_Composite* DummyNode = NewObject<UK2Node_Composite>(Graph);
-	UEdGraphPin* DummyPin = DummyNode->CreatePin(EGPD_Input, PinType, DummyNode->GetFName());
-	ContextMenuBuilder.FromPin = DummyPin;
-
-	UBlueprint* Blueprint = FBlueprintEditorUtils::FindBlueprintForGraph(Graph);
-	DumpContextActionList(Indent, ContextMenuBuilder, Blueprint, FileOutWriter);
-	FileOutWriter->Serialize(TCHAR_TO_ANSI(TEXT(",\n")), 2);
-	DummyPin->Direction = EGPD_Output;
-	DumpContextActionList(Indent, ContextMenuBuilder, Blueprint, FileOutWriter);
-}
-
-//------------------------------------------------------------------------------
-static void DumpBlueprintInfoUtils::DumpContextActionList(uint32 Indent, FGraphContextMenuBuilder ActionBuilder, UBlueprint* Blueprint, FArchive* FileOutWriter)
-{
-	ActionBuilder.Empty();
-
-	if (ActionBuilder.FromPin != nullptr)
-	{
-		FString PinTypeLog("<NONE>");
-		if (ActionBuilder.FromPin->Direction == EGPD_Input)
-		{
-			PinTypeLog = "INPUT ";
-		}
-		else
-		{
-			PinTypeLog = "OUTPUT";
-		}
-		PinTypeLog += UEdGraphSchema_K2::TypeToText(ActionBuilder.FromPin->PinType).ToString();
-
-		UE_LOG(LogBlueprintInfoDump, Display, TEXT("%sDumping pin actions: %s"), *BuildIndentString(Indent, true), *PinTypeLog);
-	}
-	
-	double MenuBuildDuration = GetContextMenuActions(ActionBuilder, Blueprint);
-
-	FString const ContextEntryIndent = BuildIndentString(Indent);
-	FString BeginContextEntry = FString::Printf(TEXT("%s{\n"), *ContextEntryIndent);
-
-	if (CommandOptions.DumpFlags & BPDUMP_RecordTiming)
-	{
-		BeginContextEntry += FString::Printf(TEXT("%s\"MenuBuildTime_Seconds\" : %f,\n"), *BuildIndentString(Indent + 1), MenuBuildDuration);
-	}
-	FileOutWriter->Serialize(TCHAR_TO_ANSI(*BeginContextEntry), BeginContextEntry.Len());
-
-	DumpContextInfo(Indent + 1, ActionBuilder, Blueprint, FileOutWriter);
-	FileOutWriter->Serialize(TCHAR_TO_ANSI(TEXT(",\n")), 2);
-	DumpActionList(Indent + 1, ActionBuilder, FileOutWriter);
-
-	FString EndContextEntry = FString::Printf(TEXT("\n%s}"), *ContextEntryIndent);
-	FileOutWriter->Serialize(TCHAR_TO_ANSI(*EndContextEntry), EndContextEntry.Len());
-}
-
-//------------------------------------------------------------------------------
-static double DumpBlueprintInfoUtils::GetContextMenuActions(FGraphContextMenuBuilder ActionBuilder, UBlueprint* Blueprint)
-{
-	ActionBuilder.Empty();
-	check(ActionBuilder.CurrentGraph != nullptr);
-
-	double MenuBuildDuration = 0.0;
-	
-	FBlueprintActionContext FilterContext;
-	FilterContext.Blueprints.Add(const_cast<UBlueprint*>(Blueprint));
-	FilterContext.Graphs.Add(const_cast<UEdGraph*>(ActionBuilder.CurrentGraph));
-
-	if (ActionBuilder.FromPin != nullptr)
-	{
-		FilterContext.Pins.Add(const_cast<UEdGraphPin*>(ActionBuilder.FromPin));
-	}
-
-	TArray<UProperty*> SelectedProperties;
-	for (UObject* SelectedObj : ActionBuilder.SelectedObjects)
-	{
-		if (UProperty* SelectedProperty = Cast<UObjectProperty>(SelectedObj))
-		{
-			FilterContext.SelectedObjects.Add(SelectedProperty);
-		}
-	}
-
-	FBlueprintActionMenuBuilder MenuBuilder(nullptr);
-	{
-		FScopedDurationTimer DurationTimer(MenuBuildDuration);
-		FBlueprintActionMenuUtils::MakeContextMenu(FilterContext, /*bIsContextSensitive =*/true, 
-			(EContextTargetFlags::TARGET_Blueprint | EContextTargetFlags::TARGET_NodeTarget | EContextTargetFlags::TARGET_PinObject | EContextTargetFlags::TARGET_SiblingPinObjects), MenuBuilder);
-	}
-	ActionBuilder.Append(MenuBuilder);
-
-	return MenuBuildDuration;
-}
-
-//------------------------------------------------------------------------------
-static void DumpBlueprintInfoUtils::DumpContextInfo(uint32 Indent, FGraphContextMenuBuilder ActionBuilder, UBlueprint* Blueprint, FArchive* FileOutWriter)
-{
-	FString const ContextEntryIndent = BuildIndentString(Indent);
-	++Indent;
-	FString const IndentedNewline = "\n" + BuildIndentString(Indent);
-	
-	FString ContextEntry = FString::Printf(TEXT("%s\"Context\" : {%s\"Graph\"   : \""), *ContextEntryIndent, *IndentedNewline);
-	ContextEntry += GetGraphTypeString(ActionBuilder.CurrentGraph) + "\"," + IndentedNewline + "\"PinType\" : \"";
-
-	if (ActionBuilder.FromPin == nullptr)
-	{
-		ContextEntry += "<NONE>\"," + IndentedNewline + "\"PinDirection\" : \"<NONE>";
-	}
-	else
-	{
-		ContextEntry += UEdGraphSchema_K2::TypeToText(ActionBuilder.FromPin->PinType).ToString();
-
-		ContextEntry += "\"," + IndentedNewline + "\"PinDirection\" : \"";
-		if (ActionBuilder.FromPin->Direction == EGPD_Input)
-		{
-			ContextEntry += "Input";
-		}
-		else
-		{
-			check(ActionBuilder.FromPin->Direction == EGPD_Output);
-			ContextEntry += "Output";
-		}
-	}	
-	ContextEntry += "\"," + IndentedNewline + "\"SelectedObjects\" : [";
-	
-	FString const ActorEntryIndent = "\n" + BuildIndentString(Indent + 1);
-	for (UObject* Selection : ActionBuilder.SelectedObjects)
-	{
-		ContextEntry += ActorEntryIndent + "\"" + CastChecked<UObjectProperty>(Selection)->PropertyClass->GetName() + "\",";
-	}
-	if (ActionBuilder.SelectedObjects.Num() > 0)
-	{
-		ContextEntry.RemoveAt(ContextEntry.Len() - 1); // remove the trailing ','
-		ContextEntry += IndentedNewline;
-	}
-	ContextEntry += "]\n" + ContextEntryIndent + "}";
-	
-	FileOutWriter->Serialize(TCHAR_TO_ANSI(*ContextEntry), ContextEntry.Len());
-}
-
-//------------------------------------------------------------------------------
-// @TODO: should probably be a platform agnostic function
-static bool DumpBlueprintInfoUtils::DoFilesDiffer(FString const& NewFilePath, FString const& OldFilePath)
-{
-	int32 ReturnCode = -1;
-	FString const CommandStr = "fc " + NewFilePath.Replace(TEXT("/"), TEXT("\\")) + " " + OldFilePath.Replace(TEXT("/"), TEXT("\\"));
-
-	// @TODO: very platform specific :(
-	bool bSuccess = FPlatformProcess::ExecProcess(
-		TEXT("cmd.exe"), 
-		*("/C \"" + CommandStr + "\""), 
-		&ReturnCode,
-		/*OutStdOut =*/nullptr, 
-		/*OutStdErr =*/nullptr);
-
-	if (!bSuccess)
-	{
-		UE_LOG(LogBlueprintInfoDump, Error, TEXT("Failed to run diff: '%s'"), *CommandStr);
-	}
-	
-	// ReturnCode
-	// 0 : no differences?
-	// 1 : differences?
-	// 2 : couldn't find file? 
-
-	return !bSuccess || (ReturnCode != 0);
-}
-
-//------------------------------------------------------------------------------
-static void DumpBlueprintInfoUtils::DiffDumpFiles(FString const& NewFilePath, FString const& OldFilePath, FString const& UserDiffCmd)
-{
-	check(!NewFilePath.IsEmpty());
-	check(FPaths::FileExists(NewFilePath));
-	check(!OldFilePath.IsEmpty() || !UserDiffCmd.IsEmpty());
-
-	UE_LOG(LogBlueprintInfoDump, Display, TEXT("Diffing: '%s'..."), *NewFilePath);
-
-	FString WorkingDirectory = FPaths::GetPath(NewFilePath);
-	// if we created a whole new folder for this file, then make the working 
-	// directory one up
-	bool const bSplitBlueprintsByFile = (CommandOptions.DumpFlags & BPDUMP_FilePerBlueprint) != 0;
-	if (bSplitBlueprintsByFile)
-	{
-		WorkingDirectory = WorkingDirectory / TEXT("..");
-	}
-	
-	FString QualifiedOldFilePath;
-	if (!OldFilePath.IsEmpty())
-	{
-		IFileManager& FileManager = IFileManager::Get();
-		bool const bIsLocalDirectory = FileManager.DirectoryExists(*(WorkingDirectory / *OldFilePath));
-		bool const bIsDirectoryPath  = !bIsLocalDirectory && FileManager.DirectoryExists(*OldFilePath);
-		
-		if (bIsLocalDirectory)
-		{
-			FString Filename = FPaths::GetCleanFilename(NewFilePath);
-			QualifiedOldFilePath = WorkingDirectory / *(OldFilePath / *Filename);
-		}
-		else if (bIsDirectoryPath)
-		{
-			FString Filename = FPaths::GetCleanFilename(NewFilePath);
-			QualifiedOldFilePath = OldFilePath / *Filename;
-		}
-		else
-		{
-			QualifiedOldFilePath = OldFilePath;
-			if (FPaths::GetPath(QualifiedOldFilePath).IsEmpty())
-			{
-				QualifiedOldFilePath = WorkingDirectory / OldFilePath;
-			}
-		}
-	}
-
-	if (!FPaths::FileExists(QualifiedOldFilePath) && !QualifiedOldFilePath.IsEmpty())
-	{
-		UE_LOG(LogBlueprintInfoDump, Error, TEXT("Cannot find file '%s' to diff against"), *QualifiedOldFilePath);
-	}
-	else if (!bSplitBlueprintsByFile || QualifiedOldFilePath.IsEmpty() || DoFilesDiffer(NewFilePath, QualifiedOldFilePath))
-	{
-		FString DiffCommand = UserDiffCmd;
-		if (DiffCommand.IsEmpty())
-		{
-			// -dw : Ignore line ending and all whitespace differences
-			DiffCommand = TEXT("p4merge.exe -dw \"{2}\" \"{1}\"");
-		}
-		
-		FString DiffArgs;
-		if (int32 ArgsIndex = DiffCommand.Find(TEXT(" ")))
-		{
-			DiffArgs = *DiffCommand + ArgsIndex + 1;
-			DiffCommand = DiffCommand.Left(ArgsIndex);
-		}
-		
-		DiffArgs = DiffArgs.Replace(TEXT("{1}"), *NewFilePath);
-		DiffArgs = DiffArgs.Replace(TEXT("{2}"), *QualifiedOldFilePath);
-
-		FProcHandle DiffProc = FPlatformProcess::CreateProc(*DiffCommand,
-			/*Params =*/              *DiffArgs,
-			/*bLaunchDetached =*/     true,
-			/*bLaunchHidden =*/       false,
-			/*bLaunchReallyHidden =*/ false,
-			/*OutProcessID =*/        nullptr,
-			/*PriorityModifier =*/    0,
-			*WorkingDirectory,
-			/*PipeWrite =*/           nullptr);
-
-		if (!DiffProc.IsValid())
-		{
-			UE_LOG(LogBlueprintInfoDump, Error, TEXT("Could not launch: '%s'"), *DiffCommand);
-		}
-	}
-}
-
-//------------------------------------------------------------------------------
-static UClass* DumpBlueprintInfoUtils::GetUserNamedClass(FString ClassName)
-{
-	UClass* FoundClass = FindObject<UClass>(ANY_PACKAGE, *ClassName);
-
-	if (FoundClass == nullptr)
-	{
-		// maybe they meant a blueprint?
-		UBlueprint* Blueprint = FindObject<UBlueprint>(ANY_PACKAGE, *ClassName);
-		// maybe we have to load that blueprint?
-		if (Blueprint == nullptr)
-		{
-			// if this loaded something...
-			if (LoadBlueprints(ClassName, /*bAlloDevBlueprints =*/true) > 0)
-			{
-				Blueprint = FindObject<UBlueprint>(ANY_PACKAGE, *ClassName);
-				if (Blueprint != nullptr)
-				{
-					FoundClass = Blueprint->GeneratedClass;
-				}
-				else
-				{
-					FoundClass = FindObject<UClass>(ANY_PACKAGE, *ClassName);
-				}
-			}
-		}
-	}
-
-	return FoundClass;
-}
-
-//------------------------------------------------------------------------------
-static int32 DumpBlueprintInfoUtils::LoadBlueprints(FString AssetName, bool bAllowDevBlueprints)
-{
-	int32 LoadedCount = 0;
-
-	TArray<FString> PackagesToLoad;
-	if (!AssetName.IsEmpty())
-	{
-		FString PackageName = FString::Printf(TEXT("*%s*%s"), *AssetName, *FPackageName::GetAssetPackageExtension());
-
-		TArray<FString> UnusedPackageNames;
-		NormalizePackageNames(UnusedPackageNames, PackagesToLoad, PackageName, NORMALIZE_ExcludeMapPackages);
-
-		for (FString& FilePath : PackagesToLoad)
-		{
-			FString PackagePath;
-			if (FPackageName::TryConvertFilenameToLongPackageName(FilePath, PackagePath))
-			{
-				FilePath = PackagePath;
-			}
-		}
-	}
-
-	FAssetRegistryModule& AssetRegistryModule = FModuleManager::LoadModuleChecked<FAssetRegistryModule>(TEXT("AssetRegistry"));
-
-	static bool bAssetRegistryLoaded = false;
-	if (!bAssetRegistryLoaded)
-	{
-		UE_LOG(LogBlueprintInfoDump, Display, TEXT("Loading the asset registry..."));
-		AssetRegistryModule.Get().SearchAllAssets(/*bSynchronousSearch =*/true);
-		UE_LOG(LogBlueprintInfoDump, Display, TEXT("Finished loading the asset registry."));
-		bAssetRegistryLoaded = true;
-	}
-	
-	TArray<FAssetData> BlueprintAssetList;
-	AssetRegistryModule.Get().GetAssetsByClass(UBlueprint::StaticClass()->GetFName(), BlueprintAssetList);
-	AssetRegistryModule.Get().GetAssetsByClass(UAnimBlueprint::StaticClass()->GetFName(), BlueprintAssetList);
-
- 	FString DevelopersRoot = FPaths::GameDevelopersDir().LeftChop(1);
-	FPackageName::TryConvertFilenameToLongPackageName(DevelopersRoot, DevelopersRoot);
-
-	bool const bLoadPackageSubset = (PackagesToLoad.Num() > 0);	
-	int32 AlreadyLoadedCount = 0;
-
-	for (FAssetData const& Asset : BlueprintAssetList)
-	{
-		if (Asset.IsAssetLoaded())
-		{
-			++AlreadyLoadedCount;
-			continue;
-		}
-
-		FString const PackageName = Asset.PackageName.ToString();
-		if (!bAllowDevBlueprints && PackageName.StartsWith(DevelopersRoot))
-		{
-			continue;
-		}
-
-		if (bLoadPackageSubset && !PackagesToLoad.Contains(PackageName))
-		{
-			continue;
-		}
-
-		FString const AssetPath = Asset.ObjectPath.ToString();
-		UE_LOG(LogBlueprintInfoDump, Display, TEXT("Loading '%s'..."), *AssetPath);
-
-		UBlueprint* LoadedBlueprint = Cast<UBlueprint>(StaticLoadObject(Asset.GetClass(), /*Outer =*/nullptr, *AssetPath));
-		if (LoadedBlueprint != nullptr)
-		{
-			++LoadedCount;
-		}
-		else
-		{
-			UE_LOG(LogBlueprintInfoDump, Warning, TEXT("Failed to load: '%s'."), *AssetPath);
-		}
-
-		if (bLoadPackageSubset && (LoadedCount >= PackagesToLoad.Num()))
-		{
-			break;
-		}
- 	}
-
-	if (!bLoadPackageSubset || (PackagesToLoad.Num() > 0))
-	{
-		int32 const TotalLoaded = LoadedCount + AlreadyLoadedCount;
-		int32 const AttemptedLoadCount = bLoadPackageSubset ? PackagesToLoad.Num() : BlueprintAssetList.Num();
-		UE_LOG(LogBlueprintInfoDump, Warning, TEXT("Successfully loaded %d/%d Blueprints (%.1f%%)."), TotalLoaded, AttemptedLoadCount, 100.f * (float)TotalLoaded / AttemptedLoadCount);
-	}
-
-	return LoadedCount;
-}
-
 /*******************************************************************************
  * UDumpBlueprintsInfoCommandlet
  ******************************************************************************/
@@ -2435,153 +18,7 @@ UDumpBlueprintsInfoCommandlet::UDumpBlueprintsInfoCommandlet(FObjectInitializer
 //------------------------------------------------------------------------------
 int32 UDumpBlueprintsInfoCommandlet::Main(FString const& Params)
 {
-	bool const bCachedIsRunning = GIsRunning;
-	GIsRunning = true;
-
-	TArray<FString> Tokens, Switches;
-	ParseCommandLine(*Params, Tokens, Switches);
-
-	DumpBlueprintInfoUtils::LevelActors.Empty();
-
-	DumpBlueprintInfoUtils::CommandOptions = DumpBlueprintInfoUtils::CommandletOptions(Switches);
-	DumpBlueprintInfoUtils::CommandletOptions const& CommandOptions = DumpBlueprintInfoUtils::CommandOptions;
-
-	bool const bSplitFilesByClass = (CommandOptions.DumpFlags & DumpBlueprintInfoUtils::BPDUMP_FilePerBlueprint) != 0;
-	bool const bDiffGeneratedFile = !CommandOptions.DiffPath.IsEmpty() || !CommandOptions.DiffCommand.IsEmpty();
-
-	UClass* const SelectedObjType = DumpBlueprintInfoUtils::CommandOptions.SelectedObjectType;
-	// if the user specified that they want a level actor selected during the 
-	// dump, then spawn one and select it (extra blueprint context actions 
-	// appear in certain situations regarding selected level actors)
-	if ((SelectedObjType != nullptr) && SelectedObjType->IsChildOf(AActor::StaticClass()))
-	{
-		DumpBlueprintInfoUtils::SpawnLevelActor(SelectedObjType, /*bSelect =*/true);
-	}
-
-	FString ActiveFilePath;
-	FArchive* FileOut = nullptr;
-
-	// this lambda is responsible for adding closing characters to the file, and 
-	// closing out the writer (and diffing the resultant file if the user deigns us to do so)
-	auto CloseFileStream = [bDiffGeneratedFile, &ActiveFilePath, &CommandOptions](FArchive** FileOutPtr)
-	{
-		FArchive*& ArFileOut = (*FileOutPtr);
-		if (ArFileOut != nullptr)
-		{
-			ArFileOut->Serialize(TCHAR_TO_ANSI(TEXT("\n}")), 2);
-			ArFileOut->Close();
-
-			if (bDiffGeneratedFile)
-			{
-				check(!ActiveFilePath.IsEmpty());
-				DumpBlueprintInfoUtils::DiffDumpFiles(ActiveFilePath, CommandOptions.DiffPath, CommandOptions.DiffCommand);
-			}
-			delete ArFileOut;
-			ArFileOut = nullptr;
-		}
-	};
-
-	
-	// this lambda is responsible for opening a file for write, and adding 
-	// opening json characters to the file (contextually tracks the active filepath as well)
-	auto OpenFileStream = [&ActiveFilePath, &CloseFileStream, &FileOut](UClass* Class)->FArchive*
-	{
-		CloseFileStream(&FileOut);
-
-		ActiveFilePath = DumpBlueprintInfoUtils::BuildDumpFilePath(Class);
-		FileOut = IFileManager::Get().CreateFileWriter(*ActiveFilePath);
-		check(FileOut != nullptr);
-		FileOut->Serialize(TCHAR_TO_ANSI(TEXT("{\n")), 2);
-
-		if (DumpBlueprintInfoUtils::DumpActionDatabaseInfo(1, FileOut))
-		{
-			FileOut->Serialize(TCHAR_TO_ANSI(TEXT(",\n")), 2);
-		}
-		return FileOut;
-	};
-	
-	bool bNeedsJsonComma = false;
-	// this lambda will dump blueprint info for the specified class, if the user
-	// set -multifile, then this will also close the existing file and open a new one for this class
-	auto WriteClassInfo = [bSplitFilesByClass, &bNeedsJsonComma, &OpenFileStream, &CloseFileStream, &FileOut](UClass* Class)
-	{
-		if (bSplitFilesByClass && (FileOut != nullptr))
-		{
-			CloseFileStream(&FileOut);
-		}
-
-		if (FileOut == nullptr)
-		{
-			FileOut = OpenFileStream(Class);
-		}
-		// if we're adding all the class entries into one file, then we need to 
-		// separate them by a comma (or invalid json)
-		else if (!bSplitFilesByClass && bNeedsJsonComma)
-		{
-			FileOut->Serialize(TCHAR_TO_ANSI(TEXT(",\n")), 2);
-		}
-
-		DumpBlueprintInfoUtils::DumpInfoForClass(1, Class, FileOut);
-		bNeedsJsonComma = true;
-	};
-
-	// this lambda is used as a precursory check to ensure that the specified 
-	// class is a blueprintable type... broken into it's own lambda to save on reuse
-	auto IsInvalidBlueprintClass = [](UClass* Class)->bool
-	{
-		return !IsValid(Class) ||
-			Class->HasAnyClassFlags(CLASS_NewerVersionExists) ||
-			FKismetEditorUtilities::IsClassABlueprintSkeleton(Class)  ||
-			!FKismetEditorUtilities::CanCreateBlueprintOfClass(Class) ||
-			(Class->GetOuterUPackage() == GetTransientPackage());
-	};
-
-	UClass* const BlueprintClass = CommandOptions.BlueprintClass;
-	if (CommandOptions.DumpFlags & DumpBlueprintInfoUtils::BPDUMP_LogHelp)
-	{
-		UE_LOG(LogBlueprintInfoDump, Display, TEXT("%s"), *DumpBlueprintInfoUtils::HelpString);
-	}
-	else if (BlueprintClass != nullptr)
-	{
-		UE_LOG(LogBlueprintInfoDump, Display, TEXT("Dumping Blueprint info..."));
-		// make sure the class that the user specified is a blueprintable type
-		if (IsInvalidBlueprintClass(BlueprintClass))
-		{
-			UE_LOG(LogBlueprintInfoDump, Error, TEXT("Cannot create a blueprint from class: '%s'"), *BlueprintClass->GetName());
-			if (FileOut != nullptr)
-			{
-				FString const InvalidClassEntry = DumpBlueprintInfoUtils::BuildIndentString(1) + "\"INVALID_BLUEPRINT_CLASS\" : \"" + BlueprintClass->GetName() + "\"";
-				FileOut->Serialize(TCHAR_TO_ANSI(*InvalidClassEntry), InvalidClassEntry.Len());
-			}
-		}
-		else
-		{
-			WriteClassInfo(BlueprintClass);
-		}
-	}
-	// if the user didn't specify a class, then we take that to mean dump ALL the classes!
-	else
-	{
-		UE_LOG(LogBlueprintInfoDump, Display, TEXT("Dumping Blueprint info..."));
-		for (TObjectIterator<UClass> ClassIt; ClassIt; ++ClassIt)
-		{
-			UClass* Class = *ClassIt;
-			if (IsInvalidBlueprintClass(Class))
-			{
-				continue;
-			}
-			WriteClassInfo(Class);
-		}
-	}
-
-	CloseFileStream(&FileOut);
-
-	// restore the globals that we forcefully overrode
-	GIsRunning = bCachedIsRunning;
-
+	UE_LOG(LogBlueprintInfoDump, Error, TEXT("DumpBlueprintsInfo has been removed - consider using the GenerateBlueprintAPI commandlet instead\n"));
 	return 0;
 }
 
-
-
-
