Fix "launch on" not correctly checking whether the project needs to be rebuilt.

#rnx

diff --git a/Engine/Source/Editor/UnrealEd/Private/PlayLevel.cpp b/Engine/Source/Editor/UnrealEd/Private/PlayLevel.cpp
--- a/Engine/Source/Editor/UnrealEd/Private/PlayLevel.cpp
+++ b/Engine/Source/Editor/UnrealEd/Private/PlayLevel.cpp
@@ -2001,13 +2001,34 @@ void UEditorEngine::PlayUsingLauncher()
 			}
 		}
 
+		// set the build/launch configuration 
+		EBuildConfiguration BuildConfiguration;
+		const ULevelEditorPlaySettings* PlayInSettings = GetDefault<ULevelEditorPlaySettings>();
+		switch (PlayInSettings->LaunchConfiguration)
+		{
+		case LaunchConfig_Debug:
+			BuildConfiguration = EBuildConfiguration::Debug;
+			break;
+		case LaunchConfig_Development:
+			BuildConfiguration = EBuildConfiguration::Development;
+			break;
+		case LaunchConfig_Test:
+			BuildConfiguration = EBuildConfiguration::Test;
+			break;
+		case LaunchConfig_Shipping:
+			BuildConfiguration = EBuildConfiguration::Shipping;
+			break;
+		default:
+			// same as the running editor
+			BuildConfiguration = FApp::GetBuildConfiguration();
+			break;
+		}
+
 		// does the project have any code?
 		FGameProjectGenerationModule& GameProjectModule = FModuleManager::LoadModuleChecked<FGameProjectGenerationModule>(TEXT("GameProjectGeneration"));
 		bPlayUsingLauncherHasCode = GameProjectModule.Get().ProjectHasCodeFiles();
 
-		const ULevelEditorPlaySettings* PlayInSettings = GetDefault<ULevelEditorPlaySettings>();
-		// Setup launch profile, keep the setting here to a minimum.
-		ILauncherProfileRef LauncherProfile = LauncherServicesModule.CreateProfile(TEXT("Launch On Device"));
+		// Figure out if we need to build anything
 		if(PlayInSettings->BuildGameBeforeLaunch == EPlayOnBuildMode::PlayOnBuild_Always)
 		{
 			bPlayUsingLauncherBuild = true;
@@ -2019,33 +2040,25 @@ void UEditorEngine::PlayUsingLauncher()
 		else if(PlayInSettings->BuildGameBeforeLaunch == EPlayOnBuildMode::PlayOnBuild_Default)
 		{
 			bPlayUsingLauncherBuild = bPlayUsingLauncherHasCode || !FApp::GetEngineIsPromotedBuild();
+			if (!bPlayUsingLauncherBuild)
+			{
+				FText Reason;
+				if (LaunchPlatform->RequiresTempTarget(bPlayUsingLauncherHasCode, BuildConfiguration, false, Reason))
+				{
+					UE_LOG(LogPlayLevel, Log, TEXT("Project requires temp target (%s)"), *Reason.ToString());
+					bPlayUsingLauncherBuild = true;
+				}
+			}
 		}
 		else if(PlayInSettings->BuildGameBeforeLaunch == EPlayOnBuildMode::PlayOnBuild_IfEditorBuiltLocally)
 		{
 			bPlayUsingLauncherBuild = !FApp::GetEngineIsPromotedBuild();
 		}
-		LauncherProfile->SetBuildGame(bPlayUsingLauncherBuild);
 
-		// set the build/launch configuration 
-		switch (PlayInSettings->LaunchConfiguration)
-		{
-		case LaunchConfig_Debug:
-			LauncherProfile->SetBuildConfiguration(EBuildConfiguration::Debug);
-			break;
-		case LaunchConfig_Development:
-			LauncherProfile->SetBuildConfiguration(EBuildConfiguration::Development);
-			break;
-		case LaunchConfig_Test:
-			LauncherProfile->SetBuildConfiguration(EBuildConfiguration::Test);
-			break;
-		case LaunchConfig_Shipping:
-			LauncherProfile->SetBuildConfiguration(EBuildConfiguration::Shipping);
-			break;
-		default:
-			// same as the running editor
-			LauncherProfile->SetBuildConfiguration(FApp::GetBuildConfiguration());
-			break;
-		}
+		// Setup launch profile, keep the setting here to a minimum.
+		ILauncherProfileRef LauncherProfile = LauncherServicesModule.CreateProfile(TEXT("Launch On Device"));
+		LauncherProfile->SetBuildGame(bPlayUsingLauncherBuild);
+		LauncherProfile->SetBuildConfiguration(BuildConfiguration);
 
 		// select the quickest cook mode based on which in editor cook mode is enabled
 		bool bIncrimentalCooking = true;
