Prevent component type registry from querying the asset registry while the asset registry is still populating. This seems to fix a cause of asset discovery slowdown

diff --git a/Engine/Source/Editor/UnrealEd/Private/ComponentTypeRegistry.cpp b/Engine/Source/Editor/UnrealEd/Private/ComponentTypeRegistry.cpp
--- a/Engine/Source/Editor/UnrealEd/Private/ComponentTypeRegistry.cpp
+++ b/Engine/Source/Editor/UnrealEd/Private/ComponentTypeRegistry.cpp
@@ -405,12 +405,20 @@ void FComponentTypeRegistryData::Tick(float)
 	if (PendingAssetData.Num() != 0)
 	{
 		FAssetRegistryModule& AssetRegistryModule = FModuleManager::LoadModuleChecked<FAssetRegistryModule>(TEXT("AssetRegistry"));
+
+		// Avoid querying the asset registry until it has finished discovery, 
+		// as doing so may force it to update temporary caches many times:
+		if(AssetRegistryModule.Get().IsLoadingAssets())
+		{
+			return;
+		}
+
 		TArray<FName> ClassNames;
 		ClassNames.Add(UActorComponent::StaticClass()->GetFName());
 		TSet<FName> DerivedClassNames;
 		AssetRegistryModule.Get().GetDerivedClassNames(ClassNames, TSet<FName>(), DerivedClassNames);
 
-		for (auto Asset : PendingAssetData)
+		for (const FAssetData& Asset : PendingAssetData)
 		{
 			const FName BPParentClassName(GET_MEMBER_NAME_CHECKED(UBlueprint, ParentClass));
 			const FString TagValue = Asset.GetTagValueRef<FString>(BPParentClassName);
