UpdateBytecodeReferences now also updates the blueprint's called functions tracking data - this prevents entries in the CalledFunctions list from getting stale
#rnx

diff --git a/Engine/Source/Editor/UnrealEd/Private/Kismet2/KismetReinstanceUtilities.cpp b/Engine/Source/Editor/UnrealEd/Private/Kismet2/KismetReinstanceUtilities.cpp
--- a/Engine/Source/Editor/UnrealEd/Private/Kismet2/KismetReinstanceUtilities.cpp
+++ b/Engine/Source/Editor/UnrealEd/Private/Kismet2/KismetReinstanceUtilities.cpp
@@ -858,6 +858,24 @@ void FBlueprintCompileReinstancer::UpdateBytecodeReferences()
 				bBPWasChanged |= (0 != ReplaceAr.GetCount());
 			}
 
+			// Update any refs in called functions array, as the bytecode was just similarly updated:
+			if(UBlueprintGeneratedClass* AsBPGC = Cast<UBlueprintGeneratedClass>(BPClass))
+			{
+				for(int32 Idx = 0; Idx < AsBPGC->CalledFunctions.Num(); ++Idx)
+				{
+					UObject** Val = FieldMappings.Find(AsBPGC->CalledFunctions[Idx]);
+					if(Val && *Val)
+					{
+						// This ::Cast should always succeed, but I'm uncomfortable making 
+						// rigid assumptions about the FieldMappings array:
+						if(UFunction* NewFn = Cast<UFunction>(*Val))
+						{
+							AsBPGC->CalledFunctions[Idx] = NewFn;
+						}
+					}
+				}
+			}
+
 			FArchiveReplaceObjectRef<UObject> ReplaceInBPAr(*DependentBP, FieldMappings, false, true, true);
 			if (ReplaceInBPAr.GetCount())
 			{
