Mirror 9001036 from Dev-Framework
CompileAllBlueprints commandlet now reinstances, ensuring that it does not later crash during GC - this cl also fixes some errors in

diff --git a/Engine/Source/Editor/UnrealEd/Private/Commandlets/CompileAllBlueprintsCommandlet.cpp b/Engine/Source/Editor/UnrealEd/Private/Commandlets/CompileAllBlueprintsCommandlet.cpp
--- a/Engine/Source/Editor/UnrealEd/Private/Commandlets/CompileAllBlueprintsCommandlet.cpp
+++ b/Engine/Source/Editor/UnrealEd/Private/Commandlets/CompileAllBlueprintsCommandlet.cpp
@@ -304,6 +304,11 @@ void UCompileAllBlueprintsCommandlet::CompileBlueprint(UBlueprint* Blueprint)
 		MessageLog.SetSourcePath(Blueprint->GetPathName());
 		MessageLog.BeginEvent(TEXT("Compile"));
 
+		if (bCompileSkeletonOnly)
+		{
+			CompileOptions.CompileType = EKismetCompileType::SkeletonOnly;
+		}
+
 		FKismetEditorUtilities::CompileBlueprint(Blueprint, EBlueprintCompileOptions::SkipGarbageCollection, &MessageLog);
 
 		MessageLog.EndEvent();
