Mirror 9001036 and 9240957 from Dev-Framework, just to get //UE4/Main into a coherent state
Remove CompileSkeletonOnly option from CompileAllBlueprints + missing include
#rnx

diff --git a/Engine/Source/Editor/UnrealEd/Private/Commandlets/CompileAllBlueprintsCommandlet.cpp b/Engine/Source/Editor/UnrealEd/Private/Commandlets/CompileAllBlueprintsCommandlet.cpp
--- a/Engine/Source/Editor/UnrealEd/Private/Commandlets/CompileAllBlueprintsCommandlet.cpp
+++ b/Engine/Source/Editor/UnrealEd/Private/Commandlets/CompileAllBlueprintsCommandlet.cpp
@@ -304,11 +304,6 @@ void UCompileAllBlueprintsCommandlet::CompileBlueprint(UBlueprint* Blueprint)
 		MessageLog.SetSourcePath(Blueprint->GetPathName());
 		MessageLog.BeginEvent(TEXT("Compile"));
 
-		if (bCompileSkeletonOnly)
-		{
-			CompileOptions.CompileType = EKismetCompileType::SkeletonOnly;
-		}
-
 		FKismetEditorUtilities::CompileBlueprint(Blueprint, EBlueprintCompileOptions::SkipGarbageCollection, &MessageLog);
 
 		MessageLog.EndEvent();
