Fixed a crash while Reloading a Dataprep Asset.

The root cause is that the function UPackageTools::HandlePackageReloaded assume that all the uclass in a package are compatible with the function FBlueprintCompilationManager::ReparentHierarchies.
This function only support the blueprint generated classes and the dataprep asset use a custom UClass as a way to store arbitrary parameterization data on a object.

diff --git a/Engine/Source/Editor/UnrealEd/Private/PackageTools.cpp b/Engine/Source/Editor/UnrealEd/Private/PackageTools.cpp
--- a/Engine/Source/Editor/UnrealEd/Private/PackageTools.cpp
+++ b/Engine/Source/Editor/UnrealEd/Private/PackageTools.cpp
@@ -43,10 +43,11 @@
 #include "Logging/MessageLog.h"
 #include "UObject/UObjectIterator.h"
 #include "ComponentReregisterContext.h"
-#include "Engine/Selection.h"
+#include "Engine/BlueprintGeneratedClass.h"
 #include "Engine/GameEngine.h"
 #include "Engine/LevelStreaming.h"
 #include "Engine/MapBuildDataRegistry.h"
+#include "Engine/Selection.h"
 
 #include "ShaderCompiler.h"
 #include "DistanceFieldAtlas.h"
@@ -861,7 +862,8 @@ UPackageTools::UPackageTools(const FObjectInitializer& ObjectInitializer)
 
 				if(OldObject && NewObject)
 				{
-					UClass* OldObjectAsClass = Cast<UClass>(OldObject);
+					// Only the blueprint generated class are supported by the FBlueprintCompilationManager so we only reparent those
+					UClass* OldObjectAsClass = Cast<UBlueprintGeneratedClass>(OldObject);
 					if(OldObjectAsClass)
 					{
 						UClass* NewObjectAsClass = Cast<UClass>(NewObject);
