ResavePackagesCommandlet: added option to rebuild navigation data. This option will checkout/save/checkin only dirty packages
ANavigationData now marks package as dirty when rebuilding

diff --git a/Engine/Source/Editor/UnrealEd/Classes/Commandlets/ResavePackagesCommandlet.h b/Engine/Source/Editor/UnrealEd/Classes/Commandlets/ResavePackagesCommandlet.h
--- a/Engine/Source/Editor/UnrealEd/Classes/Commandlets/ResavePackagesCommandlet.h
+++ b/Engine/Source/Editor/UnrealEd/Classes/Commandlets/ResavePackagesCommandlet.h
@@ -98,6 +98,9 @@ protected:
 	/** Only save packages that been saved by a licensee **/
 	bool bOnlyLicenseed;
 
+	/** Should we build navigation data for the packages we are saving? **/
+	bool bShouldBuildNavigationData;
+
 	/** Ignore package version changelist **/
 	bool bIgnoreChangelist;
 
diff --git a/Engine/Source/Editor/UnrealEd/Private/Commandlets/ContentCommandlets.cpp b/Engine/Source/Editor/UnrealEd/Private/Commandlets/ContentCommandlets.cpp
--- a/Engine/Source/Editor/UnrealEd/Private/Commandlets/ContentCommandlets.cpp
+++ b/Engine/Source/Editor/UnrealEd/Private/Commandlets/ContentCommandlets.cpp
@@ -59,6 +59,12 @@ DEFINE_LOG_CATEGORY(LogContentCommandlet);
 
 // for UResavePackagesCommandlet::PerformAdditionalOperations building lighting code
 #include "LightingBuildOptions.h"
+
+// for UResavePackagesCommandlet::PerformAdditionalOperations building navigation data
+#include "EngineUtils.h"
+#include "NavigationData.h"
+#include "AI/NavigationSystemBase.h"
+
 // For preloading FFindInBlueprintSearchManager
 #include "FindInBlueprintManager.h"
 #include "IHierarchicalLODUtilities.h"
@@ -902,6 +908,8 @@ int32 UResavePackagesCommandlet::Main( const FString& Params )
 	bOnlyLicenseed = Switches.Contains(TEXT("OnlyLicenseed"));
 	/** only process packages containing materials */
 	bOnlyMaterials = Switches.Contains(TEXT("onlymaterials"));
+	/** determine if we are building navigation data for the map packages on the pass. **/
+	bShouldBuildNavigationData = Switches.Contains(TEXT("BuildNavigationData"));
 
 	/** check for filtering packages by collection. **/
 	FString FilterByCollection;
@@ -1144,50 +1152,44 @@ int32 UResavePackagesCommandlet::Main( const FString& Params )
 
 FText UResavePackagesCommandlet::GetChangelistDescription() const
 {
-	FText ChangelistDescription;
+	FTextBuilder ChangelistDescription;
 
-	if (bShouldBuildTextureStreaming && bShouldBuildLighting && bShouldBuildReflectionCaptures)
-	{
-		ChangelistDescription = NSLOCTEXT("ContentCmdlets", "ChangelistDescriptionBuildLightingAndTextureStreamingAndReflectionCaptures", "Rebuild lightmaps & texture streaming & reflection captures.");
-	}
-	else if (bShouldBuildTextureStreaming && bShouldBuildLighting)
-	{
-		ChangelistDescription = NSLOCTEXT("ContentCmdlets", "ChangelistDescriptionBuildLightingAndTextureStreaming", "Rebuild lightmaps & texture streaming.");
-	}
-	else if (bShouldBuildTextureStreaming && bShouldBuildReflectionCaptures)
+	if (bShouldBuildLighting)
 	{
-		ChangelistDescription = NSLOCTEXT("ContentCmdlets", "ChangelistDescriptionBuildTextureStreamingAndReflectionCaptures", "Rebuild texture streaming & reflection captures.");
+		ChangelistDescription.AppendLine(NSLOCTEXT("ContentCmdlets", "ChangelistDescriptionBuildLighting", "Rebuild lightmaps."));
 	}
-	else if (bShouldBuildLighting && bShouldBuildReflectionCaptures)
-	{
-		ChangelistDescription = NSLOCTEXT("ContentCmdlets", "ChangelistDescriptionBuildLightingAndReflectionCaptures", "Rebuild lightmaps & reflection captures.");
-	}
-	else if (bShouldBuildLighting)
+
+	if (bShouldBuildTextureStreaming)
 	{
-		ChangelistDescription = NSLOCTEXT("ContentCmdlets", "ChangelistDescriptionBuildLighting", "Rebuild lightmaps.");
+		ChangelistDescription.AppendLine(NSLOCTEXT("ContentCmdlets", "ChangelistDescriptionBuildTextureStreaming", "Rebuild texture streaming."));
 	}
-	else if (bShouldBuildTextureStreaming)
+
+	if (bShouldBuildReflectionCaptures)
 	{
-		ChangelistDescription = NSLOCTEXT("ContentCmdlets", "ChangelistDescriptionBuildTextureStreaming", "Rebuild texture streaming.");
+		ChangelistDescription.AppendLine(NSLOCTEXT("ContentCmdlets", "ChangelistDescriptionBuildReflectionCaptures", "Rebuild reflection captures."));
 	}
-	else if (bShouldBuildReflectionCaptures)
+
+	if (RedirectorsToFixup.Num() > 0)
 	{
-		ChangelistDescription = NSLOCTEXT("ContentCmdlets", "ChangelistDescriptionBuildReflectionCaptures", "Rebuild reflection captures.");
+		ChangelistDescription.AppendLine(NSLOCTEXT("ContentCmdlets", "ChangelistDescriptionRedirectors", "Fixing Redirectors"));
 	}
-	else if (RedirectorsToFixup.Num() > 0)
+
+	if (bShouldBuildHLOD)
 	{
-		ChangelistDescription = NSLOCTEXT("ContentCmdlets", "ChangelistDescriptionRedirectors", "Fixing Redirectors");
+		ChangelistDescription.AppendLine(NSLOCTEXT("ContentCmdlets", "ChangelistDescriptionHLODs", "Rebuilding HLODs"));
 	}
-	else if (bShouldBuildHLOD)
+
+	if (bShouldBuildNavigationData)
 	{
-		ChangelistDescription = NSLOCTEXT("ContentCmdlets", "ChangelistDescriptionHLODs", "Rebuilding HLODs");
+		ChangelistDescription.AppendLine(NSLOCTEXT("ContentCmdlets", "ChangelistDescriptionBuildNavigationData", "Rebuilding Navigation data."));
 	}
-	else
+	
+	if (ChangelistDescription.IsEmpty())
 	{
-		ChangelistDescription = NSLOCTEXT("ContentCmdlets", "ChangelistDescription", "Resave Deprecated Packages");
+		ChangelistDescription.AppendLine(NSLOCTEXT("ContentCmdlets", "ChangelistDescription", "Resave Deprecated Packages"));
 	}
 
-	return ChangelistDescription;
+	return ChangelistDescription.ToText();
 }
 
 
@@ -1429,7 +1431,11 @@ void UResavePackagesCommandlet::PerformAdditionalOperations(class UWorld* World,
 
 	const bool bShouldBuildTextureStreamingForWorld = bShouldBuildTextureStreaming && !bShouldBuildTextureStreamingForAll;
 	const bool bBuildingNonHLODData = (bShouldBuildLighting || bShouldBuildTextureStreamingForWorld || bShouldBuildReflectionCaptures);
-	if (bBuildingNonHLODData || bShouldBuildHLOD)
+
+	// indicates if world and level packages should be checked out only if dirty after building data
+	const bool bShouldCheckoutDirtyPackageOnly = (bShouldBuildHLOD || bShouldBuildNavigationData) && !bBuildingNonHLODData;
+
+	if (bBuildingNonHLODData || bShouldBuildHLOD || bShouldBuildNavigationData)
 	{
 		bool bShouldProceedWithRebuild = true;
 
@@ -1442,7 +1448,7 @@ void UResavePackagesCommandlet::PerformAdditionalOperations(class UWorld* World,
 			IVS.RequiresHitProxies(false);
 			IVS.ShouldSimulatePhysics(false);
 			IVS.EnableTraceCollision(false);
-			IVS.CreateNavigation(false);
+			IVS.CreateNavigation(bShouldBuildNavigationData);
 			IVS.CreateAISystem(false);
 			IVS.AllowAudioPlayback(false);
 			IVS.CreatePhysicsScene(true);
@@ -1487,17 +1493,10 @@ void UResavePackagesCommandlet::PerformAdditionalOperations(class UWorld* World,
 		FString WorldPackageCheckedOutUser;
 		if (FPackageName::DoesPackageExist(World->GetOutermost()->GetName(), NULL, &WorldPackageName))
 		{
-			// If we are only building HLODs check if level can be checked out, if so add to list of files that will be saved/checked-out after rebuilding the data
-			if(bShouldBuildHLOD && !bBuildingNonHLODData)
+			// If we are checking out only dirty package, simply check if level can be checked out. They will be added to list of fsaved/checked-out files after rebuilding the data
+			if (bShouldCheckoutDirtyPackageOnly)
 			{
-				if (CanCheckoutFile(WorldPackageName, WorldPackageCheckedOutUser) || !bSkipCheckedOutFiles)
-				{
-					CheckedOutPackagesFilenames.Add(WorldPackageName);
-				}
-				else 
-				{
-					bShouldProceedWithRebuild = false;
-				}				
+				bShouldProceedWithRebuild = (CanCheckoutFile(WorldPackageName, WorldPackageCheckedOutUser) || !bSkipCheckedOutFiles);
 			}
 			else
 			{
@@ -1528,8 +1527,7 @@ void UResavePackagesCommandlet::PerformAdditionalOperations(class UWorld* World,
 			{
 				bool bShouldBeLoaded = true;
 
-				// If we are not building HLODs or are but also rebuilding lighting we check out the level file, otherwise we don't to try and ensure a minimal HLOD rebuild
-				if (!bShouldBuildHLOD || bBuildingNonHLODData)
+				if (!bShouldCheckoutDirtyPackageOnly)
 				{
 					CheckOutLevelFile(StreamingLevel->GetLoadedLevel());
 				}
@@ -1538,15 +1536,12 @@ void UResavePackagesCommandlet::PerformAdditionalOperations(class UWorld* World,
 				const FString StreamingLevelWorldAssetPackageName = StreamingLevel->GetWorldAssetPackageName();
 				if (FPackageName::DoesPackageExist(StreamingLevelWorldAssetPackageName, NULL, &StreamingLevelPackageFilename))
 				{
-					// If we are building HLODs only, we dont check out the files ahead of rebuilding the data
-					if(bShouldBuildHLOD && !bBuildingNonHLODData)
+					if (bShouldCheckoutDirtyPackageOnly)
 					{
 						FString OutUser;
-						if (CanCheckoutFile(StreamingLevelPackageFilename, OutUser) || !bSkipCheckedOutFiles)
-						{
-							CheckedOutPackagesFilenames.Add(StreamingLevelPackageFilename);
-						}
-						else 
+						bShouldBeLoaded = (CanCheckoutFile(StreamingLevelPackageFilename, OutUser) || !bSkipCheckedOutFiles);
+
+						if (!bShouldBeLoaded)
 						{
 							if (OutUser.Len())
 							{
@@ -1556,8 +1551,6 @@ void UResavePackagesCommandlet::PerformAdditionalOperations(class UWorld* World,
 							{
 								UE_LOG(LogContentCommandlet, Warning, TEXT("[REPORT] Skipping %s as it could not be checked out (not at head revision ?)"), *StreamingLevelPackageFilename);
 							}
-
-							bShouldBeLoaded = false;
 						}
 					}
 					else
@@ -1593,8 +1586,8 @@ void UResavePackagesCommandlet::PerformAdditionalOperations(class UWorld* World,
 		{
 			World->FlushLevelStreaming(EFlushLevelStreamingType::Full);
 
-			// If we are (minimally) rebuilding HLOD, set the visible streamed-in levels packages to clean (as FlushLevelStreaming will dirty their packages in this commandlet context)
-			if(bShouldBuildHLOD)
+			// If we are (minimally) rebuilding using only dirty packages, set the visible streamed-in levels packages to clean (as FlushLevelStreaming will dirty their packages in this commandlet context)
+			if (bShouldCheckoutDirtyPackageOnly)
 			{
 				for (const ULevel* Level : World->GetLevels())
 				{
@@ -1696,6 +1689,26 @@ void UResavePackagesCommandlet::PerformAdditionalOperations(class UWorld* World,
 				}
 			}
 
+			if (bShouldBuildNavigationData)
+			{
+				// Make sure navigation is added and initialized in EditorMode
+				FNavigationSystem::AddNavigationSystemToWorld(*World, FNavigationSystemRunMode::EditorMode);
+
+				// Invoke navigation data generator
+				UE_LOG(LogContentCommandlet, Display, TEXT("Building navigation data for %s"), *World->GetOutermost()->GetName());
+				FNavigationSystem::Build(*World);
+
+				// Checkout and save each dirty package
+				for (TActorIterator<ANavigationData> It(World); It; ++It)
+				{
+					UPackage* Package = It->GetOutermost();
+					if (Package != nullptr && Package->IsDirty())
+					{
+						CheckoutAndSavePackage(Package, CheckedOutPackagesFilenames, bSkipCheckedOutFiles);
+					}
+				}
+			}
+
 			if (bShouldBuildLighting)
 			{
 				FLightingBuildOptions LightingOptions;
@@ -1747,17 +1760,15 @@ void UResavePackagesCommandlet::PerformAdditionalOperations(class UWorld* World,
 					{
 						UPackage* SubLevelPackage = NextStreamingLevel->GetLoadedLevel()->GetOutermost();
 						bool bSaveSubLevelPackage = true;
-						if(bShouldBuildHLOD && !bBuildingNonHLODData)
+						if (bShouldCheckoutDirtyPackageOnly)
 						{
-							// If we are building HLOD, only save packages that were dirtied
 							bSaveSubLevelPackage = SubLevelPackage->IsDirty();
 						}
 
-						if(bSaveSubLevelPackage)
+						if (bSaveSubLevelPackage)
 						{
-							// When building HLODs we dont check out/modify maps unless dirty
 							bool bFileCheckedOut = true;
-							if(bShouldBuildHLOD && !bBuildingNonHLODData)
+							if (bShouldCheckoutDirtyPackageOnly)
 							{
 								bFileCheckedOut = CheckoutFile(StreamingLevelPackageFilename, true);
 							}
@@ -1783,8 +1794,8 @@ void UResavePackagesCommandlet::PerformAdditionalOperations(class UWorld* World,
 			}
 			else
 			{
-			UE_LOG(LogContentCommandlet, Error, TEXT("[REPORT] Failed to complete steps necessary to start a lightmass or texture streaming build of %s"), *World->GetName());
-		}
+				UE_LOG(LogContentCommandlet, Error, TEXT("[REPORT] Failed to complete steps necessary to start a lightmass or texture streaming build of %s"), *World->GetName());
+			}
 		}
 
 		if ((bShouldProceedWithRebuild == false)||(bSavePackage == false))
@@ -1808,9 +1819,8 @@ void UResavePackagesCommandlet::PerformAdditionalOperations(class UWorld* World,
 				FilesToSubmit.AddUnique(CheckedOutPackageFilename);
 			}
 
-			if(bShouldBuildHLOD && !bBuildingNonHLODData)
+			if (bShouldCheckoutDirtyPackageOnly)
 			{
-				// Don't save outer package if it isn't dirty when doing a HLOD rebuild only
 				bSavePackage = World->GetOutermost()->IsDirty();
 			}
 		}
diff --git a/Engine/Source/Runtime/NavigationSystem/Private/NavigationData.cpp b/Engine/Source/Runtime/NavigationSystem/Private/NavigationData.cpp
--- a/Engine/Source/Runtime/NavigationSystem/Private/NavigationData.cpp
+++ b/Engine/Source/Runtime/NavigationSystem/Private/NavigationData.cpp
@@ -521,6 +521,8 @@ void ANavigationData::RebuildAll()
 	
 	if (NavDataGenerator.IsValid())
 	{
+		// mark outermost package as dirty. Internal filters will dirty only for valid scenarios (i.e. commandlet or editor mode only)
+		MarkPackageDirty();
 		NavDataGenerator->RebuildAll();
 	}
 }
