Fix an assert that happen when building a meshdescription from the importer. The reserve of the sparse array can create some holes in the data which is the cause of the assert

diff --git a/Engine/Source/Editor/UnrealEd/Private/Fbx/FbxStaticMeshImport.cpp b/Engine/Source/Editor/UnrealEd/Private/Fbx/FbxStaticMeshImport.cpp
--- a/Engine/Source/Editor/UnrealEd/Private/Fbx/FbxStaticMeshImport.cpp
+++ b/Engine/Source/Editor/UnrealEd/Private/Fbx/FbxStaticMeshImport.cpp
@@ -962,6 +962,14 @@ bool UnFbx::FFbxImporter::BuildStaticMeshFromGeometry(FbxNode* Node, UStaticMesh
 			{
 				Mesh->EndGetMeshEdgeIndexForPolygon();
 			}
+			if (SkippedVertexInstance > 0)
+			{
+				//We must compact the sparse array before reserving new space
+				//When we reserve it will make a hole in the sparse array if the last reserve was not fully use
+				//The importer assume there will be no hole when importing a mesh
+				FElementIDRemappings OutRemappings;
+				MeshDescription->Compact(OutRemappings);
+			}
 		}
 	}
 
