Fix TiffLoader limited to Window platforms (Linux and Mac are now supported)

#rnx

diff --git a/Engine/Source/Editor/UnrealEd/Private/Factories/TIFFLoader.cpp b/Engine/Source/Editor/UnrealEd/Private/Factories/TIFFLoader.cpp
--- a/Engine/Source/Editor/UnrealEd/Private/Factories/TIFFLoader.cpp
+++ b/Engine/Source/Editor/UnrealEd/Private/Factories/TIFFLoader.cpp
@@ -4,9 +4,11 @@
 
 DEFINE_LOG_CATEGORY_STATIC(LogTIFFLoader, Log, All);
 
-#if PLATFORM_WINDOWS
+#if WITH_FREEIMAGE_LIB
 
+#if PLATFORM_WINDOWS
 #include "Windows/AllowWindowsPlatformTypes.h"
+#endif // PLATFORM_WINDOWS
 
 THIRD_PARTY_INCLUDES_START
 #include "FreeImage.h"
@@ -35,8 +37,9 @@ void FFreeImageWrapper::FreeImage_Initialise(bool bLoadLocalPluginsOnly)
 	if (FreeImageDllHandle == nullptr)
 	{
 		FString FreeImageDir = FPaths::Combine(FPaths::EngineDir(), TEXT("Binaries/ThirdParty/FreeImage"), FPlatformProcess::GetBinariesSubdirectory());
+		FString FreeImageLibDir = FPaths::Combine( FreeImageDir, TEXT(FREEIMAGE_LIB_FILENAME));
 		FPlatformProcess::PushDllDirectory(*FreeImageDir);
-		FreeImageDllHandle = FPlatformProcess::GetDllHandle(TEXT("FreeImage.dll"));
+		FreeImageDllHandle = FPlatformProcess::GetDllHandle(*FreeImageLibDir);
 		FPlatformProcess::PopDllDirectory(*FreeImageDir);
 	}
 
@@ -46,8 +49,9 @@ void FFreeImageWrapper::FreeImage_Initialise(bool bLoadLocalPluginsOnly)
 	}
 }
 
+#if PLATFORM_WINDOWS
 #include "Windows/HideWindowsPlatformTypes.h"
-
+#endif // PLATFORM_WINDOWS
 
 FTiffLoadHelper::FTiffLoadHelper()
 {
@@ -344,7 +348,7 @@ bool FTiffLoadHelper::IsValid()
 	return bIsValid;
 }
 
-#else //PLATFORM_WINDOWS 
+#else // WITH_FREEIMAGE_LIB 
 
 FTiffLoadHelper::FTiffLoadHelper() 
 {
@@ -378,4 +382,4 @@ bool FTiffLoadHelper::IsValid()
     return false;
 }
 
-#endif //PLATFORM_WINDOWS
+#endif // WITH_FREEIMAGE_LIB
