Add Component menu in the actor Details view now filters out allowed classes based on UClassViewerSettings AllowedClasses.

diff --git a/Engine/Source/Editor/UnrealEd/Private/ComponentTypeRegistry.cpp b/Engine/Source/Editor/UnrealEd/Private/ComponentTypeRegistry.cpp
--- a/Engine/Source/Editor/UnrealEd/Private/ComponentTypeRegistry.cpp
+++ b/Engine/Source/Editor/UnrealEd/Private/ComponentTypeRegistry.cpp
@@ -10,7 +10,6 @@
 #include "TickableEditorObject.h"
 #include "ActorFactories/ActorFactoryBasicShape.h"
 #include "Materials/Material.h"
-#include "Settings/EditorExperimentalSettings.h"
 #include "Engine/StaticMesh.h"
 #include "AssetData.h"
 
@@ -20,6 +19,7 @@
 #include "Kismet2/KismetEditorUtilities.h"
 #include "Misc/HotReloadInterface.h"
 #include "SComponentClassCombo.h"
+#include "Settings/ClassViewerSettings.h"
 
 #define LOCTEXT_NAMESPACE "ComponentTypeRegistry"
 
@@ -275,40 +275,45 @@ void FComponentTypeRegistryData::ForceRefreshComponentList()
 		{
 			InMemoryClasses.Push(Class->GetFName());
 
-			bool const bOutOfDateClass = Class->HasAnyClassFlags(CLASS_NewerVersionExists);
-			bool const bBlueprintSkeletonClass = FKismetEditorUtilities::IsClassABlueprintSkeleton(Class);
-			if (!Class->HasAnyClassFlags(CLASS_Abstract) && Class->HasMetaData(FBlueprintMetadata::MD_BlueprintSpawnableComponent) && !bOutOfDateClass && !bBlueprintSkeletonClass) //@TODO: Fold this logic together with the one in UEdGraphSchema_K2::GetAddComponentClasses
-			{
-				TArray<FString> ClassGroupNames;
-				Class->GetClassGroupNames(ClassGroupNames);
+			const bool bOutOfDateClass = Class->HasAnyClassFlags(CLASS_NewerVersionExists);
+			const bool bBlueprintSkeletonClass = FKismetEditorUtilities::IsClassABlueprintSkeleton(Class);
+			const bool bPassesAllowedClasses = GetDefault<UClassViewerSettings>()->AllowedClasses.Num() == 0 || GetDefault<UClassViewerSettings>()->AllowedClasses.Contains(Class->GetName());
 
-				if (ClassGroupNames.Contains(CommonClassGroup))
+			if (!bOutOfDateClass &&
+				!bBlueprintSkeletonClass &&
+				bPassesAllowedClasses)
+			{
+				if (!Class->HasAnyClassFlags(CLASS_Abstract) &&
+					Class->HasMetaData(FBlueprintMetadata::MD_BlueprintSpawnableComponent)) //@TODO: Fold this logic together with the one in UEdGraphSchema_K2::GetAddComponentClasses
 				{
-					FString ClassGroup = CommonClassGroup;
-					FComponentClassComboEntryPtr NewEntry(new FComponentClassComboEntry(ClassGroup, Class, ClassGroupNames.Num() <= 1, EComponentCreateAction::SpawnExistingClass));
-					SortedClassList.Add(NewEntry);
+					TArray<FString> ClassGroupNames;
+					Class->GetClassGroupNames(ClassGroupNames);
+
+					if (ClassGroupNames.Contains(CommonClassGroup))
+					{
+						FString ClassGroup = CommonClassGroup;
+						FComponentClassComboEntryPtr NewEntry(new FComponentClassComboEntry(ClassGroup, Class, ClassGroupNames.Num() <= 1, EComponentCreateAction::SpawnExistingClass));
+						SortedClassList.Add(NewEntry);
+					}
+					if (ClassGroupNames.Num() && !ClassGroupNames[0].Equals(CommonClassGroup))
+					{
+						const bool bIncludeInFilter = true;
+
+						FString ClassGroup = ClassGroupNames[0];
+						FComponentClassComboEntryPtr NewEntry(new FComponentClassComboEntry(ClassGroup, Class, bIncludeInFilter, EComponentCreateAction::SpawnExistingClass));
+						SortedClassList.Add(NewEntry);
+					}
+					else if (ClassGroupNames.Num() == 0)
+					{
+						// No class group name found. Just add it to a "custom" category
+
+						const bool bIncludeInFilter = true;
+						FString ClassGroup = LOCTEXT("CustomClassGroup", "Custom").ToString();
+						FComponentClassComboEntryPtr NewEntry(new FComponentClassComboEntry(ClassGroup, Class, bIncludeInFilter, EComponentCreateAction::SpawnExistingClass));
+						SortedClassList.Add(NewEntry);
+					}
 				}
-				if (ClassGroupNames.Num() && !ClassGroupNames[0].Equals(CommonClassGroup))
-				{
-					const bool bIncludeInFilter = true;
 
-					FString ClassGroup = ClassGroupNames[0];
-					FComponentClassComboEntryPtr NewEntry(new FComponentClassComboEntry(ClassGroup, Class, bIncludeInFilter, EComponentCreateAction::SpawnExistingClass));
-					SortedClassList.Add(NewEntry);
-				}
-				else if(ClassGroupNames.Num() == 0 )
-				{
-					// No class group name found. Just add it to a "custom" category
-					
-					const bool bIncludeInFilter = true;
-					FString ClassGroup = LOCTEXT("CustomClassGroup", "Custom").ToString();
-					FComponentClassComboEntryPtr NewEntry(new FComponentClassComboEntry(ClassGroup, Class, bIncludeInFilter, EComponentCreateAction::SpawnExistingClass));
-					SortedClassList.Add(NewEntry);
-				}
-			}
-			
-			if (!bOutOfDateClass && !bBlueprintSkeletonClass)
-			{
 				FComponentTypeEntry Entry = { Class->GetName(), FString(), Class };
 				ComponentTypeList.Add(Entry);
 			}
@@ -316,18 +321,16 @@ void FComponentTypeRegistryData::ForceRefreshComponentList()
 	}
 
 	{
-		// make sure that we add any user created classes immediately, generally this will not create anything (because assets have not been discovered yet), but asset discovery
-		// should be allowed to take place at any time:
+		// make sure that we add any user created classes immediately, generally this will not create anything (because assets have not been discovered yet), 
+		// but asset discovery should be allowed to take place at any time:
 		FAssetRegistryModule& AssetRegistryModule = FModuleManager::LoadModuleChecked<FAssetRegistryModule>(TEXT("AssetRegistry"));
 		TArray<FName> ClassNames;
 		ClassNames.Add(UActorComponent::StaticClass()->GetFName());
 		TSet<FName> DerivedClassNames;
 		AssetRegistryModule.Get().GetDerivedClassNames(ClassNames, TSet<FName>(), DerivedClassNames);
 
-		const bool bIncludeInFilter = true;
-
-		auto InMemoryClassesSet = TSet<FName>(InMemoryClasses);
-		auto OnDiskClasses = DerivedClassNames.Difference(InMemoryClassesSet);
+		TSet<FName> InMemoryClassesSet = TSet<FName>(InMemoryClasses);
+		TSet<FName> OnDiskClasses = DerivedClassNames.Difference(InMemoryClassesSet);
 
 		if (OnDiskClasses.Num() > 0)
 		{
@@ -338,8 +341,7 @@ void FComponentTypeRegistryData::ForceRefreshComponentList()
 			TMap<FString, FAssetData> BlueprintNames;
 			for (const FAssetData& Blueprint : BlueprintAssetData)
 			{
-				const auto& name = Blueprint.AssetName.ToString();
-				BlueprintNames.Add(name, Blueprint);
+				BlueprintNames.Add(Blueprint.AssetName.ToString(), Blueprint);
 			}
 
 			for (const FName& OnDiskClass : OnDiskClasses)
@@ -347,20 +349,25 @@ void FComponentTypeRegistryData::ForceRefreshComponentList()
 				FString FixedString = OnDiskClass.ToString();
 				FixedString.RemoveFromEnd(TEXT("_C"));
 
-				FAssetData AssetData;
-				if (const FAssetData* Value = BlueprintNames.Find(FixedString))
+				const bool bPassesAllowedClasses = GetDefault<UClassViewerSettings>()->AllowedClasses.Num() == 0 || GetDefault<UClassViewerSettings>()->AllowedClasses.Contains(FixedString);
+				if (bPassesAllowedClasses)
 				{
-					AssetData = *Value;
-				}
+					FAssetData AssetData;
+					if (const FAssetData* Value = BlueprintNames.Find(FixedString))
+					{
+						AssetData = *Value;
+					}
 
-				FComponentTypeEntry Entry = { FixedString, AssetData.ObjectPath.ToString(), nullptr };
-				ComponentTypeList.Add(Entry);
+					FComponentTypeEntry Entry = { FixedString, AssetData.ObjectPath.ToString(), nullptr };
+					ComponentTypeList.Add(Entry);
 
-				// The blueprint is unloaded, so we need to work out which icon to use for it using its asset data
-				const UClass* BlueprintIconClass = FClassIconFinder::GetIconClassForAssetData(AssetData);
+					// The blueprint is unloaded, so we need to work out which icon to use for it using its asset data
+					const UClass* BlueprintIconClass = FClassIconFinder::GetIconClassForAssetData(AssetData);
 
-				FComponentClassComboEntryPtr NewEntry(new FComponentClassComboEntry(BlueprintComponents, FixedString, AssetData.ObjectPath, BlueprintIconClass, bIncludeInFilter));
-				SortedClassList.Add(NewEntry);
+					const bool bIncludeInFilter = true;
+					FComponentClassComboEntryPtr NewEntry(new FComponentClassComboEntry(BlueprintComponents, FixedString, AssetData.ObjectPath, BlueprintIconClass, bIncludeInFilter));
+					SortedClassList.Add(NewEntry);
+				}
 			}
 		}
 	}
@@ -375,7 +382,6 @@ void FComponentTypeRegistryData::ForceRefreshComponentList()
 			FComponentClassComboEntryPtr& CurrentEntry = SortedClassList[ClassIndex];
 
 			const FString& CurrentHeadingText = CurrentEntry->GetHeadingText();
-
 			if (CurrentHeadingText != PreviousHeading)
 			{
 				// This avoids a redundant separator being added to the very top of the list
