Editor: Support "view modifiers" for editor viewports.

These are delegates that allow plugins and other 3rd party code to affect the transform of the viewport.

I'm also removing the call to `UpdateViewForLockedActor` in `CalcSceneView` because that violated the expectation that this method was only computing some view information based on the current transform/FOV/etc. That call came from a GDC demo crunch bug fix that we will fix again better if it comes back.

diff --git a/Engine/Source/Editor/UnrealEd/Private/LevelEditorViewport.cpp b/Engine/Source/Editor/UnrealEd/Private/LevelEditorViewport.cpp
--- a/Engine/Source/Editor/UnrealEd/Private/LevelEditorViewport.cpp
+++ b/Engine/Source/Editor/UnrealEd/Private/LevelEditorViewport.cpp
@@ -1741,6 +1741,7 @@ FLevelEditorViewportClient::FLevelEditorViewportClient(const TSharedPtr<SLevelVi
 	, bLockedCameraView(true)
 	, bReceivedFocusRecently(false)
 	, bAlwaysShowModeWidgetAfterSelectionChanges(true)
+	, bShouldApplyViewModifiers(true)
 	, SpriteCategoryVisibility()
 	, World(nullptr)
 	, TrackingTransaction()
@@ -1837,8 +1838,6 @@ FSceneView* FLevelEditorViewportClient::CalcSceneView(FSceneViewFamily* ViewFami
 {
 	bWasControlledByOtherViewport = false;
 
-	UpdateViewForLockedActor();
-
 	// set all other matching viewports to my location, if the LOD locking is enabled,
 	// unless another viewport already set me this frame (otherwise they fight)
 	if (GEditor->bEnableLODLocking)
@@ -2291,9 +2290,32 @@ void FLevelEditorViewportClient::Tick(float DeltaTime)
 
 	UpdateViewForLockedActor(DeltaTime);
 
+	if (bShouldApplyViewModifiers)
+	{
+		ApplyViewModifiers(DeltaTime);
+	}
+
 	UserIsControllingAtmosphericLightTimer = FMath::Max(UserIsControllingAtmosphericLightTimer - DeltaTime, 0.0f);
 }
 
+void FLevelEditorViewportClient::ApplyViewModifiers(float DeltaTime)
+{
+	FEditorViewportViewModifierParams Params;
+	Params.DeltaTime = DeltaTime;
+	Params.ViewportClient = this;
+
+	FMinimalViewInfo InOutPOV;
+	InOutPOV.Location = GetViewLocation();
+	InOutPOV.Rotation = GetViewRotation();
+	InOutPOV.FOV = ViewFOV;
+
+	ViewModifiers.Broadcast(Params, InOutPOV);
+
+	SetViewLocation(InOutPOV.Location);
+	SetViewRotation(InOutPOV.Rotation);
+	ViewFOV = InOutPOV.FOV;
+}
+
 void FLevelEditorViewportClient::UpdateViewForLockedActor(float DeltaTime)
 {
 	// We can't be locked to a matinee actor if this viewport doesn't allow matinee control
diff --git a/Engine/Source/Editor/UnrealEd/Public/LevelEditorViewport.h b/Engine/Source/Editor/UnrealEd/Public/LevelEditorViewport.h
--- a/Engine/Source/Editor/UnrealEd/Public/LevelEditorViewport.h
+++ b/Engine/Source/Editor/UnrealEd/Public/LevelEditorViewport.h
@@ -12,6 +12,7 @@
 #include "UObject/ObjectKey.h"
 
 struct FAssetData;
+struct FMinimalViewInfo;
 class FCanvas;
 class FDragTool;
 class HModel;
@@ -120,6 +121,16 @@ private:
 	
 };
 
+struct FEditorViewportViewModifierParams
+{
+	/** Delta time for updating view modifications. */
+	float DeltaTime = 0.f;
+	/** The viewport client calling for view modification. */
+	FEditorViewportClient* ViewportClient = nullptr;
+};
+
+/** Delegate for modifying view parameters of an editor viewport. */
+DECLARE_MULTICAST_DELEGATE_TwoParams(FEditorViewportViewModifierDelegate, const FEditorViewportViewModifierParams&, FMinimalViewInfo&);
 
 /** */
 class UNREALED_API FLevelEditorViewportClient : public FEditorViewportClient
@@ -569,6 +580,9 @@ protected:
 	/** Delegate handler for ActorMoved events */
 	void OnActorMoved(AActor* InActor);
 
+	/** Call registered view modifiers and apply the modify result */
+	void ApplyViewModifiers(float DeltaTime);
+
 	/** FEditorViewportClient Interface*/
 
 public:
@@ -748,6 +762,13 @@ public:
 
 	/** When enabled, the Unreal transform widget will become visible after an actor is selected, even if it was turned off via a show flag */
 	bool bAlwaysShowModeWidgetAfterSelectionChanges;
+
+	/** List of view modifiers to apply on view parameters. */
+	FEditorViewportViewModifierDelegate ViewModifiers;
+
+	/** Whether view modifiers should be called and applied. */
+	bool bShouldApplyViewModifiers;
+
 private:
 	/** The actors that are currently being placed in the viewport via dragging */
 	static TArray< TWeakObjectPtr< AActor > > DropPreviewActors;
