Remove duplicate prompt when pasting to hidden level

diff --git a/Engine/Source/Editor/UnrealEd/Private/EditorActor.cpp b/Engine/Source/Editor/UnrealEd/Private/EditorActor.cpp
--- a/Engine/Source/Editor/UnrealEd/Private/EditorActor.cpp
+++ b/Engine/Source/Editor/UnrealEd/Private/EditorActor.cpp
@@ -234,14 +234,14 @@ static FVector CreateLocationOffset(bool bDuplicate, bool bOffsetLocations)
 
 bool UUnrealEdEngine::WarnIfDestinationLevelIsHidden( UWorld* InWorld )
 {
-	bool result = false;
+	bool bShouldLoadHiddenLevels = true;
+	bool bShowPasteHiddenWarning = true;
 	TArray<ULevel*> Levels;
 	TArray<bool> bTheyShouldBeVisible;
 	//prepare the warning dialog
 	FSuppressableWarningDialog::FSetupInfo Info( LOCTEXT( "Warning_PasteWarningBody","You are trying to paste to a hidden level.\nSuppressing this will default to Do Not Paste" ), LOCTEXT( "Warning_PasteWarningHeader","Pasting To Hidden Level" ), "PasteHiddenWarning" );
 	Info.ConfirmText = LOCTEXT( "Warning_PasteContinue","Unhide Level and paste" );
 	Info.CancelText = LOCTEXT( "Warning_PasteCancel","Do not paste" );
-	FSuppressableWarningDialog PasteHiddenWarning( Info );
 
 	//check streaming levels first
 	for (ULevelStreaming* StreamedLevel : InWorld->GetStreamingLevels())
@@ -252,13 +252,14 @@ bool UUnrealEdEngine::WarnIfDestinationLevelIsHidden( UWorld* InWorld )
 			ULevel* Level = StreamedLevel->GetLoadedLevel();
 			if( Level && Level->IsCurrentLevel() )
 			{
-				//the streamed level is not visible - check what the user wants to do
-				FSuppressableWarningDialog::EResult DialogResult = PasteHiddenWarning.ShowModal();
-				if( ( DialogResult == FSuppressableWarningDialog::Cancel )  || ( DialogResult == FSuppressableWarningDialog::Suppressed ) )
+				if (bShowPasteHiddenWarning)
 				{
-					result = true;
+					bShowPasteHiddenWarning = false;
+					//the streamed level is not visible - check what the user wants to do
+					bShouldLoadHiddenLevels = (FSuppressableWarningDialog(Info).ShowModal() == FSuppressableWarningDialog::Confirm);
 				}
-				else
+				
+				if (bShouldLoadHiddenLevels)
 				{
 					Levels.Add(Level);
 					bTheyShouldBeVisible.Add(true);
@@ -268,17 +269,18 @@ bool UUnrealEdEngine::WarnIfDestinationLevelIsHidden( UWorld* InWorld )
 	}
 
 	//now check the active level (this handles the persistent level also)
-	if( result == false )
+	if (bShouldLoadHiddenLevels)
 	{
 		if( FLevelUtils::IsLevelVisible( InWorld->GetCurrentLevel() ) == false )
 		{	
-			//the level is not visible - check what the user wants to do
-			FSuppressableWarningDialog::EResult DialogResult = PasteHiddenWarning.ShowModal();
-			if( ( DialogResult == FSuppressableWarningDialog::Cancel )  || ( DialogResult == FSuppressableWarningDialog::Suppressed ) )
+			if (bShowPasteHiddenWarning)
 			{
-				result = true;
+				bShowPasteHiddenWarning = false;
+				//the streamed level is not visible - check what the user wants to do
+				bShouldLoadHiddenLevels = (FSuppressableWarningDialog(Info).ShowModal() == FSuppressableWarningDialog::Confirm);
 			}
-			else 
+
+			if (bShouldLoadHiddenLevels)
 			{
 				Levels.Add(InWorld->GetCurrentLevel());
 				bTheyShouldBeVisible.Add(true);
@@ -291,7 +293,7 @@ bool UUnrealEdEngine::WarnIfDestinationLevelIsHidden( UWorld* InWorld )
 	{
 		EditorLevelUtils::SetLevelsVisibility(Levels, bTheyShouldBeVisible, true);
 	}
-	return result;
+	return !bShouldLoadHiddenLevels;
 }
 
 void UUnrealEdEngine::edactPasteSelected(UWorld* InWorld, bool bDuplicate, bool bOffsetLocations, bool bWarnIfHidden, FString* SourceData)
