Editor: Fix focusing actor in viewport.

diff --git a/Engine/Source/Editor/UnrealEd/Private/LevelEditorViewport.cpp b/Engine/Source/Editor/UnrealEd/Private/LevelEditorViewport.cpp
--- a/Engine/Source/Editor/UnrealEd/Private/LevelEditorViewport.cpp
+++ b/Engine/Source/Editor/UnrealEd/Private/LevelEditorViewport.cpp
@@ -2292,7 +2292,11 @@ void FLevelEditorViewportClient::Tick(float DeltaTime)
 
 	if (bShouldApplyViewModifiers)
 	{
-		ApplyViewModifiers(DeltaTime);
+		FViewportCameraTransform& ViewTransform = GetViewTransform();
+		if (!ViewTransform.IsPlaying())
+		{
+			ApplyViewModifiers(DeltaTime);
+		}
 	}
 
 	UserIsControllingAtmosphericLightTimer = FMath::Max(UserIsControllingAtmosphericLightTimer - DeltaTime, 0.0f);
@@ -2300,20 +2304,23 @@ void FLevelEditorViewportClient::Tick(float DeltaTime)
 
 void FLevelEditorViewportClient::ApplyViewModifiers(float DeltaTime)
 {
-	FEditorViewportViewModifierParams Params;
-	Params.DeltaTime = DeltaTime;
-	Params.ViewportClient = this;
+	if (ViewModifiers.IsBound())
+	{
+		FEditorViewportViewModifierParams Params;
+		Params.DeltaTime = DeltaTime;
+		Params.ViewportClient = this;
 
-	FMinimalViewInfo InOutPOV;
-	InOutPOV.Location = GetViewLocation();
-	InOutPOV.Rotation = GetViewRotation();
-	InOutPOV.FOV = ViewFOV;
+		FMinimalViewInfo InOutPOV;
+		InOutPOV.Location = GetViewLocation();
+		InOutPOV.Rotation = GetViewRotation();
+		InOutPOV.FOV = ViewFOV;
 
-	ViewModifiers.Broadcast(Params, InOutPOV);
+		ViewModifiers.Broadcast(Params, InOutPOV);
 
-	SetViewLocation(InOutPOV.Location);
-	SetViewRotation(InOutPOV.Rotation);
-	ViewFOV = InOutPOV.FOV;
+		SetViewLocation(InOutPOV.Location);
+		SetViewRotation(InOutPOV.Rotation);
+		ViewFOV = InOutPOV.FOV;
+	}
 }
 
 void FLevelEditorViewportClient::UpdateViewForLockedActor(float DeltaTime)
