LevelPartition: When current level is part of partition update DesiredLevel accordingly

#rnx

diff --git a/Engine/Source/Editor/UnrealEd/Private/AssetSelection.cpp b/Engine/Source/Editor/UnrealEd/Private/AssetSelection.cpp
--- a/Engine/Source/Editor/UnrealEd/Private/AssetSelection.cpp
+++ b/Engine/Source/Editor/UnrealEd/Private/AssetSelection.cpp
@@ -566,6 +566,16 @@ static AActor* PrivateAddActor( UObject* Asset, UActorFactory* Factory, bool Sel
 	FSnappingUtils::ClearSnappingHelpers( bClearImmediately );
 
 	ULevel* DesiredLevel = GWorld->GetCurrentLevel();
+
+	// If DesireLevel is part of a LevelPartition find the proper DesiredLevel by asking the Partition
+	if (const ILevelPartitionInterface* LevelPartition = DesiredLevel->GetLevelPartition())
+	{
+		if (ULevel* SubLevel = LevelPartition->GetSubLevel(ActorTransform.GetLocation()))
+		{
+			DesiredLevel = SubLevel;
+		}
+	}
+
 	bool bSpawnActor = true;
 
 	if ((ObjectFlags & RF_Transactional) != 0)
