Fixed issue with missing assets when cooking on the fly after an import

Root cause: Assets' packages are not updated in PackageNameCache
Solution: Call MarkPackageDirtyForCooker on packages which have been recently saved

diff --git a/Engine/Source/Editor/UnrealEd/Private/CookOnTheFlyServer.cpp b/Engine/Source/Editor/UnrealEd/Private/CookOnTheFlyServer.cpp
--- a/Engine/Source/Editor/UnrealEd/Private/CookOnTheFlyServer.cpp
+++ b/Engine/Source/Editor/UnrealEd/Private/CookOnTheFlyServer.cpp
@@ -5454,6 +5454,24 @@ void UCookOnTheFlyServer::GenerateAssetRegistry()
 	FAssetRegistryModule& AssetRegistryModule = FModuleManager::LoadModuleChecked<FAssetRegistryModule>(TEXT("AssetRegistry"));
 	AssetRegistry = &AssetRegistryModule.Get();
 
+	// Mark package as dirty for the last ones saved
+	if(PackageNameCache != nullptr)
+	{
+		for (FName AssetFilename : ModifiedAssetFilenames)
+		{
+			const FString AssetPathOnDisk = AssetFilename.ToString();
+			if(FPaths::FileExists(AssetPathOnDisk))
+			{
+				const FString PackageName = FPackageName::FilenameToLongPackageName(AssetPathOnDisk);
+				FSoftObjectPath SoftPackage(PackageName);
+				if(UPackage* Package = Cast<UPackage>(SoftPackage.ResolveObject()))
+				{
+					MarkPackageDirtyForCooker( Package );
+				}
+			}
+		}
+	}
+
 	if (!!(CookFlags & ECookInitializationFlags::GeneratedAssetRegistry))
 	{
 		// Force a rescan of modified package files
@@ -5466,8 +5484,6 @@ void UCookOnTheFlyServer::GenerateAssetRegistry()
 
 		AssetRegistry->ScanModifiedAssetFiles(ModifiedPackageFileList);
 
-		ModifiedAssetFilenames.Reset();
-
 		// This is cook in the editor on a second pass, so refresh the generators
 		for (TPair<FName, FAssetRegistryGenerator*>& Pair : RegistryGenerators)
 		{
@@ -5477,6 +5493,8 @@ void UCookOnTheFlyServer::GenerateAssetRegistry()
 	}
 	CookFlags |= ECookInitializationFlags::GeneratedAssetRegistry;
 
+	ModifiedAssetFilenames.Reset();
+
 	double GenerateAssetRegistryTime = 0.0;
 	{
 		SCOPE_TIMER(GenerateAssetRegistryTime);
