AssetRegistryDiff: Downgrading log warning to message when an asset has multiple chunkIDs

diff --git a/Engine/Source/Editor/UnrealEd/Private/Commandlets/DiffAssetRegistriesCommandlet.cpp b/Engine/Source/Editor/UnrealEd/Private/Commandlets/DiffAssetRegistriesCommandlet.cpp
--- a/Engine/Source/Editor/UnrealEd/Private/Commandlets/DiffAssetRegistriesCommandlet.cpp
+++ b/Engine/Source/Editor/UnrealEd/Private/Commandlets/DiffAssetRegistriesCommandlet.cpp
@@ -560,7 +560,7 @@ TArray<int32> UDiffAssetRegistriesCommandlet::GetAssetChunks(FAssetRegistryState
 		{
 			if (Assets[0]->ChunkIDs.Num() > 1)
 			{
-				UE_LOG(LogDiffAssets, Warning, TEXT("Multiple ChunkIds for asset %s"), *InAssetPath.ToString());
+				UE_LOG(LogDiffAssets, Log, TEXT("Multiple ChunkIds for asset %s"), *InAssetPath.ToString());
 			}
 
 			for (int32 id : Assets[0]->ChunkIDs)
