New experimental WASD camera prototype: when enabled, shift doubles the camera speed and the mouse wheel changes the global camera speed (ideally, this would be a per-viewport setting).

diff --git a/Engine/Source/Editor/UnrealEd/Classes/Settings/LevelEditorViewportSettings.h b/Engine/Source/Editor/UnrealEd/Classes/Settings/LevelEditorViewportSettings.h
--- a/Engine/Source/Editor/UnrealEd/Classes/Settings/LevelEditorViewportSettings.h
+++ b/Engine/Source/Editor/UnrealEd/Classes/Settings/LevelEditorViewportSettings.h
@@ -222,6 +222,10 @@ class UNREALED_API ULevelEditorViewportSettings
 	UPROPERTY(EditAnywhere, config, Category=Controls)
 	TEnumAsByte<EWASDType> FlightCameraControlType;
 
+	/** Enable the use of the experimental navigation in the flight camera controls. */
+	UPROPERTY(EditAnywhere, config, Category=Controls)
+	bool FlightCameraControlExperimentalNavigation;
+
 	/** Choose the control scheme for landscape tools (ignored for pen input) */
 	UPROPERTY(EditAnywhere, config, Category=Controls)
 	ELandscapeFoliageEditorControlType LandscapeEditorControlType;
diff --git a/Engine/Source/Editor/UnrealEd/Private/EditorViewportClient.cpp b/Engine/Source/Editor/UnrealEd/Private/EditorViewportClient.cpp
--- a/Engine/Source/Editor/UnrealEd/Private/EditorViewportClient.cpp
+++ b/Engine/Source/Editor/UnrealEd/Private/EditorViewportClient.cpp
@@ -1609,7 +1609,7 @@ void FEditorViewportClient::UpdateCameraMovement( float DeltaTime )
 
 		// Do we want to remap the various WASD keys for flight input?
 		const bool bRemapWASDKeys =
-			(bUnmodifiedPress) &&
+			(bUnmodifiedPress || (GetDefault<ULevelEditorViewportSettings>()->FlightCameraControlExperimentalNavigation && IsShiftPressed())) &&
 			(GetDefault<ULevelEditorViewportSettings>()->FlightCameraControlType == WASD_Always ||
 			( bUsingFlightInput &&
 			( GetDefault<ULevelEditorViewportSettings>()->FlightCameraControlType == WASD_RMBOnly && (Viewport->KeyState(EKeys::RightMouseButton ) ||Viewport->KeyState(EKeys::MiddleMouseButton) || Viewport->KeyState(EKeys::LeftMouseButton) || bIsUsingTrackpad ) ) ) ) &&
@@ -1756,7 +1756,8 @@ void FEditorViewportClient::UpdateCameraMovement( float DeltaTime )
 		// We'll combine the regular camera speed scale (controlled by viewport toolbar setting) with
 		// the flight camera speed scale (controlled by mouse wheel) and the CameraSpeedScalar (set in the transform viewport toolbar).
 		const float CameraSpeed = GetCameraSpeed();
-		const float FinalCameraSpeedScale = FlightCameraSpeedScale * CameraSpeed * GetCameraSpeedScalar();
+		const float CameraBoost = IsShiftPressed() ? 2.0f : 1.0f;
+		const float FinalCameraSpeedScale = FlightCameraSpeedScale * CameraSpeed * GetCameraSpeedScalar() * CameraBoost;
 
 		// Only allow FOV recoil if flight camera mode is currently inactive.
 		const bool bAllowRecoilIfNoImpulse = (!bUsingFlightInput) && (!IsMatineeRecordingWindow());
@@ -3329,44 +3330,58 @@ void FEditorViewportClient::OnChangeCameraSpeed( const struct FInputEventState&
 
 	FKey Key = InputState.GetKey();
 
-	// Adjust and clamp the camera speed scale
-	if( Key == EKeys::MouseScrollUp )
+	if (GetDefault<ULevelEditorViewportSettings>()->FlightCameraControlExperimentalNavigation)
 	{
-		if( FlightCameraSpeedScale >= 2.0f )
+		if( Key == EKeys::MouseScrollUp )
 		{
-			FlightCameraSpeedScale += 0.5f;
-		}
-		else if( FlightCameraSpeedScale >= 1.0f )
-		{
-			FlightCameraSpeedScale += 0.2f;
+			GetMutableDefault<ULevelEditorViewportSettings>()->CameraSpeed = FMath::Clamp<int32>(GetDefault<ULevelEditorViewportSettings>()->CameraSpeed + 1, 1, MaxCameraSpeeds);
 		}
 		else
 		{
-			FlightCameraSpeedScale += 0.1f;
+			GetMutableDefault<ULevelEditorViewportSettings>()->CameraSpeed = FMath::Clamp<int32>(GetDefault<ULevelEditorViewportSettings>()->CameraSpeed - 1, 1, MaxCameraSpeeds);
 		}
 	}
 	else
 	{
-		if( FlightCameraSpeedScale > 2.49f )
-		{
-			FlightCameraSpeedScale -= 0.5f;
-		}
-		else if( FlightCameraSpeedScale >= 1.19f )
+		// Adjust and clamp the camera speed scale
+		if( Key == EKeys::MouseScrollUp )
 		{
-			FlightCameraSpeedScale -= 0.2f;
+			if( FlightCameraSpeedScale >= 2.0f )
+			{
+				FlightCameraSpeedScale += 0.5f;
+			}
+			else if( FlightCameraSpeedScale >= 1.0f )
+			{
+				FlightCameraSpeedScale += 0.2f;
+			}
+			else
+			{
+				FlightCameraSpeedScale += 0.1f;
+			}
 		}
 		else
 		{
-			FlightCameraSpeedScale -= 0.1f;
+			if( FlightCameraSpeedScale > 2.49f )
+			{
+				FlightCameraSpeedScale -= 0.5f;
+			}
+			else if( FlightCameraSpeedScale >= 1.19f )
+			{
+				FlightCameraSpeedScale -= 0.2f;
+			}
+			else
+			{
+				FlightCameraSpeedScale -= 0.1f;
+			}
 		}
-	}
 
-	FlightCameraSpeedScale = FMath::Clamp( FlightCameraSpeedScale, MinCameraSpeedScale, MaxCameraSpeedScale );
+		FlightCameraSpeedScale = FMath::Clamp( FlightCameraSpeedScale, MinCameraSpeedScale, MaxCameraSpeedScale );
 
-	if( FMath::IsNearlyEqual( FlightCameraSpeedScale, 1.0f, 0.01f ) )
-	{
-		// Snap to 1.0 if we're really close to that
-		FlightCameraSpeedScale = 1.0f;
+		if( FMath::IsNearlyEqual( FlightCameraSpeedScale, 1.0f, 0.01f ) )
+		{
+			// Snap to 1.0 if we're really close to that
+			FlightCameraSpeedScale = 1.0f;
+		}
 	}
 }
 
@@ -4547,7 +4562,7 @@ bool FEditorViewportClient::IsFlightCameraInputModeActive() const
 				bIsTracking &&
 				Widget->GetCurrentAxis() == EAxisList::None &&
 				( bLeftMouseButtonDown || bMiddleMouseButtonDown || bRightMouseButtonDown || bIsUsingTrackpad ) &&
-				!IsCtrlPressed() && !IsShiftPressed() && !IsAltPressed();
+				!IsCtrlPressed() && (GetDefault<ULevelEditorViewportSettings>()->FlightCameraControlExperimentalNavigation || !IsShiftPressed()) && !IsAltPressed();
 
 			return bIsMouseLooking;
 		}
