Edigrate CL 9582789 + Fix from //UE4/dev-editor-openworld

 LODActors no longer generate MessageBox when referenced actors are being deleted

#rnx

diff --git a/Engine/Source/Editor/UnrealEd/Private/EditorActor.cpp b/Engine/Source/Editor/UnrealEd/Private/EditorActor.cpp
--- a/Engine/Source/Editor/UnrealEd/Private/EditorActor.cpp
+++ b/Engine/Source/Editor/UnrealEd/Private/EditorActor.cpp
@@ -849,6 +849,7 @@ bool UUnrealEdEngine::edactDeleteSelected( UWorld* InWorld, bool bVerifyDeletion
 
 		bool bReferencedByLevelScript = bWarnAboutReferences && (nullptr != LSB && ReferencedToActorsFromLevelScriptArray.Num() > 0);
 		bool bReferencedByActor = false;
+		bool bReferencedByLODActor = false;
 		bool bReferencedBySoftReference = false;
 		TArray<UObject*>* SoftReferencingObjects = nullptr;
 
@@ -872,26 +873,32 @@ bool UUnrealEdEngine::edactDeleteSelected( UWorld* InWorld, bool bVerifyDeletion
 				{
 					continue;
 				}
-
-				// If the referencing actor is a child actor that is referencing us, do not treat it
-				// as referencing for the purposes of warning about deletion
-				UChildActorComponent* ParentComponent = ReferencingActor->GetParentComponent();
-				if (ParentComponent == nullptr || ParentComponent->GetOwner() != Actor)
+				else if (Cast<ALODActor>(ReferencingActor))
+				{
+					bReferencedByLODActor = true;
+				}
+				else
 				{
-					bReferencedByActor = true;
-
-					FText ActorReferencedMessage = FText::Format(LOCTEXT("ActorDeleteReferencedMessage",
-						"Actor {0} is referenced by {1}."),
-						FText::FromString(Actor->GetActorLabel()),
-						FText::FromString(ReferencingActor->GetActorLabel())
-					);
-					UE_LOG(LogEditorActor, Log, TEXT("%s"), *ActorReferencedMessage.ToString());
+					// If the referencing actor is a child actor that is referencing us, do not treat it
+					// as referencing for the purposes of warning about deletion
+					UChildActorComponent* ParentComponent = ReferencingActor->GetParentComponent();
+					if (ParentComponent == nullptr || ParentComponent->GetOwner() != Actor)
+					{
+						bReferencedByActor = true;
+
+						FText ActorReferencedMessage = FText::Format(LOCTEXT("ActorDeleteReferencedMessage",
+							"Actor {0} is referenced by {1}."),
+							FText::FromString(Actor->GetActorLabel()),
+							FText::FromString(ReferencingActor->GetActorLabel())
+						);
+						UE_LOG(LogEditorActor, Log, TEXT("%s"), *ActorReferencedMessage.ToString());
+					}
 				}
 			}
 		}
 
 		// We have references from one or more sources, prompt the user for feedback.
-		if (bReferencedByLevelScript || bReferencedByActor || bReferencedBySoftReference)
+		if (bReferencedByLevelScript || bReferencedByActor || bReferencedBySoftReference || bReferencedByLODActor)
 		{
 			if ((bReferencedByLevelScript && !bRequestedDeleteAllByLevel) ||
 				(bReferencedByActor && !bRequestedDeleteAllByActor) ||
@@ -994,8 +1001,10 @@ bool UUnrealEdEngine::edactDeleteSelected( UWorld* InWorld, bool bVerifyDeletion
 			{
 				FBlueprintEditorUtils::ModifyActorReferencedGraphNodes(LSB, Actor);
 			}
-			if (bReferencedByActor && ReferencingActors != nullptr)
+
+			if (bReferencedByActor || bReferencedByLODActor)
 			{
+				check(ReferencingActors != nullptr);
 				for (int32 ReferencingActorIndex = 0; ReferencingActorIndex < ReferencingActors->Num(); ReferencingActorIndex++)
 				{
 					AActor* ReferencingActor = (*ReferencingActors)[ReferencingActorIndex];
@@ -1006,6 +1015,13 @@ bool UUnrealEdEngine::edactDeleteSelected( UWorld* InWorld, bool bVerifyDeletion
 					if (LODActor)
 					{
 						LODActor->RemoveSubActor(Actor);
+
+						FText SubActorRemovedMessage = FText::Format(LOCTEXT("LODActorSubActorDeletedMessage",
+							"Sub Actor '{0}' was removed from LODActor '{1}'."),
+							FText::FromString(Actor->GetActorLabel()),
+							FText::FromString(ReferencingActor->GetActorLabel())
+						);
+						UE_LOG(LogEditorActor, Log, TEXT("%s"), *SubActorRemovedMessage.ToString());
 					}
 				}
 			}
