LevelPartitionGrid improvements
* Don't freeze level bounds for grid cell levels
* Added labels for cells + option
* Removed non selection visualizer as it's now unused
* Misc cleanup

diff --git a/Engine/Source/Editor/UnrealEd/Classes/Editor/UnrealEdEngine.h b/Engine/Source/Editor/UnrealEd/Classes/Editor/UnrealEdEngine.h
--- a/Engine/Source/Editor/UnrealEd/Classes/Editor/UnrealEdEngine.h
+++ b/Engine/Source/Editor/UnrealEd/Classes/Editor/UnrealEdEngine.h
@@ -840,11 +840,6 @@ public:
 	DECLARE_DELEGATE_RetVal_OneParam(TSharedPtr<IAssetReferenceFilter>, FOnMakeAssetReferenceFilter, const FAssetReferenceFilterContext& /*Context*/);
 	FOnMakeAssetReferenceFilter& OnMakeAssetReferenceFilter() { return OnMakeAssetReferenceFilterDelegate; }
 	TSharedPtr<IAssetReferenceFilter> MakeAssetReferenceFilter(const FAssetReferenceFilterContext& Context) { return OnMakeAssetReferenceFilterDelegate.IsBound() ? OnMakeAssetReferenceFilterDelegate.Execute(Context) : nullptr; }
-
-	/** Add Visualizers that are active even when Actor is not part of the selection. */
-	void AddNonSelectionVisualizers(AActor* Actor);
-	/** Remove Visualizers that are active even when Actor is not part of the selection */
-	void RemoveNonSelectionVisualizers(AActor* Actor);
 	
 protected:
 
@@ -865,12 +860,7 @@ protected:
 	/**
 	 * The list of visualizers to draw when selection changes
 	 */
-	TArray<FCachedComponentVisualizer> SelectionVisualizers;
-
-	/**
-	 * The list of visualizers to draw independant of selection
-	 */
-	TArray<FCachedComponentVisualizer> NonSelectionVisualizers;
+	TArray<FCachedComponentVisualizer> VisualizersForSelection;
 
 	/** Instance responsible for monitoring this editor's performance */
 	FPerformanceMonitor* PerformanceMonitor;
@@ -894,7 +884,4 @@ private:
 	* @return	Returns the number of dirty packages that require checkout. If bCheckIfAny is true, returns 1 if any packages will require checkout.
 	*/
 	int32 InternalGetNumDirtyPackagesThatNeedCheckout(bool bCheckIfAny) const;
-
-	/** Internal function to filter and add visualizers to a specific list */
-	void AddVisualizers(AActor* Actor, TArray<FCachedComponentVisualizer>& Visualizers, TFunctionRef<bool(const TSharedPtr<FComponentVisualizer>&)> Condition);
 };
diff --git a/Engine/Source/Editor/UnrealEd/Private/UnrealEdEngine.cpp b/Engine/Source/Editor/UnrealEd/Private/UnrealEdEngine.cpp
--- a/Engine/Source/Editor/UnrealEd/Private/UnrealEdEngine.cpp
+++ b/Engine/Source/Editor/UnrealEd/Private/UnrealEdEngine.cpp
@@ -1300,7 +1300,7 @@ void UUnrealEdEngine::RegisterComponentVisualizer(FName ComponentClassName, TSha
 void UUnrealEdEngine::UnregisterComponentVisualizer(FName ComponentClassName)
 {
 	TSharedPtr<FComponentVisualizer> Visualizer = FindComponentVisualizer(ComponentClassName);
-	SelectionVisualizers.RemoveAll([&Visualizer](const auto& CachedComponentVisualizer) { return CachedComponentVisualizer.Visualizer == Visualizer; });
+	VisualizersForSelection.RemoveAll([&Visualizer](const auto& CachedComponentVisualizer) { return CachedComponentVisualizer.Visualizer == Visualizer; });
 
 	ComponentVisualizerMap.Remove(ComponentClassName);
 }
@@ -1335,12 +1335,7 @@ TSharedPtr<class FComponentVisualizer> UUnrealEdEngine::FindComponentVisualizer(
 
 void UUnrealEdEngine::DrawComponentVisualizers(const FSceneView* View, FPrimitiveDrawInterface* PDI)
 {
-	for(FCachedComponentVisualizer& CachedVisualizer : SelectionVisualizers)
-	{
-		CachedVisualizer.Visualizer->DrawVisualization(CachedVisualizer.ComponentPropertyPath.GetComponent(), View, PDI);
-	}
-
-	for(FCachedComponentVisualizer& CachedVisualizer : NonSelectionVisualizers)
+	for(FCachedComponentVisualizer& CachedVisualizer : VisualizersForSelection)
 	{
 		CachedVisualizer.Visualizer->DrawVisualization(CachedVisualizer.ComponentPropertyPath.GetComponent(), View, PDI);
 	}
@@ -1349,44 +1344,19 @@ void UUnrealEdEngine::DrawComponentVisualizers(const FSceneView* View, FPrimitiv
 
 void UUnrealEdEngine::DrawComponentVisualizersHUD(const FViewport* Viewport, const FSceneView* View, FCanvas* Canvas)
 {
-	for(FCachedComponentVisualizer& CachedVisualizer : SelectionVisualizers)
-	{
-		CachedVisualizer.Visualizer->DrawVisualizationHUD(CachedVisualizer.ComponentPropertyPath.GetComponent(), Viewport, View, Canvas);
-	}
-
-	for(FCachedComponentVisualizer& CachedVisualizer : NonSelectionVisualizers)
+	for(FCachedComponentVisualizer& CachedVisualizer : VisualizersForSelection)
 	{
 		CachedVisualizer.Visualizer->DrawVisualizationHUD(CachedVisualizer.ComponentPropertyPath.GetComponent(), Viewport, View, Canvas);
 	}
 }
 
-void UUnrealEdEngine::AddVisualizers(AActor* Actor, TArray<FCachedComponentVisualizer>& Visualizers, TFunctionRef<bool(const TSharedPtr<FComponentVisualizer>&)> Condition)
-{
-	TInlineComponentArray<UActorComponent*> Components;
-	Actor->GetComponents(Components, true);
-
-	for(int32 CompIdx = 0; CompIdx < Components.Num(); CompIdx++)
-	{
-		UActorComponent* Comp = Components[CompIdx];
-		if(Comp->IsRegistered())
-		{
-			// Try and find a visualizer
-			TSharedPtr<FComponentVisualizer> Visualizer = FindComponentVisualizer(Comp->GetClass());
-			if(Visualizer.IsValid() && Condition(Visualizer))
-			{
-				Visualizers.Add(FCachedComponentVisualizer(Comp, Visualizer));
-			}
-		}
-	}
-}
-
 void UUnrealEdEngine::OnEditorSelectionChanged(UObject* SelectionThatChanged)
 {
 	if(SelectionThatChanged == GetSelectedActors())
 	{
 		// actor selection changed.  Update the list of component visualizers
 		// This is expensive so we do not search for visualizers each time they want to draw
-		SelectionVisualizers.Empty();
+		VisualizersForSelection.Empty();
 
 		// Iterate over all selected actors
 		for(FSelectionIterator It(GetSelectedActorIterator()); It; ++It)
@@ -1394,37 +1364,29 @@ void UUnrealEdEngine::OnEditorSelectionChanged(UObject* SelectionThatChanged)
 			AActor* Actor = Cast<AActor>(*It);
 			if(Actor != nullptr)
 			{
-				AddVisualizers(Actor, SelectionVisualizers, [](const TSharedPtr<FComponentVisualizer>& Visualizer) { return Visualizer->ShowWhenSelected(); });
-			}
-		}
-	}
-}
+				// Then iterate over components of that actor (and recurse through child components)
+				TInlineComponentArray<UActorComponent*> Components;
+				Actor->GetComponents(Components, true);
 
-void UUnrealEdEngine::AddNonSelectionVisualizers(AActor* Actor)
-{
-	AddVisualizers(Actor, NonSelectionVisualizers, [](const TSharedPtr<FComponentVisualizer>& Visualizer) { return !Visualizer->ShowWhenSelected(); });
-}
-
-void UUnrealEdEngine::RemoveNonSelectionVisualizers(AActor* Actor)
-{
-	TInlineComponentArray<UActorComponent*> Components;
-	Actor->GetComponents(Components);
-
-	for (int32 CompIdx = 0; CompIdx < Components.Num(); CompIdx++)
-	{
-		UActorComponent* Comp = Components[CompIdx];
-		if (Comp->IsRegistered())
-		{
-			// Try and find a visualizer
-			TSharedPtr<FComponentVisualizer> Visualizer = FindComponentVisualizer(Comp->GetClass());
-			if (Visualizer.IsValid())
-			{
-				NonSelectionVisualizers.RemoveAll([&Visualizer, Actor](const auto& CachedComponentVisualizer) { return CachedComponentVisualizer.ComponentPropertyPath.GetComponent() && CachedComponentVisualizer.ComponentPropertyPath.GetComponent()->GetOwner() == Actor; });
+				for(int32 CompIdx = 0; CompIdx < Components.Num(); CompIdx++)
+				{
+					UActorComponent* Comp = Components[CompIdx];
+					if(Comp->IsRegistered())
+					{
+						// Try and find a visualizer
+						TSharedPtr<FComponentVisualizer> Visualizer = FindComponentVisualizer(Comp->GetClass());
+						if(Visualizer.IsValid())
+						{
+							VisualizersForSelection.Add(FCachedComponentVisualizer(Comp, Visualizer));
+						}
+					}
+				}
 			}
 		}
 	}
 }
 
+
 EWriteDisallowedWarningState UUnrealEdEngine::GetWarningStateForWritePermission(const FString& PackageName) const
 {
 	EWriteDisallowedWarningState WarningState = WDWS_WarningUnnecessary;
diff --git a/Engine/Source/Editor/UnrealEd/Public/ComponentVisualizer.h b/Engine/Source/Editor/UnrealEd/Public/ComponentVisualizer.h
--- a/Engine/Source/Editor/UnrealEd/Public/ComponentVisualizer.h
+++ b/Engine/Source/Editor/UnrealEd/Public/ComponentVisualizer.h
@@ -108,8 +108,6 @@ public:
 
 	/** */
 	virtual void OnRegister() {}
-	/** Only show this visualizer if the actor is selected */
-	virtual bool ShowWhenSelected() { return true; }
 	/** Draw visualization for the supplied component */
 	virtual void DrawVisualization(const UActorComponent* Component, const FSceneView* View, FPrimitiveDrawInterface* PDI) {}
 	/** Draw HUD on viewport for the supplied component */
