Correct a bad check which was failing cooks when stage-time encrypted assets were enabled

diff --git a/Engine/Source/Editor/UnrealEd/Private/Commandlets/AssetRegistryGenerator.cpp b/Engine/Source/Editor/UnrealEd/Private/Commandlets/AssetRegistryGenerator.cpp
--- a/Engine/Source/Editor/UnrealEd/Private/Commandlets/AssetRegistryGenerator.cpp
+++ b/Engine/Source/Editor/UnrealEd/Private/Commandlets/AssetRegistryGenerator.cpp
@@ -342,7 +342,13 @@ bool FAssetRegistryGenerator::GenerateStreamingInstallManifest(int64 InExtraFlav
 					FinalChunkManifests.AddDefaulted(PakchunkIndex - FinalChunkManifests.Num() + 1);
 				}
 				checkf(PakchunkIndex < FinalChunkManifests.Num(), TEXT("Chunk %i out of range. %i manifests available"), PakchunkIndex, FinalChunkManifests.Num() - 1);
-				checkf(FinalChunkManifests[PakchunkIndex] == nullptr, TEXT("Manifest already exists for chunk %i"), PakchunkIndex);
+				checkf(FinalChunkManifests[PakchunkIndex] == nullptr || FinalChunkManifests[PakchunkIndex]->Num() == 0, TEXT("Manifest already exists for chunk %i"), PakchunkIndex);
+
+				if (FinalChunkManifests[PakchunkIndex])
+				{
+					delete FinalChunkManifests[PakchunkIndex];
+				}
+
 				FinalChunkManifests[PakchunkIndex] = NewManifest;
 			}
 		}
