PromptWhenAddingToLevelOutsideBounds - Rely on the LevelBoundsActor if there's one

diff --git a/Engine/Source/Editor/UnrealEd/Private/AssetSelection.cpp b/Engine/Source/Editor/UnrealEd/Private/AssetSelection.cpp
--- a/Engine/Source/Editor/UnrealEd/Private/AssetSelection.cpp
+++ b/Engine/Source/Editor/UnrealEd/Private/AssetSelection.cpp
@@ -483,7 +483,16 @@ namespace ActorPlacementUtils
 		}
 		if (InLevel && GetDefault<ULevelEditorMiscSettings>()->bPromptWhenAddingToLevelOutsideBounds)
 		{
-			FBox CurrentLevelBounds = ALevelBounds::CalculateLevelBounds(InLevel);
+			FBox CurrentLevelBounds(ForceInit);
+			if (InLevel->LevelBoundsActor.IsValid())
+			{
+				CurrentLevelBounds = InLevel->LevelBoundsActor.Get()->GetComponentsBoundingBox();
+			}
+			else
+			{
+				CurrentLevelBounds = ALevelBounds::CalculateLevelBounds(InLevel);
+			}
+
 			FVector BoundsExtent = CurrentLevelBounds.GetExtent();
 			if (BoundsExtent.X < GetDefault<ULevelEditorMiscSettings>()->MinimumBoundsForCheckingSize.X
 				&& BoundsExtent.Y < GetDefault<ULevelEditorMiscSettings>()->MinimumBoundsForCheckingSize.Y
