Crash fix for trying to open asset editor for invalid asset. Now logs error message instead of assert.

diff --git a/Engine/Source/Editor/UnrealEd/Private/Subsystems/AssetEditorSubsystem.cpp b/Engine/Source/Editor/UnrealEd/Private/Subsystems/AssetEditorSubsystem.cpp
--- a/Engine/Source/Editor/UnrealEd/Private/Subsystems/AssetEditorSubsystem.cpp
+++ b/Engine/Source/Editor/UnrealEd/Private/Subsystems/AssetEditorSubsystem.cpp
@@ -264,7 +264,12 @@ bool UAssetEditorSubsystem::CloseAllAssetEditors()
 
 bool UAssetEditorSubsystem::OpenEditorForAsset(UObject* Asset, const EToolkitMode::Type ToolkitMode, TSharedPtr< IToolkitHost > OpenedFromLevelEditor, const bool bShowProgressWindow)
 {
-	check(Asset);
+	if (!Asset)
+	{
+		UE_LOG(LogAssetEditorSubsystem, Error, TEXT("Opening Asset editor failed because asset is null"));
+		return false;
+	}
+
 	// @todo toolkit minor: When "Edit Here" happens in a different level editor from the one that an asset is already
 	//    being edited within, we should decide whether to disallow "Edit Here" in that case, or to close the old asset
 	//    editor and summon it in the new level editor, or to just foreground the old level editor (current behavior)
