Crash fix for trying to open asset editors for invalid assets. Now logs error message instead of assert.

diff --git a/Engine/Source/Editor/Persona/Private/SAssetFamilyShortcutBar.cpp b/Engine/Source/Editor/Persona/Private/SAssetFamilyShortcutBar.cpp
--- a/Engine/Source/Editor/Persona/Private/SAssetFamilyShortcutBar.cpp
+++ b/Engine/Source/Editor/Persona/Private/SAssetFamilyShortcutBar.cpp
@@ -192,9 +192,16 @@ public:
 	{
 		if(AssetData.IsValid())
 		{
-			TArray<UObject*> Assets;
-			Assets.Add(AssetData.GetAsset());
-			GEditor->GetEditorSubsystem<UAssetEditorSubsystem>()->OpenEditorForAssets(Assets);
+			if (UObject* AssetObject = AssetData.GetAsset())
+			{
+				TArray<UObject*> Assets;
+				Assets.Add(AssetObject);
+				GEditor->GetEditorSubsystem<UAssetEditorSubsystem>()->OpenEditorForAssets(Assets);
+			}
+			else
+			{
+				UE_LOG(LogAnimation, Error, TEXT("Asset cannot be opened: %s"), *AssetData.ObjectPath.ToString());
+			}
 		}
 	}
 
diff --git a/Engine/Source/Editor/UnrealEd/Private/Subsystems/AssetEditorSubsystem.cpp b/Engine/Source/Editor/UnrealEd/Private/Subsystems/AssetEditorSubsystem.cpp
--- a/Engine/Source/Editor/UnrealEd/Private/Subsystems/AssetEditorSubsystem.cpp
+++ b/Engine/Source/Editor/UnrealEd/Private/Subsystems/AssetEditorSubsystem.cpp
@@ -377,8 +377,32 @@ bool UAssetEditorSubsystem::OpenEditorForAsset(UObject* Asset, const EToolkitMod
 }
 
 
-bool UAssetEditorSubsystem::OpenEditorForAssets_Advanced(const TArray <UObject* >& Assets, const EToolkitMode::Type ToolkitMode, TSharedPtr< IToolkitHost > OpenedFromLevelEditor)
+bool UAssetEditorSubsystem::OpenEditorForAssets_Advanced(const TArray <UObject* >& InAssets, const EToolkitMode::Type ToolkitMode, TSharedPtr< IToolkitHost > OpenedFromLevelEditor)
 {
+	TArray<UObject*> Assets;
+	Assets.Reserve(InAssets.Num());
+	int32 NumNullAssets = 0;
+	for (UObject* Asset : InAssets)
+	{
+		if (Asset)
+		{
+			Assets.AddUnique(Asset);
+		}
+		else
+		{
+			++NumNullAssets;
+		}
+	}
+
+	if (NumNullAssets > 1)
+	{
+		UE_LOG(LogAssetEditorSubsystem, Error, TEXT("Opening Asset editors failed because of null assets"));
+	}
+	else if (NumNullAssets > 0)
+	{
+		UE_LOG(LogAssetEditorSubsystem, Error, TEXT("Opening Asset editor failed because of null asset"));
+	}
+
 	if (Assets.Num() == 1)
 	{
 		return OpenEditorForAsset(Assets[0], ToolkitMode, OpenedFromLevelEditor);
