Fixed errors occuring when loading maps when persistent HLOD static meshes changed names
Fixes the "Failed import for StaticMesh" errors which could cause nightly build failure.

Changed the name of the static meshes generated through the HLOD process to remove the "first actor name" part... which wasn't entirelly deterministic

diff --git a/Engine/Source/Developer/HierarchicalLODUtilities/Private/HierarchicalLODUtilities.cpp b/Engine/Source/Developer/HierarchicalLODUtilities/Private/HierarchicalLODUtilities.cpp
--- a/Engine/Source/Developer/HierarchicalLODUtilities/Private/HierarchicalLODUtilities.cpp
+++ b/Engine/Source/Developer/HierarchicalLODUtilities/Private/HierarchicalLODUtilities.cpp
@@ -390,8 +390,6 @@ bool FHierarchicalLODUtilities::BuildStaticMeshForLODActor(ALODActor* LODActor,
 
 	if (AllComponents.Num() > 0)
 	{
-		const AActor* FirstActor = UHLODProxy::FindFirstActor(LODActor);
-
 		TArray<UObject*> OutAssets;
 		FVector OutProxyLocation = FVector::ZeroVector;
 		UStaticMesh* MainMesh = nullptr;
@@ -401,9 +399,16 @@ bool FHierarchicalLODUtilities::BuildStaticMeshForLODActor(ALODActor* LODActor,
 		const bool bHasMeshReductionCapableModule = (MeshReductionModule.GetMeshMergingInterface() != NULL);
 
 		const IMeshMergeUtilities& MeshMergeUtilities = FModuleManager::Get().LoadModuleChecked<IMeshMergeModule>("MeshMergeUtilities").GetUtilities();
-		// should give unique name, so use level + actor name
-			
-		const FString PackageName = FString::Printf(TEXT("LOD_%s_%i_%s"), *FirstActor->GetOutermost()->GetName(), LODActor->LODLevel - 1, *FirstActor->GetName());
+		
+		// Should give a unique name, so use the LODActor tag, or if empty, the first actor name
+		FString LODActorTag = LODActor->GetLODActorTag();
+		if (LODActorTag.IsEmpty())
+		{
+			const AActor* FirstActor = UHLODProxy::FindFirstActor(LODActor);
+			LODActorTag = *FirstActor->GetName();
+		}
+		const FString PackageName = FString::Printf(TEXT("LOD_%s_%i_%s"), *(AssetsOuter->GetName()), LODActor->LODLevel - 1, *LODActorTag);
+
 		if (bHasMeshReductionCapableModule && LODSetup.bSimplifyMesh)
 		{
 			FHierarchicalLODUtilitiesModule& Module = FModuleManager::LoadModuleChecked<FHierarchicalLODUtilitiesModule>("HierarchicalLODUtilities");
@@ -445,7 +450,7 @@ bool FHierarchicalLODUtilities::BuildStaticMeshForLODActor(ALODActor* LODActor,
 			// update LOD parents before rebuild to ensure they are valid when mesh merge extensions are called.
 			LODActor->UpdateSubActorLODParents();
 
-			MeshMergeUtilities.MergeComponentsToStaticMesh(AllComponents, FirstActor->GetWorld(), MergeSettings, InBaseMaterial, AssetsOuter, PackageName, OutAssets, OutProxyLocation, LODSetup.TransitionScreenSize, true);
+			MeshMergeUtilities.MergeComponentsToStaticMesh(AllComponents, LODActor->GetWorld(), MergeSettings, InBaseMaterial, AssetsOuter, PackageName, OutAssets, OutProxyLocation, LODSetup.TransitionScreenSize, true);
 
 			// set staticmesh
 			for (UObject* Asset : OutAssets)
diff --git a/Engine/Source/Editor/UnrealEd/Private/HierarchicalLOD.cpp b/Engine/Source/Editor/UnrealEd/Private/HierarchicalLOD.cpp
--- a/Engine/Source/Editor/UnrealEd/Private/HierarchicalLOD.cpp
+++ b/Engine/Source/Editor/UnrealEd/Private/HierarchicalLOD.cpp
@@ -213,6 +213,7 @@ void FHierarchicalLODBuilder::BuildClusters(ULevel* InLevel)
 						}
 
 						ALODActor* LODActor = CreateLODActor(PreviousActorCluster, InLevel, LODId);
+						LODActor->SetLODActorTag(PreviousLODActor->GetLODActorTag());
 						LODLevelLODActors[LODId].Add(LODActor);
 
 						ValidStaticMeshActorsInLevel.RemoveAll([PreviousActorCluster](AActor* InActor) { return PreviousActorCluster.Actors.Contains(InActor); });
@@ -435,6 +436,7 @@ void FHierarchicalLODBuilder::GenerateAsSingleCluster(const int32 NumHLODLevels,
 		if (LevelCluster.IsValid())
 		{
 			ALODActor* LODActor = CreateLODActor(LevelCluster, InLevel, LODId);
+			LODActor->SetLODActorTag("SingleCluster");
 			PreviousLevelActor = LODActor;
 		}
 	}
@@ -812,14 +814,13 @@ void FHierarchicalLODBuilder::BuildMeshesForLODActors(bool bForceAll)
 					checkf(AssetsOuter != nullptr, TEXT("Failed to created outer for generated HLOD assets"));
 					AssetsOuter->Modify();
 
-					int32 CurrentLODLevel = LODIndex;
 					int32 LODActorIndex = 0;
-					TArray<ALODActor*>& LODLevel = LODLevelActors[CurrentLODLevel];
+					TArray<ALODActor*>& LODLevel = LODLevelActors[LODIndex];
 					for (ALODActor* Actor : LODLevel)
 					{
-						SlowTask.EnterProgressFrame(100.0f / (float)NumLODActors, FText::Format(LOCTEXT("HierarchicalLOD_BuildLODActorMeshesProgress", "Building LODActor Mesh {1} of {2} (LOD Level {0})"), FText::AsNumber(LODIndex + 1), FText::AsNumber(LODActorIndex), FText::AsNumber(LODLevelActors[CurrentLODLevel].Num())));
+						SlowTask.EnterProgressFrame(100.0f / (float)NumLODActors, FText::Format(LOCTEXT("HierarchicalLOD_BuildLODActorMeshesProgress", "Building LODActor Mesh {0} of {1} (LOD Level {2})"), FText::AsNumber(LODActorIndex), FText::AsNumber(LODLevelActors[LODIndex].Num()), FText::AsNumber(LODIndex + 1)));
 
-						bool bBuildSuccessful = Utilities->BuildStaticMeshForLODActor(Actor, AssetsOuter, BuildLODLevelSettings[CurrentLODLevel], BaseMaterial);
+						bool bBuildSuccessful = Utilities->BuildStaticMeshForLODActor(Actor, AssetsOuter, BuildLODLevelSettings[LODIndex], BaseMaterial);
 
 						// Report an error if the build failed
 						if (!bBuildSuccessful)
@@ -1141,6 +1142,7 @@ void FHierarchicalLODBuilder::MergeClustersAndBuildActors(ULevel* InLevel, const
 							if (AHierarchicalLODVolume* const* Volume = HLODVolumeClusters.FindKey(Cluster))
 							{
 								HLODVolumeActors.Add(LODActor, *Volume);
+								LODActor->SetLODActorTag((*Volume)->GetName());
 							}
 						}
 
diff --git a/Engine/Source/Runtime/Engine/Classes/Engine/LODActor.h b/Engine/Source/Runtime/Engine/Classes/Engine/LODActor.h
--- a/Engine/Source/Runtime/Engine/Classes/Engine/LODActor.h
+++ b/Engine/Source/Runtime/Engine/Classes/Engine/LODActor.h
@@ -105,6 +105,21 @@ public:
 	 */
 	void SetLODParent(UPrimitiveComponent* InLODParent, float InParentDrawDistance, bool bInApplyToImposters);
 
+#if WITH_EDITORONLY_DATA
+	/** 
+	 * Set the named tag which describes the source of the LOD actor (HLOD Volume, single cluster / level, etc)
+	 * Helps uniquely identifying auto generated LOD actors & naming of the proxy static mesh assets.
+	 * @param	InLODActorTag		Named tag to assign to this LODActor.
+	 */
+	void SetLODActorTag(const FString& InLODActorTag) { LODActorTag = InLODActorTag; }
+
+	/**
+	 * Get the named tag which describes the source of the LOD actor (HLOD Volume, single cluster / level, etc)
+	 * Helps uniquely identifying auto generated LOD actors & naming of the proxy static mesh assets.
+	 * @returns the named tag assigned to this LODActor.
+	 */
+	const FString& GetLODActorTag() const { return LODActorTag; }
+#endif // WITH_EDITORONLY_DATA
 
 	/** Get the key that we use to check if we need to (re)build */
 	const FName& GetKey() const { return Key; }
@@ -304,4 +319,9 @@ private:
 
 	// Sink for when CVars are changed to check to see if the maximum HLOD level value has changed
 	static FAutoConsoleVariableSink CVarSink;
+
+#if WITH_EDITORONLY_DATA
+	UPROPERTY()
+	FString LODActorTag;
+#endif // WITH_EDITORONLY_DATA
 }; 
