Fix crash reimporting with less sections

diff --git a/Engine/Source/Developer/MeshBuilder/Private/SkeletalMeshBuilder.cpp b/Engine/Source/Developer/MeshBuilder/Private/SkeletalMeshBuilder.cpp
--- a/Engine/Source/Developer/MeshBuilder/Private/SkeletalMeshBuilder.cpp
+++ b/Engine/Source/Developer/MeshBuilder/Private/SkeletalMeshBuilder.cpp
@@ -50,6 +50,37 @@ struct FSkeletalMeshVertInstanceIDAndZ
 	float Z;
 };
 
+//TODO: move that in a public header and use it everywhere we call FSkeletalMeshImportData::CopyLODImportData
+//or simply add parameter to FSkeletalMeshImportData::CopyLODImportData to do the job inside the function.
+namespace SkeletalMeshBuilderHelperNS
+{
+	void FixFaceMaterial(USkeletalMesh* SkeletalMesh, TArray<SkeletalMeshImportData::FMaterial>& RawMeshMaterials, TArray<SkeletalMeshImportData::FMeshFace>& LODFaces)
+	{
+		//Fix the material for the faces
+		TArray<int32> MaterialRemap;
+		MaterialRemap.Reserve(RawMeshMaterials.Num());
+		for (int32 MaterialIndex = 0; MaterialIndex < RawMeshMaterials.Num(); ++MaterialIndex)
+		{
+			MaterialRemap.Add(MaterialIndex);
+			FName MaterialImportName = *(RawMeshMaterials[MaterialIndex].MaterialImportName);
+			for (int32 MeshMaterialIndex = 0; MeshMaterialIndex < SkeletalMesh->Materials.Num(); ++MeshMaterialIndex)
+			{
+				FName MeshMaterialName = SkeletalMesh->Materials[MeshMaterialIndex].ImportedMaterialSlotName;
+				if (MaterialImportName == MeshMaterialName)
+				{
+					MaterialRemap[MaterialIndex] = MeshMaterialIndex;
+					break;
+				}
+			}
+		}
+		//Update all the faces
+		for (int32 FaceIndex = 0; FaceIndex < LODFaces.Num(); ++FaceIndex)
+		{
+			LODFaces[FaceIndex].MeshMaterialIndex = MaterialRemap[LODFaces[FaceIndex].MeshMaterialIndex];
+		}
+	}
+}
+
 FSkeletalMeshBuilder::FSkeletalMeshBuilder()
 {
 
@@ -93,6 +124,7 @@ bool FSkeletalMeshBuilder::Build(USkeletalMesh* SkeletalMesh, const int32 LODInd
 		TArray<int32> LODPointToRawMap;
 		SkeletalMeshImportData.CopyLODImportData(LODPoints, LODWedges, LODFaces, LODInfluences, LODPointToRawMap);
 
+		SkeletalMeshBuilderHelperNS::FixFaceMaterial(SkeletalMesh, SkeletalMeshImportData.Materials, LODFaces);
 
 		//Build the skeletalmesh using mesh utilities module
 		IMeshUtilities::MeshBuildOptions Options;
diff --git a/Engine/Source/Editor/UnrealEd/Private/Factories/SkeletalMeshImport.cpp b/Engine/Source/Editor/UnrealEd/Private/Factories/SkeletalMeshImport.cpp
--- a/Engine/Source/Editor/UnrealEd/Private/Factories/SkeletalMeshImport.cpp
+++ b/Engine/Source/Editor/UnrealEd/Private/Factories/SkeletalMeshImport.cpp
@@ -998,8 +998,6 @@ void RestoreExistingSkelMeshData(ExistingSkelMeshData* MeshData, USkeletalMesh*
 		for (int32 SectionIndex = 0; SectionIndex < NewSkelMeshLodModel.Sections.Num(); SectionIndex++)
 		{
 			int32 NewMeshSectionMaterialIndex = NewSkelMeshLodModel.Sections[SectionIndex].MaterialIndex;
-			//Get the new skelmesh section slot import name
-			FName NewMeshSectionSlotName = SkeletalMesh->Materials[NewMeshSectionMaterialIndex].ImportedMaterialSlotName;
 
 			if (RemapMaterial.IsValidIndex(NewMeshSectionMaterialIndex))
 			{
@@ -1034,36 +1032,34 @@ void RestoreExistingSkelMeshData(ExistingSkelMeshData* MeshData, USkeletalMesh*
 			{
 				break;
 			}
-
-			FName CurrentSectionImportedMaterialName = SkeletalMesh->Materials[NewSkelMeshLodModel.Sections[SectionIndex].MaterialIndex].ImportedMaterialSlotName;
+			int32 CurrentSectionMaterialIndex = NewSkelMeshLodModel.Sections[SectionIndex].MaterialIndex;
+			FName CurrentSectionImportedMaterialName = SkeletalMesh->Materials[CurrentSectionMaterialIndex].ImportedMaterialSlotName;
 			for (int32 ExistSectionIndex = 0; ExistSectionIndex < MeshData->LastImportMeshLodSectionMaterialData[SafeReimportLODIndex].Num(); ++ExistSectionIndex)
 			{
-				int32 OriginalSectionIndex = NewSkelMeshLodModel.Sections[ExistSectionIndex].OriginalDataSectionIndex;
-				if (!MeshData->LastImportMeshLodSectionMaterialData[SafeReimportLODIndex].IsValidIndex(ExistSectionIndex) || !MeshData->ExistingImportMeshLodSectionMaterialData[SafeReimportLODIndex].IsValidIndex(OriginalSectionIndex))
+				if (!MeshData->LastImportMeshLodSectionMaterialData[SafeReimportLODIndex].IsValidIndex(ExistSectionIndex) || !MeshData->ExistingImportMeshLodSectionMaterialData[SafeReimportLODIndex].IsValidIndex(ExistSectionIndex))
 				{
 					continue;
 				}
-				//Get the Last imported skelmesh section slot import name
-				FName OriginalImportMeshSectionSlotName = MeshData->LastImportMeshLodSectionMaterialData[SafeReimportLODIndex][ExistSectionIndex];
-				if (OriginalImportMeshSectionSlotName != CurrentSectionImportedMaterialName)
+				//The last import slot name is use to match the New import slot name
+				//If the user has change the slot in the editor the ExistMeshSectionSlotName will be different.
+				//This is why we use the ExistMeshSectionSlotName to restore the data and use the LastImportedMeshSectionSlotName to rematch the section
+				FName LastImportedMeshSectionSlotName = MeshData->LastImportMeshLodSectionMaterialData[SafeReimportLODIndex][ExistSectionIndex];
+				if (LastImportedMeshSectionSlotName != CurrentSectionImportedMaterialName)
 				{
+					//This material do not match
 					continue;
 				}
+				//Restore the material slot
+				FName ExistMeshSectionSlotName = MeshData->ExistingImportMeshLodSectionMaterialData[SafeReimportLODIndex][ExistSectionIndex].ImportedMaterialSlotName;
 
-				//Get the current skelmesh section slot import name
-				FName ExistMeshSectionSlotName = MeshData->ExistingImportMeshLodSectionMaterialData[SafeReimportLODIndex][OriginalSectionIndex].ImportedMaterialSlotName;
+				//Override the new section material index to use the one that the user set
+				for (int32 RemapMaterialIndex = 0; RemapMaterialIndex < SkeletalMesh->Materials.Num(); ++RemapMaterialIndex)
 				{
-					//The last import slot name match the New import slot name, but the Exist slot name is different then the last import slot name.
-					//This mean the user has change the section assign slot and the fbx file did not change it
-					//Override the new section material index to use the one that the user set
-					for (int32 RemapMaterialIndex = 0; RemapMaterialIndex < SkeletalMesh->Materials.Num(); ++RemapMaterialIndex)
+					const FSkeletalMaterial &NewSectionMaterial = SkeletalMesh->Materials[RemapMaterialIndex];
+					if (NewSectionMaterial.ImportedMaterialSlotName == ExistMeshSectionSlotName)
 					{
-						const FSkeletalMaterial &NewSectionMaterial = SkeletalMesh->Materials[RemapMaterialIndex];
-						if (NewSectionMaterial.ImportedMaterialSlotName == ExistMeshSectionSlotName)
-						{
-							NewSkelMeshLodModel.Sections[SectionIndex].MaterialIndex = RemapMaterialIndex;
-							break;
-						}
+						NewSkelMeshLodModel.Sections[SectionIndex].MaterialIndex = RemapMaterialIndex;
+						break;
 					}
 				}
 				break;
