Add delete custom attribute curve to re-import menu

diff --git a/Engine/Config/BaseEditorPerProjectUserSettings.ini b/Engine/Config/BaseEditorPerProjectUserSettings.ini
--- a/Engine/Config/BaseEditorPerProjectUserSettings.ini
+++ b/Engine/Config/BaseEditorPerProjectUserSettings.ini
@@ -676,6 +676,7 @@ FrameImportRange=(Min=0, Max=0)
 bUseDefaultSampleRate=False
 CustomSampleRate=0
 bImportCustomAttribute=True
+bDeleteExistingCustomAttributeCurves=False
 bImportBoneTracks=True
 bSetMaterialDriveParameterOnCustomAttribute=False
 bRemoveRedundantKeys=True
diff --git a/Engine/Source/Editor/UnrealEd/Classes/Factories/FbxAnimSequenceImportData.h b/Engine/Source/Editor/UnrealEd/Classes/Factories/FbxAnimSequenceImportData.h
--- a/Engine/Source/Editor/UnrealEd/Classes/Factories/FbxAnimSequenceImportData.h
+++ b/Engine/Source/Editor/UnrealEd/Classes/Factories/FbxAnimSequenceImportData.h
@@ -73,9 +73,12 @@ class UNREALED_API UFbxAnimSequenceImportData : public UFbxAssetImportData
 	/** Import if custom attribute as a curve within the animation */
 	UPROPERTY(EditAnywhere, AdvancedDisplay, config, Category = ImportSettings)
 	bool bImportCustomAttribute;
+
+	/** If true, all previous custom attribute curves will be deleted when doing a re-import. */
+	UPROPERTY(EditAnywhere, AdvancedDisplay, config, Category = ImportSettings)
+	bool bDeleteExistingCustomAttributeCurves;
 	
 	/** Import bone transform tracks. If false, this will discard any bone transform tracks. (useful for curves only animations)*/
-
 	UPROPERTY(EditAnywhere, AdvancedDisplay, config, Category = ImportSettings)
 	bool bImportBoneTracks;
 
diff --git a/Engine/Source/Editor/UnrealEd/Classes/Factories/FbxSceneImportOptionsSkeletalMesh.h b/Engine/Source/Editor/UnrealEd/Classes/Factories/FbxSceneImportOptionsSkeletalMesh.h
--- a/Engine/Source/Editor/UnrealEd/Classes/Factories/FbxSceneImportOptionsSkeletalMesh.h
+++ b/Engine/Source/Editor/UnrealEd/Classes/Factories/FbxSceneImportOptionsSkeletalMesh.h
@@ -81,6 +81,10 @@ class UFbxSceneImportOptionsSkeletalMesh : public UObject
 	UPROPERTY(EditAnywhere, AdvancedDisplay, config, Category = Animation)
 	bool bImportCustomAttribute;
 
+	/** If true, all previous custom attribute curves will be deleted when doing a re-import. */
+	UPROPERTY(EditAnywhere, AdvancedDisplay, config, Category = Animation)
+	bool bDeleteExistingCustomAttributeCurves;
+
 	/** Type of asset to import from the FBX file */
 	UPROPERTY(EditAnywhere, AdvancedDisplay, config, Category = Animation)
 	bool bPreserveLocalTransform;
diff --git a/Engine/Source/Editor/UnrealEd/Private/Fbx/FbxMainImport.cpp b/Engine/Source/Editor/UnrealEd/Private/Fbx/FbxMainImport.cpp
--- a/Engine/Source/Editor/UnrealEd/Private/Fbx/FbxMainImport.cpp
+++ b/Engine/Source/Editor/UnrealEd/Private/Fbx/FbxMainImport.cpp
@@ -479,6 +479,7 @@ void ApplyImportUIToImportOptions(UFbxImportUI* ImportUI, FBXImportOptions& InOu
 	InOutImportOptions.bRemoveRedundantKeys = ImportUI->AnimSequenceImportData->bRemoveRedundantKeys;
 	InOutImportOptions.bDoNotImportCurveWithZero = ImportUI->AnimSequenceImportData->bDoNotImportCurveWithZero;
 	InOutImportOptions.bImportCustomAttribute = ImportUI->AnimSequenceImportData->bImportCustomAttribute;
+	InOutImportOptions.bDeleteExistingCustomAttributeCurves = ImportUI->AnimSequenceImportData->bDeleteExistingCustomAttributeCurves;
 	InOutImportOptions.bImportBoneTracks = ImportUI->AnimSequenceImportData->bImportBoneTracks;
 	InOutImportOptions.bSetMaterialDriveParameterOnCustomAttribute = ImportUI->AnimSequenceImportData->bSetMaterialDriveParameterOnCustomAttribute;
 	InOutImportOptions.MaterialCurveSuffixes = ImportUI->AnimSequenceImportData->MaterialCurveSuffixes;
@@ -1567,6 +1568,7 @@ bool FFbxImporter::ImportFromFile(const FString& Filename, const FString& Type,
 							Attribs.Add(FAnalyticsEventAttribute(TEXT("AnimOpt DoNotImportCurveWithZero"), CaptureImportOptions->bDoNotImportCurveWithZero));
 							Attribs.Add(FAnalyticsEventAttribute(TEXT("AnimOpt ImportBoneTracks"), CaptureImportOptions->bImportBoneTracks));
 							Attribs.Add(FAnalyticsEventAttribute(TEXT("AnimOpt ImportCustomAttribute"), CaptureImportOptions->bImportCustomAttribute));
+							Attribs.Add(FAnalyticsEventAttribute(TEXT("AnimOpt DeleteExistingCustomAttributeCurves"), CaptureImportOptions->bDeleteExistingCustomAttributeCurves));
 							Attribs.Add(FAnalyticsEventAttribute(TEXT("AnimOpt PreserveLocalTransform"), CaptureImportOptions->bPreserveLocalTransform));
 							Attribs.Add(FAnalyticsEventAttribute(TEXT("AnimOpt RemoveRedundantKeys"), CaptureImportOptions->bRemoveRedundantKeys));
 							Attribs.Add(FAnalyticsEventAttribute(TEXT("AnimOpt Resample"), CaptureImportOptions->bResample));
diff --git a/Engine/Source/Editor/UnrealEd/Private/Fbx/FbxSceneImportOptionsSkeletalMesh.cpp b/Engine/Source/Editor/UnrealEd/Private/Fbx/FbxSceneImportOptionsSkeletalMesh.cpp
--- a/Engine/Source/Editor/UnrealEd/Private/Fbx/FbxSceneImportOptionsSkeletalMesh.cpp
+++ b/Engine/Source/Editor/UnrealEd/Private/Fbx/FbxSceneImportOptionsSkeletalMesh.cpp
@@ -23,6 +23,7 @@ UFbxSceneImportOptionsSkeletalMesh::UFbxSceneImportOptionsSkeletalMesh(const FOb
 	, bUseDefaultSampleRate(false)
 	, CustomSampleRate(0)
 	, bImportCustomAttribute(true)
+	, bDeleteExistingCustomAttributeCurves(false)
 	, bPreserveLocalTransform(false)
 	, bDeleteExistingMorphTargetCurves(false)
 {
@@ -53,6 +54,7 @@ void UFbxSceneImportOptionsSkeletalMesh::FillSkeletalMeshInmportData(UFbxSkeleta
 	AnimSequenceImportData->AnimationLength = AnimationLength;
 	AnimSequenceImportData->bDeleteExistingMorphTargetCurves = bDeleteExistingMorphTargetCurves;
 	AnimSequenceImportData->bImportCustomAttribute = bImportCustomAttribute;
+	AnimSequenceImportData->bDeleteExistingCustomAttributeCurves = bDeleteExistingCustomAttributeCurves;
 	AnimSequenceImportData->bPreserveLocalTransform = bPreserveLocalTransform;
 	AnimSequenceImportData->bUseDefaultSampleRate = bUseDefaultSampleRate;
 	AnimSequenceImportData->CustomSampleRate = CustomSampleRate;
diff --git a/Engine/Source/Editor/UnrealEd/Private/Fbx/SFbxSceneOptionWindow.cpp b/Engine/Source/Editor/UnrealEd/Private/Fbx/SFbxSceneOptionWindow.cpp
--- a/Engine/Source/Editor/UnrealEd/Private/Fbx/SFbxSceneOptionWindow.cpp
+++ b/Engine/Source/Editor/UnrealEd/Private/Fbx/SFbxSceneOptionWindow.cpp
@@ -1505,6 +1505,7 @@ void SFbxSceneOptionWindow::CopySkeletalMeshOptionsToFbxOptions(UnFbx::FBXImport
 	ImportSettings->AnimationLengthImportType = SkeletalMeshOptions->AnimationLength;
 	ImportSettings->bDeleteExistingMorphTargetCurves = SkeletalMeshOptions->bDeleteExistingMorphTargetCurves;
 	ImportSettings->bImportCustomAttribute = SkeletalMeshOptions->bImportCustomAttribute;
+	ImportSettings->bDeleteExistingCustomAttributeCurves = SkeletalMeshOptions->bDeleteExistingCustomAttributeCurves;
 	ImportSettings->bPreserveLocalTransform = SkeletalMeshOptions->bPreserveLocalTransform;
 	ImportSettings->bResample = !SkeletalMeshOptions->bUseDefaultSampleRate;
 	ImportSettings->ResampleRate = SkeletalMeshOptions->CustomSampleRate;
@@ -1528,6 +1529,7 @@ void SFbxSceneOptionWindow::CopyFbxOptionsToSkeletalMeshOptions(UnFbx::FBXImport
 	SkeletalMeshOptions->AnimationLength = ImportSettings->AnimationLengthImportType;
 	SkeletalMeshOptions->bDeleteExistingMorphTargetCurves = ImportSettings->bDeleteExistingMorphTargetCurves;
 	SkeletalMeshOptions->bImportCustomAttribute = ImportSettings->bImportCustomAttribute;
+	SkeletalMeshOptions->bDeleteExistingCustomAttributeCurves = ImportSettings->bDeleteExistingCustomAttributeCurves;
 	SkeletalMeshOptions->bPreserveLocalTransform = ImportSettings->bPreserveLocalTransform;
 	SkeletalMeshOptions->bUseDefaultSampleRate = !ImportSettings->bResample;
 	SkeletalMeshOptions->CustomSampleRate = ImportSettings->ResampleRate;
diff --git a/Engine/Source/Editor/UnrealEd/Private/SkeletalMeshEdit.cpp b/Engine/Source/Editor/UnrealEd/Private/SkeletalMeshEdit.cpp
--- a/Engine/Source/Editor/UnrealEd/Private/SkeletalMeshEdit.cpp
+++ b/Engine/Source/Editor/UnrealEd/Private/SkeletalMeshEdit.cpp
@@ -1317,22 +1317,26 @@ bool UnFbx::FFbxImporter::ImportAnimation(USkeleton* Skeleton, UAnimSequence * D
 	USkeleton* MySkeleton = DestSeq->GetSkeleton();
 	check(MySkeleton);
 
-	if (ImportOptions->bDeleteExistingMorphTargetCurves)
+	if (ImportOptions->bDeleteExistingMorphTargetCurves || ImportOptions->bDeleteExistingCustomAttributeCurves)
 	{
 		for (int32 CurveIdx=0; CurveIdx<DestSeq->RawCurveData.FloatCurves.Num(); ++CurveIdx)
 		{
 			auto& Curve = DestSeq->RawCurveData.FloatCurves[CurveIdx];
 			const FCurveMetaData* MetaData = MySkeleton->GetCurveMetaData(Curve.Name);
-			if (MetaData && MetaData->Type.bMorphtarget)
+			if (MetaData)
 			{
-				DestSeq->RawCurveData.FloatCurves.RemoveAt(CurveIdx, 1, false);
-				--CurveIdx;
+				bool bDeleteCurve = MetaData->Type.bMorphtarget ? ImportOptions->bDeleteExistingMorphTargetCurves : ImportOptions->bDeleteExistingCustomAttributeCurves;
+				if (bDeleteCurve)
+				{
+					DestSeq->RawCurveData.FloatCurves.RemoveAt(CurveIdx, 1, false);
+					--CurveIdx;
+				}
 			}
 		}
-
 		DestSeq->RawCurveData.FloatCurves.Shrink();
 	}
 
+
 	// Store float curve tracks which use to exist on the animation
 	TArray<FString> ExistingCurveNames;
 	for (int32 CurveIdx = 0; CurveIdx < DestSeq->RawCurveData.FloatCurves.Num(); ++CurveIdx)
diff --git a/Engine/Source/Editor/UnrealEd/Public/FbxImporter.h b/Engine/Source/Editor/UnrealEd/Public/FbxImporter.h
--- a/Engine/Source/Editor/UnrealEd/Public/FbxImporter.h
+++ b/Engine/Source/Editor/UnrealEd/Public/FbxImporter.h
@@ -190,6 +190,7 @@ struct FBXImportOptions
 	bool	bPreserveLocalTransform;
 	bool	bDeleteExistingMorphTargetCurves;
 	bool	bImportCustomAttribute;
+	bool	bDeleteExistingCustomAttributeCurves;
 	bool	bImportBoneTracks;
 	bool	bSetMaterialDriveParameterOnCustomAttribute;
 	bool	bRemoveRedundantKeys;
