Make sure the fbx exporter correctly configure the fbx sdk time line properties.

UDN 491461 bug report

diff --git a/Engine/Source/Editor/UnrealEd/Private/Fbx/FbxAnimationExport.cpp b/Engine/Source/Editor/UnrealEd/Private/Fbx/FbxAnimationExport.cpp
--- a/Engine/Source/Editor/UnrealEd/Private/Fbx/FbxAnimationExport.cpp
+++ b/Engine/Source/Editor/UnrealEd/Private/Fbx/FbxAnimationExport.cpp
@@ -45,19 +45,26 @@ void FFbxExporter::ExportAnimSequenceToFbx(const UAnimSequence* AnimSeq,
 	}
 
 	const float FrameRate = FMath::TruncToFloat(((AnimSeq->GetRawNumberOfFrames() - 1) / AnimSeq->SequenceLength) + 0.5f);
+	//Configure the scene time line
+	{
+		FbxGlobalSettings& SceneGlobalSettings = Scene->GetGlobalSettings();
+		double CurrentSceneFrameRate = FbxTime::GetFrameRate(SceneGlobalSettings.GetTimeMode());
+		if (!bSceneGlobalTimeLineSet || FrameRate > CurrentSceneFrameRate)
+		{
+			FbxTime::EMode ComputeTimeMode = FbxTime::ConvertFrameRateToTimeMode(FrameRate);
+			FbxTime::SetGlobalTimeMode(ComputeTimeMode, ComputeTimeMode == FbxTime::eCustom ? FrameRate : 0.0);
+			SceneGlobalSettings.SetTimeMode(ComputeTimeMode);
+			if (ComputeTimeMode == FbxTime::eCustom)
+			{
+				SceneGlobalSettings.SetCustomFrameRate(FrameRate);
+			}
+			bSceneGlobalTimeLineSet = true;
+		}
+	}
 
 	// set time correctly
 	FbxTime ExportedStartTime, ExportedStopTime;
-	if ( FMath::IsNearlyEqual(FrameRate, DEFAULT_SAMPLERATE, 1.f) )
-	{
-		ExportedStartTime.SetGlobalTimeMode(FbxTime::eFrames30);
-		ExportedStopTime.SetGlobalTimeMode(FbxTime::eFrames30);
-	}
-	else
-	{
-		ExportedStartTime.SetGlobalTimeMode(FbxTime::eCustom, FrameRate);
-		ExportedStopTime.SetGlobalTimeMode(FbxTime::eCustom, FrameRate);
-	}
+	
 
 	ExportedStartTime.SetSecondDouble(0.f);
 	ExportedStopTime.SetSecondDouble(AnimSeq->SequenceLength);
diff --git a/Engine/Source/Editor/UnrealEd/Private/Fbx/FbxMainExport.cpp b/Engine/Source/Editor/UnrealEd/Private/Fbx/FbxMainExport.cpp
--- a/Engine/Source/Editor/UnrealEd/Private/Fbx/FbxMainExport.cpp
+++ b/Engine/Source/Editor/UnrealEd/Private/Fbx/FbxMainExport.cpp
@@ -255,6 +255,9 @@ void FFbxExporter::CreateDocument()
 	// Maya use cm by default
 	Scene->GetGlobalSettings().SetSystemUnit(FbxSystemUnit::cm);
 	//FbxScene->GetGlobalSettings().SetOriginalSystemUnit( KFbxSystemUnit::m );
+
+	bSceneGlobalTimeLineSet = false;
+	Scene->GetGlobalSettings().SetTimeMode(FbxTime::eDefaultMode);
 	
 	// setup anim stack
 	AnimStack = FbxAnimStack::Create(Scene, "Unreal Take");
@@ -381,6 +384,7 @@ void FFbxExporter::CloseDocument()
 		Scene->Destroy();
 		Scene = NULL;
 	}
+	bSceneGlobalTimeLineSet = false;
 }
 
 void FFbxExporter::CreateAnimatableUserProperty(FbxNode* Node, float Value, const char* Name, const char* Label)
diff --git a/Engine/Source/Editor/UnrealEd/Private/FbxExporter.h b/Engine/Source/Editor/UnrealEd/Private/FbxExporter.h
--- a/Engine/Source/Editor/UnrealEd/Private/FbxExporter.h
+++ b/Engine/Source/Editor/UnrealEd/Private/FbxExporter.h
@@ -483,6 +483,8 @@ private:
 public:
 	/** Returns currently active FBX export options. Automation or UI dialog based options. */
 	UFbxExportOption* GetExportOptions();
+
+	bool bSceneGlobalTimeLineSet = false;
 };
 
 
