Force Delete on an asset that is open in an editor no longer closes all open asset editors.

diff --git a/Engine/Source/Editor/UnrealEd/Private/ObjectTools.cpp b/Engine/Source/Editor/UnrealEd/Private/ObjectTools.cpp
--- a/Engine/Source/Editor/UnrealEd/Private/ObjectTools.cpp
+++ b/Engine/Source/Editor/UnrealEd/Private/ObjectTools.cpp
@@ -2297,11 +2297,25 @@ namespace ObjectTools
 			return 0;
 		}
 
-		// Close all editors to avoid changing references to temporary objects used by the editor
-		if ( !FAssetEditorManager::Get().CloseAllAssetEditors() )
+		// Attempt to close all editors referencing this asset.
+		bool bClosedAllEditors = true;
+
+		for (UObject* ObjectToDelete : InObjectsToDelete)
+		{
+			const TArray<IAssetEditorInstance*> ObjectEditors = FAssetEditorManager::Get().FindEditorsForAsset(ObjectToDelete);
+			for (IAssetEditorInstance* ObjectEditorInstance : ObjectEditors)
+			{
+				if (!ObjectEditorInstance->CloseWindow())
+				{
+					bClosedAllEditors = false;
+				}
+			}
+		}
+
+		// Failed to close at least one editor. It is possible that this editor has in-memory object references
+		// which are not prepared to be changed dynamically so it is not safe to continue
+		if (!bClosedAllEditors)
 		{
-			// Failed to close at least one editor. It is possible that this editor has in-memory object references
-			// which are not prepared to be changed dynamically so it is not safe to continue
 			return 0;
 		}
 
