[NO NON-PII COMMIT MESSAGE]

diff --git a/Engine/Source/Editor/UnrealEd/Private/Settings/LevelEditorPlaySettingsCustomization.h b/Engine/Source/Editor/UnrealEd/Private/Settings/LevelEditorPlaySettingsCustomization.h
--- a/Engine/Source/Editor/UnrealEd/Private/Settings/LevelEditorPlaySettingsCustomization.h
+++ b/Engine/Source/Editor/UnrealEd/Private/Settings/LevelEditorPlaySettingsCustomization.h
@@ -187,9 +187,9 @@ public:
 					[
 						SNew(STextBlock)
 						.Font(LayoutBuilder->GetDetailFont())
-						.Text(LOCTEXT("CommonResolutionsButtonText", "Common Window Sizes"))
+						.Text(LOCTEXT("CommonResolutionsButtonText", "Common Resolutions"))
 					]
-					.ContentPadding(FMargin(6.0f, 2.0f))
+					.ContentPadding(FMargin(6,2))
 					.MenuContent()
 					[
 						MakeCommonResolutionsMenu()
@@ -197,12 +197,13 @@ public:
 					.ToolTipText(LOCTEXT("CommonResolutionsButtonTooltip", "Pick from a list of common screen resolutions"))
 				]
 				+ SHorizontalBox::Slot()
+				.Padding(0,0,6,0)
 				.AutoWidth()
 				.VAlign(VAlign_Center)
 				[
 					SNew(SButton)
 					.OnClicked(this, &SScreenResolutionCustomization::HandleSwapAspectRatioClicked)
-					.ContentPadding(FMargin(3.0f, 0.0f, 3.0f, 1.0f))
+					.ContentPadding(FMargin(3,0,3,1))
 					.Content()
 					[
 						SNew(SImage)
@@ -221,7 +222,7 @@ public:
 					+ SVerticalBox::Slot()
 					.AutoHeight()
 					[
-						WindowWidthProperty->CreatePropertyNameWidget(LOCTEXT("WindowWidthLabel", "Window Width"))
+						WindowWidthProperty->CreatePropertyNameWidget(LOCTEXT("ViewportWidthLabel", "Viewport Width"))
 					]
 					+ SVerticalBox::Slot()
 					[
@@ -229,13 +230,13 @@ public:
 					]
 				]
 				+ SHorizontalBox::Slot()
-				.Padding(8.0f, 0.0f, 0.0f, 0.0f)
+				.Padding(8,0,0,0)
 				[
 					SNew(SVerticalBox)
 					+ SVerticalBox::Slot()
 					.AutoHeight()
 					[
-						WindowHeightProperty->CreatePropertyNameWidget(LOCTEXT("WindowHeightLabel", "Window Height"))
+						WindowHeightProperty->CreatePropertyNameWidget(LOCTEXT("ViewportHeightLabel", "Viewport Height"))
 					]
 					+ SVerticalBox::Slot()
 					.AutoHeight()
@@ -526,7 +527,7 @@ public:
 		}
 		IDetailCategoryBuilder& GameViewportSettings = LayoutBuilder.EditCategory("GameViewportSettings");
 		{
-		// new window size
+		// new window resolution
 		TSharedRef<IPropertyHandle> WindowHeightHandle = LayoutBuilder.GetProperty(GET_MEMBER_NAME_CHECKED(ULevelEditorPlaySettings, NewWindowHeight));
 		TSharedRef<IPropertyHandle> WindowWidthHandle = LayoutBuilder.GetProperty(GET_MEMBER_NAME_CHECKED(ULevelEditorPlaySettings, NewWindowWidth));
 		TSharedRef<IPropertyHandle> WindowPositionHandle = LayoutBuilder.GetProperty(GET_MEMBER_NAME_CHECKED(ULevelEditorPlaySettings, NewWindowPosition));
@@ -539,12 +540,12 @@ public:
 		CenterNewWindowHandle->MarkHiddenByCustomization();
 		EmulatedDeviceHandle->MarkHiddenByCustomization();
 
-		GameViewportSettings.AddCustomRow(LOCTEXT("NewWindowSizeRow", "New Window Size"), false)
+		GameViewportSettings.AddCustomRow(LOCTEXT("NewViewportResolutionRow", "New Viewport Resolution"), false)
 			.NameContent()
 			[
 				SNew(STextBlock)
 				.Font(LayoutBuilder.GetDetailFont())
-				.Text(LOCTEXT("NewWindowSizeName", "New Window Size"))
+				.Text(LOCTEXT("NewViewportResolutionName", "New Viewport Resolution"))
 				.ToolTipText(LOCTEXT("NewWindowSizeTooltip", "Sets the width and height of floating PIE windows (in pixels)"))
 			]
 			.ValueContent()
@@ -553,7 +554,7 @@ public:
 				SNew(SScreenResolutionCustomization, &LayoutBuilder, WindowHeightHandle, WindowWidthHandle)
 			];
 
-			GameViewportSettings.AddCustomRow(LOCTEXT("NewWindowPositionRow", "New Window Position"), false)
+		GameViewportSettings.AddCustomRow(LOCTEXT("NewWindowPositionRow", "New Window Position"), false)
 			.NameContent()
 			[
 				SNew(STextBlock)
@@ -567,19 +568,19 @@ public:
 				SNew(SScreenPositionCustomization, &LayoutBuilder, WindowPositionHandle, CenterNewWindowHandle)
 			];
 
-			GameViewportSettings.AddCustomRow(LOCTEXT("SafeZonePreviewName", "Safe Zone Preview"), false)
-				.NameContent()
-				[
-					SNew(STextBlock)
-					.Font(LayoutBuilder.GetDetailFont())
-					.Text(LOCTEXT("SafeZonePreviewName", "Safe Zone Preview"))
-				]
-				.ValueContent()
-				[
-					SNew(STextBlock)
-					.Font(LayoutBuilder.GetDetailFont())
-					.Text(this, &FLevelEditorPlaySettingsCustomization::GetPreviewText)
-				];
+		GameViewportSettings.AddCustomRow(LOCTEXT("SafeZonePreviewName", "Safe Zone Preview"), false)
+			.NameContent()
+			[
+				SNew(STextBlock)
+				.Font(LayoutBuilder.GetDetailFont())
+				.Text(LOCTEXT("SafeZonePreviewName", "Safe Zone Preview"))
+			]
+			.ValueContent()
+			[
+				SNew(STextBlock)
+				.Font(LayoutBuilder.GetDetailFont())
+				.Text(this, &FLevelEditorPlaySettingsCustomization::GetPreviewText)
+			];
 		}
 
 		// play in new window settings
@@ -695,10 +696,10 @@ public:
 				.DisplayName(LOCTEXT("AdditionalLaunchOptionsLabel", "Command Line Arguments"))
 				.Visibility(TAttribute<EVisibility>::Create(TAttribute<EVisibility>::FGetter::CreateSP(this, &FLevelEditorPlaySettingsCustomization::HandleCmdLineVisibility)));
 
-			NetworkCategory.AddCustomRow(LOCTEXT("PlayInNetworkWindowDetails", "Multiplayer Window Size"), false)
+			NetworkCategory.AddCustomRow(LOCTEXT("PlayInNetworkViewportSize", "Multiplayer Viewport Size"), false)
 				.NameContent()
 				[
-					WindowHeightHandle->CreatePropertyNameWidget(LOCTEXT("ClientWindowSizeName", "Multiplayer Window Size (in pixels)"), LOCTEXT("ClientWindowSizeTooltip", "Width and Height to use when spawning additional windows."))
+					WindowHeightHandle->CreatePropertyNameWidget(LOCTEXT("ClientViewportSizeName", "Multiplayer Viewport Size (in pixels)"), LOCTEXT("ClientWindowSizeTooltip", "Width and Height to use when spawning additional windows."))
 				]
 				.ValueContent()
 				.MaxDesiredWidth(MaxPropertyWidth)
