PIE window is no longer recentered every time when using DPI scaling.

diff --git a/Engine/Source/Editor/UnrealEd/Private/PlayLevel.cpp b/Engine/Source/Editor/UnrealEd/Private/PlayLevel.cpp
--- a/Engine/Source/Editor/UnrealEd/Private/PlayLevel.cpp
+++ b/Engine/Source/Editor/UnrealEd/Private/PlayLevel.cpp
@@ -3330,9 +3330,6 @@ UGameInstance* UEditorEngine::CreatePIEGameInstance(int32 InPIEInstance, bool bI
 				// Mark the viewport as PIE viewport
 				ViewportClient->Viewport->SetPlayInEditorViewport( ViewportClient->bIsPlayInEditorViewport );
 
-				// Ensure the window has a valid size before calling BeginPlay
-				PieWindow->ReshapeWindow(PieWindow->GetPositionInScreen(), FVector2D(NewWindowWidth, NewWindowHeight));
-
 				// Change the system resolution to match our window, to make sure game and slate window are kept syncronised
 				FSystemResolution::RequestResolutionChange(NewWindowWidth, NewWindowHeight, EWindowMode::Windowed);
 
