Fix for flashing tab whenever a blueprint is marked dirty

diff --git a/Engine/Source/Editor/UnrealEd/Private/Kismet2/BlueprintEditorUtils.cpp b/Engine/Source/Editor/UnrealEd/Private/Kismet2/BlueprintEditorUtils.cpp
--- a/Engine/Source/Editor/UnrealEd/Private/Kismet2/BlueprintEditorUtils.cpp
+++ b/Engine/Source/Editor/UnrealEd/Private/Kismet2/BlueprintEditorUtils.cpp
@@ -2141,7 +2141,7 @@ void FBlueprintEditorUtils::MarkBlueprintAsModified(UBlueprint* Blueprint, FProp
 		FBlueprintEditorUtils::ClearMacroCosmeticInfoCache(Blueprint);
 	}
 	
-	IAssetEditorInstance* AssetEditor = FAssetEditorManager::Get().FindEditorForAsset(Blueprint, true);
+	IAssetEditorInstance* AssetEditor = FAssetEditorManager::Get().FindEditorForAsset(Blueprint, false);
 	if (AssetEditor)
 	{
 		FBlueprintEditor* BlueprintEditor = static_cast<FBlueprintEditor*>(AssetEditor);
