Fix VTable missing error

diff --git a/Engine/Source/Editor/UnrealEd/Classes/ThumbnailRendering/BlueprintThumbnailRenderer.h b/Engine/Source/Editor/UnrealEd/Classes/ThumbnailRendering/BlueprintThumbnailRenderer.h
--- a/Engine/Source/Editor/UnrealEd/Classes/ThumbnailRendering/BlueprintThumbnailRenderer.h
+++ b/Engine/Source/Editor/UnrealEd/Classes/ThumbnailRendering/BlueprintThumbnailRenderer.h
@@ -17,7 +17,7 @@ class FCanvas;
 class FRenderTarget;
 
 UCLASS(config=Editor,MinimalAPI)
-class UBlueprintThumbnailRenderer : public UDefaultSizedThumbnailRenderer
+class UNREALED_VTABLE UBlueprintThumbnailRenderer : public UDefaultSizedThumbnailRenderer
 {
 	GENERATED_UCLASS_BODY()
 
