Created SCustomDialog which allows for a dialog to be created with any SWidget as its contents, and any number of buttons. The dialog supports passing in an brush to use as an icon, and can be displayed as either modal or not.

diff --git a/Engine/Source/Editor/UnrealEd/Private/Dialogs/CustomDialog.cpp b/Engine/Source/Editor/UnrealEd/Private/Dialogs/CustomDialog.cpp
new file mode 100644
--- /dev/null
+++ b/Engine/Source/Editor/UnrealEd/Private/Dialogs/CustomDialog.cpp
@@ -0,0 +1,152 @@
+// Copyright 1998-2019 Epic Games, Inc. All Rights Reserved.
+
+#include "Dialogs/CustomDialog.h"
+
+#include "HAL/PlatformApplicationMisc.h"
+#include "Logging/LogMacros.h"
+#include "Widgets/Layout/SSpacer.h"
+#include "Widgets/Layout/SScrollBox.h"
+#include "Widgets/Layout/SUniformGridPanel.h"
+
+DEFINE_LOG_CATEGORY_STATIC(LogCustomDialog, Log, All);
+
+void SCustomDialog::Construct(const FArguments& InArgs)
+{
+	UE_LOG(LogCustomDialog, Log, TEXT("Dialog displayed:"), *InArgs._Title.ToString());
+
+	check(InArgs._Buttons.Num() > 0);
+	
+	TSharedPtr<SHorizontalBox> ContentBox;
+	TSharedPtr<SHorizontalBox> ButtonBox;
+
+	SWindow::Construct( SWindow::FArguments()
+		.Title(InArgs._Title)
+		.SizingRule(ESizingRule::Autosized)
+		.SupportsMaximize(false)
+		.SupportsMinimize(false)
+		[
+			SNew(SBorder)
+			.Padding(4.f)
+			.BorderImage(FEditorStyle::GetBrush( "ToolPanel.GroupBorder" ))
+			[
+				SNew(SVerticalBox)
+				+ SVerticalBox::Slot()
+				.FillHeight(1.0f)
+				[
+					SAssignNew(ContentBox, SHorizontalBox)
+				]
+				+ SVerticalBox::Slot()
+				.VAlign(VAlign_Center)
+				.AutoHeight()
+				[
+					SAssignNew(ButtonBox, SHorizontalBox)
+				]
+			]
+		] );
+
+	if (InArgs._IconBrush.IsValid())
+	{
+		const FSlateBrush* ImageBrush = FEditorStyle::GetBrush(InArgs._IconBrush);
+		if (ImageBrush != nullptr)
+		{
+			ContentBox->AddSlot()
+				.AutoWidth()
+				.VAlign(VAlign_Center)
+				.HAlign(HAlign_Left)
+				[
+					SNew(SImage)
+					.Image(ImageBrush)
+				];
+		}
+	}
+
+	if (InArgs._UseScrollBox)
+	{
+		ContentBox->AddSlot()
+		[
+			SNew(SBox)
+			.MaxDesiredHeight(InArgs._ScrollBoxMaxHeight)
+			[
+				SNew(SScrollBox)
+				+SScrollBox::Slot()
+				[
+					InArgs._DialogContent.ToSharedRef()
+				]
+			]
+		];
+	}
+	else
+	{
+		ContentBox->AddSlot()
+			.FillWidth(1.0f)
+			.VAlign(VAlign_Center)
+			.HAlign(HAlign_Left)
+			[
+				InArgs._DialogContent.ToSharedRef()
+			];
+	}
+
+	ButtonBox->AddSlot()
+		.AutoWidth()
+		[
+			SNew(SSpacer)
+			.Size(FVector2D(20.0f, 1.0f))
+		];
+
+	TSharedPtr<SUniformGridPanel> ButtonPanel;
+
+	ButtonBox->AddSlot()
+		.FillWidth(1.0f)
+		.VAlign(VAlign_Center)
+		.HAlign(HAlign_Right)
+		[
+			SAssignNew(ButtonPanel, SUniformGridPanel)
+			.SlotPadding(FEditorStyle::GetMargin("StandardDialog.SlotPadding"))
+			.MinDesiredSlotWidth(FEditorStyle::GetFloat("StandardDialog.MinDesiredSlotWidth"))
+			.MinDesiredSlotHeight(FEditorStyle::GetFloat("StandardDialog.MinDesiredSlotHeight"))
+		];
+
+	for (int i = 0; i < InArgs._Buttons.Num(); ++i)
+	{
+		const FButton& Button = InArgs._Buttons[i];
+
+		ButtonPanel->AddSlot(ButtonPanel->GetChildren()->Num(), 0)
+		[
+			SNew(SButton)
+			.OnClicked(FOnClicked::CreateSP(this, &SCustomDialog::OnButtonClicked, Button.OnClicked, i))
+			[
+				SNew(SHorizontalBox)
+				+ SHorizontalBox::Slot()
+				.VAlign(VAlign_Center)
+				.HAlign(HAlign_Center)
+				[
+					SNew(STextBlock)
+					.Text(Button.ButtonText)
+				]
+			]
+		];
+	}
+}
+
+int SCustomDialog::ShowModal()
+{
+	FSlateApplication::Get().AddModalWindow(StaticCastSharedRef<SWindow>(this->AsShared()), FGlobalTabmanager::Get()->GetRootWindow());
+
+	return LastPressedButton;
+}
+
+void SCustomDialog::Show()
+{
+	FSlateApplication::Get().AddWindow(StaticCastSharedRef<SWindow>(this->AsShared()), true);
+}
+
+/** Handle the button being clicked */
+FReply SCustomDialog::OnButtonClicked(FSimpleDelegate OnClicked, int ButtonIndex)
+{
+	LastPressedButton = ButtonIndex;
+
+	FSlateApplication::Get().RequestDestroyWindow(StaticCastSharedRef<SWindow>(this->AsShared()));
+
+	OnClicked.ExecuteIfBound();
+	return FReply::Handled();
+}
diff --git a/Engine/Source/Editor/UnrealEd/Public/Dialogs/CustomDialog.h b/Engine/Source/Editor/UnrealEd/Public/Dialogs/CustomDialog.h
new file mode 100644
--- /dev/null
+++ b/Engine/Source/Editor/UnrealEd/Public/Dialogs/CustomDialog.h
@@ -0,0 +1,77 @@
+// Copyright 1998-2019 Epic Games, Inc. All Rights Reserved.
+
+#pragma once
+
+#include "Widgets/SWindow.h"
+
+/** This is a custom dialog class, which allows any Slate widget to be used as the contents,
+ * with any number of buttons that have any text. 
+ * It also supports adding a custom icon to the dialog.
+ * Usage:
+ * TSharedRef<SCustomDialog> HelloWorldDialog = SNew(SCustomDialog)
+		.Title(FText("Hello, World!"))
+		.DialogContent( SNew(SImage).Image(FName("Hello")) )
+		.Buttons( { SCustomDialog::FButton(LOCTEXT("OK", "OK"), []() { printf("OK pressed!" ); } ),
+				SCustomDialog::FButton(LOCTEXT("Cancel", "Cancel") } );
+
+		// returns 0 when OK is pressed, 1 when Cancel is pressed
+		int ButtonPressed = CustomDialog->ShowModal(); 
+ */
+class SCustomDialog : public SWindow
+{
+public:
+	struct FButton
+	{
+		FButton(const FText& InButtonText, const FSimpleDelegate& InOnClicked = FSimpleDelegate())
+			: ButtonText(InButtonText),
+			OnClicked(InOnClicked)
+		{
+		}
+
+		FText ButtonText;
+		FSimpleDelegate OnClicked;
+	};
+
+	SLATE_BEGIN_ARGS(SCustomDialog) 
+		: _UseScrollBox(true)
+		, _ScrollBoxMaxHeight(300)
+	{
+	}
+		/** Title to display for the dialog. */
+		SLATE_ARGUMENT(FText, Title)
+
+		/** Optional icon to display in the dialog. (default: none) */
+		SLATE_ARGUMENT(FName, IconBrush)
+
+		/** Should this dialog use a scroll box for over-sized content? (default: true) */
+		SLATE_ARGUMENT(bool, UseScrollBox)
+
+		/** Max height for the scroll box (default: 300) */
+		SLATE_ARGUMENT(int32, ScrollBoxMaxHeight)
+
+		/** The buttons that this dialog should have. One or more buttons must be added.*/
+		SLATE_ARGUMENT(TArray<FButton>, Buttons)
+
+		/** Content for the dialog */
+		SLATE_ARGUMENT(TSharedPtr<SWidget>, DialogContent)
+
+	SLATE_END_ARGS()
+
+	void Construct(const FArguments& InArgs);
+
+	/** Show the dialog.
+	 * This method will return immediately.
+	 */ 
+	void Show();
+
+	/** Show a modal dialog. Will block until an input is received.
+	 * Returns the index of the button that was pressed.
+	 */
+	int ShowModal();
+
+private:
+	FReply OnButtonClicked(FSimpleDelegate OnClicked, int ButtonIndex);
+
+	/** The index of the button that was pressed last. */
+	int LastPressedButton = -1;
+};
\ No newline at end of file
