[NO NON-PII COMMIT MESSAGE]

diff --git a/Engine/Source/Editor/UnrealEd/Private/PlayLevel.cpp b/Engine/Source/Editor/UnrealEd/Private/PlayLevel.cpp
--- a/Engine/Source/Editor/UnrealEd/Private/PlayLevel.cpp
+++ b/Engine/Source/Editor/UnrealEd/Private/PlayLevel.cpp
@@ -91,6 +91,8 @@
 #include "Engine/LocalPlayer.h"
 #include "Slate/SGameLayerManager.h"
 #include "HAL/PlatformApplicationMisc.h"
+#include "Widgets/Input/SHyperlink.h"
+#include "Dialogs/CustomDialog.h"
 
 #include "IHeadMountedDisplay.h"
 #include "IXRTrackingSystem.h"
@@ -2241,6 +2243,62 @@ bool UEditorEngine::SpawnPlayFromHereStart( UWorld* World, AActor*& PlayerStart,
 	return true;
 }
 
+static bool ShowBlueprintErrorDialog( TArray<UBlueprint*> ErroredBlueprints )
+{
+	struct Local
+	{
+		static void OnHyperlinkClicked( TWeakObjectPtr<UBlueprint> InBlueprint )
+		{
+			if (UBlueprint* BlueprintToEdit = InBlueprint.Get())
+			{
+				// Open the blueprint
+				GEditor->EditObject( BlueprintToEdit );
+			}
+		}
+	};
+
+	TSharedRef<SVerticalBox> DialogContents = SNew( SVerticalBox )
+		+ SVerticalBox::Slot()
+		[
+			SNew( STextBlock )
+			.Text( NSLOCTEXT( "PlayInEditor", "PrePIE_BlueprintErrors", "One or more blueprints has an unresolved compiler error, are you sure you want to Play in Editor?" ) )
+		]
+		+ SVerticalBox::Slot()
+		[
+			SNew( STextBlock )
+		];
+
+	for (UBlueprint* Blueprint : ErroredBlueprints)
+	{
+		TWeakObjectPtr<UBlueprint> BlueprintPtr = Blueprint;
+
+		DialogContents->AddSlot()
+			.AutoHeight()
+			.HAlign(HAlign_Left)
+			[
+				SNew(SHyperlink)
+				.Style(FEditorStyle::Get(), "Common.GotoBlueprintHyperlink")
+				.OnNavigate_Static(&Local::OnHyperlinkClicked, BlueprintPtr)
+				.Text(FText::FromString(Blueprint->GetName()))
+				.ToolTipText(NSLOCTEXT("SourceHyperlink", "EditBlueprint_ToolTip", "Click to edit the blueprint"))
+			];
+	}
+
+	FText DialogTitle = NSLOCTEXT("PlayInEditor", "PrePIE_BlueprintErrorsTitle", "Blueprint Compilation Errors");
+
+	FText OKText = NSLOCTEXT("PlayInEditor", "PrePIE_OkText", "Play in Editor");
+	FText CancelText = NSLOCTEXT("Dialogs", "EAppReturnTypeCancel", "Cancel");
+
+	TSharedRef<SCustomDialog> CustomDialog = SNew(SCustomDialog)
+		.Title(DialogTitle)
+		.IconBrush("NotificationList.DefaultMessage")
+		.DialogContent(DialogContents)
+		.Buttons( { SCustomDialog::FButton(OKText), SCustomDialog::FButton(CancelText) } );
+
+	int ButtonPressed = CustomDialog->ShowModal();
+	return ButtonPressed == 0;
+}
+
 void UEditorEngine::PlayInEditor( UWorld* InWorld, bool bInSimulateInEditor, FPlayInEditorOverrides Overrides )
 {
 	// Broadcast PreBeginPIE before checks that might block PIE below (BeginPIE is broadcast below after the checks)
@@ -2328,17 +2386,9 @@ void UEditorEngine::PlayInEditor( UWorld* InWorld, bool bInSimulateInEditor, FPl
 
 	if (ErroredBlueprints.Num() && !GIsDemoMode)
 	{
-		FString ErroredBlueprintList;
-		for (UBlueprint* Blueprint : ErroredBlueprints)
-		{
-			ErroredBlueprintList += FString::Printf(TEXT("\n   %s"), *Blueprint->GetName());
-		}
-
-		FFormatNamedArguments Args;
-		Args.Add(TEXT("ErrorBlueprints"), FText::FromString(ErroredBlueprintList));
-
 		// There was at least one blueprint with an error, make sure the user is OK with that.
-		const bool bContinuePIE = EAppReturnType::Yes == FMessageDialog::Open( EAppMsgType::YesNo, FText::Format( NSLOCTEXT("PlayInEditor", "PrePIE_BlueprintErrors", "One or more blueprints has an unresolved compiler error, are you sure you want to Play in Editor?{ErrorBlueprints}"), Args ) );
+		bool bContinuePIE = ShowBlueprintErrorDialog( ErroredBlueprints );
+
 		if ( !bContinuePIE )
 		{
 			FEditorDelegates::EndPIE.Broadcast(bInSimulateInEditor);
