Build Fix: missing include files for CL 6104187

diff --git a/Engine/Source/Editor/UnrealEd/Private/Dialogs/CustomDialog.cpp b/Engine/Source/Editor/UnrealEd/Private/Dialogs/CustomDialog.cpp
--- a/Engine/Source/Editor/UnrealEd/Private/Dialogs/CustomDialog.cpp
+++ b/Engine/Source/Editor/UnrealEd/Private/Dialogs/CustomDialog.cpp
@@ -4,9 +4,13 @@
 
 #include "HAL/PlatformApplicationMisc.h"
 #include "Logging/LogMacros.h"
+#include "Widgets/Images/SImage.h"
 #include "Widgets/Layout/SSpacer.h"
+#include "Widgets/Layout/SBox.h"
 #include "Widgets/Layout/SScrollBox.h"
 #include "Widgets/Layout/SUniformGridPanel.h"
+#include "EditorStyleSet.h"
+#include "Styling/SlateBrush.h"
 
 DEFINE_LOG_CATEGORY_STATIC(LogCustomDialog, Log, All);
 
