Build fix: missing include files in PlayLevel and CustomDialog.

diff --git a/Engine/Source/Editor/UnrealEd/Private/Dialogs/CustomDialog.cpp b/Engine/Source/Editor/UnrealEd/Private/Dialogs/CustomDialog.cpp
--- a/Engine/Source/Editor/UnrealEd/Private/Dialogs/CustomDialog.cpp
+++ b/Engine/Source/Editor/UnrealEd/Private/Dialogs/CustomDialog.cpp
@@ -3,14 +3,20 @@
 #include "Dialogs/CustomDialog.h"
 
 #include "HAL/PlatformApplicationMisc.h"
+
+#include "EditorStyleSet.h"
+#include "Framework/Application/SlateApplication.h"
+#include "Framework/Docking/TabManager.h"
 #include "Logging/LogMacros.h"
+#include "Styling/SlateBrush.h"
 #include "Widgets/Images/SImage.h"
+#include "Widgets/Input/SButton.h"
+#include "Widgets/Text/STextBlock.h"
 #include "Widgets/Layout/SSpacer.h"
 #include "Widgets/Layout/SBox.h"
 #include "Widgets/Layout/SScrollBox.h"
 #include "Widgets/Layout/SUniformGridPanel.h"
-#include "EditorStyleSet.h"
-#include "Styling/SlateBrush.h"
+#include "Widgets/SBoxPanel.h"
 
 DEFINE_LOG_CATEGORY_STATIC(LogCustomDialog, Log, All);
 
diff --git a/Engine/Source/Editor/UnrealEd/Private/PlayLevel.cpp b/Engine/Source/Editor/UnrealEd/Private/PlayLevel.cpp
--- a/Engine/Source/Editor/UnrealEd/Private/PlayLevel.cpp
+++ b/Engine/Source/Editor/UnrealEd/Private/PlayLevel.cpp
@@ -106,6 +106,8 @@
 #include "EditorModeRegistry.h"
 #include "PhysicsManipulationMode.h"
 #include "CookerSettings.h"
+#include "Widgets/Text/STextBlock.h"
+#include "Widgets/SBoxPanel.h"
 
 
 DEFINE_LOG_CATEGORY_STATIC(LogPlayLevel, Log, All);
