Fixing confusing log at SkeletalMesh import

diff --git a/Engine/Source/Editor/UnrealEd/Private/Fbx/FbxSkeletalMeshImport.cpp b/Engine/Source/Editor/UnrealEd/Private/Fbx/FbxSkeletalMeshImport.cpp
--- a/Engine/Source/Editor/UnrealEd/Private/Fbx/FbxSkeletalMeshImport.cpp
+++ b/Engine/Source/Editor/UnrealEd/Private/Fbx/FbxSkeletalMeshImport.cpp
@@ -1324,7 +1324,6 @@ bool UnFbx::FFbxImporter::FillSkeletalMeshImportData(TArray<FbxNode*>& NodeArray
 	{
 		if (!(bIsReimport && ImportOptions->bImportAsSkeletalGeometry)) //Do not import bone if we import only the geometry and we are reimporting
 		{
-			AddTokenizedErrorMessage(FTokenizedMessage::Create(EMessageSeverity::Error, LOCTEXT("FbxSkeletaLMeshimport_MultipleRootFound", "Multiple roots found")), FFbxErrors::SkeletalMesh_MultipleRoots);
 			return false;
 		}
 	}
