Removing unnecessary logging during SkeletalMesh import.

diff --git a/Engine/Source/Editor/UnrealEd/Private/Fbx/FbxMainImport.cpp b/Engine/Source/Editor/UnrealEd/Private/Fbx/FbxMainImport.cpp
--- a/Engine/Source/Editor/UnrealEd/Private/Fbx/FbxMainImport.cpp
+++ b/Engine/Source/Editor/UnrealEd/Private/Fbx/FbxMainImport.cpp
@@ -2342,59 +2342,7 @@ FbxNode* FFbxImporter::GetRootSkeleton(FbxNode* Link)
 	return RootBone;
 }
 
-
-void FFbxImporter::DumpFBXNode(FbxNode* Node) 
-{
-	FbxMesh* Mesh = Node->GetMesh();
-	const FString NodeName(Node->GetName());
-
-	if(Mesh)
-	{
-		UE_LOG(LogFbx, Log, TEXT("================================================="));
-		UE_LOG(LogFbx, Log, TEXT("Dumping Node START [%s] "), *NodeName);
-		int DeformerCount = Mesh->GetDeformerCount();
-		UE_LOG(LogFbx, Log,TEXT("\tTotal Deformer Count %d."), *NodeName, DeformerCount);
-		for(int i=0; i<DeformerCount; i++)
-		{
-			FbxDeformer* Deformer = Mesh->GetDeformer(i);
-			const FString DeformerName(Deformer->GetName());
-			const FString DeformerTypeName(Deformer->GetTypeName());
-			UE_LOG(LogFbx, Log,TEXT("\t\t[Node %d] %s (Type %s)."), i+1, *DeformerName, *DeformerTypeName);
-			UE_LOG(LogFbx, Log,TEXT("================================================="));
-		}
-
-		FbxNodeAttribute* NodeAttribute = Node->GetNodeAttribute();
-		if(NodeAttribute)
-		{
-			FString NodeAttributeName(NodeAttribute->GetName());
-			FbxNodeAttribute::EType Type = NodeAttribute->GetAttributeType();
-			UE_LOG(LogFbx, Log,TEXT("\tAttribute (%s) Type (%d)."), *NodeAttributeName, (int32)Type);
-		
-			for (int i=0; i<NodeAttribute->GetNodeCount(); ++i)
-			{
-				FbxNode * Child = NodeAttribute->GetNode(i);
-
-				if (Child)
-				{
-					const FString ChildName(Child->GetName());
-					const FString ChildTypeName(Child->GetTypeName());
-					UE_LOG(LogFbx, Log,TEXT("\t\t[Node Attribute Child %d] %s (Type %s)."), i+1, *ChildName, *ChildTypeName);
-				}
-			}
-
-		}
-
-		UE_LOG(LogFbx, Log,TEXT("Dumping Node END [%s]"), *NodeName);
-	}
-
-	for(int ChildIdx=0; ChildIdx < Node->GetChildCount(); ChildIdx++)
-	{
-		FbxNode* ChildNode = Node->GetChild(ChildIdx);
-		DumpFBXNode(ChildNode);
-	}
-
-}
-
+ 
 void FFbxImporter::ApplyTransformSettingsToFbxNode(FbxNode* Node, UFbxAssetImportData* AssetData)
 {
 	check(Node);
@@ -2495,8 +2443,6 @@ void FFbxImporter::RecursiveFindFbxSkelMesh(FbxNode* Node, TArray< TArray<FbxNod
 	FbxNode* SkelMeshNode = nullptr;
 	FbxNode* NodeToAdd = Node;
 
-	DumpFBXNode(Node);
-
     if (Node->GetMesh() && Node->GetMesh()->GetDeformerCount(FbxDeformer::eSkin) > 0 )
 	{
 		SkelMeshNode = Node;
diff --git a/Engine/Source/Editor/UnrealEd/Public/FbxImporter.h b/Engine/Source/Editor/UnrealEd/Public/FbxImporter.h
--- a/Engine/Source/Editor/UnrealEd/Public/FbxImporter.h
+++ b/Engine/Source/Editor/UnrealEd/Public/FbxImporter.h
@@ -1057,10 +1057,6 @@ public:
 	template<typename TMaterialType>
 	static void ShowFbxMaterialConflictWindow(const TArray<TMaterialType>& InSourceMaterials, const TArray<TMaterialType>& InResultMaterials, TArray<int32>& RemapMaterials, TArray<bool>& FuzzyRemapMaterials, EFBXReimportDialogReturnOption& OutReturnOption, bool bIsPreviewConflict = false);
 
-
-	/** helper function **/
-	UNREALED_API static void DumpFBXNode(FbxNode* Node);
-
 	/**
 	 * Apply asset import settings for transform to an FBX node
 	 *
