Fixed issue where actors would become partially unselected when clicking on a locked actor

diff --git a/Engine/Source/Editor/UnrealEd/Private/LevelViewportClickHandlers.cpp b/Engine/Source/Editor/UnrealEd/Private/LevelViewportClickHandlers.cpp
--- a/Engine/Source/Editor/UnrealEd/Private/LevelViewportClickHandlers.cpp
+++ b/Engine/Source/Editor/UnrealEd/Private/LevelViewportClickHandlers.cpp
@@ -140,7 +140,7 @@ namespace ClickHandlers
 
 				GEditor->GetSelectedActors()->Modify();
 
-				if( bAllowSelectionChange )
+				if( bAllowSelectionChange && GEditor->CanSelectActor(Actor, true, true) )
 				{
 					// If the actor the user clicked on was already selected, then we won't bother clearing the selection
 					if( !Actor->IsSelected() )
@@ -173,7 +173,7 @@ namespace ClickHandlers
 
 				GEditor->GetSelectedActors()->Modify();
 
-				if( bAllowSelectionChange )
+				if( bAllowSelectionChange && GEditor->CanSelectActor(Actor, true, true) )
 				{
 					// Clear the selection
 					GEditor->SelectNone( false, true );
@@ -229,7 +229,7 @@ namespace ClickHandlers
 			}
 			else if ( Actor )
 			{
-				if( bAllowSelectionChange )
+				if( bAllowSelectionChange  && GEditor->CanSelectActor(Actor, true, true, true) )
 				{
 					const FScopedTransaction Transaction( NSLOCTEXT("UnrealEd", "ClickingOnActors", "Clicking on Actors") );
 					GEditor->GetSelectedActors()->Modify();
