Fix fbx file open conversion when we try to open the fbx file two time. It happen with animation.

diff --git a/Engine/Source/Editor/UnrealEd/Private/Fbx/FbxMainImport.cpp b/Engine/Source/Editor/UnrealEd/Private/Fbx/FbxMainImport.cpp
--- a/Engine/Source/Editor/UnrealEd/Private/Fbx/FbxMainImport.cpp
+++ b/Engine/Source/Editor/UnrealEd/Private/Fbx/FbxMainImport.cpp
@@ -684,7 +684,7 @@ bool FFbxImporter::GetSceneInfo(FString Filename, FbxSceneInfo& SceneInfo, bool
 		}
 		GWarn->UpdateProgress( 90, 100 );
 	case IMPORTED:
-	
+	case FIXEDANDCONVERTED:
 	default:
 		break;
 	}
@@ -1609,8 +1609,10 @@ bool FFbxImporter::ImportFromFile(const FString& Filename, const FString& Type,
 			ConvertLodPrefixToLodGroup();
 
 			MeshNamesCache.Empty();
+			CurPhase = FIXEDANDCONVERTED;
+			break;
 		}
-		
+	case FIXEDANDCONVERTED:
 	default:
 		break;
 	}
diff --git a/Engine/Source/Editor/UnrealEd/Public/FbxImporter.h b/Engine/Source/Editor/UnrealEd/Public/FbxImporter.h
--- a/Engine/Source/Editor/UnrealEd/Public/FbxImporter.h
+++ b/Engine/Source/Editor/UnrealEd/Public/FbxImporter.h
@@ -1247,7 +1247,8 @@ protected:
 	{
 		NOTSTARTED,
 		FILEOPENED,
-		IMPORTED
+		IMPORTED,
+		FIXEDANDCONVERTED,
 	};
 	
 	static TSharedPtr<FFbxImporter> StaticInstance;
