Fix issue where piloting actors with offset cameras causes the actor transform to be set incorrectly

#rn Fixed issue where piloting actors wouldn't work properly if they had a camera offset from the actor root position

diff --git a/Engine/Source/Editor/UnrealEd/Private/LevelEditorViewport.cpp b/Engine/Source/Editor/UnrealEd/Private/LevelEditorViewport.cpp
--- a/Engine/Source/Editor/UnrealEd/Private/LevelEditorViewport.cpp
+++ b/Engine/Source/Editor/UnrealEd/Private/LevelEditorViewport.cpp
@@ -3335,8 +3335,23 @@ void FLevelEditorViewportClient::MoveLockedActorToCamera()
 				}
 			}
 
-			ActiveActorLock->SetActorLocation(GCurrentLevelEditingViewportClient->GetViewLocation(), false);
-			ActiveActorLock->SetActorRotation(GCurrentLevelEditingViewportClient->GetViewRotation());
+			// If we're locked to a camera then we're reflecting the camera view and not the actor position. We need to reflect that delta when we reposition the piloted actor
+			if (bUseControllingActorViewInfo)
+			{
+				const USceneComponent* ViewComponent = Cast<USceneComponent>(FindViewComponentForActor(ActiveActorLock));
+				if (ViewComponent != nullptr)
+				{
+					const FTransform RelativeTransform = ViewComponent->GetComponentTransform().Inverse();
+					const FTransform DesiredTransform = FTransform(GCurrentLevelEditingViewportClient->GetViewRotation(), GCurrentLevelEditingViewportClient->GetViewLocation());
+
+					ActiveActorLock->SetActorTransform(ActiveActorLock->GetActorTransform() * RelativeTransform * DesiredTransform);
+				}
+			}
+			else
+			{
+				ActiveActorLock->SetActorLocation(GCurrentLevelEditingViewportClient->GetViewLocation(), false);
+				ActiveActorLock->SetActorRotation(GCurrentLevelEditingViewportClient->GetViewRotation());
+			}
 		}
 
 		if (ABrush* Brush = Cast<ABrush>(ActiveActorLock))
