Added logging to ObjectTools::ContainsWorldInUse to write out the currently active worlds.

diff --git a/Engine/Source/Editor/UnrealEd/Private/ObjectTools.cpp b/Engine/Source/Editor/UnrealEd/Private/ObjectTools.cpp
--- a/Engine/Source/Editor/UnrealEd/Private/ObjectTools.cpp
+++ b/Engine/Source/Editor/UnrealEd/Private/ObjectTools.cpp
@@ -2019,10 +2019,37 @@ namespace ObjectTools
 
 	bool ContainsWorldInUse(const TArray< UObject* >& ObjectsToDelete)
 	{
-		TArray<const UObject*> ActiveWorlds;
+		TArray<const UWorld*> WorldsToDelete;
+
+		for (const UObject* ObjectToDelete : ObjectsToDelete)
+		{
+			if (const UWorld* World = Cast<UWorld>(ObjectToDelete))
+			{
+				WorldsToDelete.AddUnique(World);
+			}
+		}
+
+		if (WorldsToDelete.Num() == 0)
+		{
+			return false;
+		}
+
+		auto GetCombinedWorldNames = [](const TArray<const UWorld*>& Worlds) -> FString
+		{ 
+			return FString::JoinBy(Worlds, TEXT(", "), 
+				[](const UWorld* World) -> FString
+				{
+					return World->GetPathName();
+				});
+		};
+
+		UE_LOG(LogObjectTools, Log, TEXT("Deleting %d worlds: %s"), WorldsToDelete.Num(), *GetCombinedWorldNames(WorldsToDelete));
+
+		TArray<const UWorld*> ActiveWorlds;
+
 		for (const FWorldContext& WorldContext : GEditor->GetWorldContexts())
 		{
-			if (UWorld* World = WorldContext.World())
+			if (const UWorld* World = WorldContext.World())
 			{
 				ActiveWorlds.AddUnique(World);
 
@@ -2030,20 +2057,22 @@ namespace ObjectTools
 				{
 					if (StreamingLevel && StreamingLevel->GetLoadedLevel() && StreamingLevel->GetLoadedLevel()->GetOuter())
 					{
-						ActiveWorlds.AddUnique(StreamingLevel->GetLoadedLevel()->GetOuter());
+						if (const UWorld* StreamingWorld = Cast<UWorld>(StreamingLevel->GetLoadedLevel()->GetOuter()))
+						{
+							ActiveWorlds.AddUnique(StreamingWorld);
+						}
 					}
 				}
 			}
 		}
 
-		for (const UObject* ObjectToDelete : ObjectsToDelete)
+		UE_LOG(LogObjectTools, Log, TEXT("Currently %d active worlds: %s"), ActiveWorlds.Num(), *GetCombinedWorldNames(ActiveWorlds));
+
+		for (const UWorld* World : WorldsToDelete)
 		{
-			if (const UWorld* World = Cast<UWorld>(ObjectToDelete))
+			if (ActiveWorlds.Contains(World))
 			{
-				if (ActiveWorlds.Contains(World))
-				{
-					return true;
-				}
+				return true;
 			}
 		}
 
