Expose import data table to blueprint/script

diff --git a/Engine/Source/Editor/UnrealEd/Classes/Factories/CSVImportFactory.h b/Engine/Source/Editor/UnrealEd/Classes/Factories/CSVImportFactory.h
--- a/Engine/Source/Editor/UnrealEd/Classes/Factories/CSVImportFactory.h
+++ b/Engine/Source/Editor/UnrealEd/Classes/Factories/CSVImportFactory.h
@@ -17,7 +17,7 @@
 DECLARE_LOG_CATEGORY_EXTERN(LogCSVImportFactory, Log, All);
 
 /** Enum to indicate what to import CSV as */
-UENUM()
+UENUM(BlueprintType)
 enum class ECSVImportType : uint8
 {
 	/** Import as UDataTable */
@@ -32,20 +32,20 @@ enum class ECSVImportType : uint8
 	ECSV_CurveLinearColor,
 };
 
-USTRUCT()
+USTRUCT(BlueprintType)
 struct FCSVImportSettings
 {
 	GENERATED_BODY()
 
 	FCSVImportSettings();
 
-	UPROPERTY()
+	UPROPERTY(BlueprintReadWrite, Category="Misc")
 	UScriptStruct* ImportRowStruct;
 
-	UPROPERTY()
+	UPROPERTY(BlueprintReadWrite, Category="Misc")
 	ECSVImportType ImportType;
 
-	UPROPERTY()
+	UPROPERTY(BlueprintReadWrite, Category="Misc")
 	TEnumAsByte<ERichCurveInterpMode> ImportCurveInterpMode;
 };
 
@@ -79,7 +79,8 @@ private:
 	/* Reimport object from the given path*/
 	EReimportResult::Type Reimport(UObject* Obj, const FString& Path);
 
-	UPROPERTY()
+public:
+	UPROPERTY(BlueprintReadWrite, Category="Automation")
 	FCSVImportSettings AutomatedImportSettings;
 };
 
