Make sure staticmesh reimport of enbedded LOD do also remap the section info map.

diff --git a/Engine/Source/Editor/UnrealEd/Private/StaticMeshEdit.cpp b/Engine/Source/Editor/UnrealEd/Private/StaticMeshEdit.cpp
--- a/Engine/Source/Editor/UnrealEd/Private/StaticMeshEdit.cpp
+++ b/Engine/Source/Editor/UnrealEd/Private/StaticMeshEdit.cpp
@@ -1514,8 +1514,9 @@ void RestoreExistingMeshData(ExistingStaticMeshData* ExistingMeshDataPtr, UStati
 				continue;
 			}
 
+			
 			//When re-importing the asset, do not touch the LOD that was imported from file, the material array is keep intact so the section should still be valid.
-			bool NoRemapForThisLOD = LodLevel == INDEX_NONE && i != 0 && !IsReductionActive(NewMesh->SourceModels[i].ReductionSettings);
+			bool NoRemapForThisLOD = LodLevel == INDEX_NONE && i != 0 && !NewMesh->SourceModels[i].bImportWithBaseMesh && !IsReductionActive(NewMesh->SourceModels[i].ReductionSettings);
 
 			FStaticMeshLODResources& LOD = NewMesh->RenderData->LODResources[i];
 			
