Fix to make orbit camera work properly with piloted actors

diff --git a/Engine/Source/Editor/UnrealEd/Private/EditorViewportClient.cpp b/Engine/Source/Editor/UnrealEd/Private/EditorViewportClient.cpp
--- a/Engine/Source/Editor/UnrealEd/Private/EditorViewportClient.cpp
+++ b/Engine/Source/Editor/UnrealEd/Private/EditorViewportClient.cpp
@@ -2295,6 +2295,11 @@ void FEditorViewportClient::InputAxisForOrbit(FViewport* InViewport, const FVect
 		SetViewLocation(RotatedViewLocation);
 
 		FEditorViewportStats::Using(IsPerspective() ? FEditorViewportStats::CAT_PERSPECTIVE_MOUSE_ORBIT_ROTATION : FEditorViewportStats::CAT_ORTHOGRAPHIC_MOUSE_ORBIT_ROTATION);
+
+		if (IsPerspective())
+		{
+			PerspectiveCameraMoved();
+		}
 	}
 	else
 	{
diff --git a/Engine/Source/Editor/UnrealEd/Private/LevelEditorViewport.cpp b/Engine/Source/Editor/UnrealEd/Private/LevelEditorViewport.cpp
--- a/Engine/Source/Editor/UnrealEd/Private/LevelEditorViewport.cpp
+++ b/Engine/Source/Editor/UnrealEd/Private/LevelEditorViewport.cpp
@@ -3335,6 +3335,9 @@ void FLevelEditorViewportClient::MoveLockedActorToCamera()
 				}
 			}
 
+			// Need to disable orbit camera before setting actor position so that the viewport camera location is converted back
+			GCurrentLevelEditingViewportClient->ToggleOrbitCamera(false);
+
 			// If we're locked to a camera then we're reflecting the camera view and not the actor position. We need to reflect that delta when we reposition the piloted actor
 			if (bUseControllingActorViewInfo)
 			{
