First pass prototype of editor modes which can generate a toolbar.  This is prototype and will likely be replaced by a more advanced system.

Simply override BuildModeToolbar on FEdMode to use this and generate a toolbar the way you would generate any other toolbar using a ToolbarBuilder.   Only usable in level editor modes for now.

This protoype involved improving tab spawning:
- Added support for tabs which can be forced off even when there is a saved layout for them. E.g for when it doesnt make sense for the tab to be visible.  Editor modes will force the mode toolbar tab hidden when no editor modes require a tab.
- Added support for inserting into tab area splitters by layout customizations. Allows code to add a manditory tab to an already saved layout which doesnt have space for the tab.

diff --git a/Engine/Source/Editor/LevelEditor/Private/SLevelEditor.cpp b/Engine/Source/Editor/LevelEditor/Private/SLevelEditor.cpp
--- a/Engine/Source/Editor/LevelEditor/Private/SLevelEditor.cpp
+++ b/Engine/Source/Editor/LevelEditor/Private/SLevelEditor.cpp
@@ -632,6 +632,10 @@ TSharedRef<SDockTab> SLevelEditor::SpawnLevelEditorTab( const FSpawnTabArgs& Arg
 			];
 
 	}
+	else if (TabIdentifier == FEditorModeTools::EditorModeToolbarTabName)
+	{
+		return GLevelEditorModeTools().MakeModeToolbarTab();
+	}
 	else if( TabIdentifier == TEXT("LevelEditorSelectionDetails") || TabIdentifier == TEXT("LevelEditorSelectionDetails2") || TabIdentifier == TEXT("LevelEditorSelectionDetails3") || TabIdentifier == TEXT("LevelEditorSelectionDetails4") )
 	{
 		TSharedRef<SDockTab> DetailsPanel = SummonDetailsPanel( TabIdentifier );
@@ -816,6 +820,11 @@ TSharedRef<SDockTab> SLevelEditor::SpawnLevelEditorTab( const FSpawnTabArgs& Arg
 	return SNew(SDockTab);
 }
 
+bool SLevelEditor::CanSpawnEditorModeToolbarTab(const FSpawnTabArgs& Args) const
+{
+	return GLevelEditorModeTools().ShouldShowModeToolbar();
+}
+
 TSharedRef<SDockTab> SLevelEditor::InvokeTab( FName TabID )
 {
 	TSharedPtr<FTabManager> LevelEditorTabManager = GetTabManager();
@@ -1024,7 +1033,7 @@ TSharedRef<SWidget> SLevelEditor::RestoreContentArea( const TSharedRef<SDockTab>
 
 		{
 			const FSlateIcon ToolbarIcon(FEditorStyle::GetStyleSetName(), "LevelEditor.Tabs.Toolbar");
-			LevelEditorTabManager->RegisterTabSpawner( "LevelEditorToolBar", FOnSpawnTab::CreateSP<SLevelEditor, FName, FString>(this, &SLevelEditor::SpawnLevelEditorTab, FName("LevelEditorToolBar"), FString()) )
+			LevelEditorTabManager->RegisterTabSpawner( "LevelEditorToolBar", FOnSpawnTab::CreateSP<SLevelEditor, FName, FString>(this, &SLevelEditor::SpawnLevelEditorTab, FName("LevelEditorToolBar"), FString()))
 				.SetDisplayName(NSLOCTEXT("LevelEditorTabs", "LevelEditorToolBar", "Toolbar"))
 				.SetTooltipText(NSLOCTEXT("LevelEditorTabs", "LevelEditorToolBarTooltipText", "Open the Toolbar tab, which provides access to the most common / important actions."))
 				.SetGroup( MenuStructure.GetLevelEditorCategory() )
@@ -1060,12 +1069,22 @@ TSharedRef<SWidget> SLevelEditor::RestoreContentArea( const TSharedRef<SDockTab>
 				.SetIcon( DetailsIcon );
 		}
 
-		const FSlateIcon ToolsIcon( FEditorStyle::GetStyleSetName(), "LevelEditor.Tabs.Modes" );
-		LevelEditorTabManager->RegisterTabSpawner( "LevelEditorToolBox", FOnSpawnTab::CreateSP<SLevelEditor, FName, FString>( this, &SLevelEditor::SpawnLevelEditorTab, FName( "LevelEditorToolBox" ), FString() ) )
-			.SetDisplayName( NSLOCTEXT( "LevelEditorTabs", "LevelEditorToolBox", "Modes" ) )
-			.SetTooltipText( NSLOCTEXT( "LevelEditorTabs", "LevelEditorToolBoxTooltipText", "Open the Modes tab, which specifies all the available editing modes." ) )
-			.SetGroup( MenuStructure.GetLevelEditorCategory() )
-			.SetIcon( ToolsIcon );
+		{
+		
+			const FSlateIcon ToolsIcon(FEditorStyle::GetStyleSetName(), "LevelEditor.Tabs.Modes");
+			LevelEditorTabManager->RegisterTabSpawner("LevelEditorToolBox", FOnSpawnTab::CreateSP<SLevelEditor, FName, FString>(this, &SLevelEditor::SpawnLevelEditorTab, FName("LevelEditorToolBox"), FString()))
+				.SetDisplayName(NSLOCTEXT("LevelEditorTabs", "LevelEditorToolBox", "Modes Panel"))
+				.SetTooltipText(NSLOCTEXT("LevelEditorTabs", "LevelEditorToolBoxTooltipText", "Open the Modes tab, which specifies all the available editing modes."))
+				.SetGroup(MenuStructure.GetLevelEditorModesCategory())
+				.SetIcon(ToolsIcon);
+
+			FCanSpawnTab CanSpawnTabDelegate = FCanSpawnTab::CreateSP(this, &SLevelEditor::CanSpawnEditorModeToolbarTab);
+			LevelEditorTabManager->RegisterTabSpawner(FEditorModeTools::EditorModeToolbarTabName, FOnSpawnTab::CreateSP<SLevelEditor, FName, FString>(this, &SLevelEditor::SpawnLevelEditorTab, FEditorModeTools::EditorModeToolbarTabName, FString()), CanSpawnTabDelegate)
+				.SetDisplayName(NSLOCTEXT("LevelEditorTabs", "LevelEditorModesTab", "Active Mode Tools"))
+				.SetTooltipText(NSLOCTEXT("LevelEditorTabs", "LevelEditorModesTabTooltipText", "Open the modes toolbar which has the active mode's tools"))
+				.SetGroup(MenuStructure.GetLevelEditorModesCategory())
+				.SetIcon(ToolsIcon);
+		}
 
 		{
 			const FSlateIcon OutlinerIcon(FEditorStyle::GetStyleSetName(), "LevelEditor.Tabs.Outliner");
@@ -1186,6 +1205,12 @@ TSharedRef<SWidget> SLevelEditor::RestoreContentArea( const TSharedRef<SDockTab>
 							->AddTab("LevelEditorToolBar", ETabState::OpenedTab)
 						)
 						->Split
+						(
+							FTabManager::NewStack()
+							->SetHideTabWell(true)
+							->AddTab(FEditorModeTools::EditorModeToolbarTabName, ETabState::ClosedTab)
+						)
+						->Split
 						(
 							FTabManager::NewStack()
 							->SetHideTabWell(true)
@@ -1227,6 +1252,10 @@ TSharedRef<SWidget> SLevelEditor::RestoreContentArea( const TSharedRef<SDockTab>
 		));
 	
 	FLayoutExtender LayoutExtender;
+
+	// Make a layout extension for the mode toolbar.  This avoids bumping level editor layout version to add a mandatory tab
+	LayoutExtender.ExtendLayout(FName("LevelEditorToolBar"), ELayoutExtensionPosition::Below, FTabManager::FTab(FEditorModeTools::EditorModeToolbarTabName, ETabState::ClosedTab));
+
 	LevelEditorModule.OnRegisterLayoutExtensions().Broadcast(LayoutExtender);
 	Layout->ProcessExtensions(LayoutExtender);
 
diff --git a/Engine/Source/Editor/LevelEditor/Private/SLevelEditor.h b/Engine/Source/Editor/LevelEditor/Private/SLevelEditor.h
--- a/Engine/Source/Editor/LevelEditor/Private/SLevelEditor.h
+++ b/Engine/Source/Editor/LevelEditor/Private/SLevelEditor.h
@@ -138,6 +138,8 @@ public:
 private:
 	
 	TSharedRef<SDockTab> SpawnLevelEditorTab(const FSpawnTabArgs& Args, FName TabIdentifier, FString InitializationPayload);
+	bool CanSpawnEditorModeToolbarTab(const FSpawnTabArgs& Args) const;
+
 	//TSharedRef<SDockTab> SpawnLevelEditorModeTab(const FSpawnTabArgs& Args, FEdMode* EditorMode);
 	TSharedRef<SDockTab> SummonDetailsPanel( FName Identifier );
 
diff --git a/Engine/Source/Editor/Persona/Private/PersonaEditorModeManager.cpp b/Engine/Source/Editor/Persona/Private/PersonaEditorModeManager.cpp
--- a/Engine/Source/Editor/Persona/Private/PersonaEditorModeManager.cpp
+++ b/Engine/Source/Editor/Persona/Private/PersonaEditorModeManager.cpp
@@ -6,9 +6,9 @@
 bool FPersonaEditorModeManager::GetCameraTarget(FSphere& OutTarget) const
 {
 	// Note: assumes all of our modes are IPersonaEditMode!
-	for(int32 ModeIndex = 0; ModeIndex < Modes.Num(); ++ModeIndex)
+	for(int32 ModeIndex = 0; ModeIndex < ActiveModes.Num(); ++ModeIndex)
 	{
-		TSharedPtr<IPersonaEditMode> EditMode = StaticCastSharedPtr<IPersonaEditMode>(Modes[ModeIndex]);
+		TSharedPtr<IPersonaEditMode> EditMode = StaticCastSharedPtr<IPersonaEditMode>(ActiveModes[ModeIndex]);
 
 		FSphere Target;
 		if (EditMode->GetCameraTarget(Target))
@@ -24,9 +24,9 @@ bool FPersonaEditorModeManager::GetCameraTarget(FSphere& OutTarget) const
 void FPersonaEditorModeManager::GetOnScreenDebugInfo(TArray<FText>& OutDebugText) const
 {
 	// Note: assumes all of our modes are IPersonaEditMode!
-	for (int32 ModeIndex = 0; ModeIndex < Modes.Num(); ++ModeIndex)
+	for (int32 ModeIndex = 0; ModeIndex < ActiveModes.Num(); ++ModeIndex)
 	{
-		TSharedPtr<IPersonaEditMode> EditMode = StaticCastSharedPtr<IPersonaEditMode>(Modes[ModeIndex]);
+		TSharedPtr<IPersonaEditMode> EditMode = StaticCastSharedPtr<IPersonaEditMode>(ActiveModes[ModeIndex]);
 		EditMode->GetOnScreenDebugInfo(OutDebugText);
 	}
 }
diff --git a/Engine/Source/Editor/UnrealEd/Private/EditorModeManager.cpp b/Engine/Source/Editor/UnrealEd/Private/EditorModeManager.cpp
--- a/Engine/Source/Editor/UnrealEd/Private/EditorModeManager.cpp
+++ b/Engine/Source/Editor/UnrealEd/Private/EditorModeManager.cpp
@@ -25,6 +25,7 @@
 #include "Editor/UnrealEdEngine.h"
 
 #include "Bookmarks/IBookmarkTypeTools.h"
+#include "Widgets/Docking/SDockTab.h"
 
 /*------------------------------------------------------------------------------
 	FEditorModeTools.
@@ -32,6 +33,8 @@
 	The master class that handles tracking of the current mode.
 ------------------------------------------------------------------------------*/
 
+const FName FEditorModeTools::EditorModeToolbarTabName = TEXT("EditorModeToolbar");
+
 FEditorModeTools::FEditorModeTools()
 	: PivotShown(false)
 	, Snapping(false)
@@ -297,21 +300,80 @@ void FEditorModeTools::ActivateDefaultMode()
 
 void FEditorModeTools::DeactivateModeAtIndex(int32 InIndex)
 {
-	check( InIndex >= 0 && InIndex < Modes.Num() );
+	check( InIndex >= 0 && InIndex < ActiveModes.Num() );
 
-	auto& Mode = Modes[InIndex];
+	auto& Mode = ActiveModes[InIndex];
 
 	Mode->Exit();
+
+	// Remove the toolbar widget
+	ActiveToolBarRows.RemoveAll(
+		[&Mode](FEdModeToolbarRow& Row)
+		{ 
+			return Row.Mode == Mode;
+		}
+	);
+
+	RebuildModeToolBar();
+
 	RecycledModes.Add( Mode->GetID(), Mode );
-	Modes.RemoveAt( InIndex );
+	ActiveModes.RemoveAt( InIndex );
+}
+
+void FEditorModeTools::RebuildModeToolBar()
+{
+	// If the tab or box is not valid the toolbar has not been opened or has been closed by the user
+	TSharedPtr<SVerticalBox> ModeToolbarBoxPinned = ModeToolbarBox.Pin();
+	if (ModeToolbarTab.IsValid() && ModeToolbarBoxPinned.IsValid())
+	{
+		ModeToolbarBoxPinned->ClearChildren();
+
+		if(ActiveToolBarRows.Num())
+		{
+			for(int32 RowIdx = 0; RowIdx < ActiveToolBarRows.Num(); ++RowIdx)
+			{
+				const FEdModeToolbarRow& Row = ActiveToolBarRows[RowIdx];
+				if (ensure(Row.ToolbarWidget.IsValid()))
+				{
+					ModeToolbarBoxPinned->AddSlot()
+					.HAlign(HAlign_Left)
+					.AutoHeight()
+					.Padding(0.0f, RowIdx > 0 ? 5.0f : 0.0f, 0.0f, 0.0f)
+					[
+						Row.ToolbarWidget.ToSharedRef()
+					];
+				}
+			}
+		}
+		else
+		{
+			ModeToolbarTab.Pin()->RequestCloseTab();
+		}
+	}
+}
+
+
+void FEditorModeTools::SpawnOrUpdateModeToolbar()
+{
+	if(ShouldShowModeToolbar())
+	{
+		if (ModeToolbarTab.IsValid())
+		{
+			RebuildModeToolBar();
+		}
+		else if (ToolkitHost.IsValid())
+		{
+			ToolkitHost.Pin()->GetTabManager()->InvokeTab(EditorModeToolbarTabName);
+		}
+	}
 }
 
 void FEditorModeTools::DeactivateMode( FEditorModeID InID )
 {
 	// Find the mode from the ID and exit it.
-	for( int32 Index = Modes.Num() - 1; Index >= 0; --Index )
+	for( int32 Index = ActiveModes.Num() - 1; Index >= 0; --Index )
 	{
-		auto& Mode = Modes[Index];
+		auto& Mode = ActiveModes[Index];
 		if( Mode->GetID() == InID )
 		{
 			DeactivateModeAtIndex(Index);
@@ -319,7 +381,7 @@ void FEditorModeTools::DeactivateMode( FEditorModeID InID )
 		}
 	}
 
-	if( Modes.Num() == 0 )
+	if( ActiveModes.Num() == 0 )
 	{
 		// Ensure the default mode is active if there are no active modes.
 		ActivateDefaultMode();
@@ -328,7 +390,7 @@ void FEditorModeTools::DeactivateMode( FEditorModeID InID )
 
 void FEditorModeTools::DeactivateAllModes()
 {
-	for( int32 Index = Modes.Num() - 1; Index >= 0; --Index )
+	for( int32 Index = ActiveModes.Num() - 1; Index >= 0; --Index )
 	{
 		DeactivateModeAtIndex(Index);
 	}
@@ -337,9 +399,9 @@ void FEditorModeTools::DeactivateAllModes()
 void FEditorModeTools::DestroyMode( FEditorModeID InID )
 {
 	// Find the mode from the ID and exit it.
-	for( int32 Index = Modes.Num() - 1; Index >= 0; --Index )
+	for( int32 Index = ActiveModes.Num() - 1; Index >= 0; --Index )
 	{
-		auto& Mode = Modes[Index];
+		auto& Mode = ActiveModes[Index];
 		if ( Mode->GetID() == InID )
 		{
 			// Deactivate and destroy
@@ -351,6 +413,32 @@ void FEditorModeTools::DestroyMode( FEditorModeID InID )
 	RecycledModes.Remove(InID);
 }
 
+TSharedRef<SDockTab> FEditorModeTools::MakeModeToolbarTab()
+{
+	TSharedRef<SDockTab> ToolbarTabRef = 	
+		SNew(SDockTab)
+		.Label(NSLOCTEXT("EditorModes", "EditorModesToolbarTitle", "Mode Toolbar"))
+		.ShouldAutosize(true)
+		.Icon(FEditorStyle::GetBrush("ToolBar.Icon"))
+		[
+			SAssignNew(ModeToolbarBox, SVerticalBox)
+			
+		];
+
+	ModeToolbarTab = ToolbarTabRef;
+
+	// Rebuild the toolbar with existing mode tools that may be active
+	RebuildModeToolBar();
+
+	return ToolbarTabRef;
+
+}
+
+bool FEditorModeTools::ShouldShowModeToolbar()
+{
+	return ActiveToolBarRows.Num() > 0;
+}
+
 void FEditorModeTools::ActivateMode(FEditorModeID InID, bool bToggle)
 {
 	static bool bReentrant = false;
@@ -407,20 +495,39 @@ void FEditorModeTools::ActivateMode(FEditorModeID InID, bool bToggle)
 	}
 
 	// Remove anything that isn't compatible with this mode
-	for (int32 ModeIndex = Modes.Num() - 1; ModeIndex >= 0; --ModeIndex)
+	for (int32 ModeIndex = ActiveModes.Num() - 1; ModeIndex >= 0; --ModeIndex)
 	{
-		const bool bModesAreCompatible = Mode->IsCompatibleWith(Modes[ModeIndex]->GetID()) || Modes[ModeIndex]->IsCompatibleWith(Mode->GetID());
+		const bool bModesAreCompatible = Mode->IsCompatibleWith(ActiveModes[ModeIndex]->GetID()) || ActiveModes[ModeIndex]->IsCompatibleWith(Mode->GetID());
 		if (!bModesAreCompatible)
 		{
 			DeactivateModeAtIndex(ModeIndex);
 		}
 	}
 
-	Modes.Add(Mode);
+	ActiveModes.Add(Mode);
 
 	// Enter the new mode
 	Mode->Enter();
 
+	// Ask the mode to build the toolbar.
+	TSharedPtr<FUICommandList> CommandList;
+	const TSharedPtr<FModeToolkit> Toolkit = Mode->GetToolkit();
+	if (Toolkit.IsValid())
+	{
+		CommandList = Toolkit->GetToolkitCommands();
+	}
+
+	FToolBarBuilder ModeToolbarBuilder(CommandList, FMultiBoxCustomization(Mode->GetModeInfo().ToolbarCustomizationName), TSharedPtr<FExtender>(), Orient_Horizontal, false);
+	Mode->BuildModeToolbar(ModeToolbarBuilder);
+
+	if (ModeToolbarBuilder.GetMultiBox()->GetBlocks().Num())
+	{
+		TSharedRef<SWidget> ToolbarWidget = ModeToolbarBuilder.MakeWidget();
+		ActiveToolBarRows.Emplace(Mode, ToolbarWidget);
+
+		SpawnOrUpdateModeToolbar();
+	}
+
 	// Update the editor UI
 	FEditorSupportDelegates::UpdateUI.Broadcast();
 }
@@ -447,7 +554,7 @@ bool FEditorModeTools::EnsureNotInMode(FEditorModeID ModeID, const FText& ErrorM
 
 FEdMode* FEditorModeTools::FindMode( FEditorModeID InID )
 {
-	for( auto& Mode : Modes )
+	for( auto& Mode : ActiveModes )
 	{
 		if( Mode->GetID() == InID )
 		{
@@ -474,7 +581,7 @@ FMatrix FEditorModeTools::GetCustomDrawingCoordinateSystem()
 			// If it doesn't want to, create it by looking at the currently selected actors list.
 
 			bool CustomCoordinateSystemProvided = false;
-			for (const auto& Mode : Modes)
+			for (const auto& Mode : ActiveModes)
 			{
 				if (Mode->GetCustomDrawingCoordinateSystem(Matrix, nullptr))
 				{
@@ -520,11 +627,11 @@ FMatrix FEditorModeTools::GetCustomInputCoordinateSystem()
 EAxisList::Type FEditorModeTools::GetWidgetAxisToDraw( FWidget::EWidgetMode InWidgetMode ) const
 {
 	EAxisList::Type OutAxis = EAxisList::All;
-	for( int Index = Modes.Num() - 1; Index >= 0 ; Index-- )
+	for( int Index = ActiveModes.Num() - 1; Index >= 0 ; Index-- )
 	{
-		if ( Modes[Index]->ShouldDrawWidget() )
+		if ( ActiveModes[Index]->ShouldDrawWidget() )
 		{
-			OutAxis = Modes[Index]->GetWidgetAxisToDraw( InWidgetMode );
+			OutAxis = ActiveModes[Index]->GetWidgetAxisToDraw( InWidgetMode );
 			break;
 		}
 	}
@@ -540,9 +647,9 @@ bool FEditorModeTools::StartTracking(FEditorViewportClient* InViewportClient, FV
 
 	CachedLocation = PivotLocation;	// Cache the pivot location
 
-	for( int32 ModeIndex = 0; ModeIndex < Modes.Num(); ++ModeIndex )
+	for( int32 ModeIndex = 0; ModeIndex < ActiveModes.Num(); ++ModeIndex )
 	{
-		const TSharedPtr<FEdMode>& Mode = Modes[ ModeIndex ];
+		const TSharedPtr<FEdMode>& Mode = ActiveModes[ ModeIndex ];
 		bTransactionHandled |= Mode->StartTracking(InViewportClient, InViewport);
 	}
 
@@ -555,9 +662,9 @@ bool FEditorModeTools::EndTracking(FEditorViewportClient* InViewportClient, FVie
 	bIsTracking = false;
 	bool bTransactionHandled = false;
 
-	for( int32 ModeIndex = 0; ModeIndex < Modes.Num(); ++ModeIndex )
+	for( int32 ModeIndex = 0; ModeIndex < ActiveModes.Num(); ++ModeIndex )
 	{
-		const TSharedPtr<FEdMode>& Mode = Modes[ ModeIndex ];
+		const TSharedPtr<FEdMode>& Mode = ActiveModes[ ModeIndex ];
 		bTransactionHandled |= Mode->EndTracking(InViewportClient, InViewportClient->Viewport);
 	}
 
@@ -569,7 +676,7 @@ bool FEditorModeTools::EndTracking(FEditorViewportClient* InViewportClient, FVie
 bool FEditorModeTools::AllowsViewportDragTool() const
 {
 	bool bCanUseDragTool = false;
-	for (const TSharedPtr<FEdMode>& Mode : Modes)
+	for (const TSharedPtr<FEdMode>& Mode : ActiveModes)
 	{
 		bCanUseDragTool |= Mode->AllowsViewportDragTool();
 	}
@@ -579,9 +686,9 @@ bool FEditorModeTools::AllowsViewportDragTool() const
 /** Notifies all active modes that a map change has occured */
 void FEditorModeTools::MapChangeNotify()
 {
-	for( int32 ModeIndex = 0; ModeIndex < Modes.Num(); ++ModeIndex )
+	for( int32 ModeIndex = 0; ModeIndex < ActiveModes.Num(); ++ModeIndex )
 	{
-		const TSharedPtr<FEdMode>& Mode = Modes[ ModeIndex ];
+		const TSharedPtr<FEdMode>& Mode = ActiveModes[ ModeIndex ];
 		Mode->MapChangeNotify();
 	}
 }
@@ -590,9 +697,9 @@ void FEditorModeTools::MapChangeNotify()
 /** Notifies all active modes to empty their selections */
 void FEditorModeTools::SelectNone()
 {
-	for( int32 ModeIndex = 0; ModeIndex < Modes.Num(); ++ModeIndex )
+	for( int32 ModeIndex = 0; ModeIndex < ActiveModes.Num(); ++ModeIndex )
 	{
-		const TSharedPtr<FEdMode>& Mode = Modes[ ModeIndex ];
+		const TSharedPtr<FEdMode>& Mode = ActiveModes[ ModeIndex ];
 		Mode->SelectNone();
 	}
 }
@@ -601,9 +708,9 @@ void FEditorModeTools::SelectNone()
 bool FEditorModeTools::BoxSelect( FBox& InBox, bool InSelect )
 {
 	bool bHandled = false;
-	for( int32 ModeIndex = 0; ModeIndex < Modes.Num(); ++ModeIndex )
+	for( int32 ModeIndex = 0; ModeIndex < ActiveModes.Num(); ++ModeIndex )
 	{
-		const TSharedPtr<FEdMode>& Mode = Modes[ ModeIndex ];
+		const TSharedPtr<FEdMode>& Mode = ActiveModes[ ModeIndex ];
 		bHandled |= Mode->BoxSelect( InBox, InSelect );
 	}
 	return bHandled;
@@ -613,9 +720,9 @@ bool FEditorModeTools::BoxSelect( FBox& InBox, bool InSelect )
 bool FEditorModeTools::FrustumSelect( const FConvexVolume& InFrustum, FEditorViewportClient* InViewportClient, bool InSelect )
 {
 	bool bHandled = false;
-	for( int32 ModeIndex = 0; ModeIndex < Modes.Num(); ++ModeIndex )
+	for( int32 ModeIndex = 0; ModeIndex < ActiveModes.Num(); ++ModeIndex )
 	{
-		const TSharedPtr<FEdMode>& Mode = Modes[ ModeIndex ];
+		const TSharedPtr<FEdMode>& Mode = ActiveModes[ ModeIndex ];
 		bHandled |= Mode->FrustumSelect( InFrustum, InViewportClient, InSelect );
 	}
 	return bHandled;
@@ -626,7 +733,7 @@ bool FEditorModeTools::FrustumSelect( const FConvexVolume& InFrustum, FEditorVie
 bool FEditorModeTools::UsesTransformWidget() const
 {
 	bool bUsesTransformWidget = false;
-	for( const auto& Mode : Modes)
+	for( const auto& Mode : ActiveModes)
 	{
 		bUsesTransformWidget |= Mode->UsesTransformWidget();
 	}
@@ -638,7 +745,7 @@ bool FEditorModeTools::UsesTransformWidget() const
 bool FEditorModeTools::UsesTransformWidget( FWidget::EWidgetMode CheckMode ) const
 {
 	bool bUsesTransformWidget = false;
-	for( const auto& Mode : Modes)
+	for( const auto& Mode : ActiveModes)
 	{
 		bUsesTransformWidget |= Mode->UsesTransformWidget(CheckMode);
 	}
@@ -649,9 +756,9 @@ bool FEditorModeTools::UsesTransformWidget( FWidget::EWidgetMode CheckMode ) con
 /** Sets the current widget axis */
 void FEditorModeTools::SetCurrentWidgetAxis( EAxisList::Type NewAxis )
 {
-	for( int32 ModeIndex = 0; ModeIndex < Modes.Num(); ++ModeIndex )
+	for( int32 ModeIndex = 0; ModeIndex < ActiveModes.Num(); ++ModeIndex )
 	{
-		const TSharedPtr<FEdMode>& Mode = Modes[ ModeIndex ];
+		const TSharedPtr<FEdMode>& Mode = ActiveModes[ ModeIndex ];
 		Mode->SetCurrentWidgetAxis( NewAxis );
 	}
 }
@@ -660,9 +767,9 @@ void FEditorModeTools::SetCurrentWidgetAxis( EAxisList::Type NewAxis )
 bool FEditorModeTools::HandleClick(FEditorViewportClient* InViewportClient,  HHitProxy *HitProxy, const FViewportClick& Click )
 {
 	bool bHandled = false;
-	for( int32 ModeIndex = 0; ModeIndex < Modes.Num(); ++ModeIndex )
+	for( int32 ModeIndex = 0; ModeIndex < ActiveModes.Num(); ++ModeIndex )
 	{
-		const TSharedPtr<FEdMode>& Mode = Modes[ ModeIndex ];
+		const TSharedPtr<FEdMode>& Mode = ActiveModes[ ModeIndex ];
 		bHandled |= Mode->HandleClick(InViewportClient, HitProxy, Click);
 	}
 
@@ -673,12 +780,12 @@ bool FEditorModeTools::HandleClick(FEditorViewportClient* InViewportClient,  HHi
 bool FEditorModeTools::ShouldDrawBrushWireframe( AActor* InActor ) const
 {
 	bool bShouldDraw = false;
-	for( const auto& Mode : Modes)
+	for( const auto& Mode : ActiveModes)
 	{
 		bShouldDraw |= Mode->ShouldDrawBrushWireframe( InActor );
 	}
 
-	if( Modes.Num() == 0 )
+	if( ActiveModes.Num() == 0 )
 	{
 		// We can get into a state where there are no active modes at editor startup if the builder brush is created before the default mode is activated.
 		// Ensure we can see the builder brush when no modes are active.
@@ -698,23 +805,23 @@ bool FEditorModeTools::ShouldDrawBrushVertices() const
 void FEditorModeTools::Tick( FEditorViewportClient* ViewportClient, float DeltaTime )
 {
 	// Remove anything pending destruction
-	for( int32 Index = Modes.Num() - 1; Index >= 0; --Index)
+	for( int32 Index = ActiveModes.Num() - 1; Index >= 0; --Index)
 	{
-		if (Modes[Index]->IsPendingDeletion())
+		if (ActiveModes[Index]->IsPendingDeletion())
 		{
 			DeactivateModeAtIndex(Index);
 		}
 	}
 	
-	if (Modes.Num() == 0)
+	if (ActiveModes.Num() == 0)
 	{
 		// Ensure the default mode is active if there are no active modes.
 		ActivateDefaultMode();
 	}
 
-	for( int32 ModeIndex = 0; ModeIndex < Modes.Num(); ++ModeIndex )
+	for( int32 ModeIndex = 0; ModeIndex < ActiveModes.Num(); ++ModeIndex )
 	{
-		const TSharedPtr<FEdMode>& Mode = Modes[ ModeIndex ];
+		const TSharedPtr<FEdMode>& Mode = ActiveModes[ ModeIndex ];
 		Mode->Tick( ViewportClient, DeltaTime );
 	}
 }
@@ -723,9 +830,9 @@ void FEditorModeTools::Tick( FEditorViewportClient* ViewportClient, float DeltaT
 bool FEditorModeTools::InputDelta( FEditorViewportClient* InViewportClient,FViewport* InViewport,FVector& InDrag,FRotator& InRot,FVector& InScale )
 {
 	bool bHandled = false;
-	for( int32 ModeIndex = 0; ModeIndex < Modes.Num(); ++ModeIndex )
+	for( int32 ModeIndex = 0; ModeIndex < ActiveModes.Num(); ++ModeIndex )
 	{
-		const TSharedPtr<FEdMode>& Mode = Modes[ ModeIndex ];
+		const TSharedPtr<FEdMode>& Mode = ActiveModes[ ModeIndex ];
 		bHandled |= Mode->InputDelta( InViewportClient, InViewport, InDrag, InRot, InScale );
 	}
 	return bHandled;
@@ -735,9 +842,9 @@ bool FEditorModeTools::InputDelta( FEditorViewportClient* InViewportClient,FView
 bool FEditorModeTools::CapturedMouseMove( FEditorViewportClient* InViewportClient, FViewport* InViewport, int32 InMouseX, int32 InMouseY )
 {
 	bool bHandled = false;
-	for( int32 ModeIndex = 0; ModeIndex < Modes.Num(); ++ModeIndex )
+	for( int32 ModeIndex = 0; ModeIndex < ActiveModes.Num(); ++ModeIndex )
 	{
-		const TSharedPtr<FEdMode>& Mode = Modes[ ModeIndex ];
+		const TSharedPtr<FEdMode>& Mode = ActiveModes[ ModeIndex ];
 		bHandled |= Mode->CapturedMouseMove( InViewportClient, InViewport, InMouseX, InMouseY );
 	}
 	return bHandled;
@@ -747,9 +854,9 @@ bool FEditorModeTools::CapturedMouseMove( FEditorViewportClient* InViewportClien
 bool FEditorModeTools::ProcessCapturedMouseMoves( FEditorViewportClient* InViewportClient, FViewport* InViewport, const TArrayView<FIntPoint>& CapturedMouseMoves )
 {
 	bool bHandled = false;
-	for( int32 ModeIndex = 0; ModeIndex < Modes.Num(); ++ModeIndex )
+	for( int32 ModeIndex = 0; ModeIndex < ActiveModes.Num(); ++ModeIndex )
 	{
-		const TSharedPtr<FEdMode>& Mode = Modes[ ModeIndex ];
+		const TSharedPtr<FEdMode>& Mode = ActiveModes[ ModeIndex ];
 		bHandled |= Mode->ProcessCapturedMouseMoves( InViewportClient, InViewport, CapturedMouseMoves );
 	}
 	return bHandled;
@@ -759,9 +866,9 @@ bool FEditorModeTools::ProcessCapturedMouseMoves( FEditorViewportClient* InViewp
 bool FEditorModeTools::InputKey(FEditorViewportClient* InViewportClient, FViewport* Viewport, FKey Key, EInputEvent Event)
 {
 	bool bHandled = false;
-	for( int32 ModeIndex = 0; ModeIndex < Modes.Num(); ++ModeIndex )
+	for( int32 ModeIndex = 0; ModeIndex < ActiveModes.Num(); ++ModeIndex )
 	{
-		const TSharedPtr<FEdMode>& Mode = Modes[ ModeIndex ];
+		const TSharedPtr<FEdMode>& Mode = ActiveModes[ ModeIndex ];
 		bHandled |= Mode->InputKey( InViewportClient, Viewport, Key, Event );
 	}
 	return bHandled;
@@ -771,9 +878,9 @@ bool FEditorModeTools::InputKey(FEditorViewportClient* InViewportClient, FViewpo
 bool FEditorModeTools::InputAxis(FEditorViewportClient* InViewportClient, FViewport* Viewport, int32 ControllerId, FKey Key, float Delta, float DeltaTime)
 {
 	bool bHandled = false;
-	for( int32 ModeIndex = 0; ModeIndex < Modes.Num(); ++ModeIndex )
+	for( int32 ModeIndex = 0; ModeIndex < ActiveModes.Num(); ++ModeIndex )
 	{
-		const TSharedPtr<FEdMode>& Mode = Modes[ ModeIndex ];
+		const TSharedPtr<FEdMode>& Mode = ActiveModes[ ModeIndex ];
 		bHandled |= Mode->InputAxis( InViewportClient, Viewport, ControllerId, Key, Delta, DeltaTime );
 	}
 	return bHandled;
@@ -782,9 +889,9 @@ bool FEditorModeTools::InputAxis(FEditorViewportClient* InViewportClient, FViewp
 bool FEditorModeTools::GetPivotForOrbit( FVector& Pivot ) const
 {
 	// Just return the first pivot point specified by a mode
-	for( int32 ModeIndex = 0; ModeIndex < Modes.Num(); ++ModeIndex )
+	for( int32 ModeIndex = 0; ModeIndex < ActiveModes.Num(); ++ModeIndex )
 	{
-		const TSharedPtr<FEdMode>& Mode = Modes[ ModeIndex ];
+		const TSharedPtr<FEdMode>& Mode = ActiveModes[ ModeIndex ];
 		if ( Mode->GetPivotForOrbit( Pivot ) )
 		{
 			return true;
@@ -796,9 +903,9 @@ bool FEditorModeTools::GetPivotForOrbit( FVector& Pivot ) const
 bool FEditorModeTools::MouseEnter( FEditorViewportClient* InViewportClient, FViewport* Viewport, int32 X, int32 Y )
 {
 	bool bHandled = false;
-	for( int32 ModeIndex = 0; ModeIndex < Modes.Num(); ++ModeIndex )
+	for( int32 ModeIndex = 0; ModeIndex < ActiveModes.Num(); ++ModeIndex )
 	{
-		const TSharedPtr<FEdMode>& Mode = Modes[ ModeIndex ];
+		const TSharedPtr<FEdMode>& Mode = ActiveModes[ ModeIndex ];
 		bHandled |= Mode->MouseEnter( InViewportClient, Viewport, X, Y );
 	}
 	return bHandled;
@@ -807,9 +914,9 @@ bool FEditorModeTools::MouseEnter( FEditorViewportClient* InViewportClient, FVie
 bool FEditorModeTools::MouseLeave( FEditorViewportClient* InViewportClient, FViewport* Viewport )
 {
 	bool bHandled = false;
-	for( int32 ModeIndex = 0; ModeIndex < Modes.Num(); ++ModeIndex )
+	for( int32 ModeIndex = 0; ModeIndex < ActiveModes.Num(); ++ModeIndex )
 	{
-		const TSharedPtr<FEdMode>& Mode = Modes[ ModeIndex ];
+		const TSharedPtr<FEdMode>& Mode = ActiveModes[ ModeIndex ];
 		bHandled |= Mode->MouseLeave( InViewportClient, Viewport );
 	}
 	return bHandled;
@@ -819,9 +926,9 @@ bool FEditorModeTools::MouseLeave( FEditorViewportClient* InViewportClient, FVie
 bool FEditorModeTools::MouseMove( FEditorViewportClient* InViewportClient, FViewport* Viewport, int32 X, int32 Y )
 {
 	bool bHandled = false;
-	for( int32 ModeIndex = 0; ModeIndex < Modes.Num(); ++ModeIndex )
+	for( int32 ModeIndex = 0; ModeIndex < ActiveModes.Num(); ++ModeIndex )
 	{
-		const TSharedPtr<FEdMode>& Mode = Modes[ ModeIndex ];
+		const TSharedPtr<FEdMode>& Mode = ActiveModes[ ModeIndex ];
 		bHandled |= Mode->MouseMove( InViewportClient, Viewport, X, Y );
 	}
 	return bHandled;
@@ -830,9 +937,9 @@ bool FEditorModeTools::MouseMove( FEditorViewportClient* InViewportClient, FView
 bool FEditorModeTools::ReceivedFocus( FEditorViewportClient* InViewportClient, FViewport* Viewport )
 {
 	bool bHandled = false;
-	for( int32 ModeIndex = 0; ModeIndex < Modes.Num(); ++ModeIndex )
+	for( int32 ModeIndex = 0; ModeIndex < ActiveModes.Num(); ++ModeIndex )
 	{
-		const TSharedPtr<FEdMode>& Mode = Modes[ ModeIndex ];
+		const TSharedPtr<FEdMode>& Mode = ActiveModes[ ModeIndex ];
 		bHandled |= Mode->ReceivedFocus( InViewportClient, Viewport );
 	}
 	return bHandled;
@@ -841,9 +948,9 @@ bool FEditorModeTools::ReceivedFocus( FEditorViewportClient* InViewportClient, F
 bool FEditorModeTools::LostFocus( FEditorViewportClient* InViewportClient, FViewport* Viewport )
 {
 	bool bHandled = false;
-	for( int32 ModeIndex = 0; ModeIndex < Modes.Num(); ++ModeIndex )
+	for( int32 ModeIndex = 0; ModeIndex < ActiveModes.Num(); ++ModeIndex )
 	{
-		const TSharedPtr<FEdMode>& Mode = Modes[ ModeIndex ];
+		const TSharedPtr<FEdMode>& Mode = ActiveModes[ ModeIndex ];
 		bHandled |= Mode->LostFocus( InViewportClient, Viewport );
 	}
 	return bHandled;
@@ -852,7 +959,7 @@ bool FEditorModeTools::LostFocus( FEditorViewportClient* InViewportClient, FView
 /** Draws all active mode components */	
 void FEditorModeTools::DrawActiveModes( const FSceneView* InView, FPrimitiveDrawInterface* PDI )
 {
-	for( const auto& Mode : Modes)
+	for( const auto& Mode : ActiveModes)
 	{
 		Mode->Draw( InView, PDI );
 	}
@@ -861,7 +968,7 @@ void FEditorModeTools::DrawActiveModes( const FSceneView* InView, FPrimitiveDraw
 /** Renders all active modes */
 void FEditorModeTools::Render( const FSceneView* InView, FViewport* Viewport, FPrimitiveDrawInterface* PDI )
 {
-	for( const auto& Mode : Modes)
+	for( const auto& Mode : ActiveModes)
 	{
 		Mode->Render( InView, Viewport, PDI );
 	}
@@ -870,7 +977,7 @@ void FEditorModeTools::Render( const FSceneView* InView, FViewport* Viewport, FP
 /** Draws the HUD for all active modes */
 void FEditorModeTools::DrawHUD( FEditorViewportClient* InViewportClient,FViewport* Viewport, const FSceneView* View, FCanvas* Canvas )
 {
-	for( const auto& Mode : Modes)
+	for( const auto& Mode : ActiveModes)
 	{
 		Mode->DrawHUD( InViewportClient, Viewport, View, Canvas );
 	}
@@ -881,9 +988,9 @@ void FEditorModeTools::PostUndo(bool bSuccess)
 {
 	if (bSuccess)
 	{
-		for( int32 ModeIndex = 0; ModeIndex < Modes.Num(); ++ModeIndex )
+		for( int32 ModeIndex = 0; ModeIndex < ActiveModes.Num(); ++ModeIndex )
 		{
-			const TSharedPtr<FEdMode>& Mode = Modes[ ModeIndex ];
+			const TSharedPtr<FEdMode>& Mode = ActiveModes[ ModeIndex ];
 			Mode->PostUndo();
 		}
 	}
@@ -897,7 +1004,7 @@ void FEditorModeTools::PostRedo(bool bSuccess)
 bool FEditorModeTools::AllowWidgetMove() const
 {
 	bool bAllow = false;
-	for( const auto& Mode : Modes)
+	for( const auto& Mode : ActiveModes)
 	{
 		bAllow |= Mode->AllowWidgetMove();
 	}
@@ -907,7 +1014,7 @@ bool FEditorModeTools::AllowWidgetMove() const
 bool FEditorModeTools::DisallowMouseDeltaTracking() const
 {
 	bool bDisallow = false;
-	for( const auto& Mode : Modes)
+	for( const auto& Mode : ActiveModes)
 	{
 		bDisallow |= Mode->DisallowMouseDeltaTracking();
 	}
@@ -917,7 +1024,7 @@ bool FEditorModeTools::DisallowMouseDeltaTracking() const
 bool FEditorModeTools::GetCursor(EMouseCursor::Type& OutCursor) const
 {
 	bool bHandled = false;
-	for( const auto& Mode : Modes)
+	for( const auto& Mode : ActiveModes)
 	{
 		bHandled |= Mode->GetCursor(OutCursor);
 	}
@@ -927,7 +1034,7 @@ bool FEditorModeTools::GetCursor(EMouseCursor::Type& OutCursor) const
 bool FEditorModeTools::GetOverrideCursorVisibility(bool& bWantsOverride, bool& bHardwareCursorVisible, bool bSoftwareCursorVisible) const
 {
 	bool bHandled = false;
-	for (const auto& Mode : Modes)
+	for (const auto& Mode : ActiveModes)
 	{
 		bHandled |= Mode->GetOverrideCursorVisibility(bWantsOverride, bHardwareCursorVisible, bSoftwareCursorVisible);
 	}
@@ -937,7 +1044,7 @@ bool FEditorModeTools::GetOverrideCursorVisibility(bool& bWantsOverride, bool& b
 bool FEditorModeTools::PreConvertMouseMovement(FEditorViewportClient* InViewportClient)
 {
 	bool bHandled = false;
-	for (const auto& Mode : Modes)
+	for (const auto& Mode : ActiveModes)
 	{
 		bHandled |= Mode->PreConvertMouseMovement(InViewportClient);
 	}
@@ -947,7 +1054,7 @@ bool FEditorModeTools::PreConvertMouseMovement(FEditorViewportClient* InViewport
 bool FEditorModeTools::PostConvertMouseMovement(FEditorViewportClient* InViewportClient)
 {
 	bool bHandled = false;
-	for (const auto& Mode : Modes)
+	for (const auto& Mode : ActiveModes)
 	{
 		bHandled |= Mode->PostConvertMouseMovement(InViewportClient);
 	}
@@ -1009,11 +1116,11 @@ void FEditorModeTools::LoadWidgetSettings(void)
 
 FVector FEditorModeTools::GetWidgetLocation() const
 {
-	for (int Index = Modes.Num() - 1; Index >= 0 ; Index--)
+	for (int Index = ActiveModes.Num() - 1; Index >= 0 ; Index--)
 	{
-		if ( Modes[Index]->UsesTransformWidget() )
+		if ( ActiveModes[Index]->UsesTransformWidget() )
 		{
-			 return Modes[Index]->GetWidgetLocation();
+			 return ActiveModes[Index]->GetWidgetLocation();
 		}
 	}
 	
@@ -1126,15 +1233,15 @@ void FEditorModeTools::ClearAllBookmarks( FEditorViewportClient* InViewportClien
 
 void FEditorModeTools::AddReferencedObjects( FReferenceCollector& Collector )
 {
-	for( int32 x = 0 ; x < Modes.Num() ; ++x )
+	for( int32 x = 0 ; x < ActiveModes.Num() ; ++x )
 	{
-		Modes[x]->AddReferencedObjects( Collector );
+		ActiveModes[x]->AddReferencedObjects( Collector );
 	}
 }
 
 FEdMode* FEditorModeTools::GetActiveMode( FEditorModeID InID )
 {
-	for( auto& Mode : Modes )
+	for( auto& Mode : ActiveModes )
 	{
 		if( Mode->GetID() == InID )
 		{
@@ -1146,7 +1253,7 @@ FEdMode* FEditorModeTools::GetActiveMode( FEditorModeID InID )
 
 const FEdMode* FEditorModeTools::GetActiveMode( FEditorModeID InID ) const
 {
-	for (const auto& Mode : Modes)
+	for (const auto& Mode : ActiveModes)
 	{
 		if (Mode->GetID() == InID)
 		{
@@ -1191,14 +1298,14 @@ void FEditorModeTools::GetActiveModes( TArray<FEdMode*>& OutActiveModes )
 {
 	OutActiveModes.Empty();
 	// Copy into an array.  Do not let users modify the active list directly.
-	for( auto& Mode : Modes)
+	for( auto& Mode : ActiveModes)
 	{
 		OutActiveModes.Add(Mode.Get());
 	}
 }
 bool FEditorModeTools::CanCycleWidgetMode() const
 {
-	for (auto& Mode : Modes)
+	for (auto& Mode : ActiveModes)
 	{
 		if (Mode->CanCycleWidgetMode())
 		{
diff --git a/Engine/Source/Editor/UnrealEd/Private/EditorModeRegistry.cpp b/Engine/Source/Editor/UnrealEd/Private/EditorModeRegistry.cpp
--- a/Engine/Source/Editor/UnrealEd/Private/EditorModeRegistry.cpp
+++ b/Engine/Source/Editor/UnrealEd/Private/EditorModeRegistry.cpp
@@ -32,6 +32,7 @@ FEditorModeInfo::FEditorModeInfo(
 	int32 InPriorityOrder
 	)
 	: ID(InID)
+	, ToolbarCustomizationName(*(InID.ToString()+TEXT("Toolbar")))
 	, Name(InName)
 	, IconBrush(InIconBrush)
 	, bVisible(InIsVisible)
diff --git a/Engine/Source/Editor/UnrealEd/Public/EdMode.h b/Engine/Source/Editor/UnrealEd/Public/EdMode.h
--- a/Engine/Source/Editor/UnrealEd/Public/EdMode.h
+++ b/Engine/Source/Editor/UnrealEd/Public/EdMode.h
@@ -278,12 +278,22 @@ public:
 	/** True if this mode uses a toolkit mode (eventually they all should) */
 	virtual bool UsesToolkits() const;
 
+	/** Gets the toolkit created by this mode */
+	TSharedPtr<FModeToolkit> GetToolkit() { return Toolkit; }
+
 	/** Returns the world this toolkit is editing */
 	UWorld* GetWorld() const;
 
 	/** Returns the owning mode manager for this mode */
 	class FEditorModeTools* GetModeManager() const;
 
+	/** 
+	 * Called when the editor mode should rebuild its toolbar 
+	 *
+	 * @param ToolbarBuilder	The builder which should be used to add toolbar widgets
+	 */
+	virtual void BuildModeToolbar(class FToolBarBuilder& ToolbarBuilder) {}
+
 	// Property Widgets
 
 	/**
diff --git a/Engine/Source/Editor/UnrealEd/Public/EditorModeManager.h b/Engine/Source/Editor/UnrealEd/Public/EditorModeManager.h
--- a/Engine/Source/Editor/UnrealEd/Public/EditorModeManager.h
+++ b/Engine/Source/Editor/UnrealEd/Public/EditorModeManager.h
@@ -81,11 +81,24 @@ public:
 	 */
 	void DestroyMode(FEditorModeID InID);
 
+	/**
+	 * Creates the mode toolbar tab if needed
+	 */
+	TSharedRef<SDockTab> MakeModeToolbarTab();
+
+	/**
+	 * Whether or not the mode toolbar should be shown.  If any active modes generated a toolbar this method will return true
+	 */
+	bool ShouldShowModeToolbar();
+
 protected:
 	
 	/** Deactivates the editor mode at the specified index */
 	void DeactivateModeAtIndex( int32 InIndex );
 		
+private:
+	void RebuildModeToolBar();
+	void SpawnOrUpdateModeToolbar();
 public:
 
 	/**
@@ -420,6 +433,9 @@ public:
 	/** Is the viewport UI hidden? */
 	bool IsViewportUIHidden() const { return bHideViewportUI; }
 
+	/** The toolbar tab name that should be used as the tab identifier */
+	static const FName EditorModeToolbarTabName;
+
 	bool PivotShown;
 	bool Snapping;
 	bool SnappedActor;
@@ -482,7 +498,7 @@ protected:
 	TArray<FEditorModeID> DefaultModeIDs;
 
 	/** A list of active editor modes. */
-	TArray< TSharedPtr<FEdMode> > Modes;
+	TArray< TSharedPtr<FEdMode> > ActiveModes;
 
 	/** The host of the toolkits created by these modes */
 	TWeakPtr<IToolkitHost> ToolkitHost;
@@ -505,6 +521,18 @@ protected:
 	/** if true the current selection has a scene component */
 	bool bSelectionHasSceneComponent;
 private:
+	struct FEdModeToolbarRow
+	{
+		FEdModeToolbarRow(TSharedPtr<FEdMode>& InMode, TSharedRef<SWidget>& InToolbarWidget)
+			: Mode(InMode)
+			, ToolbarWidget(InToolbarWidget)
+		{}
+		TSharedPtr<FEdMode> Mode;
+		TSharedPtr<SWidget> ToolbarWidget;
+	};
+
+	/** All toolbar rows generated by active modes.  There will be one row per active mode that generates a toolbar */
+	TArray<FEdModeToolbarRow> ActiveToolBarRows;
 
 	/** The coordinate system the widget is operating within. */
 	ECoordSystem CoordSystem;
@@ -515,6 +543,12 @@ private:
 	/** Multicast delegate that is broadcast when a widget mode is changed */
 	FWidgetModeChangedEvent WidgetModeChangedEvent;
 
+	/** The dock tab for any modes that generate a toolbar */
+	TWeakPtr<SDockTab> ModeToolbarTab;
+
+	/** The actual toolbar rows will be placed in this vertical box */
+	TWeakPtr<SVerticalBox> ModeToolbarBox;
+
 	/** Flag set between calls to StartTracking() and EndTracking() */
 	bool bIsTracking;
 };
diff --git a/Engine/Source/Editor/UnrealEd/Public/EditorModeRegistry.h b/Engine/Source/Editor/UnrealEd/Public/EditorModeRegistry.h
--- a/Engine/Source/Editor/UnrealEd/Public/EditorModeRegistry.h
+++ b/Engine/Source/Editor/UnrealEd/Public/EditorModeRegistry.h
@@ -52,6 +52,9 @@ struct FEditorModeInfo
 	/** The mode ID */
 	FEditorModeID ID;
 
+	/** Name of the toolbar this mode uses and can be used by external systems to customize that mode toolbar */
+	FName ToolbarCustomizationName;
+
 	/** Name for the editor to display */
 	FText Name;
 
diff --git a/Engine/Source/Editor/WorkspaceMenuStructure/Private/WorkspaceMenuStructureModule.cpp b/Engine/Source/Editor/WorkspaceMenuStructure/Private/WorkspaceMenuStructureModule.cpp
--- a/Engine/Source/Editor/WorkspaceMenuStructure/Private/WorkspaceMenuStructureModule.cpp
+++ b/Engine/Source/Editor/WorkspaceMenuStructure/Private/WorkspaceMenuStructureModule.cpp
@@ -76,7 +76,7 @@ public:
 		LevelEditorCategory->ClearItems();
 		LevelEditorViewportsCategory = LevelEditorCategory->AddGroup(LOCTEXT( "WorkspaceMenu_LevelEditorViewportCategory", "Viewports" ), LOCTEXT( "WorkspaceMenu_LevelEditorViewportCategoryTooltip", "Open a Viewport tab." ), FSlateIcon(FEditorStyle::GetStyleSetName(), "LevelEditor.Tabs.Viewports"), true);
 		LevelEditorDetailsCategory = LevelEditorCategory->AddGroup(LOCTEXT("WorkspaceMenu_LevelEditorDetailCategory", "Details" ), LOCTEXT("WorkspaceMenu_LevelEditorDetailCategoryTooltip", "Open a Details tab." ), FSlateIcon(FEditorStyle::GetStyleSetName(), "LevelEditor.Tabs.Details"), true );
-		LevelEditorModesCategory = LevelEditorCategory->AddGroup(LOCTEXT("WorkspaceMenu_LevelEditorToolsCategory", "Tools" ), FSlateIcon(FEditorStyle::GetStyleSetName(), "LevelEditor.Tabs.EditorModes"), true );
+		LevelEditorModesCategory = LevelEditorCategory->AddGroup(LOCTEXT("WorkspaceMenu_LevelEditorToolsCategory", "Editor Modes" ), FSlateIcon(FEditorStyle::GetStyleSetName(), "LevelEditor.Tabs.EditorModes"), true );
 	}
 
 	void ResetToolsCategory()
diff --git a/Engine/Source/Runtime/Slate/Private/Framework/Docking/TabManager.cpp b/Engine/Source/Runtime/Slate/Private/Framework/Docking/TabManager.cpp
--- a/Engine/Source/Runtime/Slate/Private/Framework/Docking/TabManager.cpp
+++ b/Engine/Source/Runtime/Slate/Private/Framework/Docking/TabManager.cpp
@@ -415,6 +415,8 @@ void FTabManager::FLayout::ProcessExtensions(const FLayoutExtender& Extender)
 				TSharedPtr<FTabManager::FStack> Stack = Child->AsStack();
 				if (Stack.IsValid())
 				{
+					StackToParentSplitterMap.Add(Stack.Get(), &Splitter);
+
 					AllStacks.Add(Stack.Get());
 
 					for (FTabManager::FTab& Tab : Stack->Tabs)
@@ -446,6 +448,7 @@ void FTabManager::FLayout::ProcessExtensions(const FLayoutExtender& Extender)
 			return AllDefinedTabs.Contains(TabId);
 		}
 
+		TMap<FTabManager::FStack*, FTabManager::FSplitter*> StackToParentSplitterMap;
 		TArray<FTabManager::FStack*> AllStacks;
 		TSet<FTabId> AllDefinedTabs;
 	};
@@ -455,11 +458,12 @@ void FTabManager::FLayout::ProcessExtensions(const FLayoutExtender& Extender)
 
 	for (FTabManager::FStack* Stack : AllTabs.AllStacks)
 	{
+		FSplitter* ParentSplitter = AllTabs.StackToParentSplitterMap.FindRef(Stack);
 		for (int32 TabIndex = 0; TabIndex < Stack->Tabs.Num();)
 		{
 			FTabId TabId = Stack->Tabs[TabIndex].TabId;
 
-			Extender.FindExtensions(TabId, ELayoutExtensionPosition::Before, ExtendedTabs);
+			Extender.FindTabExtensions(TabId, ELayoutExtensionPosition::Before, ExtendedTabs);
 			for (FTab& NewTab : ExtendedTabs)
 			{
 				if (!AllTabs.Contains(NewTab.TabId))
@@ -470,7 +474,7 @@ void FTabManager::FLayout::ProcessExtensions(const FLayoutExtender& Extender)
 
 			++TabIndex;
 
-			Extender.FindExtensions(TabId, ELayoutExtensionPosition::After, ExtendedTabs);
+			Extender.FindTabExtensions(TabId, ELayoutExtensionPosition::After, ExtendedTabs);
 			for (FTab& NewTab : ExtendedTabs)
 			{
 				if (!AllTabs.Contains(NewTab.TabId))
@@ -478,6 +482,42 @@ void FTabManager::FLayout::ProcessExtensions(const FLayoutExtender& Extender)
 					Stack->Tabs.Insert(NewTab, TabIndex++);
 				}
 			}
+
+	
+			if (ParentSplitter)
+			{
+				Extender.FindTabExtensions(TabId, ELayoutExtensionPosition::Below, ExtendedTabs);
+				if (ExtendedTabs.Num())
+				{
+					for (FTab& NewTab : ExtendedTabs)
+					{
+						if (!AllTabs.Contains(NewTab.TabId))
+						{
+							ParentSplitter->InsertAfter(Stack->AsShared(),
+								FTabManager::NewStack()
+								->SetHideTabWell(true)
+								->AddTab(NewTab)
+							);
+						}
+					}
+				}
+
+				Extender.FindTabExtensions(TabId, ELayoutExtensionPosition::Above, ExtendedTabs);
+				if (ExtendedTabs.Num())
+				{
+					for (FTab& NewTab : ExtendedTabs)
+					{
+						if (!AllTabs.Contains(NewTab.TabId))
+						{
+							ParentSplitter->InsertBefore(Stack->AsShared(),
+								FTabManager::NewStack()
+								->SetHideTabWell(true)
+								->AddTab(NewTab)
+							);
+						}
+					}
+				}
+			}
 		}
 	}
 }
@@ -766,13 +806,13 @@ void FTabManager::SavePersistentLayout()
 	OnPersistLayout_Handler.ExecuteIfBound( MyLayout );
 }
 
-FTabSpawnerEntry& FTabManager::RegisterTabSpawner( const FName TabId, const FOnSpawnTab& OnSpawnTab )
+FTabSpawnerEntry& FTabManager::RegisterTabSpawner(const FName TabId, const FOnSpawnTab& OnSpawnTab, const FCanSpawnTab& CanSpawnTab)
 {
-	ensure( !TabSpawner.Contains(TabId) );
-	ensure( !FGlobalTabmanager::Get()->IsLegacyTabType(TabId) );
+	ensure(!TabSpawner.Contains(TabId));
+	ensure(!FGlobalTabmanager::Get()->IsLegacyTabType(TabId));
 
-	TSharedRef<FTabSpawnerEntry> NewSpawnerEntry = MakeShareable( new FTabSpawnerEntry( TabId, OnSpawnTab ) );
-	TabSpawner.Add( TabId, NewSpawnerEntry );
+	TSharedRef<FTabSpawnerEntry> NewSpawnerEntry = MakeShareable(new FTabSpawnerEntry(TabId, OnSpawnTab, CanSpawnTab));
+	TabSpawner.Add(TabId, NewSpawnerEntry);
 	return NewSpawnerEntry.Get();
 }
 
@@ -914,9 +954,15 @@ void FTabManager::PopulateTabSpawnerMenu_Helper( FMenuBuilder& PopulateMe, FPopu
 
 void FTabManager::MakeSpawnerMenuEntry( FMenuBuilder &PopulateMe, const TSharedPtr<FTabSpawnerEntry> &SpawnerNode ) 
 {
-	auto CanExecuteMenuEntry = [](TAttribute<ETabSpawnerMenuType::Type> InMenuType) -> bool
+	auto CanExecuteMenuEntry = [](TWeakPtr<FTabSpawnerEntry> SpawnerNode) -> bool
 	{
-		return InMenuType.Get() == ETabSpawnerMenuType::Enabled;
+		TSharedPtr<FTabSpawnerEntry> SpawnerNodePinned = SpawnerNode.Pin();
+		if (SpawnerNodePinned.IsValid() && SpawnerNodePinned->MenuType.Get() == ETabSpawnerMenuType::Enabled)
+		{
+			return SpawnerNodePinned->CanSpawnTab.IsBound() ? SpawnerNodePinned->CanSpawnTab.Execute(FSpawnTabArgs(TSharedPtr<SWindow>(), SpawnerNodePinned->TabType)) : true;
+		}
+
+		return false;
 	};
 
 	if ( SpawnerNode->MenuType.Get() != ETabSpawnerMenuType::Hidden )
@@ -927,7 +973,7 @@ void FTabManager::MakeSpawnerMenuEntry( FMenuBuilder &PopulateMe, const TSharedP
 			SpawnerNode->GetIcon(),
 			FUIAction(
 			FExecuteAction::CreateSP(SharedThis(this), &FTabManager::InvokeTabForMenu, SpawnerNode->TabType),
-			FCanExecuteAction::CreateStatic(CanExecuteMenuEntry, SpawnerNode->MenuType),
+			FCanExecuteAction::CreateStatic(CanExecuteMenuEntry, TWeakPtr<FTabSpawnerEntry>(SpawnerNode)),
 			FIsActionChecked::CreateSP(SpawnerNode.ToSharedRef(), &FTabSpawnerEntry::IsSoleTabInstanceSpawned)
 			),
 			NAME_None,
@@ -1077,7 +1123,7 @@ void FTabManager::RestoreDocumentTab( FName PlaceholderId, ESearchPreference::Ty
 
 TSharedRef<SDockTab> FTabManager::InvokeTab( const FTabId& TabId )
 {
-	TSharedRef<SDockTab> NewTab = InvokeTab_Internal( TabId );
+	TSharedPtr<SDockTab> NewTab = InvokeTab_Internal( TabId );
 	TSharedPtr<SWindow> ParentWindowPtr = NewTab->GetParentWindow();
 	if ((NewTab->GetTabRole() == ETabRole::MajorTab || NewTab->GetTabRole() == ETabRole::NomadTab) && ParentWindowPtr.IsValid() && ParentWindowPtr != FGlobalTabmanager::Get()->GetRootWindow())
 	{
@@ -1086,10 +1132,11 @@ TSharedRef<SDockTab> FTabManager::InvokeTab( const FTabId& TabId )
 #if PLATFORM_MAC
 	FPlatformApplicationMisc::bChachedMacMenuStateNeedsUpdate = true;
 #endif
-	return NewTab;
+	// Note: we expect tabs that are manually invoked to always succeed
+	return NewTab.ToSharedRef();
 }
 
-TSharedRef<SDockTab> FTabManager::InvokeTab_Internal( const FTabId& TabId )
+TSharedPtr<SDockTab> FTabManager::InvokeTab_Internal( const FTabId& TabId )
 {
 	// Tab Spawning Rules:
 	// 
@@ -1136,10 +1183,14 @@ TSharedRef<SDockTab> FTabManager::InvokeTab_Internal( const FTabId& TabId )
 
 	if (StackToSpawnIn.IsValid())
 	{
-		const TSharedRef<SDockTab> NewTab = SpawnTab( TabId, TSharedPtr<SWindow>() );
+		const TSharedPtr<SDockTab> NewTab = SpawnTab( TabId, TSharedPtr<SWindow>() );
+
+		if(NewTab.IsValid())
+		{
+			StackToSpawnIn->OpenTab(NewTab.ToSharedRef());
+			NewTab->PlaySpawnAnim();
+		}
 
-		StackToSpawnIn->OpenTab( NewTab );
-		NewTab->PlaySpawnAnim();
 		return NewTab;
 	}
 	else if ( FGlobalTabmanager::Get() != SharedThis(this) && NomadTabSpawner->Contains(TabId.TabType) )
@@ -1272,14 +1323,17 @@ TSharedRef<SDockingNode> FTabManager::RestoreArea_Helper( const TSharedRef<FLayo
 
 			if (SomeTab.TabState == ETabState::OpenedTab && IsValidTabForSpawning(SomeTab))
 			{
-				const TSharedRef<SDockTab> NewTabWidget = SpawnTab( SomeTab.TabId, ParentWindow );
+				const TSharedPtr<SDockTab> NewTabWidget = SpawnTab( SomeTab.TabId, ParentWindow );
 
-				if(SomeTab.TabId == NodeAsStack->ForegroundTabId)
+				if (NewTabWidget.IsValid())
 				{
-					WidgetToActivate = NewTabWidget;
-				}
+					if (SomeTab.TabId == NodeAsStack->ForegroundTabId)
+					{
+						WidgetToActivate = NewTabWidget;
+					}
 
-				NewStackWidget->AddTabWidget( NewTabWidget );
+					NewStackWidget->AddTabWidget(NewTabWidget.ToSharedRef());
+				}
 			}
 		}
 		
@@ -1374,7 +1428,7 @@ TSharedRef<SDockingNode> FTabManager::RestoreArea_Helper( const TSharedRef<FLayo
 		ensureMsgf( false, TEXT("Unexpected node type") );
 
 		TSharedRef<SDockingTabStack> NewStackWidget = SNew(SDockingTabStack, FTabManager::NewStack());
-		NewStackWidget->OpenTab( SpawnTab( FName(NAME_None), ParentWindow ) );
+		NewStackWidget->OpenTab(SpawnTab( FName(NAME_None), ParentWindow ).ToSharedRef());
 		return NewStackWidget;
 	}
 }
@@ -1392,10 +1446,21 @@ void FTabManager::RestoreSplitterContent( const TSharedRef<FSplitter>& SplitterN
 	}
 }
 
-bool FTabManager::CanSpawnTab(FName TabId)
+bool FTabManager::CanSpawnTab(FName TabId) const
 {
-	TSharedPtr<FTabSpawnerEntry> Spawner = FindTabSpawnerFor(TabId);
-	return Spawner.IsValid();
+	return HasTabSpawner(TabId);
+}
+
+bool FTabManager::HasTabSpawner(FName TabId) const
+{
+	// Look for a spawner in this tab manager.
+	const TSharedRef<FTabSpawnerEntry>* Spawner = TabSpawner.Find(TabId);
+	if (Spawner == nullptr)
+	{
+		Spawner = NomadTabSpawner->Find(TabId);
+	}
+
+	return Spawner != nullptr;
 }
 
 bool FTabManager::IsValidTabForSpawning( const FTab& SomeTab ) const
@@ -1405,24 +1470,35 @@ bool FTabManager::IsValidTabForSpawning( const FTab& SomeTab ) const
 	return ( !NomadSpawner || !NomadSpawner->Get().IsSoleTabInstanceSpawned() );
 }
 
-TSharedRef<SDockTab> FTabManager::SpawnTab( const FTabId& TabId, const TSharedPtr<SWindow>& ParentWindow )
+TSharedPtr<SDockTab> FTabManager::SpawnTab( const FTabId& TabId, const TSharedPtr<SWindow>& ParentWindow )
 {
 	TSharedPtr<SDockTab> NewTabWidget;
 
+	// Whether or not the spawner overrode the ability for the tab to even spawn.  This is not a failure case.
+	bool bSpawningAllowedBySpawner = true;
 	// Do we know how to spawn such a tab?
 	TSharedPtr<FTabSpawnerEntry> Spawner = FindTabSpawnerFor(TabId.TabType);
 	if ( Spawner.IsValid() )
 	{
-		NewTabWidget = Spawner->OnSpawnTab.Execute( FSpawnTabArgs( ParentWindow, TabId ) );
-		NewTabWidget->SetLayoutIdentifier( TabId );
-		NewTabWidget->ProvideDefaultLabel( Spawner->GetDisplayName().IsEmpty() ? FText::FromName( Spawner->TabType ) : Spawner->GetDisplayName() );
-		NewTabWidget->ProvideDefaultIcon( Spawner->GetIcon().GetIcon() );
-	
-		// The spawner tracks that last tab it spawned
-		Spawner->SpawnedTabPtr = NewTabWidget;
+		if (Spawner->CanSpawnTab.IsBound())
+		{
+			bSpawningAllowedBySpawner = Spawner->CanSpawnTab.Execute(FSpawnTabArgs(ParentWindow, TabId));
+		}
+
+		if(bSpawningAllowedBySpawner)
+		{
+			NewTabWidget = Spawner->OnSpawnTab.Execute(FSpawnTabArgs(ParentWindow, TabId));
+			NewTabWidget->SetLayoutIdentifier(TabId);
+			NewTabWidget->ProvideDefaultLabel(Spawner->GetDisplayName().IsEmpty() ? FText::FromName(Spawner->TabType) : Spawner->GetDisplayName());
+			NewTabWidget->ProvideDefaultIcon(Spawner->GetIcon().GetIcon());
+
+			// The spawner tracks that last tab it spawned
+			Spawner->SpawnedTabPtr = NewTabWidget;
+		}
 	}
 
-	if ( !NewTabWidget.IsValid() )
+	// The tab was allowed to be spawned but failed for some reason
+	if (bSpawningAllowedBySpawner && !NewTabWidget.IsValid() )
 	{
 		// We don't know how to spawn this tab.
 		// Make a dummy tab so that things aren't entirely broken.
@@ -1443,9 +1519,12 @@ TSharedRef<SDockTab> FTabManager::SpawnTab( const FTabId& TabId, const TSharedPt
 		NewTabWidget->SetLayoutIdentifier(TabId);
 	}
 
-	NewTabWidget->SetTabManager( SharedThis(this) );
+	if (NewTabWidget.IsValid())
+	{
+		NewTabWidget->SetTabManager(SharedThis(this));
+	}
 
-	return NewTabWidget.ToSharedRef();
+	return NewTabWidget;
 }
 
 TSharedPtr<SDockTab> FTabManager::FindExistingLiveTab( const FTabId& TabId ) const
@@ -1717,18 +1796,17 @@ TSharedPtr<FTabManager::FStack> FTabManager::FindTabUnderNode( const FTabMatcher
 TSharedPtr<FTabSpawnerEntry> FTabManager::FindTabSpawnerFor(FName TabId)
 {
 	// Look for a spawner in this tab manager.
-	TSharedRef<FTabSpawnerEntry>* Spawner = TabSpawner.Find( TabId );
-	if (Spawner == NULL)
+	TSharedRef<FTabSpawnerEntry>* Spawner = TabSpawner.Find(TabId);
+	if (Spawner == nullptr)
 	{
-		Spawner = NomadTabSpawner->Find( TabId );
+		Spawner = NomadTabSpawner->Find(TabId);
 	}
 
-	return (Spawner != NULL)
+	return (Spawner != nullptr)
 		? TSharedPtr<FTabSpawnerEntry>(*Spawner)
 		: TSharedPtr<FTabSpawnerEntry>();
 }
 
-
 int32 FTabManager::FindTabInCollapsedAreas( const FTabMatcher& Matcher )
 {
 	for ( int32 CollapsedDockAreaIndex=0; CollapsedDockAreaIndex < CollapsedDockAreas.Num(); ++CollapsedDockAreaIndex )
@@ -1855,13 +1933,13 @@ void FGlobalTabmanager::SetActiveTab( const TSharedPtr<class SDockTab>& NewActiv
 	}	
 }
 
-FTabSpawnerEntry& FGlobalTabmanager::RegisterNomadTabSpawner( const FName TabId, const FOnSpawnTab& OnSpawnTab )
+FTabSpawnerEntry& FGlobalTabmanager::RegisterNomadTabSpawner(const FName TabId, const FOnSpawnTab& OnSpawnTab, const FCanSpawnTab& CanSpawnTab)
 {
-	ensure( !NomadTabSpawner->Contains(TabId) );
-	ensure( !IsLegacyTabType(TabId) );
+	ensure(!NomadTabSpawner->Contains(TabId));
+	ensure(!IsLegacyTabType(TabId));
 
-	TSharedRef<FTabSpawnerEntry> NewSpawnerEntry = MakeShareable( new FTabSpawnerEntry( TabId, OnSpawnTab ) );
-	NomadTabSpawner->Add( TabId, NewSpawnerEntry );
+	TSharedRef<FTabSpawnerEntry> NewSpawnerEntry = MakeShareable(new FTabSpawnerEntry(TabId, OnSpawnTab, CanSpawnTab));
+	NomadTabSpawner->Add(TabId, NewSpawnerEntry);
 	return NewSpawnerEntry.Get();
 }
 
diff --git a/Engine/Source/Runtime/Slate/Public/Framework/Docking/LayoutExtender.h b/Engine/Source/Runtime/Slate/Public/Framework/Docking/LayoutExtender.h
--- a/Engine/Source/Runtime/Slate/Public/Framework/Docking/LayoutExtender.h
+++ b/Engine/Source/Runtime/Slate/Public/Framework/Docking/LayoutExtender.h
@@ -16,7 +16,11 @@ enum class ELayoutExtensionPosition
 	/** Extend the layout before the specified element */
 	Before,
 	/** Extend the layout after the specified element */
-	After
+	After,
+	/** Extend the layout above the specified element in the parent splitter */
+	Above,
+	/** Extend the layout below the specified element in the parent splitter */
+	Below,
 };
 
 /** Class used for extending default layouts */
@@ -51,7 +55,7 @@ public:
 	 * @param OutValues				The container to populate with extended tabs
 	 */
 	template<typename AllocatorType>
-	void FindExtensions(FTabId TabId, ELayoutExtensionPosition Position, TArray<FTabManager::FTab, AllocatorType>& OutValues) const
+	void FindTabExtensions(FTabId TabId, ELayoutExtensionPosition Position, TArray<FTabManager::FTab, AllocatorType>& OutValues) const
 	{
 		OutValues.Reset();
 
diff --git a/Engine/Source/Runtime/Slate/Public/Framework/Docking/TabManager.h b/Engine/Source/Runtime/Slate/Public/Framework/Docking/TabManager.h
--- a/Engine/Source/Runtime/Slate/Public/Framework/Docking/TabManager.h
+++ b/Engine/Source/Runtime/Slate/Public/Framework/Docking/TabManager.h
@@ -110,7 +110,7 @@ class FSpawnTabArgs
  * Invoked when a tab needs to be spawned.
  */
 DECLARE_DELEGATE_RetVal_OneParam( TSharedRef<SDockTab>, FOnSpawnTab, const FSpawnTabArgs& );
-
+DECLARE_DELEGATE_RetVal_OneParam(bool, FCanSpawnTab, const FSpawnTabArgs&);
 /**
  * Allows users to provide custom logic when searching for a tab to reuse.
  * The TabId that is being searched for is provided as a courtesy, but does not have to be respected.
@@ -130,10 +130,11 @@ namespace ETabSpawnerMenuType
 
 struct FTabSpawnerEntry : public FWorkspaceItem
 {
-	FTabSpawnerEntry( const FName& InTabType, const FOnSpawnTab& InSpawnTabMethod )
-		: FWorkspaceItem( FText(), FSlateIcon(), false )
-		, TabType( InTabType )
-		, OnSpawnTab( InSpawnTabMethod )
+	FTabSpawnerEntry(const FName& InTabType, const FOnSpawnTab& InSpawnTabMethod, const FCanSpawnTab& InCanSpawnTab)
+		: FWorkspaceItem(FText(), FSlateIcon(), false)
+		, TabType(InTabType)
+		, OnSpawnTab(InSpawnTabMethod)
+		, CanSpawnTab(InCanSpawnTab)
 		, OnFindTabToReuse()
 		, MenuType(ETabSpawnerMenuType::Enabled)
 		, bAutoGenerateMenuEntry(true)
@@ -191,6 +192,7 @@ struct FTabSpawnerEntry : public FWorkspaceItem
 private:
 	FName TabType;
 	FOnSpawnTab OnSpawnTab;
+	FCanSpawnTab CanSpawnTab;
 	/** When this method is not provided, we assume that the tab should only allow 0 or 1 instances */
 	FOnFindTabToReuse OnFindTabToReuse;
 	/** Whether this menu item should be enabled, disabled, or hidden */
@@ -294,6 +296,12 @@ class SLATE_API FTabManager : public TSharedFromThis<FTabManager>
 					return SharedThis(this);
 				}
 
+				TSharedRef<FStack> AddTab(const FTab& Tab)
+				{
+					Tabs.Add(Tab);
+					return SharedThis(this);
+				}
+
 				TSharedRef<FStack> SetSizeCoefficient( const float InSizeCoefficient )
 				{
 					SizeCoefficient = InSizeCoefficient;
@@ -349,6 +357,23 @@ class SLATE_API FTabManager : public TSharedFromThis<FTabManager>
 					return SharedThis(this);
 				}
 
+				TSharedRef<FSplitter> InsertBefore(TSharedRef<FLayoutNode> NodeToInsertBefore, TSharedRef<FLayoutNode> NodeToInsert)
+				{
+					int32 InsertAtIndex = ChildNodes.Find(NodeToInsertBefore);
+					check(InsertAtIndex != INDEX_NONE);
+					ChildNodes.Insert(NodeToInsert, InsertAtIndex);
+					return SharedThis(this);
+				}
+
+				TSharedRef<FSplitter> InsertAfter(TSharedRef<FLayoutNode> NodeToInsertAfter, TSharedRef<FLayoutNode> NodeToInsert)
+				{
+					int32 InsertAtIndex = ChildNodes.Find(NodeToInsertAfter);
+					check(InsertAtIndex != INDEX_NONE);
+					ChildNodes.Insert(NodeToInsert, InsertAtIndex + 1);
+					return SharedThis(this);
+				}
+
+
 				TSharedRef<FSplitter> SetSizeCoefficient( const float InSizeCoefficient )
 				{
 					SizeCoefficient = InSizeCoefficient;
@@ -626,10 +651,11 @@ class SLATE_API FTabManager : public TSharedFromThis<FTabManager>
 		 * InvokeTab().
 		 * @param TabId The TabId to register the spawner for.
 		 * @param OnSpawnTab The callback that will be used to spawn the tab.
+		 * @param CanSpawnTab The callback that will be used to ask if spawning the tab is allowed
 		 * @return The registration entry for the spawner.
 		 */
-		FTabSpawnerEntry& RegisterTabSpawner( const FName TabId, const FOnSpawnTab& OnSpawnTab );
-		
+		FTabSpawnerEntry& RegisterTabSpawner(const FName TabId, const FOnSpawnTab& OnSpawnTab, const FCanSpawnTab& CanSpawnTab = FCanSpawnTab());
+
 		/**
 		 * Unregisters the tab spawner matching the provided TabId.
 		 * @param TabId The TabId to remove the spawner for.
@@ -732,7 +758,10 @@ class SLATE_API FTabManager : public TSharedFromThis<FTabManager>
 		void ClearLocalWorkspaceMenuCategories();
 
 		/** @return true if the tab has a factory registered for it that allows it to be spawned. */
-		bool CanSpawnTab( FName TabId );
+		UE_DEPRECATED(4.23, "CanSpawnTab has been replaced by HasTabSpawner")
+		bool CanSpawnTab(FName TabId) const;
+
+		bool HasTabSpawner(FName TabId) const;
 
 		/** Returns the owner tab (if it exists) */
 		TSharedPtr<SDockTab> GetOwnerTab() { return OwnerTabPtr.Pin(); }
@@ -750,7 +779,7 @@ class SLATE_API FTabManager : public TSharedFromThis<FTabManager>
 
 		void MakeSpawnerMenuEntry( FMenuBuilder &PopulateMe, const TSharedPtr<FTabSpawnerEntry> &SpawnerNode );
 
-		TSharedRef<SDockTab> InvokeTab_Internal( const FTabId& TabId );
+		TSharedPtr<SDockTab> InvokeTab_Internal( const FTabId& TabId );
 
 		/** Finds the last major or nomad tab in a particular window. */
 		TSharedPtr<SDockTab> FindLastTabInWindow(TSharedPtr<SWindow> Window) const;
@@ -773,7 +802,7 @@ class SLATE_API FTabManager : public TSharedFromThis<FTabManager>
 		void RestoreSplitterContent( const TSharedRef<FSplitter>& SplitterNode, const TSharedRef<class SDockingSplitter>& SplitterWidget, const TSharedPtr<SWindow>& ParentWindow );
 		
 		bool IsValidTabForSpawning( const FTab& SomeTab ) const;
-		TSharedRef<SDockTab> SpawnTab( const FTabId& TabId, const TSharedPtr<SWindow>& ParentWindow );
+		TSharedPtr<SDockTab> SpawnTab( const FTabId& TabId, const TSharedPtr<SWindow>& ParentWindow );
 
 		TSharedPtr<class SDockingTabStack> FindTabInLiveAreas( const FTabMatcher& TabMatcher ) const;
 		static TSharedPtr<class SDockingTabStack> FindTabInLiveArea( const FTabMatcher& TabMatcher, const TSharedRef<SDockingArea>& InArea );
@@ -812,6 +841,7 @@ class SLATE_API FTabManager : public TSharedFromThis<FTabManager>
 		FTabSpawner TabSpawner;
 		TSharedRef<FTabSpawner> NomadTabSpawner;
 		TSharedPtr<FTabSpawnerEntry> FindTabSpawnerFor(FName TabId);
+		bool HasTabSpawnerFor(FName TabId) const;
 
 		TArray< TWeakPtr<SDockingArea> > DockAreas;
 		TArray< TSharedRef<FTabManager::FArea> > CollapsedDockAreas;
@@ -899,8 +929,17 @@ public:
 	/** Activate the NewActiveTab. If NewActiveTab is NULL, the active tab is cleared. */
 	void SetActiveTab( const TSharedPtr<SDockTab>& NewActiveTab );
 
-	FTabSpawnerEntry& RegisterNomadTabSpawner( const FName TabId, const FOnSpawnTab& OnSpawnTab );
-	
+	/**
+	 * Register a new normad tab spawner with the tab manager.  The spawner will be called when anyone calls
+	 * InvokeTab().
+	 * A nomad tab is a tab that can be placed with major tabs or minor tabs in any tab well
+	 * @param TabId The TabId to register the spawner for.
+	 * @param OnSpawnTab The callback that will be used to spawn the tab.
+	 * @param CanSpawnTab The callback that will be used to ask if spawning the tab is allowed
+	 * @return The registration entry for the spawner.
+	 */
+	FTabSpawnerEntry& RegisterNomadTabSpawner( const FName TabId, const FOnSpawnTab& OnSpawnTab, const FCanSpawnTab& CanSpawnTab = FCanSpawnTab());
+
 	void UnregisterNomadTabSpawner( const FName TabId );
 
 	void SetApplicationTitle( const FText& AppTitle );
