Fix for crash caused by hitting delete while dragging and actor in the Scene outliner.

diff --git a/Engine/Source/Editor/UnrealEd/Private/EditorActor.cpp b/Engine/Source/Editor/UnrealEd/Private/EditorActor.cpp
--- a/Engine/Source/Editor/UnrealEd/Private/EditorActor.cpp
+++ b/Engine/Source/Editor/UnrealEd/Private/EditorActor.cpp
@@ -702,6 +702,8 @@ bool UUnrealEdEngine::edactDeleteSelected( UWorld* InWorld, bool bVerifyDeletion
 
 	const double StartSeconds = FPlatformTime::Seconds();
 
+	FSlateApplication::Get().CancelDragDrop();
+
 	if (GetSelectedComponentCount() > 0)
 	{
 		TArray<UActorComponent*> SelectedEditableComponents;
