CIS Fixes

diff --git a/Engine/Source/Editor/UnrealEd/Private/EditorModeManager.cpp b/Engine/Source/Editor/UnrealEd/Private/EditorModeManager.cpp
--- a/Engine/Source/Editor/UnrealEd/Private/EditorModeManager.cpp
+++ b/Engine/Source/Editor/UnrealEd/Private/EditorModeManager.cpp
@@ -14,7 +14,6 @@
 #include "EditorSupportDelegates.h"
 #include "EdMode.h"
 #include "Toolkits/IToolkitHost.h"
-
 #include "Framework/Notifications/NotificationManager.h"
 #include "Widgets/Notifications/SNotificationList.h"
 #include "Engine/LevelStreaming.h"
@@ -23,9 +22,12 @@
 #include "Editor/EditorEngine.h"
 #include "UnrealEdGlobals.h"
 #include "Editor/UnrealEdEngine.h"
-
 #include "Bookmarks/IBookmarkTypeTools.h"
 #include "Widgets/Docking/SDockTab.h"
+#include "EditorStyleSet.h"
+#include "Framework/Commands/UICommandList.h"
+#include "Framework/MultiBox/MultiBoxBuilder.h"
+#include "Toolkits/BaseToolkit.h"
 
 /*------------------------------------------------------------------------------
 	FEditorModeTools.
