Build fix missing headers

diff --git a/Engine/Source/Editor/UnrealEd/Private/AssetSelection.cpp b/Engine/Source/Editor/UnrealEd/Private/AssetSelection.cpp
--- a/Engine/Source/Editor/UnrealEd/Private/AssetSelection.cpp
+++ b/Engine/Source/Editor/UnrealEd/Private/AssetSelection.cpp
@@ -45,6 +45,7 @@
 #include "Framework/Notifications/NotificationManager.h"
 #include "Widgets/Notifications/SNotificationList.h"
 #include "Settings/LevelEditorMiscSettings.h"
+#include "Engine/LevelStreaming.h"
 #include "Engine/LevelBounds.h"
 #include "SourceControlHelpers.h"
 #include "Dialogs/Dialogs.h"
diff --git a/Engine/Source/Editor/UnrealEd/Private/EditorActor.cpp b/Engine/Source/Editor/UnrealEd/Private/EditorActor.cpp
--- a/Engine/Source/Editor/UnrealEd/Private/EditorActor.cpp
+++ b/Engine/Source/Editor/UnrealEd/Private/EditorActor.cpp
@@ -64,6 +64,7 @@
 #include "IAssetTools.h"
 #include "AssetToolsModule.h"
 #include "AssetSelection.h"
+#include "Framework/Application/SlateApplication.h"
 
 #define LOCTEXT_NAMESPACE "UnrealEd.EditorActor"
 
