Fix crash when changing skeletal mesh and the persona viewport is showing some bones (selected or all hierarchy)

diff --git a/Engine/Source/Editor/UnrealEd/Classes/Animation/DebugSkelMeshComponent.h b/Engine/Source/Editor/UnrealEd/Classes/Animation/DebugSkelMeshComponent.h
--- a/Engine/Source/Editor/UnrealEd/Classes/Animation/DebugSkelMeshComponent.h
+++ b/Engine/Source/Editor/UnrealEd/Classes/Animation/DebugSkelMeshComponent.h
@@ -510,7 +510,13 @@ public:
 	 */
 	virtual FTransform GetDrawTransform(int32 BoneIndex) const
 	{
-		return GetComponentSpaceTransforms()[BoneIndex];
+		const TArray<FTransform>& SpaceTransforms = GetComponentSpaceTransforms();
+		if (SpaceTransforms.IsValidIndex(BoneIndex))
+		{
+			return SpaceTransforms[BoneIndex];
+		}
+
+		return FTransform::Identity;
 	}
 
 };
