Fbx skeletal mesh importer now send a token error if there is any nan or zero scale in the bone transform.

diff --git a/Engine/Source/Editor/UnrealEd/Private/Fbx/FbxSkeletalMeshImport.cpp b/Engine/Source/Editor/UnrealEd/Private/Fbx/FbxSkeletalMeshImport.cpp
--- a/Engine/Source/Editor/UnrealEd/Private/Fbx/FbxSkeletalMeshImport.cpp
+++ b/Engine/Source/Editor/UnrealEd/Private/Fbx/FbxSkeletalMeshImport.cpp
@@ -1222,6 +1222,41 @@ bool UnFbx::FFbxImporter::ImportBone(TArray<FbxNode*>& NodeArray, FSkeletalMeshI
 		BoneName = UTF8_TO_TCHAR( MakeName( LinkName ) );
 		Bone.Name = BoneName;
 
+		//Check for nan and for zero scale
+		if (AllowImportBoneErrorAndWarning)
+		{
+			bool bFoundNan = false;
+			bool bFoundZeroScale = false;
+			for (int32 i = 0; i < 4; ++i)
+			{
+				if (i < 3)
+				{
+					if (FMath::IsNaN(LocalLinkT[i]) || FMath::IsNaN(LocalLinkS[i]))
+					{
+						bFoundNan = true;
+					}
+					if (FMath::IsNearlyZero(LocalLinkS[i]))
+					{
+						bFoundZeroScale = true;
+					}
+				}
+				if (FMath::IsNaN(LocalLinkQ[i]))
+				{
+					bFoundNan = true;
+				}
+			}
+
+			if (bFoundNan)
+			{
+				AddTokenizedErrorMessage(FTokenizedMessage::Create(EMessageSeverity::Error, FText::Format(LOCTEXT("BoneTransformNAN", "Found NAN value in bone transform. Bone name: [{0}]"), FText::FromString(BoneName))), FFbxErrors::SkeletalMesh_InvalidBindPose);
+			}
+
+			if (bFoundZeroScale)
+			{
+				AddTokenizedErrorMessage(FTokenizedMessage::Create(EMessageSeverity::Error, FText::Format(LOCTEXT("BoneTransformZeroScale", "Found zero scale value in bone transform. Bone name: [{0}]"), FText::FromString(BoneName))), FFbxErrors::SkeletalMesh_InvalidBindPose);
+			}
+		}
+
 		SkeletalMeshImportData::FJointPos& JointMatrix = Bone.BonePos;
 		FbxSkeleton* Skeleton = Link->GetSkeleton();
 		if (Skeleton)
