Fbx staticmesh importer: do not fail import if not all parts are completly degenerated.

diff --git a/Engine/Source/Editor/UnrealEd/Private/Fbx/FbxStaticMeshImport.cpp b/Engine/Source/Editor/UnrealEd/Private/Fbx/FbxStaticMeshImport.cpp
--- a/Engine/Source/Editor/UnrealEd/Private/Fbx/FbxStaticMeshImport.cpp
+++ b/Engine/Source/Editor/UnrealEd/Private/Fbx/FbxStaticMeshImport.cpp
@@ -271,6 +271,7 @@ bool UnFbx::FFbxImporter::BuildStaticMeshFromGeometry(FbxNode* Node, UStaticMesh
 {
 	check(StaticMesh->SourceModels.IsValidIndex(LODIndex));
 	FbxMesh* Mesh = Node->GetMesh();
+	FString FbxNodeName = UTF8_TO_TCHAR(Node->GetName());
 	FStaticMeshSourceModel& SrcModel = StaticMesh->SourceModels[LODIndex];
 	
 	FMeshDescription* MeshDescription = StaticMesh->GetMeshDescription(LODIndex);
@@ -391,7 +392,7 @@ bool UnFbx::FFbxImporter::BuildStaticMeshFromGeometry(FbxNode* Node, UStaticMesh
 	if (!Mesh->IsTriangleMesh())
 	{
 		if(!GIsAutomationTesting)
-		UE_LOG(LogFbx, Display, TEXT("Triangulating static mesh %s"), UTF8_TO_TCHAR(Node->GetName()));
+		UE_LOG(LogFbx, Display, TEXT("Triangulating static mesh %s"), *FbxNodeName);
 
 		const bool bReplace = true;
 		FbxNodeAttribute* ConvertedNode = GeometryConverter->Triangulate(Mesh, bReplace);
@@ -910,7 +911,8 @@ bool UnFbx::FFbxImporter::BuildStaticMeshFromGeometry(FbxNode* Node, UStaticMesh
 	{
 		FFormatNamedArguments Arguments;
 		Arguments.Add( TEXT("MeshName"), FText::FromString(StaticMesh->GetName()));
-		FText ErrorMsg = LOCTEXT("MeshHasNoRenderableTriangles", "{MeshName} could not be created because all of its polygons are degenerate.");
+		Arguments.Add( TEXT("PartName"), FText::FromString(FbxNodeName));
+		FText ErrorMsg = LOCTEXT("MeshHasNoRenderableTriangles", "Mesh name: [{MeshName}] part name: [{PartName}]  could not be created because all of its polygons are degenerate.");
 		AddTokenizedErrorMessage(FTokenizedMessage::Create(EMessageSeverity::Error, FText::Format(ErrorMsg, Arguments)), FFbxErrors::StaticMesh_AllTrianglesDegenerate);
 	}
 
