Look at the imported mesh bounding box before importing it in case it is too small

diff --git a/Engine/Source/Editor/UnrealEd/Private/Fbx/FbxStaticMeshImport.cpp b/Engine/Source/Editor/UnrealEd/Private/Fbx/FbxStaticMeshImport.cpp
--- a/Engine/Source/Editor/UnrealEd/Private/Fbx/FbxStaticMeshImport.cpp
+++ b/Engine/Source/Editor/UnrealEd/Private/Fbx/FbxStaticMeshImport.cpp
@@ -69,6 +69,14 @@ static FbxString GetNodeNameWithoutNamespace( FbxNode* Node )
 	}
 }
 
+void CreateTokenizedErrorForDegeneratedPart(UnFbx::FFbxImporter* FbxImporter, const FString MeshName, const FString NodeName)
+{
+	FFormatNamedArguments Arguments;
+	Arguments.Add(TEXT("MeshName"), FText::FromString(MeshName));
+	Arguments.Add(TEXT("PartName"), FText::FromString(NodeName));
+	FText ErrorMsg = LOCTEXT("MeshHasNoRenderableTriangles", "Mesh name: [{MeshName}] part name: [{PartName}]  could not be created because all of its polygons are degenerate.");
+	FbxImporter->AddTokenizedErrorMessage(FTokenizedMessage::Create(EMessageSeverity::Error, FText::Format(ErrorMsg, Arguments)), FFbxErrors::StaticMesh_AllTrianglesDegenerate);
+}
 
 //-------------------------------------------------------------------------
 //
@@ -272,7 +280,6 @@ bool UnFbx::FFbxImporter::BuildStaticMeshFromGeometry(FbxNode* Node, UStaticMesh
 	check(StaticMesh->SourceModels.IsValidIndex(LODIndex));
 	FbxMesh* Mesh = Node->GetMesh();
 	FString FbxNodeName = UTF8_TO_TCHAR(Node->GetName());
-	FStaticMeshSourceModel& SrcModel = StaticMesh->SourceModels[LODIndex];
 	
 	FMeshDescription* MeshDescription = StaticMesh->GetMeshDescription(LODIndex);
 	//The mesh description should have been created before calling BuildStaticMeshFromGeometry
@@ -909,11 +916,7 @@ bool UnFbx::FFbxImporter::BuildStaticMeshFromGeometry(FbxNode* Node, UStaticMesh
 	
 	if (!bHasNonDegeneratePolygons)
 	{
-		FFormatNamedArguments Arguments;
-		Arguments.Add( TEXT("MeshName"), FText::FromString(StaticMesh->GetName()));
-		Arguments.Add( TEXT("PartName"), FText::FromString(FbxNodeName));
-		FText ErrorMsg = LOCTEXT("MeshHasNoRenderableTriangles", "Mesh name: [{MeshName}] part name: [{PartName}]  could not be created because all of its polygons are degenerate.");
-		AddTokenizedErrorMessage(FTokenizedMessage::Create(EMessageSeverity::Error, FText::Format(ErrorMsg, Arguments)), FFbxErrors::StaticMesh_AllTrianglesDegenerate);
+		CreateTokenizedErrorForDegeneratedPart(this, StaticMesh->GetName(), FbxNodeName);
 	}
 
 	bool bIsValidMesh = bHasNonDegeneratePolygons;
@@ -1533,7 +1536,9 @@ UStaticMesh* UnFbx::FFbxImporter::ImportStaticMeshAsSingle(UObject* InParent, TA
 	StaticMesh->LightMapResolution = 64;
 	StaticMesh->LightMapCoordinateIndex = 1;
 
+	float SqrBoundingBoxThreshold = THRESH_POINTS_ARE_NEAR * THRESH_POINTS_ARE_NEAR;
 
+	bool bAllDegenerated = true;
 	TArray<FFbxMaterial> MeshMaterials;
 	for (MeshIndex = 0; MeshIndex < MeshNodeArray.Num(); MeshIndex++ )
 	{
@@ -1541,14 +1546,33 @@ UStaticMesh* UnFbx::FFbxImporter::ImportStaticMeshAsSingle(UObject* InParent, TA
 
 		if (Node->GetMesh())
 		{
-			if (!BuildStaticMeshFromGeometry(Node, StaticMesh, MeshMaterials, LODIndex,
-											 VertexColorImportOption, ExistingVertexColorData, ImportOptions->VertexOverrideColor))
+			Node->GetMesh()->ComputeBBox();
+			FbxDouble3 BoxExtend;
+			BoxExtend[0] = Node->GetMesh()->BBoxMax.Get()[0] - Node->GetMesh()->BBoxMin.Get()[0];
+			BoxExtend[1] = Node->GetMesh()->BBoxMax.Get()[1] - Node->GetMesh()->BBoxMin.Get()[1];
+			BoxExtend[2] = Node->GetMesh()->BBoxMax.Get()[2] - Node->GetMesh()->BBoxMin.Get()[2];
+			double SqrSize = (BoxExtend[0] * BoxExtend[0]) + (BoxExtend[1] * BoxExtend[1]) + (BoxExtend[2] * BoxExtend[2]);
+			//If the bounding box of the mesh part is smaller then the position threshold, the part is consider degenerated and will be skip.
+			if (SqrSize > SqrBoundingBoxThreshold)
+			{
+				bAllDegenerated = false;
+				if (!BuildStaticMeshFromGeometry(Node, StaticMesh, MeshMaterials, LODIndex,
+					VertexColorImportOption, ExistingVertexColorData, ImportOptions->VertexOverrideColor))
+				{
+					bBuildStatus = false;
+					break;
+				}
+			}
+			else
 			{
-				bBuildStatus = false;
-				break;
+				CreateTokenizedErrorForDegeneratedPart(this, StaticMesh->GetName(), UTF8_TO_TCHAR(Node->GetName()));
 			}
 		}
 	}
+	if (bAllDegenerated)
+	{
+		bBuildStatus = false;
+	}
 
 	if (bBuildStatus)
 	{
