When a skeletalmesh LOD get imported and the skeleton is not compatible which make the import fail, prevent adding material slot.

diff --git a/Engine/Source/Editor/UnrealEd/Private/FbxMeshUtils.cpp b/Engine/Source/Editor/UnrealEd/Private/FbxMeshUtils.cpp
--- a/Engine/Source/Editor/UnrealEd/Private/FbxMeshUtils.cpp
+++ b/Engine/Source/Editor/UnrealEd/Private/FbxMeshUtils.cpp
@@ -255,6 +255,23 @@ namespace FbxMeshUtils
 		// Check the file extension for FBX. Anything that isn't .FBX is rejected
 		const FString FileExtension = FPaths::GetExtension(Filename);
 		const bool bIsFBX = FCString::Stricmp(*FileExtension, TEXT("FBX")) == 0;
+		bool bSceneIsCleanUp = false;
+		TArray< TArray<FbxNode*>* > MeshArray;
+		auto CleanUpScene = [&bSceneIsCleanUp, &MeshArray, &FFbxImporter]()
+		{
+			if (bSceneIsCleanUp)
+			{
+				return;
+			}
+			bSceneIsCleanUp = true;
+			// Cleanup
+			for (int32 i = 0; i < MeshArray.Num(); i++)
+			{
+				delete MeshArray[i];
+			}
+			FFbxImporter->ReleaseScene();
+			FFbxImporter = nullptr;
+		};
 
 		//Skip none fbx file
 		if (!bIsFBX)
@@ -298,6 +315,26 @@ namespace FbxMeshUtils
 			ClothingAsset->UnbindFromSkeletalMesh(SelectedSkelMesh, LODLevel);
 		}
 
+		//Lambda to call to re-apply the clothing
+		auto ReapplyClothing = [&SelectedSkelMesh, &ClothingAssetsInUse, &ClothingAssetSectionIndices, &ClothingAssetInternalLodIndices, &ImportedResource, &LODLevel]()
+		{
+			// Re-apply our clothing assets
+			int32 NumClothingAssetsToApply = ClothingAssetsInUse.Num();
+			if (ImportedResource && ImportedResource->LODModels.IsValidIndex(LODLevel))
+			{
+				FSkeletalMeshLODModel& LodModel = ImportedResource->LODModels[LODLevel];
+				for (int32 AssetIndex = 0; AssetIndex < NumClothingAssetsToApply; ++AssetIndex)
+				{
+					// Only if the same section exists
+					if (LodModel.Sections.IsValidIndex(ClothingAssetSectionIndices[AssetIndex]))
+					{
+						UClothingAssetBase* AssetToApply = ClothingAssetsInUse[AssetIndex];
+						AssetToApply->BindToSkeletalMesh(SelectedSkelMesh, LODLevel, ClothingAssetSectionIndices[AssetIndex], ClothingAssetInternalLodIndices[AssetIndex]);
+					}
+				}
+			}
+		};
+
 		// don't import material and animation
 		UnFbx::FBXImportOptions* ImportOptions = FFbxImporter->GetImportOptions();
 			
@@ -346,6 +383,7 @@ namespace FbxMeshUtils
 
 		if ( !FFbxImporter->ImportFromFile( *Filename, FPaths::GetExtension( Filename ), true ) )
 		{
+			ReapplyClothing();
 			// Log the error message and fail the import.
 			FFbxImporter->AddTokenizedErrorMessage(FTokenizedMessage::Create(EMessageSeverity::Error, LOCTEXT("FBXImport_ParseFailed", "FBX file parsing failed.")), FFbxErrors::Generic_FBXFileParseFailed);
 		}
@@ -353,7 +391,6 @@ namespace FbxMeshUtils
 		{
 			bool bUseLODs = true;
 			int32 MaxLODLevel = 0;
-			TArray< TArray<FbxNode*>* > MeshArray;
 			TArray<FString> LODStrings;
 			TArray<FbxNode*>* MeshObject = NULL;;
 
@@ -363,8 +400,9 @@ namespace FbxMeshUtils
 			// Nothing found, error out
 			if (MeshArray.Num() == 0)
 			{
+				ReapplyClothing();
 				FFbxImporter->AddTokenizedErrorMessage(FTokenizedMessage::Create(EMessageSeverity::Error, LOCTEXT("FBXImport_NoMesh", "No meshes were found in file.")), FFbxErrors::Generic_MeshNotFound);
-				FFbxImporter->ReleaseScene();
+				CleanUpScene();
 				return false;
 			}
 
@@ -403,6 +441,7 @@ namespace FbxMeshUtils
 			int32 SelectedLOD = LODLevel;
 			if (SelectedLOD > SelectedSkelMesh->GetLODNum())
 			{
+				ReapplyClothing();
 				// Make sure they don't manage to select a bad LOD index
 				FFbxImporter->AddTokenizedErrorMessage(FTokenizedMessage::Create(EMessageSeverity::Warning, FText::Format(LOCTEXT("FBXImport_InvalidLODIdx", "Invalid mesh LOD index {0}, no prior LOD index exists"), FText::AsNumber(SelectedLOD))), FFbxErrors::Generic_Mesh_LOD_InvalidIndex);
 			}
@@ -458,7 +497,7 @@ namespace FbxMeshUtils
 				}
 					
 				ExistingSkelMeshData* SkelMeshDataPtr = nullptr;
-				if (SelectedSkelMesh->GetLODNum() > LODLevel)
+				if (SelectedSkelMesh->GetLODNum() > SelectedLOD)
 				{
 					SelectedSkelMesh->PreEditChange(NULL);
 					SkelMeshDataPtr = SaveExistingSkelMeshData(SelectedSkelMesh, true, SelectedLOD);
@@ -476,61 +515,52 @@ namespace FbxMeshUtils
 				ImportSkeletalMeshArgs.Name = NAME_None;
 				ImportSkeletalMeshArgs.Flags = RF_Transient;
 				ImportSkeletalMeshArgs.TemplateImportData = TempAssetImportData;
-				ImportSkeletalMeshArgs.LodIndex = LODLevel;
+				ImportSkeletalMeshArgs.LodIndex = SelectedLOD;
 				ImportSkeletalMeshArgs.OrderedMaterialNames = OrderedMaterialNames.Num() > 0 ? &OrderedMaterialNames : nullptr;
 				ImportSkeletalMeshArgs.ImportMaterialOriginalNameData = &ImportMaterialOriginalNameData;
 				ImportSkeletalMeshArgs.ImportMeshSectionsData = &ImportMeshLodData[0];
 				ImportSkeletalMeshArgs.OutData = &OutData;
 
 				TempSkelMesh = (USkeletalMesh*)FFbxImporter->ImportSkeletalMesh( ImportSkeletalMeshArgs );
-
-				// Add imported mesh to existing model
-				bool bMeshImportSuccess = false;
-				if( TempSkelMesh )
+				// Add the new imported LOD to the existing model (check skeleton compatibility)
+				if( TempSkelMesh  && FFbxImporter->ImportSkeletalMeshLOD(TempSkelMesh, SelectedSkelMesh, SelectedLOD, true, nullptr, TempAssetImportData))
 				{
-					bMeshImportSuccess = FFbxImporter->ImportSkeletalMeshLOD(TempSkelMesh, SelectedSkelMesh, SelectedLOD, true, nullptr, TempAssetImportData);
-
+					TComponentReregisterContext<USkinnedMeshComponent> ReregisterContext;
+					SelectedSkelMesh->ReleaseResources();
+					SelectedSkelMesh->ReleaseResourcesFence.Wait();
+					
 					//Update the import data for this lod
-					UnFbx::FFbxImporter::UpdateSkeletalMeshImportData(SelectedSkelMesh, nullptr, LODLevel, &ImportMaterialOriginalNameData, &ImportMeshLodData);
+					UnFbx::FFbxImporter::UpdateSkeletalMeshImportData(SelectedSkelMesh, nullptr, SelectedLOD, &ImportMaterialOriginalNameData, &ImportMeshLodData);
 
 					if (SkelMeshDataPtr != nullptr)
 					{
 						RestoreExistingSkelMeshData(SkelMeshDataPtr, SelectedSkelMesh, SelectedLOD, false, ImportOptions->bImportAsSkeletalSkinning);
 					}
-					SelectedSkelMesh->PostEditChange();
-					// Mark package containing skeletal mesh as dirty.
-					SelectedSkelMesh->MarkPackageDirty();
 
-					// Now iterate over all skeletal mesh components re-initialising them.
-					for (TObjectIterator<USkeletalMeshComponent> It; It; ++It)
+					if (ImportOptions->bImportMorph)
 					{
-						USkeletalMeshComponent* SkelComp = *It;
-						if (SkelComp->SkeletalMesh == SelectedSkelMesh)
-						{
-							FComponentReregisterContext ReregisterContext(SkelComp);
-						}
+						FFbxImporter->ImportFbxMorphTarget(SkelMeshNodeArray, SelectedSkelMesh, SelectedSkelMesh->GetOutermost(), SelectedLOD, OutData);
 					}
-				}
 
-				if(ImportOptions->bImportMorph)
-				{
-					FFbxImporter->ImportFbxMorphTarget(SkelMeshNodeArray, SelectedSkelMesh, SelectedSkelMesh->GetOutermost(), SelectedLOD, OutData);
-					//If we have import some morph target we have to rebuild the render resources since morph target are now using GPU
-					if (SelectedSkelMesh->MorphTargets.Num() > 0)
-					{
-						SelectedSkelMesh->ReleaseResources();
-						//Rebuild the resources with a post edit change since we have added some morph targets
-						SelectedSkelMesh->PostEditChange();
-					}
-				}
-
-				if (bMeshImportSuccess)
-				{
 					bSuccess = true;
 					// Set LOD source filename
 					SelectedSkelMesh->GetLODInfo(SelectedLOD)->SourceImportFilename = UAssetImportData::SanitizeImportFilename(Filename, nullptr);
 					SelectedSkelMesh->GetLODInfo(SelectedLOD)->bImportWithBaseMesh = false;
 
+					ReapplyClothing();
+
+					//Must be the last step because it cleanup the fbx importer to import the alternate skinning FBX
+					if (bMustReimportAlternateSkinWeightProfile)
+					{
+						//We cannot use anymore the FFbxImporter after the cleanup
+						CleanUpScene();
+						FLODUtilities::ReimportAlternateSkinWeight(SelectedSkelMesh, SelectedLOD, false);
+					}
+
+					SelectedSkelMesh->PostEditChange();
+					SelectedSkelMesh->InitResources();
+					SelectedSkelMesh->MarkPackageDirty();
+
 					// Notification of success
 					FNotificationInfo NotificationInfo(FText::GetEmpty());
 					NotificationInfo.Text = FText::Format(NSLOCTEXT("UnrealEd", "LODImportSuccessful", "Mesh for LOD {0} imported successfully!"), FText::AsNumber(SelectedLOD));
@@ -539,6 +569,7 @@ namespace FbxMeshUtils
 				}
 				else
 				{
+					ReapplyClothing();
 					// Notification of failure
 					FNotificationInfo NotificationInfo(FText::GetEmpty());
 					NotificationInfo.Text = FText::Format(NSLOCTEXT("UnrealEd", "LODImportFail", "Failed to import mesh for LOD {0}!"), FText::AsNumber(SelectedLOD));
@@ -546,36 +577,8 @@ namespace FbxMeshUtils
 					FSlateNotificationManager::Get().AddNotification(NotificationInfo);
 				}
 			}
-
-			// Cleanup
-			for (int32 i=0; i<MeshArray.Num(); i++)
-			{
-				delete MeshArray[i];
-			}					
 		}
-		FFbxImporter->ReleaseScene();
-
-		// Re-apply our clothing assets
-		int32 NumClothingAssetsToApply = ClothingAssetsInUse.Num();
-		if(ImportedResource && ImportedResource->LODModels.IsValidIndex(LODLevel))
-		{
-			FSkeletalMeshLODModel& LodModel = ImportedResource->LODModels[LODLevel];
-			for(int32 AssetIndex = 0; AssetIndex < NumClothingAssetsToApply; ++AssetIndex)
-			{
-				// Only if the same section exists
-				if(LodModel.Sections.IsValidIndex(ClothingAssetSectionIndices[AssetIndex]))
-				{
-					UClothingAssetBase* AssetToApply = ClothingAssetsInUse[AssetIndex];
-					AssetToApply->BindToSkeletalMesh(SelectedSkelMesh, LODLevel, ClothingAssetSectionIndices[AssetIndex], ClothingAssetInternalLodIndices[AssetIndex]);
-				}
-			}
-		}
-
-		if (bMustReimportAlternateSkinWeightProfile)
-		{
-			FLODUtilities::ReimportAlternateSkinWeight(SelectedSkelMesh, LODLevel, true);
-		}
-
+		CleanUpScene();
 		return bSuccess;
 	}
 
