PR #5865: Fix recursive surface snapping by ignoring attached children actors in raycast to world (Contributed by [REDACTED_PERSON])

diff --git a/Engine/Source/Editor/UnrealEd/Private/LevelEditorViewport.cpp b/Engine/Source/Editor/UnrealEd/Private/LevelEditorViewport.cpp
--- a/Engine/Source/Editor/UnrealEd/Private/LevelEditorViewport.cpp
+++ b/Engine/Source/Editor/UnrealEd/Private/LevelEditorViewport.cpp
@@ -2306,6 +2306,17 @@ void TrimLineToFrustum(const FConvexVolume& Frustum, FVector& Start, FVector& En
 	}
 }
 
+static void GetAttachedActorsRecursive(const AActor* InActor, TArray<AActor*>& OutActors)
+{
+	TArray<AActor*> AttachedActors;
+	InActor->GetAttachedActors(AttachedActors);
+	for (AActor* AttachedActor : AttachedActors)
+	{
+		GetAttachedActorsRecursive(AttachedActor, OutActors);
+	}
+	OutActors.Append(AttachedActors);
+};
+
 void FLevelEditorViewportClient::ProjectActorsIntoWorld(const TArray<AActor*>& Actors, FViewport* InViewport, const FVector& Drag, const FRotator& Rot)
 {
 	// Compile an array of selected actors
@@ -2361,6 +2372,10 @@ void FLevelEditorViewportClient::ProjectActorsIntoWorld(const TArray<AActor*>& A
 
 		if (bIsOnScreen)
 		{
+			TArray<AActor*> IgnoreActors; 
+			IgnoreActors.Append(Actors);  // Add the whole list of actors so you can't hit the moving set with the ray
+			GetAttachedActorsRecursive(Actor, IgnoreActors);
+
 			// Determine how we're going to attempt to project the object onto the world
 			if (CurrentAxis == EAxisList::XY || CurrentAxis == EAxisList::XZ || CurrentAxis == EAxisList::YZ)
 			{
@@ -2379,11 +2394,11 @@ void FLevelEditorViewportClient::ProjectActorsIntoWorld(const TArray<AActor*>& A
 
 				TrimLineToFrustum(Frustum, RayStart, RayEnd);
 
-				TraceResult = FActorPositioning::TraceWorldForPosition(*GetWorld(), *SceneView, RayStart, RayEnd, &Actors);
+				TraceResult = FActorPositioning::TraceWorldForPosition(*GetWorld(), *SceneView, RayStart, RayEnd, &IgnoreActors);
 			}
 			else
 			{
-				TraceResult = FActorPositioning::TraceWorldForPosition(Cursor, *SceneView, &Actors);
+				TraceResult = FActorPositioning::TraceWorldForPosition(Cursor, *SceneView, &IgnoreActors);
 			}
 		}
 				
