Submitting fix for issue blocking Mac builds.

diff --git a/Engine/Source/Editor/UnrealEd/Public/SColorGradientEditor.h b/Engine/Source/Editor/UnrealEd/Public/SColorGradientEditor.h
--- a/Engine/Source/Editor/UnrealEd/Public/SColorGradientEditor.h
+++ b/Engine/Source/Editor/UnrealEd/Public/SColorGradientEditor.h
@@ -17,7 +17,7 @@ class FCurveOwnerInterface;
 class FPaintArgs;
 class FSlateWindowElementList;
 
-struct FGradientStopMark
+struct UNREALED_API FGradientStopMark
 {
 public:
 
