Updated or added multiple copyright boilerplates

diff --git a/Engine/Source/Editor/ClassViewer/Private/ClassViewerFilter.cpp b/Engine/Source/Editor/ClassViewer/Private/ClassViewerFilter.cpp
--- a/Engine/Source/Editor/ClassViewer/Private/ClassViewerFilter.cpp
+++ b/Engine/Source/Editor/ClassViewer/Private/ClassViewerFilter.cpp
@@ -1,3 +1,5 @@
+// Copyright 1998-2019 Epic Games, Inc. All Rights Reserved.
+
 #include "ClassViewerFilter.h"
 
 #include "UnloadedBlueprintData.h"
diff --git a/Engine/Source/Editor/UnrealEd/Classes/Preferences/AnimationBlueprintEditorOptions.h b/Engine/Source/Editor/UnrealEd/Classes/Preferences/AnimationBlueprintEditorOptions.h
--- a/Engine/Source/Editor/UnrealEd/Classes/Preferences/AnimationBlueprintEditorOptions.h
+++ b/Engine/Source/Editor/UnrealEd/Classes/Preferences/AnimationBlueprintEditorOptions.h
@@ -1,3 +1,5 @@
+// Copyright 1998-2019 Epic Games, Inc. All Rights Reserved.
+
 /**
  * A configuration class used by the UAnimationBlueprint Editor to save editor
  * settings across sessions.
diff --git a/Engine/Source/Editor/UnrealEd/Classes/Preferences/BlueprintEditorOptions.h b/Engine/Source/Editor/UnrealEd/Classes/Preferences/BlueprintEditorOptions.h
--- a/Engine/Source/Editor/UnrealEd/Classes/Preferences/BlueprintEditorOptions.h
+++ b/Engine/Source/Editor/UnrealEd/Classes/Preferences/BlueprintEditorOptions.h
@@ -1,3 +1,5 @@
+// Copyright 1998-2019 Epic Games, Inc. All Rights Reserved.
+
 /**
  * A configuration class used by the UBlueprint Editor to save editor
  * settings across sessions.
