Fixed issue where actors spawned by child actor components could be selected, updated to handle cases of nested child actor components

diff --git a/Engine/Source/Editor/UnrealEd/Private/EditorSelectUtils.cpp b/Engine/Source/Editor/UnrealEd/Private/EditorSelectUtils.cpp
--- a/Engine/Source/Editor/UnrealEd/Private/EditorSelectUtils.cpp
+++ b/Engine/Source/Editor/UnrealEd/Private/EditorSelectUtils.cpp
@@ -561,6 +561,15 @@ void UUnrealEdEngine::SelectActor(AActor* Actor, bool bInSelected, bool bNotify,
 	// Select the actor and update its internals.
 	if( !bSelectionHandled )
 	{
+		if (bInSelected)
+		{
+			// If trying to select an Actor spawned by a ChildActorComponent, instead iterate up the hierarchy until we find a valid actor to select
+			while (Actor->IsChildActor())
+			{
+				Actor = Actor->GetParentComponent()->GetOwner();
+			}
+		}
+
 		if (UActorGroupingUtils::IsGroupingActive())
 		{
 			// if this actor is a group, do a group select/deselect
