Fixed misleading comment

#rnx

diff --git a/Engine/Source/Editor/UnrealEd/Public/AssetImportTask.h b/Engine/Source/Editor/UnrealEd/Public/AssetImportTask.h
--- a/Engine/Source/Editor/UnrealEd/Public/AssetImportTask.h
+++ b/Engine/Source/Editor/UnrealEd/Public/AssetImportTask.h
@@ -27,7 +27,7 @@ public:
 	UPROPERTY(BlueprintReadWrite, Category = "Asset Import Task")
 	FString Filename;
 
-	/** Content path in the projects content directory where asset will be imported */
+	/** Path where asset will be imported to */
 	UPROPERTY(BlueprintReadWrite, Category = "Asset Import Task")
 	FString DestinationPath;
 
