Fixing build failure due to missed file.

diff --git a/Engine/Source/Editor/UnrealEd/Private/PropertyEditorTestObject.cpp b/Engine/Source/Editor/UnrealEd/Private/PropertyEditorTestObject.cpp
--- a/Engine/Source/Editor/UnrealEd/Private/PropertyEditorTestObject.cpp
+++ b/Engine/Source/Editor/UnrealEd/Private/PropertyEditorTestObject.cpp
@@ -1,8 +1,3 @@
 // Copyright 1998-2019 Epic Games, Inc. All Rights Reserved.
 
-#include "Editor/PropertyEditorTestObject.h"
-
-UPropertyEditorTestObject::UPropertyEditorTestObject(const FObjectInitializer& ObjectInitializer)
-	: Super(ObjectInitializer)
-{
-}
+#include "Editor/PropertyEditorTestObject.h"
\ No newline at end of file
