Piloting camera goes 180 to -180 fix.

diff --git a/Engine/Source/Editor/UnrealEd/Private/LevelEditorViewport.cpp b/Engine/Source/Editor/UnrealEd/Private/LevelEditorViewport.cpp
--- a/Engine/Source/Editor/UnrealEd/Private/LevelEditorViewport.cpp
+++ b/Engine/Source/Editor/UnrealEd/Private/LevelEditorViewport.cpp
@@ -3363,6 +3363,13 @@ void FLevelEditorViewportClient::MoveLockedActorToCamera()
 			// Need to disable orbit camera before setting actor position so that the viewport camera location is converted back
 			GCurrentLevelEditingViewportClient->ToggleOrbitCamera(false);
 
+			USceneComponent* ActiveActorLockComponent = ActiveActorLock->GetRootComponent();
+			TOptional<FRotator> PreviousRotator;
+			if (ActiveActorLockComponent)
+			{
+				PreviousRotator = ActiveActorLockComponent->RelativeRotation;
+			}
+
 			// If we're locked to a camera then we're reflecting the camera view and not the actor position. We need to reflect that delta when we reposition the piloted actor
 			if (bUseControllingActorViewInfo)
 			{
@@ -3371,7 +3378,6 @@ void FLevelEditorViewportClient::MoveLockedActorToCamera()
 				{
 					const FTransform RelativeTransform = ViewComponent->GetComponentTransform().Inverse();
 					const FTransform DesiredTransform = FTransform(GCurrentLevelEditingViewportClient->GetViewRotation(), GCurrentLevelEditingViewportClient->GetViewLocation());
-
 					ActiveActorLock->SetActorTransform(ActiveActorLock->GetActorTransform() * RelativeTransform * DesiredTransform);
 				}
 			}
@@ -3380,6 +3386,20 @@ void FLevelEditorViewportClient::MoveLockedActorToCamera()
 				ActiveActorLock->SetActorLocation(GCurrentLevelEditingViewportClient->GetViewLocation(), false);
 				ActiveActorLock->SetActorRotation(GCurrentLevelEditingViewportClient->GetViewRotation());
 			}
+
+			if (PreviousRotator.IsSet())
+			{
+				const FRotator Rot = PreviousRotator.GetValue();
+				FRotator ActorRotWind, ActorRotRem;
+				Rot.GetWindingAndRemainder(ActorRotWind, ActorRotRem);
+				const FQuat ActorQ = ActorRotRem.Quaternion();
+				const FQuat ResultQ = ActiveActorLockComponent->RelativeRotation.Quaternion();
+				FRotator NewActorRotRem = FRotator(ResultQ);
+				ActorRotRem.SetClosestToMe(NewActorRotRem);
+				FRotator DeltaRot = NewActorRotRem - ActorRotRem;
+				DeltaRot.Normalize();
+				ActiveActorLockComponent->SetRelativeRotationExact(Rot + DeltaRot);
+			}
 		}
 
 		if (ABrush* Brush = Cast<ABrush>(ActiveActorLock))
