Changes so that component visualizers work better with components of ChildActorComponent actors.
UActorComponent* pointers are no longer cached anywhere; instead a class FComponentPropertyPath holds a list of properties to traverse from the topmost parent owner actor. This is used to uniquely identify the component being visualized/edited, both in custom component visualizers and in the editor rendering code.

Fixed a regression where archetype spline editing inside the BP editor window was no longer working.

diff --git a/Engine/Source/Editor/ComponentVisualizers/Private/SplineComponentVisualizer.cpp b/Engine/Source/Editor/ComponentVisualizers/Private/SplineComponentVisualizer.cpp
--- a/Engine/Source/Editor/ComponentVisualizers/Private/SplineComponentVisualizer.cpp
+++ b/Engine/Source/Editor/ComponentVisualizers/Private/SplineComponentVisualizer.cpp
@@ -411,13 +411,13 @@ bool FSplineComponentVisualizer::VisProxyHandleClick(FEditorViewportClient* InVi
 	{
 		const USplineComponent* SplineComp = CastChecked<const USplineComponent>(VisProxy->Component.Get());
 
-		SplineCompPropName = GetComponentPropertyName(SplineComp);
-		if(SplineCompPropName.IsValid())
-		{
-			AActor* OldSplineOwningActor = SplineOwningActor.Get();
-			SplineOwningActor = SplineComp->GetOwner();
+		AActor* OldSplineOwningActor = SplinePropertyPath.GetParentOwningActor();
+		SplinePropertyPath = FComponentPropertyPath(SplineComp);
+		AActor* NewSplineOwningActor = SplinePropertyPath.GetParentOwningActor();
 
-			if (OldSplineOwningActor != SplineOwningActor)
+		if (SplinePropertyPath.IsValid())
+		{
+			if (OldSplineOwningActor != NewSplineOwningActor)
 			{
 				// Reset selection state if we are selecting a different actor to the one previously selected
 				ChangeSelectionState(INDEX_NONE, false);
@@ -443,7 +443,7 @@ bool FSplineComponentVisualizer::VisProxyHandleClick(FEditorViewportClient* InVi
 
 				if (LastKeyIndexSelected == INDEX_NONE)
 				{
-					SplineOwningActor = nullptr;
+					SplinePropertyPath.Reset();
 					return false;
 				}
 
@@ -467,7 +467,7 @@ bool FSplineComponentVisualizer::VisProxyHandleClick(FEditorViewportClient* InVi
 
 				if (LastKeyIndexSelected == INDEX_NONE)
 				{
-					SplineOwningActor = nullptr;
+					SplinePropertyPath.Reset();
 					return false;
 				}
 
@@ -521,7 +521,7 @@ bool FSplineComponentVisualizer::VisProxyHandleClick(FEditorViewportClient* InVi
 		}
 		else
 		{
-			SplineOwningActor = nullptr;
+			SplinePropertyPath.Reset();
 		}
 	}
 
@@ -530,7 +530,7 @@ bool FSplineComponentVisualizer::VisProxyHandleClick(FEditorViewportClient* InVi
 
 USplineComponent* FSplineComponentVisualizer::GetEditedSplineComponent() const
 {
-	return Cast<USplineComponent>(GetComponentFromPropertyName(SplineOwningActor.Get(), SplineCompPropName));
+	return Cast<USplineComponent>(SplinePropertyPath.GetComponent());
 }
 
 
@@ -766,8 +766,7 @@ bool FSplineComponentVisualizer::HandleInputKey(FEditorViewportClient* ViewportC
 
 void FSplineComponentVisualizer::EndEditing()
 {
-	SplineOwningActor = NULL;
-	SplineCompPropName.Clear();
+	SplinePropertyPath.Reset();
 	ChangeSelectionState(INDEX_NONE, false);
 	SelectedSegmentIndex = INDEX_NONE;
 	SelectedTangentHandle = INDEX_NONE;
@@ -1209,9 +1208,9 @@ void FSplineComponentVisualizer::OnResetToDefault()
 	const FScopedTransaction Transaction(LOCTEXT("ResetToDefault", "Reset to Default"));
 
 	SplineComp->Modify();
-	if (SplineOwningActor.IsValid())
+	if (AActor* Owner = SplineComp->GetOwner())
 	{
-		SplineOwningActor.Get()->Modify();
+		Owner->Modify();
 	}
 
 	SplineComp->bSplineHasBeenEdited = false;
@@ -1222,9 +1221,9 @@ void FSplineComponentVisualizer::OnResetToDefault()
 	SelectedTangentHandle = INDEX_NONE;
 	SelectedTangentHandleType = ESelectedTangentHandle::None;
 
-	if (SplineOwningActor.IsValid())
+	if (AActor* Owner = SplineComp->GetOwner())
 	{
-		SplineOwningActor.Get()->PostEditMove(false);
+		Owner->PostEditMove(false);
 	}
 
 	GEditor->RedrawLevelEditingViewports(true);
diff --git a/Engine/Source/Editor/ComponentVisualizers/Private/SplineMeshComponentVisualizer.cpp b/Engine/Source/Editor/ComponentVisualizers/Private/SplineMeshComponentVisualizer.cpp
--- a/Engine/Source/Editor/ComponentVisualizers/Private/SplineMeshComponentVisualizer.cpp
+++ b/Engine/Source/Editor/ComponentVisualizers/Private/SplineMeshComponentVisualizer.cpp
@@ -114,10 +114,18 @@ bool FSplineMeshComponentVisualizer::VisProxyHandleClick(FEditorViewportClient*
 	{
 		const USplineMeshComponent* SplineMeshComp = CastChecked<const USplineMeshComponent>(VisProxy->Component.Get());
 
-		SplineMeshCompPropName = GetComponentPropertyName(SplineMeshComp);
-		if (SplineMeshCompPropName.IsValid())
+		AActor* OldSplineMeshOwningActor = SplineMeshPropertyPath.GetParentOwningActor();
+		SplineMeshPropertyPath = FComponentPropertyPath(SplineMeshComp);
+		AActor* NewSplineMeshOwningActor = SplineMeshPropertyPath.GetParentOwningActor();
+
+		if (SplineMeshPropertyPath.IsValid())
 		{
-			SplineMeshOwningActor = SplineMeshComp->GetOwner();
+			if (OldSplineMeshOwningActor != NewSplineMeshOwningActor)
+			{
+				SelectedKey = INDEX_NONE;
+				SelectedTangentHandle = INDEX_NONE;
+				SelectedTangentHandleType = ESelectedTangentHandle::None;
+			}
 
 			if (VisProxy->IsA(HSplineMeshKeyProxy::StaticGetType()))
 			{
@@ -146,7 +154,7 @@ bool FSplineMeshComponentVisualizer::VisProxyHandleClick(FEditorViewportClient*
 		}
 		else
 		{
-			SplineMeshOwningActor = NULL;
+			SplineMeshPropertyPath.Reset();
 		}
 	}
 
@@ -155,7 +163,7 @@ bool FSplineMeshComponentVisualizer::VisProxyHandleClick(FEditorViewportClient*
 
 USplineMeshComponent* FSplineMeshComponentVisualizer::GetEditedSplineMeshComponent() const
 {
-	return Cast<USplineMeshComponent>(GetComponentFromPropertyName(SplineMeshOwningActor.Get(), SplineMeshCompPropName));
+	return Cast<USplineMeshComponent>(SplineMeshPropertyPath.GetComponent());
 }
 
 
@@ -364,8 +372,7 @@ bool FSplineMeshComponentVisualizer::HandleInputKey(FEditorViewportClient* Viewp
 
 void FSplineMeshComponentVisualizer::EndEditing()
 {
-	SplineMeshOwningActor = NULL;
-	SplineMeshCompPropName.Clear();
+	SplineMeshPropertyPath.Reset();
 	SelectedKey = INDEX_NONE;
 	SelectedTangentHandle = INDEX_NONE;
 	SelectedTangentHandleType = ESelectedTangentHandle::None;
@@ -381,9 +388,9 @@ TSharedPtr<SWidget> FSplineMeshComponentVisualizer::GenerateContextMenu() const
 void FSplineMeshComponentVisualizer::NotifyComponentModified()
 {
 	// Notify of change so any CS is re-run
-	if (SplineMeshOwningActor.IsValid())
+	if (SplineMeshPropertyPath.IsValid())
 	{
-		SplineMeshOwningActor.Get()->PostEditMove(true);
+		SplineMeshPropertyPath.GetParentOwningActor()->PostEditMove(true);
 	}
 
 	GEditor->RedrawLevelEditingViewports(true);
diff --git a/Engine/Source/Editor/ComponentVisualizers/Public/SplineComponentVisualizer.h b/Engine/Source/Editor/ComponentVisualizers/Public/SplineComponentVisualizer.h
--- a/Engine/Source/Editor/ComponentVisualizers/Public/SplineComponentVisualizer.h
+++ b/Engine/Source/Editor/ComponentVisualizers/Public/SplineComponentVisualizer.h
@@ -135,11 +135,8 @@ protected:
 	/** Output log commands */
 	TSharedPtr<FUICommandList> SplineComponentVisualizerActions;
 
-	/** Actor that owns the currently edited spline */
-	TWeakObjectPtr<AActor> SplineOwningActor;
-
-	/** Name of property on the actor that references the spline we are editing */
-	FPropertyNameAndIndex SplineCompPropName;
+	/** Property path from the parent actor to the component */
+	FComponentPropertyPath SplinePropertyPath;
 
 	/** Index of keys we have selected */
 	TSet<int32> SelectedKeys;
diff --git a/Engine/Source/Editor/ComponentVisualizers/Public/SplineMeshComponentVisualizer.h b/Engine/Source/Editor/ComponentVisualizers/Public/SplineMeshComponentVisualizer.h
--- a/Engine/Source/Editor/ComponentVisualizers/Public/SplineMeshComponentVisualizer.h
+++ b/Engine/Source/Editor/ComponentVisualizers/Public/SplineMeshComponentVisualizer.h
@@ -84,11 +84,8 @@ protected:
 	/** Syncs changes made by the visualizer in the actual component */
 	void NotifyComponentModified();
 
-	/** Actor that owns the currently edited spline */
-	TWeakObjectPtr<AActor> SplineMeshOwningActor;
-
-	/** Name of property on the actor that references the spline we are editing */
-	FPropertyNameAndIndex SplineMeshCompPropName;
+	/** Property path from the parent actor to the component */
+	FComponentPropertyPath SplineMeshPropertyPath;
 
 	/** Index of the key we selected */
 	int32 SelectedKey;
diff --git a/Engine/Source/Editor/DetailCustomizations/Private/SplineComponentDetails.cpp b/Engine/Source/Editor/DetailCustomizations/Private/SplineComponentDetails.cpp
--- a/Engine/Source/Editor/DetailCustomizations/Private/SplineComponentDetails.cpp
+++ b/Engine/Source/Editor/DetailCustomizations/Private/SplineComponentDetails.cpp
@@ -200,7 +200,7 @@ FSplinePointDetails::FSplinePointDetails(USplineComponent* InOwningSplineCompone
 	: SplineComp(nullptr)
 {
 
-	TSharedPtr<FComponentVisualizer> Visualizer = GUnrealEd->FindComponentVisualizer(USplineComponent::StaticClass());
+	TSharedPtr<FComponentVisualizer> Visualizer = GUnrealEd->FindComponentVisualizer(InOwningSplineComponent->GetClass());
 	SplineVisualizer = (FSplineComponentVisualizer*)Visualizer.Get();
 	check(SplineVisualizer);
 
@@ -433,7 +433,8 @@ void FSplinePointDetails::GenerateChildContent(IDetailChildrenBuilder& ChildrenB
 			if (ClassIterator->IsChildOf(USplineMetadataDetailsFactoryBase::StaticClass()) && !ClassIterator->HasAnyClassFlags(CLASS_Abstract | CLASS_Deprecated | CLASS_NewerVersionExists))
 			{
 				USplineMetadataDetailsFactoryBase* Factory = ClassIterator->GetDefaultObject<USplineMetadataDetailsFactoryBase>();
-				if (Factory->GetMetadataClass() == SplineComp->GetSplinePointsMetadata()->GetClass())
+				USplineMetadata* SplineMetadata = SplineComp->GetSplinePointsMetadata();
+				if (SplineMetadata && SplineMetadata->GetClass() == Factory->GetMetadataClass())
 				{
 					SplineMetaDataDetails = Factory->Create();
 					IDetailGroup& Group = ChildrenBuilder.AddGroup(SplineMetaDataDetails->GetName(), SplineMetaDataDetails->GetDisplayName());
diff --git a/Engine/Source/Editor/UnrealEd/Classes/Editor/UnrealEdEngine.h b/Engine/Source/Editor/UnrealEd/Classes/Editor/UnrealEdEngine.h
--- a/Engine/Source/Editor/UnrealEd/Classes/Editor/UnrealEdEngine.h
+++ b/Engine/Source/Editor/UnrealEd/Classes/Editor/UnrealEdEngine.h
@@ -845,9 +845,6 @@ protected:
 	/** Called when global editor selection changes */
 	void OnEditorSelectionChanged(UObject* SelectionThatChanged);
 
-	/** Called when blueprint objects are replaced so that we can update the cached visualizer selection */
-	void ReplaceCachedVisualizerObjects(const TMap<UObject*, UObject*>& ReplacementMap);
-
 	EWriteDisallowedWarningState GetWarningStateForWritePermission(const FString& PackageName) const;
 	
 	/** The package auto-saver instance used by the editor */
diff --git a/Engine/Source/Editor/UnrealEd/Private/ComponentVisualizer.cpp b/Engine/Source/Editor/UnrealEd/Private/ComponentVisualizer.cpp
--- a/Engine/Source/Editor/UnrealEd/Private/ComponentVisualizer.cpp
+++ b/Engine/Source/Editor/UnrealEd/Private/ComponentVisualizer.cpp
@@ -7,50 +7,153 @@
 IMPLEMENT_HIT_PROXY(HComponentVisProxy, HHitProxy);
 
 
-FComponentVisualizer::FPropertyNameAndIndex FComponentVisualizer::GetComponentPropertyName(const UActorComponent* Component)
+static TTuple<FName, int32> GetActorPropertyNameAndIndexForComponent(const AActor* Actor, const UActorComponent* Component)
 {
-	if (Component)
+	if (Actor != nullptr && Component != nullptr)
 	{
-		const AActor* CompOwner = Component->GetOwner();
-		if (CompOwner)
+		// Iterate over UObject* fields of this actor
+		UClass* ActorClass = Actor->GetClass();
+		for (TFieldIterator<UObjectProperty> It(ActorClass); It; ++It)
 		{
-			// Iterate over UObject* fields of this actor
-			UClass* ActorClass = CompOwner->GetClass();
-			for (TFieldIterator<UObjectProperty> It(ActorClass); It; ++It)
+			// See if this property points to the component in question
+			UObjectProperty* ObjectProp = *It;
+			for (int32 Index = 0; Index < ObjectProp->ArrayDim; ++Index)
 			{
-				// See if this property points to the component in question
-				UObjectProperty* ObjectProp = *It;
-				for (int32 Index = 0; Index < ObjectProp->ArrayDim; ++Index)
+				UObject* Object = ObjectProp->GetObjectPropertyValue(ObjectProp->ContainerPtrToValuePtr<void>(Actor, Index));
+				if (Object == Component)
 				{
-					UObject* Object = ObjectProp->GetObjectPropertyValue(ObjectProp->ContainerPtrToValuePtr<void>(CompOwner, Index));
+					// It does! Return this name
+					return MakeTuple(ObjectProp->GetFName(), Index);
+				}
+			}
+		}
+
+		// If nothing found, look in TArray<UObject*> fields
+		for (TFieldIterator<UArrayProperty> It(ActorClass); It; ++It)
+		{
+			UArrayProperty* ArrayProp = *It;
+			if (UObjectProperty* InnerProp = Cast<UObjectProperty>(It->Inner))
+			{
+				FScriptArrayHelper ArrayHelper(ArrayProp, ArrayProp->ContainerPtrToValuePtr<void>(Actor));
+				for (int32 Index = 0; Index < ArrayHelper.Num(); ++Index)
+				{
+					UObject* Object = InnerProp->GetObjectPropertyValue(ArrayHelper.GetRawPtr(Index));
 					if (Object == Component)
 					{
-						// It does! Return this name
-						return FPropertyNameAndIndex(ObjectProp->GetFName(), Index);
+						return MakeTuple(ArrayProp->GetFName(), Index);
 					}
 				}
 			}
+		}
+	}
+
+	return MakeTuple<FName, int32>(NAME_None, INDEX_NONE);
+}
+
+
+void FComponentPropertyPath::Set(const UActorComponent* Component)
+{
+	AActor* Actor = Component->GetOwner();
+
+	UChildActorComponent* ParentComponent = Actor->GetParentComponent();
+	if (ParentComponent)
+	{
+		Set(ParentComponent);	// recurse to next parent
+	}
+	else
+	{
+		// If there are no further parents, store this one (the outermost)
+		ParentOwningActor = Actor;
+	}
+
+	// Add the next property to the chain after the recursion, so they are added outermost-first
+	PropertyChain.Add(GetActorPropertyNameAndIndexForComponent(Actor, Component));
+}
 
-			// If nothing found, look in TArray<UObject*> fields
-			for (TFieldIterator<UArrayProperty> It(ActorClass); It; ++It)
+
+UActorComponent* FComponentPropertyPath::GetComponent() const
+{
+	UActorComponent* Result = nullptr;
+	const AActor* Actor = ParentOwningActor.Get();
+	if (Actor)
+	{
+		int32 Level = 0;
+		for (TPair<FName, int32> PropertyNameAndIndex : PropertyChain)
+		{
+			FName PropertyName = PropertyNameAndIndex.Get<0>();
+			int32 PropertyIndex = PropertyNameAndIndex.Get<1>();
+			Result = nullptr;
+
+			if (PropertyName != NAME_None && PropertyIndex != INDEX_NONE)
 			{
-				UArrayProperty* ArrayProp = *It;
-				if (UObjectProperty* InnerProp = Cast<UObjectProperty>(It->Inner))
+				UClass* ActorClass = Actor->GetClass();
+				UProperty* Prop = FindField<UProperty>(ActorClass, PropertyName);
+
+				if (UObjectProperty* ObjectProp = Cast<UObjectProperty>(Prop))
+				{
+					UObject* Object = ObjectProp->GetObjectPropertyValue(ObjectProp->ContainerPtrToValuePtr<void>(Actor, PropertyIndex));
+					Result = Cast<UActorComponent>(Object);
+				}
+				else if (UArrayProperty* ArrayProp = Cast<UArrayProperty>(Prop))
 				{
-					FScriptArrayHelper ArrayHelper(ArrayProp, ArrayProp->ContainerPtrToValuePtr<void>(CompOwner));
-					for (int32 Index = 0; Index < ArrayHelper.Num(); ++Index)
+					if (UObjectProperty* InnerProp = Cast<UObjectProperty>(ArrayProp->Inner))
 					{
-						UObject* Object = InnerProp->GetObjectPropertyValue(ArrayHelper.GetRawPtr(Index));
-						if (Object == Component)
-						{
-							return FPropertyNameAndIndex(ArrayProp->GetFName(), Index);
-						}
+						FScriptArrayHelper ArrayHelper(ArrayProp, ArrayProp->ContainerPtrToValuePtr<void>(Actor));
+						UObject* Object = InnerProp->GetObjectPropertyValue(ArrayHelper.GetRawPtr(PropertyIndex));
+						Result = Cast<UActorComponent>(Object);
 					}
 				}
 			}
+
+			Level++;
+
+			if (Level < PropertyChain.Num())
+			{
+				UChildActorComponent* ChildActorComponent = Cast<UChildActorComponent>(Result);
+				if (Result == nullptr)
+				{
+					break;
+				}
+
+				Actor = ChildActorComponent->GetChildActor();
+			}
+		}
+	}
+
+	return Result;
+}
+
+
+bool FComponentPropertyPath::IsValid() const
+{
+	if (!ParentOwningActor.IsValid())
+	{
+		return false;
+	}
+
+	for (TPair<FName, int32> PropertyNameAndIndex : PropertyChain)
+	{
+		FName PropertyName = PropertyNameAndIndex.Get<0>();
+		int32 PropertyIndex = PropertyNameAndIndex.Get<1>();
+
+		if (PropertyName == NAME_None || PropertyIndex == INDEX_NONE)
+		{
+			return false;
 		}
 	}
 
+	return true;
+}
+
+
+FComponentVisualizer::FPropertyNameAndIndex FComponentVisualizer::GetComponentPropertyName(const UActorComponent* Component)
+{
+	if (Component)
+	{
+		TTuple<FName, int32> Result = GetActorPropertyNameAndIndexForComponent(Component->GetOwner(), Component);
+		return FPropertyNameAndIndex(Result.Get<0>(), Result.Get<1>());
+	}
+
 	// Didn't find actor property referencing this component
 	return FPropertyNameAndIndex();
 }
diff --git a/Engine/Source/Editor/UnrealEd/Private/UnrealEdEngine.cpp b/Engine/Source/Editor/UnrealEd/Private/UnrealEdEngine.cpp
--- a/Engine/Source/Editor/UnrealEd/Private/UnrealEdEngine.cpp
+++ b/Engine/Source/Editor/UnrealEd/Private/UnrealEdEngine.cpp
@@ -101,7 +101,6 @@ void UUnrealEdEngine::Init(IEngineLoop* InEngineLoop)
 	FEditorSupportDelegates::PostWindowsMessage.AddUObject(this, &UUnrealEdEngine::OnPostWindowsMessage);
 
 	USelection::SelectionChangedEvent.AddUObject(this, &UUnrealEdEngine::OnEditorSelectionChanged);
-	OnObjectsReplaced().AddUObject(this, &UUnrealEdEngine::ReplaceCachedVisualizerObjects);
 
 	// Initialize the snap manager
 	FSnappingUtils::InitEditorSnappingTools();
@@ -1336,7 +1335,7 @@ void UUnrealEdEngine::DrawComponentVisualizers(const FSceneView* View, FPrimitiv
 {
 	for(FCachedComponentVisualizer& CachedVisualizer : VisualizersForSelection)
 	{
-		CachedVisualizer.Visualizer->DrawVisualization(CachedVisualizer.Component.Get(), View, PDI);
+		CachedVisualizer.Visualizer->DrawVisualization(CachedVisualizer.ComponentPropertyPath.GetComponent(), View, PDI);
 	}
 }
 
@@ -1345,7 +1344,7 @@ void UUnrealEdEngine::DrawComponentVisualizersHUD(const FViewport* Viewport, con
 {
 	for(FCachedComponentVisualizer& CachedVisualizer : VisualizersForSelection)
 	{
-		CachedVisualizer.Visualizer->DrawVisualizationHUD(CachedVisualizer.Component.Get(), Viewport, View, Canvas);
+		CachedVisualizer.Visualizer->DrawVisualizationHUD(CachedVisualizer.ComponentPropertyPath.GetComponent(), Viewport, View, Canvas);
 	}
 }
 
@@ -1363,9 +1362,9 @@ void UUnrealEdEngine::OnEditorSelectionChanged(UObject* SelectionThatChanged)
 			AActor* Actor = Cast<AActor>(*It);
 			if(Actor != nullptr)
 			{
-				// Then iterate over components of that actor
+				// Then iterate over components of that actor (and recurse through child components)
 				TInlineComponentArray<UActorComponent*> Components;
-				Actor->GetComponents(Components);
+				Actor->GetComponents(Components, true);
 
 				for(int32 CompIdx = 0; CompIdx < Components.Num(); CompIdx++)
 				{
@@ -1385,18 +1384,6 @@ void UUnrealEdEngine::OnEditorSelectionChanged(UObject* SelectionThatChanged)
 	}
 }
 
-void UUnrealEdEngine::ReplaceCachedVisualizerObjects(const TMap<UObject*, UObject*>& ReplacementMap)
-{
-	for(FCachedComponentVisualizer& Visualizer : VisualizersForSelection)
-	{
-		UObject* OldObject = Visualizer.Component.Get(true);
-		UActorComponent* NewComponent = Cast<UActorComponent>(ReplacementMap.FindRef(OldObject));
-		if(NewComponent)
-		{
-			Visualizer.Component = NewComponent;
-		}
-	}
-}
 
 EWriteDisallowedWarningState UUnrealEdEngine::GetWarningStateForWritePermission(const FString& PackageName) const
 {
diff --git a/Engine/Source/Editor/UnrealEd/Public/ComponentVisualizer.h b/Engine/Source/Editor/UnrealEd/Public/ComponentVisualizer.h
--- a/Engine/Source/Editor/UnrealEd/Public/ComponentVisualizer.h
+++ b/Engine/Source/Editor/UnrealEd/Public/ComponentVisualizer.h
@@ -34,6 +34,43 @@ struct HComponentVisProxy : public HHitProxy
 };
 
 
+/**
+ * Describes a chain of properties from the parent actor of a given component, to the component itself.
+ */
+class UNREALED_API FComponentPropertyPath
+{
+public:
+
+	FComponentPropertyPath() = default;
+	explicit FComponentPropertyPath(const UActorComponent* Component) { Set(Component); }
+
+	/** Resets the property path */
+	void Reset()
+	{
+		ParentOwningActor = nullptr;
+		PropertyChain.Reset();
+	}
+
+	/** Gets the parent owning actor for the component, or nullptr if it is not valid */
+	AActor* GetParentOwningActor() const { return ParentOwningActor.Get(); }
+
+	/** Gets a pointer to the component, or nullptr if it is not valid */
+	UActorComponent* GetComponent() const;
+
+	/** Determines whether the property path is valid or not */
+	bool IsValid() const;
+
+private:
+
+	/** Sets the component referred to by the object */
+	void Set(const UActorComponent* Component);
+
+	TWeakObjectPtr<AActor> ParentOwningActor;
+	TArray<TTuple<FName, int32>> PropertyChain;
+};
+
+
+
 /** Base class for a component visualizer, that draw editor information for a particular component class */
 class UNREALED_API FComponentVisualizer : public TSharedFromThis<FComponentVisualizer>
 {
@@ -88,10 +125,13 @@ public:
 		int32 Index;
 	};
 
+
 	/** Find the name of the property that points to this component */
+	UE_DEPRECATED(4.25, "Please use the FComponentPropertyPath class to build property name paths for components.")
 	static FPropertyNameAndIndex GetComponentPropertyName(const UActorComponent* Component);
 
 	/** Get a component pointer from the property name */
+	UE_DEPRECATED(4.25, "Please use the FComponentPropertyPath::GetComponent() to retrieve a component pointer from a property name path.")
 	static UActorComponent* GetComponentFromPropertyName(const AActor* CompOwner, const FPropertyNameAndIndex& Property);
 
 	/** Notify that a component property has been modified */
@@ -103,11 +143,11 @@ public:
 
 struct FCachedComponentVisualizer
 {
-	TWeakObjectPtr<UActorComponent> Component;
+	FComponentPropertyPath ComponentPropertyPath;
 	TSharedPtr<FComponentVisualizer> Visualizer;
 	
 	FCachedComponentVisualizer(UActorComponent* InComponent, TSharedPtr<FComponentVisualizer>& InVisualizer)
-		: Component(InComponent)
+		: ComponentPropertyPath(InComponent)
 		, Visualizer(InVisualizer)
 	{}
 };
