Added AccessibleParams to SCustomDialog to avoid occasional crash.

diff --git a/Engine/Source/Editor/UnrealEd/Public/Dialogs/CustomDialog.h b/Engine/Source/Editor/UnrealEd/Public/Dialogs/CustomDialog.h
--- a/Engine/Source/Editor/UnrealEd/Public/Dialogs/CustomDialog.h
+++ b/Engine/Source/Editor/UnrealEd/Public/Dialogs/CustomDialog.h
@@ -36,6 +36,7 @@ public:
 		: _UseScrollBox(true)
 		, _ScrollBoxMaxHeight(300)
 	{
+		_AccessibleParams = FAccessibleWidgetData(EAccessibleBehavior::Auto);
 	}
 		/** Title to display for the dialog. */
 		SLATE_ARGUMENT(FText, Title)
