Missing headers
#rnx

diff --git a/Engine/Source/Editor/LevelEditor/Public/LevelEditorMenuContext.h b/Engine/Source/Editor/LevelEditor/Public/LevelEditorMenuContext.h
new file mode 100644
--- /dev/null
+++ b/Engine/Source/Editor/LevelEditor/Public/LevelEditorMenuContext.h
@@ -0,0 +1,19 @@
+// Copyright 1998-2019 Epic Games, Inc. All Rights Reserved.
+
+#pragma once
+
+#include "CoreMinimal.h"
+#include "UObject/ObjectMacros.h"
+#include "UObject/Object.h"
+
+#include "LevelEditorMenuContext.generated.h"
+
+class SLevelEditor;
+
+UCLASS()
+class ULevelEditorMenuContext : public UObject
+{
+	GENERATED_BODY()
+public:
+	TWeakPtr<SLevelEditor> SlateLevelEditor;
+};
diff --git a/Engine/Source/Editor/UnrealEd/Public/Toolkits/AssetEditorToolkitMenuContext.h b/Engine/Source/Editor/UnrealEd/Public/Toolkits/AssetEditorToolkitMenuContext.h
new file mode 100644
--- /dev/null
+++ b/Engine/Source/Editor/UnrealEd/Public/Toolkits/AssetEditorToolkitMenuContext.h
@@ -0,0 +1,20 @@
+// Copyright 1998-2019 Epic Games, Inc. All Rights Reserved.
+
+#pragma once
+
+#include "CoreMinimal.h"
+#include "UObject/GCObject.h"
+
+#include "AssetEditorToolkitMenuContext.generated.h"
+
+class FAssetEditorToolkit;
+
+UCLASS()
+class UNREALED_API UAssetEditorToolkitMenuContext : public UObject
+{
+	GENERATED_BODY()
+
+public:
+
+	TWeakPtr<FAssetEditorToolkit> Toolkit;
+};
