Updated string

diff --git a/Engine/Source/Editor/UnrealEd/Private/Factories/ActorFactory.cpp b/Engine/Source/Editor/UnrealEd/Private/Factories/ActorFactory.cpp
--- a/Engine/Source/Editor/UnrealEd/Private/Factories/ActorFactory.cpp
+++ b/Engine/Source/Editor/UnrealEd/Private/Factories/ActorFactory.cpp
@@ -461,7 +461,7 @@ UActorFactoryDeferredDecal
 UActorFactoryDeferredDecal::UActorFactoryDeferredDecal(const FObjectInitializer& ObjectInitializer)
 	: Super(ObjectInitializer)
 { 
-	DisplayName = LOCTEXT("DeferredDecalDisplayName", "Deferred Decal");
+	DisplayName = LOCTEXT("DecalDisplayName", "Decal");
 	NewActorClass = ADecalActor::StaticClass();
 	bUseSurfaceOrientation = true;
 }
