1. Fixed issue that layer visibility desynchronized when it was modified while the actor was placed in a hidden level (either sublevel or persistent level). After the level is visible again (and only when it becomes visible), the layers (and their actors) are refreshed.
2. Fixed typo in comment: visibilty --> visibility

diff --git a/Engine/Source/Editor/UnrealEd/Private/EditorLevelUtils.cpp b/Engine/Source/Editor/UnrealEd/Private/EditorLevelUtils.cpp
--- a/Engine/Source/Editor/UnrealEd/Private/EditorLevelUtils.cpp
+++ b/Engine/Source/Editor/UnrealEd/Private/EditorLevelUtils.cpp
@@ -914,7 +914,7 @@ void UEditorLevelUtils::SetLevelVisibility(ULevel* Level, bool bShouldBeVisible,
 	// Also, intentionally do not force layers visible for the p-level
 	if (Level->IsPersistentLevel())
 	{
-		//create a transaction so we can undo the visibilty toggle
+		// Create a transaction so we can undo the visibility toggle
 		const FScopedTransaction Transaction(LOCTEXT("ToggleLevelVisibility", "Toggle Level Visibility"));
 		if (Level->bIsVisible != bShouldBeVisible && ModifyMode == ELevelVisibilityDirtyMode::ModifyOnChange)
 		{
@@ -1094,6 +1094,8 @@ void UEditorLevelUtils::SetLevelVisibility(ULevel* Level, bool bShouldBeVisible,
 		}
 	}
 
+	Level->bIsVisible = bShouldBeVisible;
+
 	FEditorDelegates::RefreshLayerBrowser.Broadcast();
 
 	// Notify the Scene Outliner, as new Actors may be present in the world.
@@ -1120,8 +1122,6 @@ void UEditorLevelUtils::SetLevelVisibility(ULevel* Level, bool bShouldBeVisible,
 		GEditor->NoteSelectionChange();
 	}
 
-	Level->bIsVisible = bShouldBeVisible;
-
 	if (Level->bIsLightingScenario)
 	{
 		Level->OwningWorld->PropagateLightingScenarioChange();
diff --git a/Engine/Source/Editor/UnrealEd/Private/Layers/Layers.cpp b/Engine/Source/Editor/UnrealEd/Private/Layers/Layers.cpp
--- a/Engine/Source/Editor/UnrealEd/Private/Layers/Layers.cpp
+++ b/Engine/Source/Editor/UnrealEd/Private/Layers/Layers.cpp
@@ -23,19 +23,31 @@ FLayers::FLayers( const TWeakObjectPtr< UEditorEngine >& InEditor )
 FLayers::~FLayers()
 {
 	FEditorDelegates::MapChange.RemoveAll(this);
+	// Remove all callback functions from FEditorDelegates::RefreshLayerBrowser.Broadcast()
+	FEditorDelegates::RefreshLayerBrowser.RemoveAll(this);
 }
 
 
 void FLayers::Initialize()
 {
 	FEditorDelegates::MapChange.AddRaw(this, &FLayers::OnEditorMapChange);
+	// Add callback function to FEditorDelegates::RefreshLayerBrowser.Broadcast()
+	FEditorDelegates::RefreshLayerBrowser.AddRaw(this, &FLayers::OnEditorRefreshLayerBrowser);
 }
 
-void FLayers::OnEditorMapChange(uint32 MapChangeFlags)
+void FLayers::OnEditorMapChange(const uint32 MapChangeFlags)
 {
 	LayersChanged.Broadcast( ELayersAction::Reset, NULL, NAME_None );
 }
 
+void FLayers::OnEditorRefreshLayerBrowser()
+{
+	// bNotifySelectionChange is false because the functions calling FEditorDelegates::RefreshLayerBrowser.Broadcast usually call GEditor->NoteSelectionChange
+	const bool bNotifySelectionChange = false;
+	const bool bRedrawViewports = false;
+	UpdateAllActorsVisibility(bNotifySelectionChange, bRedrawViewports);
+}
+
 
 ///////////////////////////////////////////////////////////////////////////////////////////////////////////////////////
 //
diff --git a/Engine/Source/Editor/UnrealEd/Private/Layers/Layers.h b/Engine/Source/Editor/UnrealEd/Private/Layers/Layers.h
--- a/Engine/Source/Editor/UnrealEd/Private/Layers/Layers.h
+++ b/Engine/Source/Editor/UnrealEd/Private/Layers/Layers.h
@@ -139,10 +139,14 @@ protected:
 
 private:
 
-	/** 
+	/**
 	 * Delegate handlers
 	 **/ 
-	void OnEditorMapChange(uint32 MapChangeFlags);
+	void OnEditorMapChange(const uint32 MapChangeFlags);
+	/**
+	 * Delegate handler for FEditorDelegates::RefreshLayerBrowser. It internally calls UpdateAllActorsVisibility to refresh the actors of each layer.
+	 **/
+	void OnEditorRefreshLayerBrowser();
 
 	/**
 	 *	FLayers Constructor
