1. Removed duplicated typedef in IWorldTreeItem.h, which is already defined in WorldBrowserDragDrop.h (file which is already imported by IWorldTreeItem.h itself); and 2. Reverted wrong const (that is not the signature of its delegate function and it is a primative type)

diff --git a/Engine/Source/Editor/UnrealEd/Private/Layers/Layers.cpp b/Engine/Source/Editor/UnrealEd/Private/Layers/Layers.cpp
--- a/Engine/Source/Editor/UnrealEd/Private/Layers/Layers.cpp
+++ b/Engine/Source/Editor/UnrealEd/Private/Layers/Layers.cpp
@@ -35,7 +35,7 @@ void FLayers::Initialize()
 	FEditorDelegates::RefreshLayerBrowser.AddRaw(this, &FLayers::OnEditorRefreshLayerBrowser);
 }
 
-void FLayers::OnEditorMapChange(const uint32 MapChangeFlags)
+void FLayers::OnEditorMapChange(uint32 MapChangeFlags)
 {
 	LayersChanged.Broadcast( ELayersAction::Reset, NULL, NAME_None );
 }
diff --git a/Engine/Source/Editor/UnrealEd/Private/Layers/Layers.h b/Engine/Source/Editor/UnrealEd/Private/Layers/Layers.h
--- a/Engine/Source/Editor/UnrealEd/Private/Layers/Layers.h
+++ b/Engine/Source/Editor/UnrealEd/Private/Layers/Layers.h
@@ -142,7 +142,7 @@ private:
 	/**
 	 * Delegate handlers
 	 **/ 
-	void OnEditorMapChange(const uint32 MapChangeFlags);
+	void OnEditorMapChange(uint32 MapChangeFlags);
 	/**
 	 * Delegate handler for FEditorDelegates::RefreshLayerBrowser. It internally calls UpdateAllActorsVisibility to refresh the actors of each layer.
 	 **/
diff --git a/Engine/Source/Editor/WorldBrowser/Private/WorldBrowserDragDrop.h b/Engine/Source/Editor/WorldBrowser/Private/WorldBrowserDragDrop.h
--- a/Engine/Source/Editor/WorldBrowser/Private/WorldBrowserDragDrop.h
+++ b/Engine/Source/Editor/WorldBrowser/Private/WorldBrowserDragDrop.h
@@ -11,6 +11,7 @@ namespace WorldHierarchy
 {
 	struct IWorldTreeItem;
 	typedef TSharedPtr<IWorldTreeItem> FWorldTreeItemPtr;
+	typedef TSharedRef<IWorldTreeItem> FWorldTreeItemRef;
 
 	struct FValidationInfo
 	{
diff --git a/Engine/Source/Editor/WorldBrowser/Public/IWorldTreeItem.h b/Engine/Source/Editor/WorldBrowser/Public/IWorldTreeItem.h
--- a/Engine/Source/Editor/WorldBrowser/Public/IWorldTreeItem.h
+++ b/Engine/Source/Editor/WorldBrowser/Public/IWorldTreeItem.h
@@ -18,10 +18,6 @@ class FMenuBuilder;
 
 namespace WorldHierarchy
 {
-	struct IWorldTreeItem;
-	typedef TSharedPtr<IWorldTreeItem> FWorldTreeItemPtr;
-	typedef TSharedRef<IWorldTreeItem> FWorldTreeItemRef;
-
 	/** Interface for validating drag/drop movement */
 	struct IDropTarget
 	{
