Created a more efficient way to update the visibility of multiple levels simultaneously. Created UEditorLevelUtils::SetLevelsVisibility as a complementary version of UEditorLevelUtils::SetLevelVisibility, which is more efficient when the visibility of multiple levels must be updated simultaneously.

diff --git a/Engine/Source/Editor/UnrealEd/Private/EditorActor.cpp b/Engine/Source/Editor/UnrealEd/Private/EditorActor.cpp
--- a/Engine/Source/Editor/UnrealEd/Private/EditorActor.cpp
+++ b/Engine/Source/Editor/UnrealEd/Private/EditorActor.cpp
@@ -232,13 +232,15 @@ static FVector CreateLocationOffset(bool bDuplicate, bool bOffsetLocations)
 bool UUnrealEdEngine::WarnIfDestinationLevelIsHidden( UWorld* InWorld )
 {
 	bool result = false;
+	TArray<ULevel*> Levels;
+	TArray<bool> bTheyShouldBeVisible;
 	//prepare the warning dialog
 	FSuppressableWarningDialog::FSetupInfo Info( LOCTEXT( "Warning_PasteWarningBody","You are trying to paste to a hidden level.\nSuppressing this will default to Do Not Paste" ), LOCTEXT( "Warning_PasteWarningHeader","Pasting To Hidden Level" ), "PasteHiddenWarning" );
 	Info.ConfirmText = LOCTEXT( "Warning_PasteContinue","Unhide Level and paste" );
 	Info.CancelText = LOCTEXT( "Warning_PasteCancel","Do not paste" );
 	FSuppressableWarningDialog PasteHiddenWarning( Info );
 
-	//check streaming levels first	
+	//check streaming levels first
 	for (ULevelStreaming* StreamedLevel : InWorld->GetStreamingLevels())
 	{
 		//this is the active level - check if it is visible
@@ -255,7 +257,8 @@ bool UUnrealEdEngine::WarnIfDestinationLevelIsHidden( UWorld* InWorld )
 				}
 				else
 				{
-					EditorLevelUtils::SetLevelVisibility( Level, true, true );					
+					Levels.Add(Level);
+					bTheyShouldBeVisible.Add(true);
 				}
 			}
 		}
@@ -274,11 +277,17 @@ bool UUnrealEdEngine::WarnIfDestinationLevelIsHidden( UWorld* InWorld )
 			}
 			else 
 			{
-				EditorLevelUtils::SetLevelVisibility( InWorld->GetCurrentLevel(), true, true );
+				Levels.Add(InWorld->GetCurrentLevel());
+				bTheyShouldBeVisible.Add(true);
 			}
 
 		}
 	}
+	// For efficiency, set visibility of all levels at once
+	if (Levels.Num() > 0)
+	{
+		EditorLevelUtils::SetLevelsVisibility(Levels, bTheyShouldBeVisible, true);
+	}
 	return result;
 }
 
diff --git a/Engine/Source/Editor/UnrealEd/Private/EditorEngine.cpp b/Engine/Source/Editor/UnrealEd/Private/EditorEngine.cpp
--- a/Engine/Source/Editor/UnrealEd/Private/EditorEngine.cpp
+++ b/Engine/Source/Editor/UnrealEd/Private/EditorEngine.cpp
@@ -2499,11 +2499,19 @@ bool UEditorEngine::WarnAboutHiddenLevels( UWorld* InWorld, bool bIncludePersist
 
 			// follow up using SetLevelVisibility - streaming should now be completed so we can show actors, layers, 
 			// BSPs etc. without too big a performance hit.
+			TArray<ULevel*> LoadedLevels;
+			TArray<bool> bTheyShouldBeVisible;
 			for( int32 HiddenLevelIdx = 0; HiddenLevelIdx < HiddenLevels.Num(); ++HiddenLevelIdx )
 			{
 				check(HiddenLevels[ HiddenLevelIdx ]->GetLoadedLevel());
 				ULevel* LoadedLevel = HiddenLevels[ HiddenLevelIdx ]->GetLoadedLevel();
-				EditorLevelUtils::SetLevelVisibility( LoadedLevel, true, false );
+				LoadedLevels.Add(LoadedLevel);
+				bTheyShouldBeVisible.Add(true);
+			}
+			// For efficiency, set visibility of all levels at once
+			if (LoadedLevels.Num() > 0)
+			{
+				EditorLevelUtils::SetLevelsVisibility(LoadedLevels, bTheyShouldBeVisible, false);
 			}
 
 			FEditorSupportDelegates::RedrawAllViewports.Broadcast();
diff --git a/Engine/Source/Editor/UnrealEd/Private/EditorLevelUtils.cpp b/Engine/Source/Editor/UnrealEd/Private/EditorLevelUtils.cpp
--- a/Engine/Source/Editor/UnrealEd/Private/EditorLevelUtils.cpp
+++ b/Engine/Source/Editor/UnrealEd/Private/EditorLevelUtils.cpp
@@ -901,7 +901,7 @@ void UEditorLevelUtils::SetLevelVisibilityTemporarily(ULevel* Level, bool bShoul
 	}
 }
 
-void UEditorLevelUtils::SetLevelVisibility(ULevel* Level, bool bShouldBeVisible, bool bForceLayersVisible, ELevelVisibilityDirtyMode ModifyMode)
+void SetLevelVisibilityNoGlobalUpdateInternal(ULevel* Level, const bool bShouldBeVisible, const bool bForceLayersVisible, const ELevelVisibilityDirtyMode ModifyMode)
 {
 	// Nothing to do
 	if (Level == NULL)
@@ -971,7 +971,6 @@ void UEditorLevelUtils::SetLevelVisibility(ULevel* Level, bool bShouldBeVisible,
 		}
 
 		Level->GetWorld()->OnLevelsChanged().Broadcast();
-		FEditorSupportDelegates::RedrawAllViewports.Broadcast();
 	}
 	else
 	{
@@ -1096,17 +1095,12 @@ void UEditorLevelUtils::SetLevelVisibility(ULevel* Level, bool bShouldBeVisible,
 
 	Level->bIsVisible = bShouldBeVisible;
 
-	FEditorDelegates::RefreshLayerBrowser.Broadcast();
-
-	// Notify the Scene Outliner, as new Actors may be present in the world.
-	GEngine->BroadcastLevelActorListChanged();
-
-	// If the level is being hidden, deselect actors and surfaces that belong to this level.
+	// If the level is being hidden, deselect actors and surfaces that belong to this level. (Part 1/2)
 	if (!bShouldBeVisible && ModifyMode == ELevelVisibilityDirtyMode::ModifyOnChange)
 	{
 		USelection* SelectedActors = GEditor->GetSelectedActors();
 		SelectedActors->Modify();
-		TArray<AActor*>& Actors = Level->Actors;
+		const TArray<AActor*>& Actors = Level->Actors;
 		for (int32 ActorIndex = 0; ActorIndex < Actors.Num(); ++ActorIndex)
 		{
 			AActor* Actor = Actors[ActorIndex];
@@ -1116,10 +1110,7 @@ void UEditorLevelUtils::SetLevelVisibility(ULevel* Level, bool bShouldBeVisible,
 			}
 		}
 
-		DeselectAllSurfacesInLevel(Level);
-
-		// Tell the editor selection status was changed.
-		GEditor->NoteSelectionChange();
+		UEditorLevelUtils::DeselectAllSurfacesInLevel(Level);
 	}
 
 	if (Level->bIsLightingScenario)
@@ -1128,6 +1119,75 @@ void UEditorLevelUtils::SetLevelVisibility(ULevel* Level, bool bShouldBeVisible,
 	}
 }
 
+void UEditorLevelUtils::SetLevelVisibility(ULevel* Level, const bool bShouldBeVisible, const bool bForceLayersVisible, const ELevelVisibilityDirtyMode ModifyMode)
+{
+	TArray<ULevel*> Levels({ Level });
+	TArray<bool> bTheyShouldBeVisible({ bShouldBeVisible });
+	SetLevelsVisibility(Levels, bTheyShouldBeVisible, bForceLayersVisible, ModifyMode);
+}
+
+void UEditorLevelUtils::SetLevelsVisibility(const TArray<ULevel*>& Levels, const TArray<bool>& bTheyShouldBeVisible, const bool bForceLayersVisible, const ELevelVisibilityDirtyMode ModifyMode)
+{
+	// Nothing to do
+	if (Levels.Num() == 0 || Levels.Num() != bTheyShouldBeVisible.Num())
+	{
+		return;
+	}
+
+	// Perform SetLevelVisibilityNoGlobalUpdateInternal for each Level
+	for (int32 LevelIndex = 0; LevelIndex < Levels.Num(); ++LevelIndex)
+	{
+		ULevel* Level = Levels[LevelIndex];
+		if (Level)
+		{
+			SetLevelVisibilityNoGlobalUpdateInternal(Level, bTheyShouldBeVisible[LevelIndex], bForceLayersVisible, ModifyMode);
+		}
+	}
+
+	// If at least 1 persistent level, then RedrawAllViewports.Broadcast
+	for (int32 LevelIndex = 0; LevelIndex < Levels.Num(); ++LevelIndex)
+	{
+		ULevel* Level = Levels[LevelIndex];
+		if (Level && Level->IsPersistentLevel())
+		{
+			FEditorSupportDelegates::RedrawAllViewports.Broadcast();
+			break;
+		}
+	}
+
+	// If at least 1 level becomes visible, force layers to update their actor status
+	// Otherwise, changes made on the layers for actors belonging to a non-visible level would not work
+	if (GEditor->Layers.IsValid())
+	{
+		for (int32 LevelIndex = 0; LevelIndex < bTheyShouldBeVisible.Num(); ++LevelIndex)
+		{
+			if (bTheyShouldBeVisible[LevelIndex])
+			{
+				// Equivalent to GEditor->Layers->UpdateAllActorsVisibilityDefault();
+				FEditorDelegates::RefreshLayerBrowser.Broadcast();
+				break;
+			}
+		}
+	}
+
+	// Notify the Scene Outliner, as new Actors may be present in the world.
+	GEngine->BroadcastLevelActorListChanged();
+
+	// If the level is being hidden, deselect actors and surfaces that belong to this level. (Part 2/2)
+	if (ModifyMode == ELevelVisibilityDirtyMode::ModifyOnChange)
+	{
+		for (int32 LevelIndex = 0; LevelIndex < bTheyShouldBeVisible.Num(); ++LevelIndex)
+		{
+			if (!bTheyShouldBeVisible[LevelIndex])
+			{
+				// Tell the editor selection status was changed.
+				GEditor->NoteSelectionChange();
+				break;
+			}
+		}
+	}
+}
+
 void UEditorLevelUtils::GetWorlds(UWorld* InWorld, TArray<UWorld*>& OutWorlds, bool bIncludeInWorld, bool bOnlyEditorVisible)
 {
 	OutWorlds.Empty();
diff --git a/Engine/Source/Editor/UnrealEd/Public/EditorLevelUtils.h b/Engine/Source/Editor/UnrealEd/Public/EditorLevelUtils.h
--- a/Engine/Source/Editor/UnrealEd/Public/EditorLevelUtils.h
+++ b/Engine/Source/Editor/UnrealEd/Public/EditorLevelUtils.h
@@ -163,14 +163,26 @@ public:
 	static UNREALED_API void SetLevelVisibilityTemporarily(ULevel* Level, bool bShouldBeVisible);
 
 	/**
-	 * Sets a level's visibility in the editor.
+	 * Sets a level's visibility in the editor. Less efficient than SetLevelsVisibility when changing the visibility of multiple levels simultaneously.
 	 *
 	 * @param	Level					The level to modify.
 	 * @param	bShouldBeVisible		The level's new visibility state.
 	 * @param	bForceLayersVisible		If true and the level is visible, force the level's layers to be visible.
+	 * @param	ModifyMode				ELevelVisibilityDirtyMode mode value.
 	 */
 	UFUNCTION(BlueprintCallable, Category = "Editor Scripting | Level Utility")
-	static UNREALED_API void SetLevelVisibility(ULevel* Level, bool bShouldBeVisible, bool bForceLayersVisible, ELevelVisibilityDirtyMode ModifyMode = ELevelVisibilityDirtyMode::ModifyOnChange);
+	static UNREALED_API void SetLevelVisibility(ULevel* Level, const bool bShouldBeVisible, const bool bForceLayersVisible, const ELevelVisibilityDirtyMode ModifyMode = ELevelVisibilityDirtyMode::ModifyOnChange);
+
+	/**
+	 * Sets a level's visibility in the editor. More efficient than SetLevelsVisibility when changing the visibility of multiple levels simultaneously.
+	 *
+	 * @param	Levels					The levels to modify.
+	 * @param	bShouldBeVisible		The level's new visibility state for each level.
+	 * @param	bForceLayersVisible		If true and the level is visible, force the level's layers to be visible.
+	 * @param	ModifyMode				ELevelVisibilityDirtyMode mode value.
+	 */
+	UFUNCTION(BlueprintCallable, Category = "Editor Scripting | Level Utility")
+	static UNREALED_API void SetLevelsVisibility(const TArray<ULevel*>& Levels, const TArray<bool>& bShouldBeVisible, const bool bForceLayersVisible, const ELevelVisibilityDirtyMode ModifyMode = ELevelVisibilityDirtyMode::ModifyOnChange);
 	
 	/** 
 	 * Deselects all BSP surfaces in this level 
diff --git a/Engine/Source/Editor/WorldBrowser/Private/LevelCollectionModel.cpp b/Engine/Source/Editor/WorldBrowser/Private/LevelCollectionModel.cpp
--- a/Engine/Source/Editor/WorldBrowser/Private/LevelCollectionModel.cpp
+++ b/Engine/Source/Editor/WorldBrowser/Private/LevelCollectionModel.cpp
@@ -460,10 +460,15 @@ void FLevelCollectionModel::HideLevels(const FLevelModelList& InLevelList)
 		return;
 	}
 	
+	// For efficiency, set visibility of all levels at once
+	TArray<FLevelModel*> LevelModels;
+	TArray<bool> bVisible;
 	for (auto It = InLevelList.CreateConstIterator(); It; ++It)
 	{
-		(*It)->SetVisible(false);
+		LevelModels.Add(It->Get());
+		bVisible.Add(false);
 	}
+	FLevelModel::SetVisible(LevelModels, bVisible);
 
 	RequestUpdateAllLevels();
 }
@@ -477,10 +482,15 @@ void FLevelCollectionModel::ShowLevels(const FLevelModelList& InLevelList)
 	
 	OnPreShowLevels(InLevelList);
 
+	// For efficiency, set visibility of all levels at once
+	TArray<FLevelModel*> LevelModels;
+	TArray<bool> bVisible;
 	for (auto It = InLevelList.CreateConstIterator(); It; ++It)
 	{
-		(*It)->SetVisible(true);
+		LevelModels.Add(It->Get());
+		bVisible.Add(true);
 	}
+	FLevelModel::SetVisible(LevelModels, bVisible);
 
 	RequestUpdateAllLevels();
 }
diff --git a/Engine/Source/Editor/WorldBrowser/Private/LevelModel.cpp b/Engine/Source/Editor/WorldBrowser/Private/LevelModel.cpp
--- a/Engine/Source/Editor/WorldBrowser/Private/LevelModel.cpp
+++ b/Engine/Source/Editor/WorldBrowser/Private/LevelModel.cpp
@@ -243,26 +243,55 @@ void FLevelModel::LoadLevel()
 
 void FLevelModel::SetVisible(bool bVisible)
 {
-	//don't create unnecessary transactions
-	if (IsVisible() == bVisible)
+	TArray<FLevelModel*> LevelModels({ this });
+	TArray<bool> bAreVisible({ bVisible });
+	FLevelModel::SetVisible(LevelModels, bAreVisible);
+}
+
+void FLevelModel::SetVisible(TArray<FLevelModel*>& LevelModels, const TArray<bool>& bAreVisible)
+{
+	// Don't create unnecessary transactions
+	if (LevelModels.Num() == 0 || LevelModels.Num() != bAreVisible.Num())
 	{
 		return;
 	}
 
-	const bool oldIsDirty = IsDirty();
+	TArray<ULevel*> Levels;
+	TArray<bool> bOldAreDirty;
+	for (int32 LevelModelIdx = 0; LevelModelIdx < LevelModels.Num(); ++LevelModelIdx)
+	{
+		FLevelModel* LevelModel = LevelModels[LevelModelIdx];
+		// If visibility does not change, omit it
+		if (LevelModel->IsVisible() == bAreVisible[LevelModelIdx])
+		{
+			continue;
+		}
+		// Otherwise, add Level and IsDirty into each TArray
+		Levels.Add(LevelModel->GetLevelObject());
+		bOldAreDirty.Add(LevelModel->IsDirty());
+	}
+	// Don't create unnecessary transactions
+	if (Levels.Num() == 0)
+	{
+		return;
+	}
 
 	const FScopedTransaction Transaction(LOCTEXT("ToggleVisibility", "Toggle Level Visibility"));
 
-	//this call hides all owned actors, etc
-	EditorLevelUtils::SetLevelVisibility( GetLevelObject(), bVisible, false );
+	// This call hides all owned actors, etc.
+	EditorLevelUtils::SetLevelsVisibility(Levels, bAreVisible, false);
 
-	if (!oldIsDirty)
+	// Note that Levels.Num() != LevelModels.Num() given the continue in the first for loop
+	for (int32 LevelModelIdx = 0; LevelModelIdx < Levels.Num(); ++LevelModelIdx)
 	{
-		// don't set the dirty flag if we're just changing the visibility of the level within the editor
-		ULevel* Level = GetLevelObject();
-		if (Level)
+		if (!bOldAreDirty[LevelModelIdx])
 		{
-			Level->GetOutermost()->SetDirtyFlag(false);
+			// Don't set the dirty flag if we're just changing the visibility of the level within the editor
+			ULevel* Level = Levels[LevelModelIdx];
+			if (Level)
+			{
+				Level->GetOutermost()->SetDirtyFlag(false);
+			}
 		}
 	}
 }
diff --git a/Engine/Source/Editor/WorldBrowser/Private/LevelModel.h b/Engine/Source/Editor/WorldBrowser/Private/LevelModel.h
--- a/Engine/Source/Editor/WorldBrowser/Private/LevelModel.h
+++ b/Engine/Source/Editor/WorldBrowser/Private/LevelModel.h
@@ -135,6 +135,9 @@ public:
 	/** Sets the Level's visibility */
 	virtual void SetVisible(bool bVisible);
 
+	/** Static function analog to the non-static SetVisible but applied to an TArray of FLevelModel elements */
+	static void SetVisible(TArray<FLevelModel*>& LevelModels, const TArray<bool>& bAreVisible);
+
 	/** Sets the Level's locked/unlocked state */
 	void SetLocked(bool bLocked);
 
