Accidentally submitted a largely commented testing version of PropertyEditorTestObject.

diff --git a/Engine/Source/Editor/UnrealEd/Classes/Editor/PropertyEditorTestObject.h b/Engine/Source/Editor/UnrealEd/Classes/Editor/PropertyEditorTestObject.h
--- a/Engine/Source/Editor/UnrealEd/Classes/Editor/PropertyEditorTestObject.h
+++ b/Engine/Source/Editor/UnrealEd/Classes/Editor/PropertyEditorTestObject.h
@@ -128,7 +128,7 @@ class UPropertyEditorTestObject : public UObject
 {
     GENERATED_BODY()
 
-/*	UPROPERTY(EditAnywhere, Category = BasicProperties)
+	UPROPERTY(EditAnywhere, Category = BasicProperties)
 	int8 Int8Property;
 
 	UPROPERTY(EditAnywhere, Category = BasicProperties)
@@ -303,7 +303,7 @@ class UPropertyEditorTestObject : public UObject
 	bool bEditConditionStructWithMultipleInstances2;
 
 	UPROPERTY(EditAnywhere, Category=StructTests, meta=(editcondition = "bEditConditionStructWithMultipleInstances2"))
-	FPropertyEditorTestBasicStruct StructWithMultipleInstances2; */
+	FPropertyEditorTestBasicStruct StructWithMultipleInstances2;
 
 	UPROPERTY(EditAnywhere, Category=StructTests)
 	FRichCurve RichCurve;
@@ -332,7 +332,7 @@ class UPropertyEditorTestObject : public UObject
 	UPROPERTY(EditAnywhere, Category=StructTests, meta=(AllowedClasses="Material,Texture"))
 	FSoftObjectPath MaterialOrTextureAssetReference;
 
-	/*UPROPERTY(EditAnywhere, Category=StructTests, meta=(InlineEditConditionToggle))
+	UPROPERTY(EditAnywhere, Category=StructTests, meta=(InlineEditConditionToggle))
 	bool bEditCondition;
 
 	UPROPERTY(EditAnywhere, Category=AdvancedProperties, meta=(editcondition = "bEditCondition"))
@@ -519,5 +519,5 @@ class UPropertyEditorTestObject : public UObject
 	bool bVisibleWhenStatic;
 
 	UPROPERTY(EditAnywhere, Category = EditConditionHides, meta = (EditConditionHides, EditCondition = "PropertyThatHides == EComponentMobility::Stationary"))
-	int32 VisibleWhenStationary;*/
+	int32 VisibleWhenStationary;
 };
