Fixed static analysis warning.

diff --git a/Engine/Source/Editor/UnrealEd/Private/LevelEditorViewport.cpp b/Engine/Source/Editor/UnrealEd/Private/LevelEditorViewport.cpp
--- a/Engine/Source/Editor/UnrealEd/Private/LevelEditorViewport.cpp
+++ b/Engine/Source/Editor/UnrealEd/Private/LevelEditorViewport.cpp
@@ -3387,7 +3387,7 @@ void FLevelEditorViewportClient::MoveLockedActorToCamera()
 				ActiveActorLock->SetActorRotation(GCurrentLevelEditingViewportClient->GetViewRotation());
 			}
 
-			if (PreviousRotator.IsSet())
+			if (ActiveActorLockComponent)
 			{
 				const FRotator Rot = PreviousRotator.GetValue();
 				FRotator ActorRotWind, ActorRotRem;
