Improved phrasing of asset editor restoration dialog to make it clear that you're re-opening asset editors rather than restoring package data after a crash

#rnx

diff --git a/Engine/Source/Editor/UnrealEd/Private/Toolkits/AssetEditorManager.cpp b/Engine/Source/Editor/UnrealEd/Private/Toolkits/AssetEditorManager.cpp
--- a/Engine/Source/Editor/UnrealEd/Private/Toolkits/AssetEditorManager.cpp
+++ b/Engine/Source/Editor/UnrealEd/Private/Toolkits/AssetEditorManager.cpp
@@ -653,20 +653,22 @@ void FAssetEditorManager::SpawnRestorePreviouslyOpenAssetsNotification(const boo
 		}
 	};
 
-	FNotificationInfo Info(bCleanShutdown 
-		? LOCTEXT("RestoreOpenAssetsAfterClose_Message", "Assets were open when the Editor was last closed, would you like to restore them now?")
-		: LOCTEXT("RestoreOpenAssetsAfterCrash", "The Editor did not shut down cleanly, would you like to attempt to restore previously open assets now?")
-		);
+	FText NotificationMessage = bCleanShutdown
+		? LOCTEXT("ReopenAssetEditorsAfterClose", "{0} asset {0}|plural(one=editor was,other=editors were) open when the editor was last closed. Would you like to re-open them?")
+		: LOCTEXT("ReopenAssetEditorsAfterCrash", "{0} asset {0}|plural(one=editor was,other=editors were) open when the editor quit unexpectedly. Would you like to re-open them?");
+	NotificationMessage = FText::Format(NotificationMessage, AssetsToOpen.Num());
+
+	FNotificationInfo Info = FNotificationInfo(NotificationMessage);
 
 	// Add the buttons
 	Info.ButtonDetails.Add(FNotificationButtonInfo(
-		LOCTEXT("RestoreOpenAssetsAfterClose_Confirm", "Restore Now"), 
+		LOCTEXT("ReopenAssetEditors_Confirm", "Open"), 
 		FText(), 
 		FSimpleDelegate::CreateRaw(this, &FAssetEditorManager::OnConfirmRestorePreviouslyOpenAssets, AssetsToOpen), 
 		SNotificationItem::CS_None
 		));
 	Info.ButtonDetails.Add(FNotificationButtonInfo(
-		LOCTEXT("RestoreOpenAssetsAfterClose_Cancel", "Don't Restore"), 
+		LOCTEXT("ReopenAssetEditors_Cancel", "Cancel"), 
 		FText(), 
 		FSimpleDelegate::CreateRaw(this, &FAssetEditorManager::OnCancelRestorePreviouslyOpenAssets), 
 		SNotificationItem::CS_None
