PR #5885: UE-33018: Case sensitive string compares instead of FName for folder renaming (Contributed by [REDACTED_PERSON])

diff --git a/Engine/Source/Editor/UnrealEd/Private/EditorActorFolders.cpp b/Engine/Source/Editor/UnrealEd/Private/EditorActorFolders.cpp
--- a/Engine/Source/Editor/UnrealEd/Private/EditorActorFolders.cpp
+++ b/Engine/Source/Editor/UnrealEd/Private/EditorActorFolders.cpp
@@ -440,17 +440,18 @@ void FActorFolders::DeleteFolder(UWorld& InWorld, FName FolderToDelete)
 
 bool FActorFolders::RenameFolderInWorld(UWorld& World, FName OldPath, FName NewPath)
 {
-	if (OldPath.IsNone() || OldPath == NewPath || PathIsChildOf(NewPath.ToString(), OldPath.ToString()))
+	const FString OldPathString = OldPath.ToString();
+	const FString NewPathString = NewPath.ToString();
+
+	if (OldPath.IsNone() || OldPathString.Equals(NewPathString) || PathIsChildOf(NewPathString, OldPathString))
 	{
 		return false;
 	}
 
 	const FScopedTransaction Transaction(LOCTEXT("UndoAction_RenameFolder", "Rename Folder"));
 
-	const FString OldPathString = OldPath.ToString();
-	const FString NewPathString = NewPath.ToString();
-
 	TSet<FName> RenamedFolders;
+	bool RenamedFolder = false;
 
 	// Move any folders we currently hold - old ones will be deleted later
 	UEditorActorFolders& FoldersInWorld = GetOrCreateFoldersForWorld(World);
@@ -465,7 +466,19 @@ bool FActorFolders::RenameFolderInWorld(UWorld& World, FName OldPath, FName NewP
 		if (OldPath == Path || PathIsChildOf(FolderPath, OldPathString))
 		{
 			const FName NewFolder = OldPathToNewPath(OldPathString, NewPathString, FolderPath);
-			if (!FoldersInWorld.Folders.Contains(NewFolder))
+			
+			// Needs to be done this way otherwise case insensitive comparison is used.
+			bool ContainsFolder = false;
+			for (const auto& FolderPair : FoldersInWorld.Folders)
+			{
+				if (FolderPair.Key.IsEqual(NewFolder, ENameCase::CaseSensitive))
+				{
+					ContainsFolder = true;
+					break;
+				}
+			}
+
+			if (!ContainsFolder)
 			{
 				// Use the existing properties for the folder if we have them
 				if (FActorFolderProps* ExistingProperties = FoldersInWorld.Folders.Find(Path))
@@ -480,26 +493,39 @@ bool FActorFolders::RenameFolderInWorld(UWorld& World, FName OldPath, FName NewP
 				OnFolderMove.Broadcast(World, Path, NewFolder);
 				OnFolderCreate.Broadcast(World, NewFolder);
 			}
-			RenamedFolders.Add(Path);
+
+			// case insensive compare as we don't want to remove the folder if it has the same name
+			if (Path != NewFolder)
+			{
+				RenamedFolders.Add(Path);
+			}
+
+			RenamedFolder = true;
 		}
 	}
 
 	// Now that we have folders created, move any actors that ultimately reside in that folder too
 	for (auto ActorIt = FActorIterator(&World); ActorIt; ++ActorIt)
 	{
-		const FName& OldActorPath = ActorIt->GetFolderPath();
-		
-		AActor* Actor = *ActorIt;
+		// copy, otherwise it returns the new value when set later
+		const FName OldActorPath = ActorIt->GetFolderPath();
 		if (OldActorPath.IsNone())
 		{
 			continue;
 		}
 
-		if (Actor->GetFolderPath() == OldPath || PathIsChildOf(OldActorPath.ToString(), OldPathString))
+		if (OldActorPath == OldPath || PathIsChildOf(OldActorPath.ToString(), OldPathString))
 		{
-			RenamedFolders.Add(OldActorPath);
-
 			ActorIt->SetFolderPath_Recursively(OldPathToNewPath(OldPathString, NewPathString, OldActorPath.ToString()));
+			const FName& NewActorPath = ActorIt->GetFolderPath();
+
+			// case insensive compare as we don't want to remove the folder if it has the same name
+			if (OldActorPath != NewActorPath)
+			{
+				RenamedFolders.Add(OldActorPath);
+			}
+
+			RenamedFolder = true;
 		}
 	}
 
@@ -510,7 +536,7 @@ bool FActorFolders::RenameFolderInWorld(UWorld& World, FName OldPath, FName NewP
 		OnFolderDelete.Broadcast(World, Path);
 	}
 
-	return RenamedFolders.Num() != 0;
+	return RenamedFolder;
 }
 
 bool FActorFolders::AddFolderToWorld(UWorld& InWorld, FName Path)
diff --git a/Engine/Source/Runtime/Engine/Private/ActorEditor.cpp b/Engine/Source/Runtime/Engine/Private/ActorEditor.cpp
--- a/Engine/Source/Runtime/Engine/Private/ActorEditor.cpp
+++ b/Engine/Source/Runtime/Engine/Private/ActorEditor.cpp
@@ -885,7 +885,7 @@ const FName& AActor::GetFolderPath() const
 
 void AActor::SetFolderPath(const FName& NewFolderPath)
 {
-	if (NewFolderPath != FolderPath)
+	if (!NewFolderPath.IsEqual(FolderPath, ENameCase::CaseSensitive))
 	{
 		Modify();
 
