Add full original imported data to be able to rebuild morph targets and alternate skinning

diff --git a/Engine/Source/Editor/UnrealEd/Private/Fbx/FbxFactory.cpp b/Engine/Source/Editor/UnrealEd/Private/Fbx/FbxFactory.cpp
--- a/Engine/Source/Editor/UnrealEd/Private/Fbx/FbxFactory.cpp
+++ b/Engine/Source/Editor/UnrealEd/Private/Fbx/FbxFactory.cpp
@@ -687,7 +687,7 @@ UObject* UFbxFactory::FactoryCreateFile
 								uint32 bImportTextures = ImportOptions->bImportTextures;
 								ImportOptions->bImportTextures = 0;
 
-								FbxImporter->ImportFbxMorphTarget(SkelMeshNodeArray, Cast<USkeletalMesh>(CreatedObject), InParent, ImportedSuccessfulLodIndex, OutData);
+								FbxImporter->ImportFbxMorphTarget(SkelMeshNodeArray, Cast<USkeletalMesh>(CreatedObject), ImportedSuccessfulLodIndex, OutData);
 							
 								ImportOptions->bImportMaterials = !!bImportMaterials;
 								ImportOptions->bImportTextures = !!bImportTextures;
diff --git a/Engine/Source/Editor/UnrealEd/Private/Fbx/FbxSceneImportFactory.cpp b/Engine/Source/Editor/UnrealEd/Private/Fbx/FbxSceneImportFactory.cpp
--- a/Engine/Source/Editor/UnrealEd/Private/Fbx/FbxSceneImportFactory.cpp
+++ b/Engine/Source/Editor/UnrealEd/Private/Fbx/FbxSceneImportFactory.cpp
@@ -2041,7 +2041,7 @@ UObject* UFbxSceneImportFactory::ImportOneSkeletalMesh(void* VoidRootNodeToImpor
 				NewSkelMesh->GetImportedModel()->LODModels.IsValidIndex(LODIndex))
 			{
 				// TODO: Disable material importing when importing morph targets
-				FbxImporter->ImportFbxMorphTarget(SkelMeshNodeArray, NewSkelMesh, Pkg, LODIndex, OutData);
+				FbxImporter->ImportFbxMorphTarget(SkelMeshNodeArray, NewSkelMesh, LODIndex, OutData);
 			}
 		}
 	}
diff --git a/Engine/Source/Editor/UnrealEd/Private/Fbx/FbxSkeletalMeshImport.cpp b/Engine/Source/Editor/UnrealEd/Private/Fbx/FbxSkeletalMeshImport.cpp
--- a/Engine/Source/Editor/UnrealEd/Private/Fbx/FbxSkeletalMeshImport.cpp
+++ b/Engine/Source/Editor/UnrealEd/Private/Fbx/FbxSkeletalMeshImport.cpp
@@ -2633,7 +2633,7 @@ USkeletalMesh* UnFbx::FFbxImporter::ReimportSkeletalMesh(USkeletalMesh* Mesh, UF
 					NewMesh->GetImportedModel()->LODModels.IsValidIndex(LODIndex))
 				{
 					// @fixme: @question : where do they import this morph? where to? What morph target sets?
-					ImportFbxMorphTarget(SkelMeshNodeArray, NewMesh, NewMesh->GetOutermost(), ImportedSuccessfulLodIndex, OutData);
+					ImportFbxMorphTarget(SkelMeshNodeArray, NewMesh, ImportedSuccessfulLodIndex, OutData);
 				}
 			}
 		}
@@ -4143,276 +4143,12 @@ bool UnFbx::FFbxImporter::ImportSkeletalMeshLOD(USkeletalMesh* InSkeletalMesh, U
 	return true;
 }
 
-struct FMeshDataBundle
-{
-	TArray< FVector > Vertices;
-	TArray< uint32 > Indices;
-	TArray< FVector2D > UVs;
-	TArray< uint32 > SmoothingGroups;
-	TArray<SkeletalMeshImportData::FTriangle> Faces;
-};
-
-static void ConvertImportDataToMeshData(const FSkeletalMeshImportData& ImportData, FMeshDataBundle& MeshDataBundle)
-{
-	for ( const SkeletalMeshImportData::FTriangle& Face : ImportData.Faces )
-	{
-		SkeletalMeshImportData::FTriangle FaceTriangle;
-		FaceTriangle = Face;
-		for ( int32 i = 0; i < 3; ++i )
-		{
-			const SkeletalMeshImportData::FVertex& Wedge = ImportData.Wedges[ Face.WedgeIndex[i] ];
-			int32 FaceWedgeIndex = MeshDataBundle.Indices.Add(Wedge.VertexIndex);
-			MeshDataBundle.UVs.Add(Wedge.UVs[0]);
-			FaceTriangle.WedgeIndex[i] = FaceWedgeIndex;
-		}
-		MeshDataBundle.Faces.Add(FaceTriangle);
-		MeshDataBundle.SmoothingGroups.Add(Face.SmoothingGroups);
-	}
-
-	MeshDataBundle.Vertices = ImportData.Points;
-}
-
-/**
-* A class encapsulating morph target processing that occurs during import on a separate thread
-*/
-class FAsyncImportMorphTargetWork : public FNonAbandonableTask
-{
-public:
-	FAsyncImportMorphTargetWork(FSkeletalMeshLODModel* InLODModel, const FReferenceSkeleton& InRefSkeleton, const FSkeletalMeshImportData& InBaseImportData, TArray<FVector>&& InMorphLODPoints,
-		FBXImportOptions* InImportOptions, TArray< FMorphTargetDelta >& InMorphDeltas, TArray<uint32>& InBaseIndexData, TArray< uint32 >& InBaseWedgePointIndices,
-		TMap<uint32,uint32>& InWedgePointToVertexIndexMap, const FOverlappingCorners& InOverlappingCorners,
-		const TSet<uint32> InModifiedPoints, const TMultiMap< int32, int32 >& InWedgeToFaces, const FMeshDataBundle& InMeshDataBundle, const TArray<FVector>& InTangentZ)
-		: LODModel(InLODModel)
-		, RefSkeleton(InRefSkeleton)
-		, BaseImportData(InBaseImportData)
-		, MorphLODPoints(InMorphLODPoints)
-		, ImportOptions(InImportOptions)
-		, MorphTargetDeltas(InMorphDeltas)
-		, BaseIndexData (InBaseIndexData)
-		, BaseWedgePointIndices(InBaseWedgePointIndices)
-		, WedgePointToVertexIndexMap(InWedgePointToVertexIndexMap)
-		, OverlappingCorners(InOverlappingCorners)
-		, ModifiedPoints(InModifiedPoints)
-		, WedgeToFaces(InWedgeToFaces)
-		, MeshDataBundle(InMeshDataBundle)
-		, BaseTangentZ(InTangentZ)
-		, TangentZ(InTangentZ)
-	{
-		MeshUtilities = &FModuleManager::Get().LoadModuleChecked<IMeshUtilities>("MeshUtilities");
-	}
-
-	void PrepareTangents()
-	{
-		TArray<bool> WasProcessed;
-		WasProcessed.Empty(MeshDataBundle.Indices.Num());
-		WasProcessed.AddZeroed(MeshDataBundle.Indices.Num());
-
-		TArray< int32 > WedgeFaces;
-		TArray< int32 > OverlappingWedgesDummy;
-		TArray< int32 > OtherOverlappingWedgesDummy;
-
-		// For each ModifiedPoints, reset the tangents for the affected wedges
-		for (int32 WedgeIdx = 0; WedgeIdx < MeshDataBundle.Indices.Num(); ++WedgeIdx)
-		{
-			int32 PointIdx = MeshDataBundle.Indices[WedgeIdx];
-
-			if ( ModifiedPoints.Find( PointIdx ) != nullptr )
-			{
-				TangentZ[ WedgeIdx ] = FVector::ZeroVector;
-
-				const TArray<int32>& OverlappingWedges = FindIncludingNoOverlapping(OverlappingCorners, WedgeIdx, OverlappingWedgesDummy);
-
-				for ( const int32 OverlappingWedgeIndex : OverlappingWedges )
-				{
-					if ( WasProcessed[ OverlappingWedgeIndex ] )
-					{
-						continue;
-					}
-
-					WasProcessed[ OverlappingWedgeIndex ] = true;
-
-					WedgeFaces.Reset();
-					WedgeToFaces.MultiFind( OverlappingWedgeIndex, WedgeFaces );
-
-					for ( const int32 FaceIndex : WedgeFaces )
-					{
-						for ( int32 CornerIndex = 0; CornerIndex < 3; ++CornerIndex )
-						{
-							int32 WedgeIndex = MeshDataBundle.Faces[FaceIndex].WedgeIndex[CornerIndex];
-
-							TangentZ[ WedgeIndex ] = FVector::ZeroVector;
-
-							const TArray<int32>& OtherOverlappingWedges = FindIncludingNoOverlapping(OverlappingCorners, WedgeIndex, OtherOverlappingWedgesDummy);
-
-							for ( const int32 OtherDupVert : OtherOverlappingWedges )
-							{
-								TArray< int32 > OtherWedgeFaces;
-								WedgeToFaces.MultiFind( OtherDupVert, OtherWedgeFaces );
-
-								for ( const int32 OtherFaceIndex : OtherWedgeFaces )
-								{
-									for ( int32 OtherCornerIndex = 0; OtherCornerIndex < 3; ++OtherCornerIndex )
-									{
-										int32 OtherWedgeIndex = MeshDataBundle.Faces[OtherFaceIndex].WedgeIndex[OtherCornerIndex];
-
-										TangentZ[ OtherWedgeIndex ] = FVector::ZeroVector;
-									}
-								}
-							}
-						}
-					}
-				}
-			}
-		}
-	}
-
-	void ComputeTangents()
-	{
-		bool bComputeNormals = !ImportOptions->ShouldImportNormals() || !BaseImportData.bHasNormals;
-		bool bComputeTangents = !ImportOptions->ShouldImportTangents() || !BaseImportData.bHasTangents;
-		bool bUseMikkTSpace = (ImportOptions->NormalGenerationMethod == EFBXNormalGenerationMethod::MikkTSpace) && (!ImportOptions->ShouldImportNormals() || !ImportOptions->ShouldImportTangents());
-
-		check(MorphLODPoints.Num() == MeshDataBundle.Vertices.Num());
-
-		ETangentOptions::Type TangentOptions = ETangentOptions::BlendOverlappingNormals;
-
-		// MikkTSpace should be use only when the user want to recompute the normals or tangents otherwise should always fallback on builtin
-		if ( bUseMikkTSpace && ( bComputeNormals || bComputeTangents ) )
-		{
-			TangentOptions = (ETangentOptions::Type)(TangentOptions | ETangentOptions::UseMikkTSpace);
-		}
-
-		MeshUtilities->CalculateNormals( MorphLODPoints, MeshDataBundle.Indices, MeshDataBundle.UVs, MeshDataBundle.SmoothingGroups, TangentOptions, TangentZ );
-	}
-
-	void ComputeMorphDeltas()
-	{
-		TArray<bool> WasProcessed;
-		WasProcessed.Empty(LODModel->NumVertices );
-		WasProcessed.AddZeroed(LODModel->NumVertices);
-
-		for ( int32 Idx = 0; Idx < BaseIndexData.Num(); ++Idx )
-		{
-			uint32 BaseVertIdx = BaseIndexData[Idx];
-			// check for duplicate processing
-			if( !WasProcessed[BaseVertIdx] )
-			{
-				// mark this base vertex as already processed
-				WasProcessed[BaseVertIdx] = true;
-
-				// clothing can add extra verts, and we won't have source point, so we ignore those
-				if ( BaseWedgePointIndices.IsValidIndex(BaseVertIdx) )
-				{
-					// get the base mesh's original wedge point index
-					uint32 BasePointIdx = BaseWedgePointIndices[ BaseVertIdx ];
-					if (MeshDataBundle.Vertices.IsValidIndex(BasePointIdx) && MorphLODPoints.IsValidIndex(BasePointIdx))
-					{
-						FVector BasePosition = MeshDataBundle.Vertices[BasePointIdx];
-						FVector TargetPosition = MorphLODPoints[BasePointIdx];
-
-					FVector PositionDelta = TargetPosition - BasePosition;
-
-						uint32* VertexIdx = WedgePointToVertexIndexMap.Find(BasePointIdx);
-
-					FVector NormalDeltaZ = FVector::ZeroVector;
-
-						if (VertexIdx != nullptr)
-					{
-							FVector BaseNormal = BaseTangentZ[*VertexIdx];
-							FVector TargetNormal = TangentZ[*VertexIdx];
-
-						NormalDeltaZ = TargetNormal - BaseNormal;
-					}
-
-					// check if position actually changed much
-						if (PositionDelta.SizeSquared() > FMath::Square(THRESH_POINTS_ARE_NEAR) ||
-						// since we can't get imported morphtarget normal from FBX
-						// we can't compare normal unless it's calculated
-						// this is special flag to ignore normal diff
-							((ImportOptions->ShouldImportNormals() == false) && NormalDeltaZ.SizeSquared() > 0.01f))
-					{
-						// create a new entry
-						FMorphTargetDelta NewVertex;
-						// position delta
-						NewVertex.PositionDelta = PositionDelta;
-						// normal delta
-						NewVertex.TangentZDelta = NormalDeltaZ;
-						// index of base mesh vert this entry is to modify
-						NewVertex.SourceIdx = BaseVertIdx;
-
-						// add it to the list of changed verts
-							MorphTargetDeltas.Add(NewVertex);
-						}
-					}
-				}
-			}
-		}
-	}
-
-	void DoWork()
-	{
-		PrepareTangents();
-		ComputeTangents();
-		ComputeMorphDeltas();
-	}
-
-	FORCEINLINE TStatId GetStatId() const
-	{
-		RETURN_QUICK_DECLARE_CYCLE_STAT(FAsyncImportMorphTargetWork, STATGROUP_ThreadPoolAsyncTasks);
-	}
-
-private:
-
-	const TArray<int32>& FindIncludingNoOverlapping(const FOverlappingCorners& Corners, int32 Key, TArray<int32>& NoOverlapping)
-	{
-		const TArray<int32>& Found = Corners.FindIfOverlapping(Key);
-		if (Found.Num() > 0)
-		{
-			return Found;
-		}
-		else
-		{
-			NoOverlapping.Reset(1);
-			NoOverlapping.Add(Key);
-			return NoOverlapping;
-		}
-	}
 
-	FSkeletalMeshLODModel* LODModel;
-	// @todo not thread safe
-	const FReferenceSkeleton& RefSkeleton;
-	const FSkeletalMeshImportData& BaseImportData;
-	const TArray<FVector> MorphLODPoints;
-
-	IMeshUtilities* MeshUtilities;
-	FBXImportOptions* ImportOptions;
-
-	TArray< FMorphTargetDelta >& MorphTargetDeltas;
-	TArray< uint32 >& BaseIndexData;
-	TArray< uint32 >& BaseWedgePointIndices;
-	TMap<uint32,uint32>& WedgePointToVertexIndexMap;
-
-	const FOverlappingCorners& OverlappingCorners;
-	const TSet<uint32> ModifiedPoints;
-	const TMultiMap< int32, int32 >& WedgeToFaces;
-	const FMeshDataBundle& MeshDataBundle;
-
-	const TArray<FVector>& BaseTangentZ;
-	TArray<FVector> TangentZ;
-};
-
-void UnFbx::FFbxImporter::ImportMorphTargetsInternal( TArray<FbxNode*>& SkelMeshNodeArray, USkeletalMesh* BaseSkelMesh, UObject* InParent, int32 LODIndex, const FSkeletalMeshImportData &BaseImportData)
+void UnFbx::FFbxImporter::ImportMorphTargetsInternal( TArray<FbxNode*>& SkelMeshNodeArray, USkeletalMesh* BaseSkelMesh, int32 LODIndex, FSkeletalMeshImportData &BaseImportData)
 {
 	FbxString ShapeNodeName;
 	TMap<FString, TArray<FbxShape*>> ShapeNameToShapeArray;
 
-	// Temp arrays to keep track of data being used by threads
-	TArray< TArray< FMorphTargetDelta >* > Results;
-	TArray<UMorphTarget*> MorphTargets;
-
-	// Array of pending tasks that are not complete
-	TIndirectArray<FAsyncTask<FAsyncImportMorphTargetWork> > PendingWork;
-
 	GWarn->BeginSlowTask( NSLOCTEXT("FbxImporter", "BeginGeneratingMorphModelsTask", "Generating Morph Models"), true);
 
 	//Use this TMap to store the name of the blendshape map by the fbx uniqueID. This allow to assign unique name to each blend shape.
@@ -4521,179 +4257,37 @@ void UnFbx::FFbxImporter::ImportMorphTargetsInternal( TArray<FbxNode*>& SkelMesh
 		}
 	} // for NodeIndex
 
-	// Prepare base data
-	FSkeletalMeshLODModel& BaseLODModel = BaseSkelMesh->GetImportedModel()->LODModels[ LODIndex ];
-
-	// Calculate overlapping corners and tangents
-	bool bComputeNormals = !ImportOptions->ShouldImportNormals() || !BaseImportData.bHasNormals;
-	bool bComputeTangents = !ImportOptions->ShouldImportTangents() || !BaseImportData.bHasTangents;
-	bool bUseMikkTSpace = (ImportOptions->NormalGenerationMethod == EFBXNormalGenerationMethod::MikkTSpace) && (!ImportOptions->ShouldImportNormals() || !ImportOptions->ShouldImportTangents());
-
-	FMeshDataBundle MeshDataBundle;
-	ConvertImportDataToMeshData( BaseImportData, MeshDataBundle);
-
-	IMeshUtilities& MeshUtilities = FModuleManager::Get().LoadModuleChecked<IMeshUtilities>("MeshUtilities");
-
-	ETangentOptions::Type TangentOptions = ETangentOptions::BlendOverlappingNormals;
-
-	// MikkTSpace should be use only when the user want to recompute the normals or tangents otherwise should always fallback on builtin
-	if ( bUseMikkTSpace && ( bComputeNormals || bComputeTangents ) )
-	{
-		TangentOptions = (ETangentOptions::Type)(TangentOptions | ETangentOptions::UseMikkTSpace);
-	}
-
-	FOverlappingCorners OverlappingVertices;
-	MeshUtilities.CalculateOverlappingCorners(MeshDataBundle.Vertices, MeshDataBundle.Indices, false, OverlappingVertices );
-
-	TArray<FVector> TangentZ;
-	MeshUtilities.CalculateNormals(MeshDataBundle.Vertices, MeshDataBundle.Indices, MeshDataBundle.UVs, MeshDataBundle.SmoothingGroups, TangentOptions, TangentZ );
-
-	TArray< uint32 > BaseWedgePointIndices;
-	if (BaseLODModel.RawPointIndices.GetBulkDataSize())
-	{
-		BaseWedgePointIndices.Empty(BaseLODModel.RawPointIndices.GetElementCount());
-		BaseWedgePointIndices.AddUninitialized(BaseLODModel.RawPointIndices.GetElementCount());
-		FMemory::Memcpy(BaseWedgePointIndices.GetData(), BaseLODModel.RawPointIndices.Lock(LOCK_READ_ONLY), BaseLODModel.RawPointIndices.GetBulkDataSize());
-		BaseLODModel.RawPointIndices.Unlock();
-	}
-
-	TArray<uint32> BaseIndexData = BaseLODModel.IndexBuffer;
-
-	TMap<uint32,uint32> WedgePointToVertexIndexMap;
-	// Build a mapping of wedge point indices to vertex indices for fast lookup later.
-	for( int32 Idx = 0; Idx < MeshDataBundle.Indices.Num(); ++Idx )
-	{
-		WedgePointToVertexIndexMap.Add(MeshDataBundle.Indices[Idx], Idx);
-	}
-
-	// Create a map from wedge indices to faces
-	TMultiMap< int32, int32 > WedgeToFaces;
-	for (int32 FaceIndex = 0; FaceIndex < MeshDataBundle.Faces.Num(); FaceIndex++)
-	{
-		const SkeletalMeshImportData::FTriangle& Face = MeshDataBundle.Faces[FaceIndex];
-		for (int32 CornerIndex = 0; CornerIndex < 3; ++CornerIndex)
-		{
-			WedgeToFaces.AddUnique(Face.WedgeIndex[CornerIndex], FaceIndex);
-		}
-	}
-
-	int32 NumCompleted = 0;
-	int32 NumTasks = 0;
-	int32 MaxShapeInProcess = FPlatformMisc::NumberOfCoresIncludingHyperthreads();
-
-	int32 ShapeIndex = 0;
-	int32 TotalShapeCount = ShapeNameToShapeArray.Num();
-	// iterate through shapename, and create morphtarget
 	for (auto Iter = ShapeNameToShapeArray.CreateIterator(); Iter; ++Iter)
 	{
-		int32 CurrentNumTasks = PendingWork.Num();
-		while (CurrentNumTasks >= MaxShapeInProcess)
-		{
-			//Wait until the first slot is available
-			PendingWork[0].EnsureCompletion();
-			for (int32 TaskIndex = PendingWork.Num()-1; TaskIndex >= 0; --TaskIndex)
-			{
-				if (PendingWork[TaskIndex].IsDone())
-				{
-					PendingWork.RemoveAt(TaskIndex);
-					++NumCompleted;
-					FFormatNamedArguments Args;
-					Args.Add(TEXT("NumCompleted"), NumCompleted);
-					Args.Add(TEXT("NumTasks"), TotalShapeCount);
-					GWarn->StatusUpdate(NumCompleted, TotalShapeCount, FText::Format(LOCTEXT("ImportingMorphTargetStatus", "Importing Morph Target: {NumCompleted} of {NumTasks}"), Args));
-				}
-			}
-			CurrentNumTasks = PendingWork.Num();
-		}
 		FString ShapeName = Iter.Key();
 		TArray< FbxShape* >& ShapeArray = Iter.Value();
-		FSkeletalMeshImportData ShapeImportData = BaseImportData;
-
-		TSet<uint32> ModifiedPoints;
-		GatherPointsForMorphTarget( &ShapeImportData, SkelMeshNodeArray, &ShapeArray, ModifiedPoints );
+		FSkeletalMeshImportData ShapeImportData;
+		BaseImportData.CopyDataNeedByMorphTargetImport(ShapeImportData);
 
-		// See if this morph target already exists.
-		UMorphTarget * MorphTarget = FindObject<UMorphTarget>(BaseSkelMesh, *ShapeName);
-		// we only create new one for LOD0, otherwise don't create new one
-		if (!MorphTarget)
+		//Store the rebuild morph data into the base import data, this will allow us to rebuild the morph data in case the mesh is rebuild and the vertex count change because of options (max bone per section, normals compute...)
+		int32 MorphTargetIndex;
+		if (BaseImportData.MorphTargetNames.Find(ShapeName, MorphTargetIndex))
 		{
-			if (LODIndex == 0)
-			{
-				MorphTarget = NewObject<UMorphTarget>(BaseSkelMesh, FName(*ShapeName));
-			}
-			else
-			{
-				AddTokenizedErrorMessage(FTokenizedMessage::Create(EMessageSeverity::Error, FText::Format(FText::FromString("Could not find the {0} morphtarget for LOD {1}. \
-					Make sure the name for morphtarget matches with LOD 0"), FText::FromString(ShapeName), FText::FromString(FString::FromInt(LODIndex)))),
-					FFbxErrors::SkeletalMesh_LOD_MissingMorphTarget);
-			}
+			BaseImportData.MorphTargetNames.RemoveAt(MorphTargetIndex);
+			BaseImportData.MorphTargetModifiedPoints.RemoveAt(MorphTargetIndex);
+			BaseImportData.MorphTargets.RemoveAt(MorphTargetIndex);
 		}
-
-		if (MorphTarget)
-		{
-			MorphTargets.Add( MorphTarget );
-			int32 NewMorphDeltasIdx = Results.Add( new TArray< FMorphTargetDelta >() );
-
-			TArray< FMorphTargetDelta >* Deltas = Results[ NewMorphDeltasIdx ];
-
-			FAsyncTask<FAsyncImportMorphTargetWork>* NewWork = new FAsyncTask<FAsyncImportMorphTargetWork>( &BaseLODModel, BaseSkelMesh->RefSkeleton, BaseImportData,
-				MoveTemp( ShapeImportData.Points ), ImportOptions, *Deltas, BaseIndexData, BaseWedgePointIndices, WedgePointToVertexIndexMap, OverlappingVertices, MoveTemp( ModifiedPoints ), WedgeToFaces, MeshDataBundle, TangentZ);
-			PendingWork.Add(NewWork);
-
-			NewWork->StartBackgroundTask(GLargeThreadPool);
-			CurrentNumTasks++;
-			NumTasks++;
-		}
-
-		++ShapeIndex;
-	}
-
-	// Wait for all importing tasks to complete
-	for ( int32 TaskIndex = 0; TaskIndex < PendingWork.Num(); ++TaskIndex )
-	{
-		PendingWork[TaskIndex].EnsureCompletion();
-
-		++NumCompleted;
-
-		FFormatNamedArguments Args;
-		Args.Add( TEXT("NumCompleted"), NumCompleted );
-		Args.Add( TEXT("NumTasks"), TotalShapeCount );
-		GWarn->StatusUpdate( NumCompleted, NumTasks, FText::Format( LOCTEXT("ImportingMorphTargetStatus", "Importing Morph Target: {NumCompleted} of {NumTasks}"), Args ) );
+		BaseImportData.MorphTargetNames.Add(ShapeName);
+		TSet<uint32>& ModifiedPoints = BaseImportData.MorphTargetModifiedPoints.AddDefaulted_GetRef();
+		GatherPointsForMorphTarget(&ShapeImportData, SkelMeshNodeArray, &ShapeArray, ModifiedPoints);
+		BaseImportData.MorphTargets.Add(ShapeImportData);
+		check(BaseImportData.MorphTargetNames.Num() == BaseImportData.MorphTargets.Num() && BaseImportData.MorphTargetNames.Num() == BaseImportData.MorphTargetModifiedPoints.Num());
 	}
 
-	bool bNeedToInvalidateRegisteredMorph = false;
-	// Create morph streams for each morph target we are importing.
-	// This has to happen on a single thread since the skeletal meshes' bulk data is locked and cant be accessed by multiple threads simultaneously
-	for (int32 Index = 0; Index < MorphTargets.Num(); Index++)
-	{
-		FFormatNamedArguments Args;
-		Args.Add( TEXT("NumCompleted"), Index+1 );
-		Args.Add(TEXT("NumTasks"), MorphTargets.Num());
-		GWarn->StatusUpdate(Index + 1, MorphTargets.Num(), FText::Format(LOCTEXT("BuildingMorphTargetRenderDataStatus", "Building Morph Target Render Data: {NumCompleted} of {NumTasks}"), Args));
-
-		UMorphTarget* MorphTarget = MorphTargets[Index];
-		MorphTarget->PopulateDeltas(*Results[Index], LODIndex, BaseLODModel.Sections, ImportOptions->ShouldImportNormals() == false);
-
-		// register does mark package as dirty
-		if (MorphTarget->HasValidData())
-		{
-			bNeedToInvalidateRegisteredMorph |= BaseSkelMesh->RegisterMorphTarget(MorphTarget, false);
-		}
-
-		delete Results[Index];
-		Results[Index] = nullptr;
-	}
-
-	if (bNeedToInvalidateRegisteredMorph)
-	{
-		BaseSkelMesh->InitMorphTargetsAndRebuildRenderData();
-	}
+	//Build MorphTargets
+	// Calculate overlapping corners and tangents
+	FLODUtilities::BuildMorphTargets(BaseSkelMesh, BaseImportData, LODIndex, ImportOptions->ShouldImportNormals(), ImportOptions->ShouldImportTangents(), (ImportOptions->NormalGenerationMethod == EFBXNormalGenerationMethod::MikkTSpace));
 
 	GWarn->EndSlowTask();
 }
 
 // Import Morph target
-void UnFbx::FFbxImporter::ImportFbxMorphTarget(TArray<FbxNode*> &SkelMeshNodeArray, USkeletalMesh* BaseSkelMesh, UObject* InParent, int32 LODIndex, const FSkeletalMeshImportData &BaseSkeletalMeshImportData)
+void UnFbx::FFbxImporter::ImportFbxMorphTarget(TArray<FbxNode*> &SkelMeshNodeArray, USkeletalMesh* BaseSkelMesh, int32 LODIndex, FSkeletalMeshImportData &BaseSkeletalMeshImportData)
 {
 	bool bHasMorph = false;
 	int32 NodeIndex;
@@ -4713,7 +4307,7 @@ void UnFbx::FFbxImporter::ImportFbxMorphTarget(TArray<FbxNode*> &SkelMeshNodeArr
 	
 	if (bHasMorph)
 	{
-		ImportMorphTargetsInternal( SkelMeshNodeArray, BaseSkelMesh, InParent, LODIndex, BaseSkeletalMeshImportData );
+		ImportMorphTargetsInternal( SkelMeshNodeArray, BaseSkelMesh, LODIndex, BaseSkeletalMeshImportData );
 	}
 }
 
diff --git a/Engine/Source/Editor/UnrealEd/Private/FbxMeshUtils.cpp b/Engine/Source/Editor/UnrealEd/Private/FbxMeshUtils.cpp
--- a/Engine/Source/Editor/UnrealEd/Private/FbxMeshUtils.cpp
+++ b/Engine/Source/Editor/UnrealEd/Private/FbxMeshUtils.cpp
@@ -539,7 +539,7 @@ namespace FbxMeshUtils
 
 					if (ImportOptions->bImportMorph)
 					{
-						FFbxImporter->ImportFbxMorphTarget(SkelMeshNodeArray, SelectedSkelMesh, SelectedSkelMesh->GetOutermost(), SelectedLOD, OutData);
+						FFbxImporter->ImportFbxMorphTarget(SkelMeshNodeArray, SelectedSkelMesh, SelectedLOD, OutData);
 					}
 
 					bSuccess = true;
diff --git a/Engine/Source/Editor/UnrealEd/Private/LODUtilities.cpp b/Engine/Source/Editor/UnrealEd/Private/LODUtilities.cpp
--- a/Engine/Source/Editor/UnrealEd/Private/LODUtilities.cpp
+++ b/Engine/Source/Editor/UnrealEd/Private/LODUtilities.cpp
@@ -14,6 +14,7 @@
 #include "Engine/SkeletalMesh.h"
 #include "EditorFramework/AssetImportData.h"
 #include "MeshUtilities.h"
+#include "OverlappingCorners.h"
 #include "Assets/ClothingAsset.h"
 
 #include "IAssetTools.h"
@@ -745,11 +746,15 @@ void FLODUtilities::ApplyMorphTargetsToLOD(USkeletalMesh* SkeletalMesh, int32 So
 	TargetLODModel.GetNonClothVertices(TargetVertices);
 	//Create the base triangle indices per section
 	TArray<TArray<uint32>> BaseTriangleIndices;
-	int32 SectionCount = BaseLODModel.NumNonClothingSections();
+	int32 SectionCount = BaseLODModel.Sections.Num();
 	BaseTriangleIndices.AddDefaulted(SectionCount);
 	for (int32 SectionIndex = 0; SectionIndex < SectionCount; ++SectionIndex)
 	{
 		const FSkelMeshSection& Section = BaseLODModel.Sections[SectionIndex];
+		if (Section.HasClothingData())
+		{
+			continue;
+		}
 		uint32 TriangleCount = Section.NumTriangles;
 		for (uint32 TriangleIndex = 0; TriangleIndex < TriangleCount; ++TriangleIndex)
 		{
@@ -1443,7 +1448,7 @@ void MatchVertexIndexUsingPosition(
 	}
 }
 
-bool FLODUtilities::UpdateAlternateSkinWeights(USkeletalMesh* SkeletalMeshDest, const FName& ProfileNameDest, USkeletalMesh* SkeletalMeshSrc, const UnFbx::FBXImportOptions& ImportOptions, int32 LODIndexDest, int32 LODIndexSrc)
+bool FLODUtilities::UpdateAlternateSkinWeights(USkeletalMesh* SkeletalMeshDest, const FName& ProfileNameDest, const UnFbx::FBXImportOptions& ImportOptions, int32 LODIndexDest)
 {
 	//Ensure log message only once
 	bool bNoMatchMsgDone = false;
@@ -1464,15 +1469,15 @@ bool FLODUtilities::UpdateAlternateSkinWeights(USkeletalMesh* SkeletalMeshDest,
 	int32 PointNumberDest = ImportDataDest.Points.Num();
 	int32 VertexNumberDest = ImportDataDest.Points.Num();
 
-	//Grab all the source structure
-	check(SkeletalMeshSrc);
-	check(SkeletalMeshSrc->GetImportedModel());
-	check(SkeletalMeshSrc->GetImportedModel()->LODModels.IsValidIndex(LODIndexSrc));
-	FSkeletalMeshLODModel& LODModelSrc = SkeletalMeshSrc->GetImportedModel()->LODModels[LODIndexSrc];
-	//The source model is a fresh import and the data need to be there
-	check(!LODModelSrc.RawSkeletalMeshBulkData.IsEmpty());
-	FSkeletalMeshImportData ImportDataSrc;
-	LODModelSrc.RawSkeletalMeshBulkData.LoadRawMesh(ImportDataSrc);
+	int32 ProfileIndex = 0;
+	if (!ImportDataDest.AlternateInfluenceProfileNames.Find(ProfileNameDest.ToString(), ProfileIndex))
+	{
+		UE_LOG(LogLODUtilities, Error, TEXT("Failed to import Skin Weight Profile the alternate skinning imported source data is not available."), SkeletalMeshDest ? *SkeletalMeshDest->GetName() : TEXT("NULL"));
+		return false;
+	}
+	check(ImportDataDest.AlternateInfluences.IsValidIndex(ProfileIndex));
+	//The data must be there and must be verified before getting here
+	const FSkeletalMeshImportData& ImportDataSrc = ImportDataDest.AlternateInfluences[ProfileIndex];
 	int32 PointNumberSrc = ImportDataSrc.Points.Num();
 	int32 VertexNumberSrc = ImportDataSrc.Points.Num();
 	int32 InfluenceNumberSrc = ImportDataSrc.Influences.Num();
@@ -1569,7 +1574,7 @@ bool FLODUtilities::UpdateAlternateSkinWeights(USkeletalMesh* SkeletalMeshDest,
 	
 	for (int32 InfluenceIndexSrc = 0; InfluenceIndexSrc < InfluenceNumberSrc; ++InfluenceIndexSrc)
 	{
-		SkeletalMeshImportData::FRawBoneInfluence& InfluenceSrc = ImportDataSrc.Influences[InfluenceIndexSrc];
+		const SkeletalMeshImportData::FRawBoneInfluence& InfluenceSrc = ImportDataSrc.Influences[InfluenceIndexSrc];
 		uint32 VertexIndexSource = InfluenceSrc.VertexIndex;
 		uint32 BoneIndexSource = InfluenceSrc.BoneIndex;
 		float Weight = InfluenceSrc.Weight;
@@ -1729,6 +1734,53 @@ bool FLODUtilities::UpdateAlternateSkinWeights(USkeletalMesh* SkeletalMeshDest,
 	return bBuildSuccess;
 }
 
+bool FLODUtilities::UpdateAlternateSkinWeights(USkeletalMesh* SkeletalMeshDest, const FName& ProfileNameDest, USkeletalMesh* SkeletalMeshSrc, const UnFbx::FBXImportOptions& ImportOptions, int32 LODIndexDest, int32 LODIndexSrc)
+{
+	//Grab all the destination structure
+	check(SkeletalMeshDest);
+	check(SkeletalMeshDest->GetImportedModel());
+	check(SkeletalMeshDest->GetImportedModel()->LODModels.IsValidIndex(LODIndexDest));
+	FSkeletalMeshLODModel& LODModelDest = SkeletalMeshDest->GetImportedModel()->LODModels[LODIndexDest];
+	if (LODModelDest.RawSkeletalMeshBulkData.IsEmpty())
+	{
+		UE_LOG(LogLODUtilities, Error, TEXT("Failed to import Skin Weight Profile as the target skeletal mesh (%s) requires reimporting first."), SkeletalMeshDest ? *SkeletalMeshDest->GetName() : TEXT("NULL"));
+		//Very old asset will not have this data, we cannot add alternate until the asset is reimported
+		return false;
+	}
+	FSkeletalMeshImportData ImportDataDest;
+	LODModelDest.RawSkeletalMeshBulkData.LoadRawMesh(ImportDataDest);
+	int32 PointNumberDest = ImportDataDest.Points.Num();
+	int32 VertexNumberDest = ImportDataDest.Points.Num();
+
+	//Grab all the source structure
+	check(SkeletalMeshSrc);
+	check(SkeletalMeshSrc->GetImportedModel());
+	check(SkeletalMeshSrc->GetImportedModel()->LODModels.IsValidIndex(LODIndexSrc));
+	FSkeletalMeshLODModel& LODModelSrc = SkeletalMeshSrc->GetImportedModel()->LODModels[LODIndexSrc];
+	//The source model is a fresh import and the data need to be there
+	check(!LODModelSrc.RawSkeletalMeshBulkData.IsEmpty());
+	FSkeletalMeshImportData ImportDataSrc;
+	LODModelSrc.RawSkeletalMeshBulkData.LoadRawMesh(ImportDataSrc);
+	
+	//Remove all unnecessary array data from the structure (this will save a lot of memory)
+	ImportDataSrc.KeepAlternateSkinningBuildDataOnly();
+
+	//Replace the data into the destination bulk data and save it
+	int32 ProfileIndex = 0;
+	if (ImportDataDest.AlternateInfluenceProfileNames.Find(ProfileNameDest.ToString(), ProfileIndex))
+	{
+		ImportDataDest.AlternateInfluenceProfileNames.RemoveAt(ProfileIndex);
+		ImportDataDest.AlternateInfluences.RemoveAt(ProfileIndex);
+	}
+	ImportDataDest.AlternateInfluenceProfileNames.Add(ProfileNameDest.ToString());
+	ImportDataDest.AlternateInfluences.Add(ImportDataSrc);
+	//Resave the bulk data with the new or refreshed data
+	LODModelDest.RawSkeletalMeshBulkData.SaveRawMesh(ImportDataDest);
+
+	//Build the alternate buffer with all the data into the bulk
+	return UpdateAlternateSkinWeights(SkeletalMeshDest, ProfileNameDest, ImportOptions, LODIndexDest);
+}
+
 void FLODUtilities::GenerateImportedSkinWeightProfileData(const FSkeletalMeshLODModel& LODModelDest, FImportedSkinWeightProfileData &ImportedProfileData)
 {
 	//Add the override buffer with the alternate influence data
@@ -2197,4 +2249,452 @@ FString FLODUtilities::PickSkinWeightFBXPath(int32 LODIndex)
 	return PickedFileName;
 }
 
+//////////////////////////////////////////////////////////////////////////
+// Morph targets build code
+//
+
+struct FMeshDataBundle
+{
+	TArray< FVector > Vertices;
+	TArray< uint32 > Indices;
+	TArray< FVector2D > UVs;
+	TArray< uint32 > SmoothingGroups;
+	TArray<SkeletalMeshImportData::FTriangle> Faces;
+};
+
+static void ConvertImportDataToMeshData(const FSkeletalMeshImportData& ImportData, FMeshDataBundle& MeshDataBundle)
+{
+	for (const SkeletalMeshImportData::FTriangle& Face : ImportData.Faces)
+	{
+		SkeletalMeshImportData::FTriangle FaceTriangle;
+		FaceTriangle = Face;
+		for (int32 i = 0; i < 3; ++i)
+		{
+			const SkeletalMeshImportData::FVertex& Wedge = ImportData.Wedges[Face.WedgeIndex[i]];
+			int32 FaceWedgeIndex = MeshDataBundle.Indices.Add(Wedge.VertexIndex);
+			MeshDataBundle.UVs.Add(Wedge.UVs[0]);
+			FaceTriangle.WedgeIndex[i] = FaceWedgeIndex;
+		}
+		MeshDataBundle.Faces.Add(FaceTriangle);
+		MeshDataBundle.SmoothingGroups.Add(Face.SmoothingGroups);
+	}
+
+	MeshDataBundle.Vertices = ImportData.Points;
+}
+
+/**
+* A class encapsulating morph target processing that occurs during import on a separate thread
+*/
+class FAsyncImportMorphTargetWork : public FNonAbandonableTask
+{
+public:
+	FAsyncImportMorphTargetWork(FSkeletalMeshLODModel* InLODModel, const FReferenceSkeleton& InRefSkeleton, const FSkeletalMeshImportData& InBaseImportData, TArray<FVector>&& InMorphLODPoints,
+		TArray< FMorphTargetDelta >& InMorphDeltas, TArray<uint32>& InBaseIndexData, TArray< uint32 >& InBaseWedgePointIndices,
+		TMap<uint32, uint32>& InWedgePointToVertexIndexMap, const FOverlappingCorners& InOverlappingCorners,
+		const TSet<uint32> InModifiedPoints, const TMultiMap< int32, int32 >& InWedgeToFaces, const FMeshDataBundle& InMeshDataBundle, const TArray<FVector>& InTangentZ,
+		bool InShouldImportNormals, bool InShouldImportTangents, bool InbUseMikkTSpace)
+		: LODModel(InLODModel)
+		, RefSkeleton(InRefSkeleton)
+		, BaseImportData(InBaseImportData)
+		, MorphLODPoints(InMorphLODPoints)
+		, MorphTargetDeltas(InMorphDeltas)
+		, BaseIndexData(InBaseIndexData)
+		, BaseWedgePointIndices(InBaseWedgePointIndices)
+		, WedgePointToVertexIndexMap(InWedgePointToVertexIndexMap)
+		, OverlappingCorners(InOverlappingCorners)
+		, ModifiedPoints(InModifiedPoints)
+		, WedgeToFaces(InWedgeToFaces)
+		, MeshDataBundle(InMeshDataBundle)
+		, BaseTangentZ(InTangentZ)
+		, TangentZ(InTangentZ)
+		, ShouldImportNormals(InShouldImportNormals)
+		, ShouldImportTangents(InShouldImportTangents)
+		, bUseMikkTSpace(InbUseMikkTSpace)
+	{
+		MeshUtilities = &FModuleManager::Get().LoadModuleChecked<IMeshUtilities>("MeshUtilities");
+	}
+
+	void PrepareTangents()
+	{
+		TArray<bool> WasProcessed;
+		WasProcessed.Empty(MeshDataBundle.Indices.Num());
+		WasProcessed.AddZeroed(MeshDataBundle.Indices.Num());
+
+		TArray< int32 > WedgeFaces;
+		TArray< int32 > OverlappingWedgesDummy;
+		TArray< int32 > OtherOverlappingWedgesDummy;
+
+		// For each ModifiedPoints, reset the tangents for the affected wedges
+		for (int32 WedgeIdx = 0; WedgeIdx < MeshDataBundle.Indices.Num(); ++WedgeIdx)
+		{
+			int32 PointIdx = MeshDataBundle.Indices[WedgeIdx];
+
+			if (ModifiedPoints.Find(PointIdx) != nullptr)
+			{
+				TangentZ[WedgeIdx] = FVector::ZeroVector;
+
+				const TArray<int32>& OverlappingWedges = FindIncludingNoOverlapping(OverlappingCorners, WedgeIdx, OverlappingWedgesDummy);
+
+				for (const int32 OverlappingWedgeIndex : OverlappingWedges)
+				{
+					if (WasProcessed[OverlappingWedgeIndex])
+					{
+						continue;
+					}
+
+					WasProcessed[OverlappingWedgeIndex] = true;
+
+					WedgeFaces.Reset();
+					WedgeToFaces.MultiFind(OverlappingWedgeIndex, WedgeFaces);
+
+					for (const int32 FaceIndex : WedgeFaces)
+					{
+						for (int32 CornerIndex = 0; CornerIndex < 3; ++CornerIndex)
+						{
+							int32 WedgeIndex = MeshDataBundle.Faces[FaceIndex].WedgeIndex[CornerIndex];
+
+							TangentZ[WedgeIndex] = FVector::ZeroVector;
+
+							const TArray<int32>& OtherOverlappingWedges = FindIncludingNoOverlapping(OverlappingCorners, WedgeIndex, OtherOverlappingWedgesDummy);
+
+							for (const int32 OtherDupVert : OtherOverlappingWedges)
+							{
+								TArray< int32 > OtherWedgeFaces;
+								WedgeToFaces.MultiFind(OtherDupVert, OtherWedgeFaces);
+
+								for (const int32 OtherFaceIndex : OtherWedgeFaces)
+								{
+									for (int32 OtherCornerIndex = 0; OtherCornerIndex < 3; ++OtherCornerIndex)
+									{
+										int32 OtherWedgeIndex = MeshDataBundle.Faces[OtherFaceIndex].WedgeIndex[OtherCornerIndex];
+
+										TangentZ[OtherWedgeIndex] = FVector::ZeroVector;
+									}
+								}
+							}
+						}
+					}
+				}
+			}
+		}
+	}
+
+	void ComputeTangents()
+	{
+		bool bComputeNormals = !ShouldImportNormals || !BaseImportData.bHasNormals;
+		bool bComputeTangents = !ShouldImportTangents || !BaseImportData.bHasTangents;
+		bool bUseMikkTSpaceFinal = bUseMikkTSpace && (!ShouldImportNormals || !ShouldImportTangents);
+
+		check(MorphLODPoints.Num() == MeshDataBundle.Vertices.Num());
+
+		ETangentOptions::Type TangentOptions = ETangentOptions::BlendOverlappingNormals;
+
+		// MikkTSpace should be use only when the user want to recompute the normals or tangents otherwise should always fallback on builtin
+		if (bUseMikkTSpaceFinal && (bComputeNormals || bComputeTangents))
+		{
+			TangentOptions = (ETangentOptions::Type)(TangentOptions | ETangentOptions::UseMikkTSpace);
+		}
+
+		MeshUtilities->CalculateNormals(MorphLODPoints, MeshDataBundle.Indices, MeshDataBundle.UVs, MeshDataBundle.SmoothingGroups, TangentOptions, TangentZ);
+	}
+
+	void ComputeMorphDeltas()
+	{
+		TArray<bool> WasProcessed;
+		WasProcessed.Empty(LODModel->NumVertices);
+		WasProcessed.AddZeroed(LODModel->NumVertices);
+
+		for (int32 Idx = 0; Idx < BaseIndexData.Num(); ++Idx)
+		{
+			uint32 BaseVertIdx = BaseIndexData[Idx];
+			// check for duplicate processing
+			if (!WasProcessed[BaseVertIdx])
+			{
+				// mark this base vertex as already processed
+				WasProcessed[BaseVertIdx] = true;
+
+				// clothing can add extra verts, and we won't have source point, so we ignore those
+				if (BaseWedgePointIndices.IsValidIndex(BaseVertIdx))
+				{
+					// get the base mesh's original wedge point index
+					uint32 BasePointIdx = BaseWedgePointIndices[BaseVertIdx];
+					if (MeshDataBundle.Vertices.IsValidIndex(BasePointIdx) && MorphLODPoints.IsValidIndex(BasePointIdx))
+					{
+						FVector BasePosition = MeshDataBundle.Vertices[BasePointIdx];
+						FVector TargetPosition = MorphLODPoints[BasePointIdx];
+
+						FVector PositionDelta = TargetPosition - BasePosition;
+
+						uint32* VertexIdx = WedgePointToVertexIndexMap.Find(BasePointIdx);
+
+						FVector NormalDeltaZ = FVector::ZeroVector;
+
+						if (VertexIdx != nullptr)
+						{
+							FVector BaseNormal = BaseTangentZ[*VertexIdx];
+							FVector TargetNormal = TangentZ[*VertexIdx];
+
+							NormalDeltaZ = TargetNormal - BaseNormal;
+						}
+
+						// check if position actually changed much
+						if (PositionDelta.SizeSquared() > FMath::Square(THRESH_POINTS_ARE_NEAR) ||
+							// since we can't get imported morphtarget normal from FBX
+							// we can't compare normal unless it's calculated
+							// this is special flag to ignore normal diff
+							((ShouldImportNormals == false) && NormalDeltaZ.SizeSquared() > 0.01f))
+						{
+							// create a new entry
+							FMorphTargetDelta NewVertex;
+							// position delta
+							NewVertex.PositionDelta = PositionDelta;
+							// normal delta
+							NewVertex.TangentZDelta = NormalDeltaZ;
+							// index of base mesh vert this entry is to modify
+							NewVertex.SourceIdx = BaseVertIdx;
+
+							// add it to the list of changed verts
+							MorphTargetDeltas.Add(NewVertex);
+						}
+					}
+				}
+			}
+		}
+	}
+
+	void DoWork()
+	{
+		PrepareTangents();
+		ComputeTangents();
+		ComputeMorphDeltas();
+	}
+
+	FORCEINLINE TStatId GetStatId() const
+	{
+		RETURN_QUICK_DECLARE_CYCLE_STAT(FAsyncImportMorphTargetWork, STATGROUP_ThreadPoolAsyncTasks);
+	}
+
+private:
+
+	const TArray<int32>& FindIncludingNoOverlapping(const FOverlappingCorners& Corners, int32 Key, TArray<int32>& NoOverlapping)
+	{
+		const TArray<int32>& Found = Corners.FindIfOverlapping(Key);
+		if (Found.Num() > 0)
+		{
+			return Found;
+		}
+		else
+		{
+			NoOverlapping.Reset(1);
+			NoOverlapping.Add(Key);
+			return NoOverlapping;
+		}
+	}
+
+	FSkeletalMeshLODModel* LODModel;
+	// @todo not thread safe
+	const FReferenceSkeleton& RefSkeleton;
+	const FSkeletalMeshImportData& BaseImportData;
+	const TArray<FVector> MorphLODPoints;
+
+	IMeshUtilities* MeshUtilities;
+
+	TArray< FMorphTargetDelta >& MorphTargetDeltas;
+	TArray< uint32 >& BaseIndexData;
+	TArray< uint32 >& BaseWedgePointIndices;
+	TMap<uint32, uint32>& WedgePointToVertexIndexMap;
+
+	const FOverlappingCorners& OverlappingCorners;
+	const TSet<uint32> ModifiedPoints;
+	const TMultiMap< int32, int32 >& WedgeToFaces;
+	const FMeshDataBundle& MeshDataBundle;
+
+	const TArray<FVector>& BaseTangentZ;
+	TArray<FVector> TangentZ;
+	bool ShouldImportNormals;
+	bool ShouldImportTangents;
+	bool bUseMikkTSpace;
+};
+
+void FLODUtilities::BuildMorphTargets(USkeletalMesh* BaseSkelMesh, FSkeletalMeshImportData &BaseImportData, int32 LODIndex, bool ShouldImportNormals, bool ShouldImportTangents, bool bUseMikkTSpace)
+{
+	bool bComputeNormals = !ShouldImportNormals || !BaseImportData.bHasNormals;
+	bool bComputeTangents = !ShouldImportTangents || !BaseImportData.bHasTangents;
+	bool bUseMikkTSpaceFinal = bUseMikkTSpace && (!ShouldImportNormals || !ShouldImportTangents);
+
+	// Prepare base data
+	FSkeletalMeshLODModel& BaseLODModel = BaseSkelMesh->GetImportedModel()->LODModels[LODIndex];
+
+	FMeshDataBundle MeshDataBundle;
+	ConvertImportDataToMeshData(BaseImportData, MeshDataBundle);
+
+	IMeshUtilities& MeshUtilities = FModuleManager::Get().LoadModuleChecked<IMeshUtilities>("MeshUtilities");
+
+	ETangentOptions::Type TangentOptions = ETangentOptions::BlendOverlappingNormals;
+
+	// MikkTSpace should be use only when the user want to recompute the normals or tangents otherwise should always fallback on builtin
+	if (bUseMikkTSpaceFinal && (bComputeNormals || bComputeTangents))
+	{
+		TangentOptions = (ETangentOptions::Type)(TangentOptions | ETangentOptions::UseMikkTSpace);
+	}
+
+	FOverlappingCorners OverlappingVertices;
+	MeshUtilities.CalculateOverlappingCorners(MeshDataBundle.Vertices, MeshDataBundle.Indices, false, OverlappingVertices);
+
+	TArray<FVector> TangentZ;
+	MeshUtilities.CalculateNormals(MeshDataBundle.Vertices, MeshDataBundle.Indices, MeshDataBundle.UVs, MeshDataBundle.SmoothingGroups, TangentOptions, TangentZ);
+
+	TArray< uint32 > BaseWedgePointIndices;
+	if (BaseLODModel.RawPointIndices.GetBulkDataSize())
+	{
+		BaseWedgePointIndices.Empty(BaseLODModel.RawPointIndices.GetElementCount());
+		BaseWedgePointIndices.AddUninitialized(BaseLODModel.RawPointIndices.GetElementCount());
+		FMemory::Memcpy(BaseWedgePointIndices.GetData(), BaseLODModel.RawPointIndices.Lock(LOCK_READ_ONLY), BaseLODModel.RawPointIndices.GetBulkDataSize());
+		BaseLODModel.RawPointIndices.Unlock();
+	}
+
+	TArray<uint32> BaseIndexData = BaseLODModel.IndexBuffer;
+
+	TMap<uint32, uint32> WedgePointToVertexIndexMap;
+	// Build a mapping of wedge point indices to vertex indices for fast lookup later.
+	for (int32 Idx = 0; Idx < MeshDataBundle.Indices.Num(); ++Idx)
+	{
+		WedgePointToVertexIndexMap.Add(MeshDataBundle.Indices[Idx], Idx);
+	}
+
+	// Create a map from wedge indices to faces
+	TMultiMap< int32, int32 > WedgeToFaces;
+	for (int32 FaceIndex = 0; FaceIndex < MeshDataBundle.Faces.Num(); FaceIndex++)
+	{
+		const SkeletalMeshImportData::FTriangle& Face = MeshDataBundle.Faces[FaceIndex];
+		for (int32 CornerIndex = 0; CornerIndex < 3; ++CornerIndex)
+		{
+			WedgeToFaces.AddUnique(Face.WedgeIndex[CornerIndex], FaceIndex);
+		}
+	}
+
+	// Temp arrays to keep track of data being used by threads
+	TArray< TArray< FMorphTargetDelta >* > Results;
+	TArray<UMorphTarget*> MorphTargets;
+
+	// Array of pending tasks that are not complete
+	TIndirectArray<FAsyncTask<FAsyncImportMorphTargetWork> > PendingWork;
+
+	int32 NumCompleted = 0;
+	int32 NumTasks = 0;
+	int32 MaxShapeInProcess = FPlatformMisc::NumberOfCoresIncludingHyperthreads();
+
+	int32 ShapeIndex = 0;
+	int32 TotalShapeCount = BaseImportData.MorphTargetNames.Num();
+
+	// iterate through shapename, and create morphtarget
+	for (int32 MorphTargetIndex = 0; MorphTargetIndex < BaseImportData.MorphTargetNames.Num(); ++MorphTargetIndex)
+	{
+		int32 CurrentNumTasks = PendingWork.Num();
+		while (CurrentNumTasks >= MaxShapeInProcess)
+		{
+			//Wait until the first slot is available
+			PendingWork[0].EnsureCompletion();
+			for (int32 TaskIndex = PendingWork.Num() - 1; TaskIndex >= 0; --TaskIndex)
+			{
+				if (PendingWork[TaskIndex].IsDone())
+				{
+					PendingWork.RemoveAt(TaskIndex);
+					++NumCompleted;
+					FFormatNamedArguments Args;
+					Args.Add(TEXT("NumCompleted"), NumCompleted);
+					Args.Add(TEXT("NumTasks"), TotalShapeCount);
+					GWarn->StatusUpdate(NumCompleted, TotalShapeCount, FText::Format(LOCTEXT("ImportingMorphTargetStatus", "Importing Morph Target: {NumCompleted} of {NumTasks}"), Args));
+				}
+			}
+			CurrentNumTasks = PendingWork.Num();
+		}
+
+		check(BaseImportData.MorphTargetNames.IsValidIndex(MorphTargetIndex));
+		check(BaseImportData.MorphTargetModifiedPoints.IsValidIndex(MorphTargetIndex));
+		check(BaseImportData.MorphTargets.IsValidIndex(MorphTargetIndex));
+
+		FString& ShapeName = BaseImportData.MorphTargetNames[MorphTargetIndex];
+		FSkeletalMeshImportData& ShapeImportData = BaseImportData.MorphTargets[MorphTargetIndex];
+		TSet<uint32>& ModifiedPoints = BaseImportData.MorphTargetModifiedPoints[MorphTargetIndex];
+
+		// See if this morph target already exists.
+		UMorphTarget * MorphTarget = FindObject<UMorphTarget>(BaseSkelMesh, *ShapeName);
+		// we only create new one for LOD0, otherwise don't create new one
+		if (!MorphTarget)
+		{
+			if (LODIndex == 0)
+			{
+				MorphTarget = NewObject<UMorphTarget>(BaseSkelMesh, FName(*ShapeName));
+			}
+			else
+			{
+				/*AddTokenizedErrorMessage(FTokenizedMessage::Create(EMessageSeverity::Error, FText::Format(FText::FromString("Could not find the {0} morphtarget for LOD {1}. \
+					Make sure the name for morphtarget matches with LOD 0"), FText::FromString(ShapeName), FText::FromString(FString::FromInt(LODIndex)))),
+					FFbxErrors::SkeletalMesh_LOD_MissingMorphTarget);*/
+			}
+		}
+
+		if (MorphTarget)
+		{
+			MorphTargets.Add(MorphTarget);
+			int32 NewMorphDeltasIdx = Results.Add(new TArray< FMorphTargetDelta >());
+
+			TArray< FMorphTargetDelta >* Deltas = Results[NewMorphDeltasIdx];
+
+			FAsyncTask<FAsyncImportMorphTargetWork>* NewWork = new FAsyncTask<FAsyncImportMorphTargetWork>(&BaseLODModel, BaseSkelMesh->RefSkeleton, BaseImportData,
+				MoveTemp(ShapeImportData.Points), *Deltas, BaseIndexData, BaseWedgePointIndices, WedgePointToVertexIndexMap, OverlappingVertices, MoveTemp(ModifiedPoints), WedgeToFaces, MeshDataBundle, TangentZ,
+				ShouldImportNormals, ShouldImportTangents, bUseMikkTSpace);
+			PendingWork.Add(NewWork);
+
+			NewWork->StartBackgroundTask(GLargeThreadPool);
+			CurrentNumTasks++;
+			NumTasks++;
+		}
+
+		++ShapeIndex;
+	}
+
+	// Wait for all importing tasks to complete
+	for (int32 TaskIndex = 0; TaskIndex < PendingWork.Num(); ++TaskIndex)
+	{
+		PendingWork[TaskIndex].EnsureCompletion();
+
+		++NumCompleted;
+
+		FFormatNamedArguments Args;
+		Args.Add(TEXT("NumCompleted"), NumCompleted);
+		Args.Add(TEXT("NumTasks"), TotalShapeCount);
+		GWarn->StatusUpdate(NumCompleted, NumTasks, FText::Format(LOCTEXT("ImportingMorphTargetStatus", "Importing Morph Target: {NumCompleted} of {NumTasks}"), Args));
+	}
+
+	bool bNeedToInvalidateRegisteredMorph = false;
+	// Create morph streams for each morph target we are importing.
+	// This has to happen on a single thread since the skeletal meshes' bulk data is locked and cant be accessed by multiple threads simultaneously
+	for (int32 Index = 0; Index < MorphTargets.Num(); Index++)
+	{
+		FFormatNamedArguments Args;
+		Args.Add(TEXT("NumCompleted"), Index + 1);
+		Args.Add(TEXT("NumTasks"), MorphTargets.Num());
+		GWarn->StatusUpdate(Index + 1, MorphTargets.Num(), FText::Format(LOCTEXT("BuildingMorphTargetRenderDataStatus", "Building Morph Target Render Data: {NumCompleted} of {NumTasks}"), Args));
+
+		UMorphTarget* MorphTarget = MorphTargets[Index];
+		MorphTarget->PopulateDeltas(*Results[Index], LODIndex, BaseLODModel.Sections, ShouldImportNormals == false);
+
+		// register does mark package as dirty
+		if (MorphTarget->HasValidData())
+		{
+			bNeedToInvalidateRegisteredMorph |= BaseSkelMesh->RegisterMorphTarget(MorphTarget, false);
+		}
+
+		delete Results[Index];
+		Results[Index] = nullptr;
+	}
+
+	if (bNeedToInvalidateRegisteredMorph)
+	{
+		BaseSkelMesh->InitMorphTargetsAndRebuildRenderData();
+	}
+}
+
 #undef LOCTEXT_NAMESPACE // "LODUtilities"
\ No newline at end of file
diff --git a/Engine/Source/Editor/UnrealEd/Public/FbxImporter.h b/Engine/Source/Editor/UnrealEd/Public/FbxImporter.h
--- a/Engine/Source/Editor/UnrealEd/Public/FbxImporter.h
+++ b/Engine/Source/Editor/UnrealEd/Public/FbxImporter.h
@@ -908,7 +908,7 @@ public:
 	 * @param BaseSkelMesh - base Skeletal Mesh
 	 * @param LODIndex - LOD index
 	 */
-	UNREALED_API void ImportFbxMorphTarget(TArray<FbxNode*> &SkelMeshNodeArray, USkeletalMesh* BaseSkelMesh, UObject* Parent, int32 LODIndex, const FSkeletalMeshImportData &BaseSkeletalMeshImportData);
+	UNREALED_API void ImportFbxMorphTarget(TArray<FbxNode*> &SkelMeshNodeArray, USkeletalMesh* BaseSkelMesh, int32 LODIndex, FSkeletalMeshImportData &BaseSkeletalMeshImportData);
 
 	/**
 	 * Import LOD object for skeletal mesh
@@ -1173,7 +1173,7 @@ private:
 	 * @param BaseSkelMesh - base Skeletal Mesh
 	 * @param LODIndex - LOD index of the skeletal mesh
 	 */
-	void ImportMorphTargetsInternal( TArray<FbxNode*>& SkelMeshNodeArray, USkeletalMesh* BaseSkelMesh, UObject* Parent, int32 LODIndex, const FSkeletalMeshImportData &BaseSkeletalMeshImportData);
+	void ImportMorphTargetsInternal( TArray<FbxNode*>& SkelMeshNodeArray, USkeletalMesh* BaseSkelMesh, int32 LODIndex, FSkeletalMeshImportData &BaseSkeletalMeshImportData);
 
 	/**
 	* sub-method called from ImportSkeletalMeshLOD method
diff --git a/Engine/Source/Editor/UnrealEd/Public/LODUtilities.h b/Engine/Source/Editor/UnrealEd/Public/LODUtilities.h
--- a/Engine/Source/Editor/UnrealEd/Public/LODUtilities.h
+++ b/Engine/Source/Editor/UnrealEd/Public/LODUtilities.h
@@ -109,18 +109,34 @@ public:
 	 */
 	static FString PickSkinWeightFBXPath(int32 LODIndex);
 
+	/*
+	 * Build the morph targets for the specified LOD. The function use the Morph target data stored in the FSkeletalMeshImportData ImportData structure
+	 */
+	static void BuildMorphTargets(USkeletalMesh* SkeletalMesh, class FSkeletalMeshImportData &ImportData, int32 LODIndex, bool ShouldImportNormals, bool ShouldImportTangents, bool bUseMikkTSpace);
+
 private:
 	FLODUtilities() {}
 
 	/**
 	 *	This function apply the skinning weights from asource skeletal mesh to the destination skeletal mesh.
 	 *  The Destination will receive the weights has the alternate weights.
+	 *  We extract the imported skinning weight data from the SkeletalMeshSrc and we save the imported raw data into the destination mesh.
+	 *  Then we call UpdateAlternateSkinWeights without the SkeletalMeshSrc
 	 *
 	 * @param SkeletalMeshDest - The skeletal mesh that will receive the alternate skinning weights.
 	 * @param SkeletalMeshSrc - The skeletal mesh that contain the alternate skinning weights.
-	 * @param LODIndex - the destination LOD
+	 * @param LODIndexDest - the destination LOD
+	 * @param LODIndexSrc - the Source LOD index
 	 */
 	static bool UpdateAlternateSkinWeights(USkeletalMesh* SkeletalMeshDest, const FName& ProfileNameDest, USkeletalMesh* SkeletalMeshSrc, const UnFbx::FBXImportOptions& ImportOptions, int32 LODIndexDest, int32 LODIndexSrc);
+	/**
+	 *	This function apply the skinning weights from the saved imported skinning weight data to the destination skeletal mesh.
+	 *  The Destination will receive the weights has the alternate weights.
+	 *
+	 * @param SkeletalMeshDest - The skeletal mesh that will receive the alternate skinning weights.
+	 * @param LODIndexDest - the destination LOD
+	 */
+	static bool UpdateAlternateSkinWeights(USkeletalMesh* SkeletalMeshDest, const FName& ProfileNameDest, const UnFbx::FBXImportOptions& ImportOptions, int32 LODIndexDest);
 	/** Generate the editor-only data stored for a skin weight profile (relies on bone indices) */
 	static void GenerateImportedSkinWeightProfileData(const FSkeletalMeshLODModel& LODModelDest, FImportedSkinWeightProfileData& ImportedProfileData);
 	/** Re-generate all (editor-only) skin weight profile, used whenever we rebuild the skeletal mesh data which could change the chunking and bone indices */
diff --git a/Engine/Source/Runtime/Engine/Public/Rendering/SkeletalMeshLODImporterData.cpp b/Engine/Source/Runtime/Engine/Public/Rendering/SkeletalMeshLODImporterData.cpp
--- a/Engine/Source/Runtime/Engine/Public/Rendering/SkeletalMeshLODImporterData.cpp
+++ b/Engine/Source/Runtime/Engine/Public/Rendering/SkeletalMeshLODImporterData.cpp
@@ -8,6 +8,30 @@
 #include "Engine/SkeletalMesh.h"
 #include "Factories/FbxSkeletalMeshImportData.h"
 
+void FSkeletalMeshImportData::CopyDataNeedByMorphTargetImport(FSkeletalMeshImportData& Other) const
+{
+	Other.Points = Points;
+	Other.PointToRawMap = PointToRawMap;
+	Other.bDiffPose = bDiffPose;
+	Other.bUseT0AsRefPose = bUseT0AsRefPose;
+}
+
+void FSkeletalMeshImportData::KeepAlternateSkinningBuildDataOnly()
+{
+	//No need of any alternate restore data, since we are this data if this function is called
+	AlternateInfluenceProfileNames.Empty();
+	AlternateInfluences.Empty();
+	
+	//No need of the morph target restore data
+	MorphTargetModifiedPoints.Empty();
+	MorphTargetNames.Empty();
+	MorphTargets.Empty();
+
+	//Remove material array and PointToRawMap
+	Materials.Empty();
+	PointToRawMap.Empty();
+}
+
 /**
 * Takes an imported bone name, removes any leading or trailing spaces, and converts the remaining spaces to dashes.
 */
@@ -467,6 +491,7 @@ enum
 	// Engine raw mesh version:
 	RAW_SKELETAL_MESH_BULKDATA_VER_INITIAL = 0,
 	RAW_SKELETAL_MESH_BULKDATA_VER_AlternateInfluence = 1,
+	RAW_SKELETAL_MESH_BULKDATA_VER_RebuildSystem = 2,
 	// Add new raw mesh versions here.
 
 	RAW_SKELETAL_MESH_BULKDATA_VER_PLUS_ONE,
@@ -514,11 +539,55 @@ FArchive& operator<<(FArchive& Ar, FSkeletalMeshImportData& RawMesh)
 	{
 		ProcessImportMeshInfluences(RawMesh);
 	}
+
+	
+	if (Version >= RAW_SKELETAL_MESH_BULKDATA_VER_RebuildSystem)
+	{
+		Ar << RawMesh.MorphTargets;
+		Ar << RawMesh.MorphTargetModifiedPoints;
+		Ar << RawMesh.MorphTargetNames;
+		Ar << RawMesh.AlternateInfluences;
+		Ar << RawMesh.AlternateInfluenceProfileNames;
+	}
+	else if (Ar.IsLoading())
+	{
+		RawMesh.MorphTargets.Empty();
+		RawMesh.MorphTargetModifiedPoints.Empty();
+		RawMesh.MorphTargetNames.Empty();
+		RawMesh.AlternateInfluences.Empty();
+		RawMesh.AlternateInfluenceProfileNames.Empty();
+	}
+
 	return Ar;
 }
 
 void FRawSkeletalMeshBulkData::Serialize(FArchive& Ar, UObject* Owner)
 {
+	if (Ar.IsLoading())
+	{
+		//Save the custom version so we can load FReductionSkeletalMeshData later
+		SerializeLoadingCustomVersionContainer = Ar.GetCustomVersions();
+		bUseSerializeLoadingCustomVersion = false;
+		const FCustomVersionContainer& RegisteredCustomVersionContainer = FCustomVersionContainer::GetRegistered();
+		const FCustomVersionArray& ArchiveCustomVersionArray = SerializeLoadingCustomVersionContainer.GetAllVersions();
+		for (const FCustomVersion& ArchiveVersion : ArchiveCustomVersionArray)
+		{
+			const FCustomVersion* RegisteredVersion = RegisteredCustomVersionContainer.GetVersion(ArchiveVersion.Key);
+			if (!RegisteredVersion || RegisteredVersion->Version != ArchiveVersion.Version)
+			{
+				bUseSerializeLoadingCustomVersion = true;
+				break;
+			}
+		}
+	}
+
+	if (Ar.IsSaving() && bUseSerializeLoadingCustomVersion == true)
+	{
+		//We need to update the FReductionSkeletalMeshData serialize version to the latest in case we save the Parent bulkdata
+		FSkeletalMeshImportData MeshImportData;
+		LoadRawMesh(MeshImportData);
+		SaveRawMesh(MeshImportData);
+	}
 	BulkData.Serialize(Ar, Owner);
 	Ar << Guid;
 	Ar << bGuidIsHash;
@@ -526,6 +595,10 @@ void FRawSkeletalMeshBulkData::Serialize(FArchive& Ar, UObject* Owner)
 
 void FRawSkeletalMeshBulkData::SaveRawMesh(FSkeletalMeshImportData& InMesh)
 {
+	//Saving the bulk data mean we do not need anymore the SerializeLoadingCustomVersionContainer of the parent bulk data
+	SerializeLoadingCustomVersionContainer.Empty();
+	bUseSerializeLoadingCustomVersion = false;
+
 	TArray<uint8> TempBytes;
 	FMemoryWriter Ar(TempBytes, /*bIsPersistent=*/ true);
 	Ar << InMesh;
@@ -545,6 +618,11 @@ void FRawSkeletalMeshBulkData::LoadRawMesh(FSkeletalMeshImportData& OutMesh)
 			BulkData.Lock(LOCK_READ_ONLY), BulkData.GetElementCount(),
 			/*bInFreeOnClose=*/ false, /*bIsPersistent=*/ true
 		);
+		//If we load the bulkdata content and the bulk was using customVersion, we must set the same until we save the reduction data
+		if (bUseSerializeLoadingCustomVersion)
+		{
+			Ar.SetCustomVersions(SerializeLoadingCustomVersionContainer);
+		}
 		Ar << OutMesh;
 		BulkData.Unlock();
 	}
diff --git a/Engine/Source/Runtime/Engine/Public/Rendering/SkeletalMeshLODImporterData.h b/Engine/Source/Runtime/Engine/Public/Rendering/SkeletalMeshLODImporterData.h
--- a/Engine/Source/Runtime/Engine/Public/Rendering/SkeletalMeshLODImporterData.h
+++ b/Engine/Source/Runtime/Engine/Public/Rendering/SkeletalMeshLODImporterData.h
@@ -284,20 +284,31 @@ template <> struct TIsPODType<SkeletalMeshImportData::FVertInfluence> { enum { V
 class ENGINE_API FSkeletalMeshImportData
 {
 public:
-	TArray <SkeletalMeshImportData::FMaterial>			Materials;				// Materials
-	TArray <FVector>			Points;											// 3D Points
-	TArray <SkeletalMeshImportData::FVertex>			Wedges;					// Wedges
-	TArray <SkeletalMeshImportData::FTriangle>			Faces;					// Faces
-	TArray <SkeletalMeshImportData::FBone>				RefBonesBinary;			// Reference Skeleton
-	TArray <SkeletalMeshImportData::FRawBoneInfluence>	Influences;				// Influences
-	TArray <int32>				PointToRawMap;									// Mapping from current point index to the original import point index
-	uint32	NumTexCoords;														// The number of texture coordinate sets
-	uint32	MaxMaterialIndex;													// The max material index found on a triangle
-	bool 	bHasVertexColors; 													// If true there are vertex colors in the imported file
-	bool	bHasNormals;														// If true there are normals in the imported file
-	bool	bHasTangents;														// If true there are tangents in the imported file
-	bool	bUseT0AsRefPose;													// If true, then the pose at time=0 will be used instead of the ref pose
-	bool	bDiffPose;															// If true, one of the bones has a different pose at time=0 vs the ref pose
+	TArray <SkeletalMeshImportData::FMaterial> Materials;
+	TArray <FVector> Points;
+	TArray <SkeletalMeshImportData::FVertex> Wedges;
+	TArray <SkeletalMeshImportData::FTriangle> Faces;
+	TArray <SkeletalMeshImportData::FBone> RefBonesBinary;
+	TArray <SkeletalMeshImportData::FRawBoneInfluence> Influences;
+	TArray <int32> PointToRawMap;	// Mapping from current point index to the original import point index
+	uint32 NumTexCoords; // The number of texture coordinate sets
+	uint32 MaxMaterialIndex; // The max material index found on a triangle
+	bool bHasVertexColors; // If true there are vertex colors in the imported file
+	bool bHasNormals; // If true there are normals in the imported file
+	bool bHasTangents; // If true there are tangents in the imported file
+	bool bUseT0AsRefPose; // If true, then the pose at time=0 will be used instead of the ref pose
+	bool bDiffPose; // If true, one of the bones has a different pose at time=0 vs the ref pose
+
+	// Morph targets imported(i.e. FBX) data. The name is the morph target name
+	TArray<FSkeletalMeshImportData> MorphTargets;
+	TArray<TSet<uint32>> MorphTargetModifiedPoints;
+	TArray<FString> MorphTargetNames;
+	
+	// Alternate influence imported(i.e. FBX) data. The name is the alternate skinning profile name
+	TArray<FSkeletalMeshImportData> AlternateInfluences;
+	TArray<FString> AlternateInfluenceProfileNames;
+
+	//////////////////////////////////////////////////////////////////////////
 
 	FSkeletalMeshImportData()
 		: NumTexCoords(0)
@@ -311,6 +322,17 @@ public:
 
 	}
 
+	/*
+	 * Copy only unnecessary array data from the structure to build the morph target (this will save a lot of memory)
+	 */
+	void CopyDataNeedByMorphTargetImport(FSkeletalMeshImportData& Other) const;
+
+	/*
+	 * Remove all unnecessary array data from the structure (this will save a lot of memory)
+	 * We only need Points, Influences and RefBonesBinary arrays
+	 */
+	void KeepAlternateSkinningBuildDataOnly();
+
 	/**
 	* Copy mesh data for importing a single LOD
 	*
@@ -394,6 +416,10 @@ class FRawSkeletalMeshBulkData
 	/** If true, the GUID is actually a hash of the contents. */
 	bool bGuidIsHash;
 
+	//The custom version when this was load
+	FCustomVersionContainer SerializeLoadingCustomVersionContainer;
+	bool bUseSerializeLoadingCustomVersion = false;
+
 public:
 	/** Default constructor. */
 	ENGINE_API FRawSkeletalMeshBulkData();
