FLevelEditorViewportClient no longer uses CastChecked and then checks the resulting pointer.

diff --git a/Engine/Source/Editor/UnrealEd/Private/LevelEditorViewport.cpp b/Engine/Source/Editor/UnrealEd/Private/LevelEditorViewport.cpp
--- a/Engine/Source/Editor/UnrealEd/Private/LevelEditorViewport.cpp
+++ b/Engine/Source/Editor/UnrealEd/Private/LevelEditorViewport.cpp
@@ -2988,14 +2988,12 @@ void FLevelEditorViewportClient::TrackingStopped()
 				TInlineComponentArray<USceneComponent*> ComponentsToMove;
 				for (FSelectedEditableComponentIterator EditableComponentIt(GEditor->GetSelectedEditableComponentIterator()); EditableComponentIt; ++EditableComponentIt)
 				{
-					USceneComponent* SceneComponent = CastChecked<USceneComponent>(*EditableComponentIt);
+					USceneComponent* SceneComponent = Cast<USceneComponent>(*EditableComponentIt);
 					if (SceneComponent)
 					{
-						USceneComponent* SelectedComponent = Cast<USceneComponent>(*EditableComponentIt);
-
 						// Check to see if any parent is selected
 						bool bParentAlsoSelected = false;
-						USceneComponent* Parent = SelectedComponent->GetAttachParent();
+						USceneComponent* Parent = SceneComponent->GetAttachParent();
 						while (Parent != nullptr)
 						{
 							if (ComponentSelection->IsSelected(Parent))
@@ -3010,7 +3008,7 @@ void FLevelEditorViewportClient::TrackingStopped()
 						// If no parent of this component is also in the selection set, move it!
 						if (!bParentAlsoSelected)
 						{
-							ComponentsToMove.Add(SelectedComponent);
+							ComponentsToMove.Add(SceneComponent);
 						}
 					}
 				}
