Enhacement request from JIRA UE-77383: Fixed missing include.

diff --git a/Engine/Source/Editor/UnrealEd/Public/Layers/LayersSubsystem.h b/Engine/Source/Editor/UnrealEd/Public/Layers/LayersSubsystem.h
--- a/Engine/Source/Editor/UnrealEd/Public/Layers/LayersSubsystem.h
+++ b/Engine/Source/Editor/UnrealEd/Public/Layers/LayersSubsystem.h
@@ -3,6 +3,7 @@
 #pragma once
 
 #include "CoreMinimal.h"
+#include "EditorSubsystem.h"
 #include "Layers/ILayers.h"
 
 #include "LayersSubsystem.generated.h"
