Refactoring the majority of AssetEditorManager cases to AssetEditorSubsystem

diff --git a/Engine/Source/Editor/UnrealEd/Private/Tests/AutomationEditorCommon.cpp b/Engine/Source/Editor/UnrealEd/Private/Tests/AutomationEditorCommon.cpp
--- a/Engine/Source/Editor/UnrealEd/Private/Tests/AutomationEditorCommon.cpp
+++ b/Engine/Source/Editor/UnrealEd/Private/Tests/AutomationEditorCommon.cpp
@@ -1,4 +1,4 @@
-﻿// Copyright 1998-2019 Epic Games, Inc. All Rights Reserved.
+// Copyright 1998-2019 Epic Games, Inc. All Rights Reserved.
 
 #include "Tests/AutomationEditorCommon.h"
 #include "UObject/UnrealType.h"
@@ -25,7 +25,7 @@
 #include "ARFilter.h"
 #include "AssetRegistryModule.h"
 #include "Tests/AutomationCommon.h"
-#include "Toolkits/AssetEditorManager.h"
+
 #include "LevelEditor.h"
 #include "Interfaces/IMainFrameModule.h"
 #include "ShaderCompiler.h"
@@ -35,6 +35,7 @@
 #include "ILauncherWorker.h"
 #include "CookOnTheSide/CookOnTheFlyServer.h"
 #include "LightingBuildOptions.h"
+#include "Subsystems/AssetEditorSubsystem.h"
 
 
 #define COOK_TIMEOUT 3600
@@ -731,9 +732,9 @@ bool FOpenEditorForAssetCommand::Update()
 	UObject* Object = StaticLoadObject(UObject::StaticClass(), NULL, *AssetName);
 	if ( Object )
 	{
-		FAssetEditorManager::Get().OpenEditorForAsset(Object);
+		GEditor->GetEditorSubsystem<UAssetEditorSubsystem>()->OpenEditorForAsset(Object);
 		//This checks to see if the asset sub editor is loaded.
-		if ( FAssetEditorManager::Get().FindEditorForAsset(Object, true) != NULL )
+		if ( GEditor->GetEditorSubsystem<UAssetEditorSubsystem>()->FindEditorForAsset(Object, true) != NULL )
 		{
 			UE_LOG(LogEditorAutomationTests, Log, TEXT("Verified asset editor for: %s."), *AssetName);
 			UE_LOG(LogEditorAutomationTests, Display, TEXT("The editor successfully loaded for: %s."), *AssetName);
@@ -752,10 +753,10 @@ bool FOpenEditorForAssetCommand::Update()
 */
 bool FCloseAllAssetEditorsCommand::Update()
 {
-	FAssetEditorManager::Get().CloseAllAssetEditors();
+	GEditor->GetEditorSubsystem<UAssetEditorSubsystem>()->CloseAllAssetEditors();
 
 	//Get all assets currently being tracked with open editors and make sure they are not still opened.
-	if ( FAssetEditorManager::Get().GetAllEditedAssets().Num() >= 1 )
+	if ( GEditor->GetEditorSubsystem<UAssetEditorSubsystem>()->GetAllEditedAssets().Num() >= 1 )
 	{
 		UE_LOG(LogEditorAutomationTests, Warning, TEXT("Not all of the editors were closed."));
 		return true;
diff --git a/Engine/Source/Editor/UnrealEd/Private/Toolkits/AssetEditorToolkit.cpp b/Engine/Source/Editor/UnrealEd/Private/Toolkits/AssetEditorToolkit.cpp
--- a/Engine/Source/Editor/UnrealEd/Private/Toolkits/AssetEditorToolkit.cpp
+++ b/Engine/Source/Editor/UnrealEd/Private/Toolkits/AssetEditorToolkit.cpp
@@ -29,6 +29,7 @@
 #include "AssetToolsModule.h"
 #include "Toolkits/AssetEditorToolkitMenuContext.h"
 #include "ToolMenus.h"
+#include "Subsystems/AssetEditorSubsystem.h"
 
 #define LOCTEXT_NAMESPACE "AssetEditorToolkit"
 
@@ -283,7 +284,7 @@ void FAssetEditorToolkit::InitAssetEditor( const EToolkitMode::Type Mode, const
 	}
 
 	// NOTE: Currently, the AssetEditorManager will keep a hard reference to our object as we're editing it
-	FAssetEditorManager::Get().NotifyAssetsOpened( EditingObjects, this );
+	GEditor->GetEditorSubsystem<UAssetEditorSubsystem>()->NotifyAssetsOpened( EditingObjects, this );
 }
 
 
@@ -292,7 +293,7 @@ FAssetEditorToolkit::~FAssetEditorToolkit()
 	EditingObjects.Empty();
 
 	// We're no longer editing this object, so let the editor know
-	FAssetEditorManager::Get().NotifyEditorClosed( this );
+	GEditor->GetEditorSubsystem<UAssetEditorSubsystem>()->NotifyEditorClosed( this );
 }
 
 
@@ -501,14 +502,14 @@ void FAssetEditorToolkit::GetSaveableObjects(TArray<UObject*>& OutObjects) const
 void FAssetEditorToolkit::AddEditingObject(UObject* Object)
 {
 	EditingObjects.Add(Object);
-	FAssetEditorManager::Get().NotifyAssetOpened( Object, this );
+	GEditor->GetEditorSubsystem<UAssetEditorSubsystem>()->NotifyAssetOpened( Object, this );
 }
 
 
 void FAssetEditorToolkit::RemoveEditingObject(UObject* Object)
 {
 	EditingObjects.Remove(Object);
-	FAssetEditorManager::Get().NotifyAssetClosed( Object, this );
+	GEditor->GetEditorSubsystem<UAssetEditorSubsystem>()->NotifyAssetClosed( Object, this );
 }
 
 
@@ -572,7 +573,7 @@ void FAssetEditorToolkit::SaveAssetAs_Execute()
 	}
 
 	// close existing asset editors for resaved assets
-	FAssetEditorManager& AssetEditorManager = FAssetEditorManager::Get();
+	UAssetEditorSubsystem* AssetEditorSubsystem = GEditor->GetEditorSubsystem<UAssetEditorSubsystem>();
 
 	/* @todo editor: Persona does not behave well when closing specific objects
 	for (int32 Index = 0; Index < ObjectsToSave.Num(); ++Index)
@@ -601,10 +602,10 @@ void FAssetEditorToolkit::SaveAssetAs_Execute()
 	}
 	for (auto Object : EditingObjects)
 	{
-		AssetEditorManager.CloseAllEditorsForAsset(Object);
-		FAssetEditorManager::Get().NotifyAssetClosed(Object, this);
+		AssetEditorSubsystem->CloseAllEditorsForAsset(Object);
+		GEditor->GetEditorSubsystem<UAssetEditorSubsystem>()->NotifyAssetClosed(Object, this);
 	}
-	AssetEditorManager.OpenEditorForAssets(ObjectsToReopen, ToolkitMode, MyToolkitHost.ToSharedRef());
+	AssetEditorSubsystem->OpenEditorForAssets(ObjectsToReopen, ToolkitMode, MyToolkitHost.ToSharedRef());
 	// end hack
 }
 
@@ -739,7 +740,7 @@ void FAssetEditorToolkit::SwitchToStandaloneEditor_Execute( TWeakPtr< FAssetEdit
 
 	if( ObjectsToEdit.Num() > 0 )
 	{
-		ensure( FAssetEditorManager::Get().OpenEditorForAssets( ObjectsToEdit, EToolkitMode::Standalone, PreviousWorldCentricToolkitHost.ToSharedRef() ) );
+		ensure( GEditor->GetEditorSubsystem<UAssetEditorSubsystem>()->OpenEditorForAssets( ObjectsToEdit, EToolkitMode::Standalone, PreviousWorldCentricToolkitHost.ToSharedRef() ) );
 	}
 }
 
@@ -787,7 +788,7 @@ void FAssetEditorToolkit::SwitchToWorldCentricEditor_Execute( TWeakPtr< FAssetEd
 
 	if( ObjectsToEdit.Num() > 0 )
 	{
-		ensure( FAssetEditorManager::Get().OpenEditorForAssets( ObjectsToEdit, EToolkitMode::WorldCentric, WorldCentricLevelEditor ) );
+		ensure( GEditor->GetEditorSubsystem<UAssetEditorSubsystem>()->OpenEditorForAssets( ObjectsToEdit, EToolkitMode::WorldCentric, WorldCentricLevelEditor ) );
 	}
 }
 
diff --git a/Engine/Source/Editor/UnrealEd/Private/Toolkits/SGlobalOpenAssetDialog.cpp b/Engine/Source/Editor/UnrealEd/Private/Toolkits/SGlobalOpenAssetDialog.cpp
--- a/Engine/Source/Editor/UnrealEd/Private/Toolkits/SGlobalOpenAssetDialog.cpp
+++ b/Engine/Source/Editor/UnrealEd/Private/Toolkits/SGlobalOpenAssetDialog.cpp
@@ -9,7 +9,8 @@
 
 #include "IContentBrowserSingleton.h"
 #include "ContentBrowserModule.h"
-#include "Toolkits/AssetEditorManager.h"
+
+#include "Subsystems/AssetEditorSubsystem.h"
 
 #define LOCTEXT_NAMESPACE "SGlobalOpenAssetDialog"
 
@@ -63,7 +64,7 @@ void SGlobalOpenAssetDialog::OnAssetSelectedFromPicker(const FAssetData& AssetDa
 {
 	if (UObject* ObjectToEdit = AssetData.GetAsset())
 	{
-		FAssetEditorManager::Get().OpenEditorForAsset(ObjectToEdit);
+		GEditor->GetEditorSubsystem<UAssetEditorSubsystem>()->OpenEditorForAsset(ObjectToEdit);
 	}
 }
 
@@ -73,7 +74,7 @@ void SGlobalOpenAssetDialog::OnPressedEnterOnAssetsInPicker(const TArray<FAssetD
 	{
 		if (UObject* ObjectToEdit = AssetIt->GetAsset())
 		{
-			FAssetEditorManager::Get().OpenEditorForAsset(ObjectToEdit);
+			GEditor->GetEditorSubsystem<UAssetEditorSubsystem>()->OpenEditorForAsset(ObjectToEdit);
 		}
 	}
 }
diff --git a/Engine/Source/Editor/UnrealEd/Private/Toolkits/SGlobalTabSwitchingDialog.cpp b/Engine/Source/Editor/UnrealEd/Private/Toolkits/SGlobalTabSwitchingDialog.cpp
--- a/Engine/Source/Editor/UnrealEd/Private/Toolkits/SGlobalTabSwitchingDialog.cpp
+++ b/Engine/Source/Editor/UnrealEd/Private/Toolkits/SGlobalTabSwitchingDialog.cpp
@@ -16,7 +16,7 @@
 
 #include "EngineGlobals.h"
 #include "Editor.h"
-#include "Toolkits/AssetEditorManager.h"
+
 #include "LevelEditor.h"
 #include "Editor/WorkspaceMenuStructure/Public/WorkspaceMenuStructure.h"
 #include "Editor/WorkspaceMenuStructure/Public/WorkspaceMenuStructureModule.h"
@@ -24,6 +24,7 @@
 #if PLATFORM_MAC
 #include "Mac/MacApplication.h"
 #endif
+#include "Subsystems/AssetEditorSubsystem.h"
 
 #define LOCTEXT_NAMESPACE "SGlobalTabSwitchingDialog"
 
@@ -78,7 +79,7 @@ public:
 	FTabSwitchingListItem_Asset(UObject* InAsset)
 		: MyAsset(InAsset)
 	{
-		if (IAssetEditorInstance* EditorInstance = FAssetEditorManager::Get().FindEditorForAsset(MyAsset, /*bFocusIfOpen=*/ false))
+		if (IAssetEditorInstance* EditorInstance = GEditor->GetEditorSubsystem<UAssetEditorSubsystem>()->FindEditorForAsset(MyAsset, /*bFocusIfOpen=*/ false))
 		{
 			LastAccessTime = EditorInstance->GetLastActivationTime();
 		}
@@ -144,12 +145,12 @@ public:
 
 	virtual void ActivateTab() override
 	{
-		FAssetEditorManager::Get().FindEditorForAsset(MyAsset, /*bFocusIfOpen=*/ true);
+		GEditor->GetEditorSubsystem<UAssetEditorSubsystem>()->FindEditorForAsset(MyAsset, /*bFocusIfOpen=*/ true);
 	}
 
 	virtual TSharedPtr<FTabManager> GetAssociatedTabManager() override
 	{
-		IAssetEditorInstance* Instance = FAssetEditorManager::Get().FindEditorForAsset(MyAsset, /*bFocusIfOpen=*/ false);
+		IAssetEditorInstance* Instance = GEditor->GetEditorSubsystem<UAssetEditorSubsystem>()->FindEditorForAsset(MyAsset, /*bFocusIfOpen=*/ false);
 		if (Instance)
 		{
 			return Instance->GetAssociatedTabManager();
@@ -389,7 +390,7 @@ void SGlobalTabSwitchingDialog::Construct(const FArguments& InArgs, FVector2D In
 	AssetThumbnailPool = MakeShareable(new FAssetThumbnailPool(128));
 
 	// Populate the list with open asset editors
-	TArray<UObject*> OpenAssetList = FAssetEditorManager::Get().GetAllEditedAssets();
+	TArray<UObject*> OpenAssetList = GEditor->GetEditorSubsystem<UAssetEditorSubsystem>()->GetAllEditedAssets();
 	for (UObject* OpenAsset : OpenAssetList)
 	{
 		if (OpenAsset->GetOuter() != GetTransientPackage())
diff --git a/Engine/Source/Editor/UnrealEd/Public/Toolkits/AssetEditorManager.h b/Engine/Source/Editor/UnrealEd/Public/Toolkits/AssetEditorManager.h
--- a/Engine/Source/Editor/UnrealEd/Public/Toolkits/AssetEditorManager.h
+++ b/Engine/Source/Editor/UnrealEd/Public/Toolkits/AssetEditorManager.h
@@ -13,33 +13,7 @@ class FMessageEndpoint;
 class IToolkitHost;
 struct FAssetEditorRequestOpenAsset;
 
-/**
- * This class keeps track of a currently open asset editor; allowing it to be
- * brought into focus, closed, etc..., without concern for how the editor was
- * implemented.
- */
-class UNREALED_API IAssetEditorInstance
-{
-public:
 
-	virtual FName GetEditorName() const = 0;
-	virtual void FocusWindow(UObject* ObjectToFocusOn = nullptr) = 0;
-	virtual bool CloseWindow() = 0;
-	virtual bool IsPrimaryEditor() const = 0;
-	virtual void InvokeTab(const struct FTabId& TabId) = 0;
-	virtual TSharedPtr<class FTabManager> GetAssociatedTabManager() = 0;
-	virtual double GetLastActivationTime() = 0;
-	virtual void RemoveEditingAsset(UObject* Asset) = 0;
-};
-
-/** The way that editors were requested to close */
-enum class EAssetEditorCloseReason : uint8
-{
-	CloseAllEditorsForAsset,
-	CloseOtherEditors,
-	RemoveAssetFromAllEditors,
-	CloseAllAssetEditors,
-};
 
 /**
  * Implements a manager for Editor windows that are currently open and the assets they are editing.
diff --git a/Engine/Source/Editor/UnrealEd/Public/Toolkits/AssetEditorToolkit.h b/Engine/Source/Editor/UnrealEd/Public/Toolkits/AssetEditorToolkit.h
--- a/Engine/Source/Editor/UnrealEd/Public/Toolkits/AssetEditorToolkit.h
+++ b/Engine/Source/Editor/UnrealEd/Public/Toolkits/AssetEditorToolkit.h
@@ -7,13 +7,14 @@
 #include "UObject/GCObject.h"
 #include "Framework/Docking/TabManager.h"
 #include "Toolkits/IToolkit.h"
-#include "Toolkits/AssetEditorManager.h"
+
 #include "Framework/Commands/UICommandList.h"
 #include "Framework/MultiBox/MultiBoxExtender.h"
 #include "Framework/Docking/LayoutService.h"
 #include "Toolkits/IToolkitHost.h"
 #include "Toolkits/BaseToolkit.h"
 #include "UnrealEdMisc.h"
+#include "Subsystems/AssetEditorSubsystem.h"
 
 class FAssetEditorModeManager;
 class FMenuBuilder;
diff --git a/Engine/Source/Editor/VREditor/UI/VREditorUISystem.cpp b/Engine/Source/Editor/VREditor/UI/VREditorUISystem.cpp
--- a/Engine/Source/Editor/VREditor/UI/VREditorUISystem.cpp
+++ b/Engine/Source/Editor/VREditor/UI/VREditorUISystem.cpp
@@ -69,7 +69,7 @@
 
 // Actor Details, Modes
 #include "SEditorViewport.h"
-#include "Toolkits/AssetEditorManager.h"
+
 #include "ILevelEditor.h"
 #include "LevelEditor.h"
 
@@ -104,6 +104,7 @@
 #include "Engine/Selection.h"
 #include "VREditorFloatingCameraUI.h"
 #include "ScopedTransaction.h"
+#include "Subsystems/AssetEditorSubsystem.h"
 
 #define LOCTEXT_NAMESPACE "VREditor"
 
@@ -292,7 +293,7 @@ void UVREditorUISystem::Shutdown()
 
 	// Remove the proxy tab manager, we don't want to steal tabs any more.
 	FGlobalTabmanager::Get()->SetProxyTabManager(TSharedPtr<FProxyTabmanager>());
-	FAssetEditorManager::Get().OnAssetEditorOpened().RemoveAll(this);
+	GEditor->GetEditorSubsystem<UAssetEditorSubsystem>()->OnAssetEditorOpened().RemoveAll(this);
 
 	VRMode = nullptr;
 	DraggingUI = nullptr;
@@ -1081,7 +1082,7 @@ void UVREditorUISystem::CreateUIs()
 			// We're going to start stealing tabs from the global tab manager inserting them into the world instead.
 			FGlobalTabmanager::Get()->SetProxyTabManager( ProxyTabManager );
 
-			FAssetEditorManager::Get().OnAssetEditorOpened().AddUObject( this, &UVREditorUISystem::OnAssetEditorOpened );
+			GEditor->GetEditorSubsystem<UAssetEditorSubsystem>()->OnAssetEditorOpened().AddUObject( this, &UVREditorUISystem::OnAssetEditorOpened );
 		}
 		// create the world settings menu
 		{
@@ -1139,7 +1140,7 @@ void UVREditorUISystem::OnAssetEditorOpened(UObject* Asset)
 	
 	{
 		// We need to disable drag drop on the tabs spawned in VR mode.
-		TArray<IAssetEditorInstance*> Editors = FAssetEditorManager::Get().FindEditorsForAsset(Asset);
+		TArray<IAssetEditorInstance*> Editors = GEditor->GetEditorSubsystem<UAssetEditorSubsystem>()->FindEditorsForAsset(Asset);
 		for ( IAssetEditorInstance* Editor : Editors )
 		{
 			if (Editor->GetAssociatedTabManager().IsValid())
diff --git a/Engine/Source/Editor/VREditor/VREditorActions.cpp b/Engine/Source/Editor/VREditor/VREditorActions.cpp
--- a/Engine/Source/Editor/VREditor/VREditorActions.cpp
+++ b/Engine/Source/Editor/VREditor/VREditorActions.cpp
@@ -13,7 +13,7 @@
 #include "Framework/Application/SlateApplication.h"
 #include "ILevelEditor.h"
 #include "LevelEditor.h"
-#include "Toolkits/AssetEditorManager.h"
+
 #include "Developer/AssetTools/Public/IAssetTools.h"
 #include "Developer/AssetTools/Public/AssetToolsModule.h"
 #include "Modules/ModuleManager.h"
@@ -30,6 +30,7 @@
 #include "LevelEditorActions.h"
 #include "UObject/UObjectIterator.h"
 #include "Factories/Factory.h"
+#include "Subsystems/AssetEditorSubsystem.h"
 
 #define LOCTEXT_NAMESPACE "VREditorActions"
 
@@ -416,12 +417,12 @@ void FVREditorActionCallbacks::CreateNewSequence(UVREditorMode* InVRMode)
 	InVRMode->GetUISystem().SequencerOpenendFromRadialMenu(true);
 
 	// Open the Sequencer window
-	FAssetEditorManager::Get().OpenEditorForAsset(NewAsset);
+	GEditor->GetEditorSubsystem<UAssetEditorSubsystem>()->OpenEditorForAsset(NewAsset);
 }
 
 void FVREditorActionCallbacks::CloseSequencer(UMovieSceneSequence* OpenSequence)
 {
-	IAssetEditorInstance* SequencerEditor = FAssetEditorManager::Get().FindEditorForAsset(OpenSequence, false);
+	IAssetEditorInstance* SequencerEditor = GEditor->GetEditorSubsystem<UAssetEditorSubsystem>()->FindEditorForAsset(OpenSequence, false);
 	if (SequencerEditor != nullptr)
 	{
 		SequencerEditor->CloseWindow();
diff --git a/Engine/Source/Editor/WorldBrowser/Private/LevelModel.cpp b/Engine/Source/Editor/WorldBrowser/Private/LevelModel.cpp
--- a/Engine/Source/Editor/WorldBrowser/Private/LevelModel.cpp
+++ b/Engine/Source/Editor/WorldBrowser/Private/LevelModel.cpp
@@ -15,9 +15,10 @@
 #include "ActorEditorUtils.h"
 
 #include "Engine/LevelScriptBlueprint.h"
-#include "Toolkits/AssetEditorManager.h"
+
 #include "LevelCollectionModel.h"
 #include "AssetRegistryModule.h"
+#include "Subsystems/AssetEditorSubsystem.h"
 
 #define LOCTEXT_NAMESPACE "WorldBrowser"
 
@@ -471,7 +472,7 @@ void FLevelModel::OpenKismet()
 	ULevelScriptBlueprint* LevelScriptBlueprint = Level->GetLevelScriptBlueprint();
 	if (LevelScriptBlueprint)
 	{
-		FAssetEditorManager::Get().OpenEditorForAsset(LevelScriptBlueprint);
+		GEditor->GetEditorSubsystem<UAssetEditorSubsystem>()->OpenEditorForAsset(LevelScriptBlueprint);
 	}
 	else
 	{
