Adding missing files

diff --git a/Engine/Source/Editor/UnrealEd/Private/Subsystems/AssetEditorSubsystem.cpp b/Engine/Source/Editor/UnrealEd/Private/Subsystems/AssetEditorSubsystem.cpp
new file mode 100644
--- /dev/null
+++ b/Engine/Source/Editor/UnrealEd/Private/Subsystems/AssetEditorSubsystem.cpp
@@ -0,0 +1,774 @@
+// Copyright 1998-2019 Epic Games, Inc. All Rights Reserved.
+
+#include "Subsystems/AssetEditorSubsystem.h"
+#include "AssetEditorMessages.h"
+#include "MessageEndpoint.h"
+#include "EngineAnalytics.h"
+#include "AnalyticsEventAttribute.h"
+#include "UObject/Package.h"
+#include "CoreGlobals.h"
+#include "AssetToolsModule.h"
+#include "LevelEditor.h"
+#include "Toolkits/AssetEditorToolkit.h"
+#include "Toolkits/SimpleAssetEditor.h"
+#include "Engine/MapBuildDataRegistry.h"
+#include "ContentBrowserModule.h"
+#include "MRUFavoritesList.h"
+#include "Settings/EditorLoadingSavingSettings.h"
+#include "Widgets/Notifications/SNotificationList.h"
+#include "Framework/Notifications/NotificationManager.h"
+#include "PackageTools.h"
+#include "UObject/PackageReload.h"
+#include "Interfaces/IAnalyticsProvider.h"
+
+
+#define LOCTEXT_NAMESPACE "AssetEditorSubsystem"
+
+DEFINE_LOG_CATEGORY_STATIC(LogAssetEditorSubsystem, Log, All);
+
+
+UAssetEditorSubsystem::UAssetEditorSubsystem()
+	: bSavingOnShutdown(false)
+	, bRequestRestorePreviouslyOpenAssets(false)
+{
+	// Message bus to receive requests to load assets
+//	MessageEndpoint = FMessageEndpoint::Builder("UAssetEditorSubsystem")
+	//	.Handling<FAssetEditorRequestOpenAsset>(this, &UAssetEditorSubsystem::HandleRequestOpenAssetMessage)
+	//	.WithInbox();
+
+	if (MessageEndpoint.IsValid())
+	{
+		MessageEndpoint->Subscribe<FAssetEditorRequestOpenAsset>();
+	}
+
+	TickDelegate = FTickerDelegate::CreateUObject(this, &UAssetEditorSubsystem::HandleTicker);
+	FTicker::GetCoreTicker().AddTicker(TickDelegate, 1.f);
+
+	FCoreUObjectDelegates::OnPackageReloaded.AddUObject(this, &UAssetEditorSubsystem::HandlePackageReloaded);
+}
+
+void UAssetEditorSubsystem::Deinitialize()
+{
+	FCoreUObjectDelegates::OnPackageReloaded.RemoveAll(this);
+
+	SaveOpenAssetEditors(true);
+
+	TGuardValue<bool> GuardOnShutdown(bSavingOnShutdown, true);
+
+	CloseAllAssetEditors();
+
+	// Don't attempt to report usage stats if analytics isn't available
+	if (FEngineAnalytics::IsAvailable())
+	{
+		TArray<FAnalyticsEventAttribute> EditorUsageAttribs;
+		EditorUsageAttribs.Empty(2);
+		for (auto Iter = EditorUsageAnalytics.CreateConstIterator(); Iter; ++Iter)
+		{
+			const FAssetEditorAnalyticInfo& Data = Iter.Value();
+			EditorUsageAttribs.Reset();
+			EditorUsageAttribs.Emplace(TEXT("TotalDuration.Seconds"), FString::Printf(TEXT("%.1f"), Data.SumDuration.GetTotalSeconds()));
+			EditorUsageAttribs.Emplace(TEXT("OpenedInstances.Count"), FString::Printf(TEXT("%d"), Data.NumTimesOpened));
+
+			const FString EventName = FString::Printf(TEXT("Editor.Usage.%s"), *Iter.Key().ToString());
+			FEngineAnalytics::GetProvider().RecordEvent(EventName, EditorUsageAttribs);
+		}
+	}
+}
+
+IAssetEditorInstance* UAssetEditorSubsystem::FindEditorForAsset(UObject* Asset, bool bFocusIfOpen)
+{
+	const TArray<IAssetEditorInstance*> AssetEditors = FindEditorsForAsset(Asset);
+
+	IAssetEditorInstance* const * PrimaryEditor = AssetEditors.FindByPredicate([](IAssetEditorInstance* Editor) { return Editor->IsPrimaryEditor(); });
+
+	const bool bEditorOpen = PrimaryEditor != NULL;
+	if (bEditorOpen && bFocusIfOpen)
+	{
+		// @todo toolkit minor: We may need to handle this differently for world-centric vs standalone.  (multiple level editors, etc)
+		(*PrimaryEditor)->FocusWindow(Asset);
+	}
+
+	return bEditorOpen ? *PrimaryEditor : NULL;
+}
+
+
+TArray<IAssetEditorInstance*> UAssetEditorSubsystem::FindEditorsForAsset(UObject* Asset)
+{
+	TArray<IAssetEditorInstance*> AssetEditors;
+	OpenedAssets.MultiFind(Asset, AssetEditors);
+	return AssetEditors;
+}
+
+
+int32 UAssetEditorSubsystem::CloseAllEditorsForAsset(UObject* Asset)
+{
+	TArray<IAssetEditorInstance*> EditorInstances = FindEditorsForAsset(Asset);
+
+	for (IAssetEditorInstance* EditorInstance : EditorInstances)
+	{
+		if (EditorInstance)
+		{
+			EditorInstance->CloseWindow();
+		}
+	}
+
+	AssetEditorRequestCloseEvent.Broadcast(Asset, EAssetEditorCloseReason::CloseAllEditorsForAsset);
+
+	return EditorInstances.Num();
+}
+
+void UAssetEditorSubsystem::RemoveAssetFromAllEditors(UObject* Asset)
+{
+	TArray<IAssetEditorInstance*> EditorInstances = FindEditorsForAsset(Asset);
+
+	for (IAssetEditorInstance* EditorIter : EditorInstances)
+	{
+		if (EditorIter)
+		{
+			EditorIter->RemoveEditingAsset(Asset);
+		}
+	}
+
+	AssetEditorRequestCloseEvent.Broadcast(Asset, EAssetEditorCloseReason::RemoveAssetFromAllEditors);
+}
+
+
+void UAssetEditorSubsystem::CloseOtherEditors(UObject* Asset, IAssetEditorInstance* OnlyEditor)
+{
+	TArray<UObject*> AllAssets;
+	for (TMultiMap<UObject*, IAssetEditorInstance*>::TIterator It(OpenedAssets); It; ++It)
+	{
+		IAssetEditorInstance* Editor = It.Value();
+		if (Asset == It.Key() && Editor != OnlyEditor)
+		{
+			Editor->CloseWindow();
+		}
+	}
+
+	AssetEditorRequestCloseEvent.Broadcast(Asset, EAssetEditorCloseReason::CloseOtherEditors);
+}
+
+
+TArray<UObject*> UAssetEditorSubsystem::GetAllEditedAssets()
+{
+	TArray<UObject*> AllAssets;
+	for (TMultiMap<UObject*, IAssetEditorInstance*>::TIterator It(OpenedAssets); It; ++It)
+	{
+		UObject* Asset = It.Key();
+		if (Asset != NULL)
+		{
+			AllAssets.AddUnique(Asset);
+		}
+	}
+	return AllAssets;
+}
+
+
+void UAssetEditorSubsystem::NotifyAssetOpened(UObject* Asset, IAssetEditorInstance* InInstance)
+{
+	if (!OpenedEditors.Contains(InInstance))
+	{
+		FOpenedEditorTime EditorTime;
+		EditorTime.EditorName = InInstance->GetEditorName();
+		EditorTime.OpenedTime = FDateTime::UtcNow();
+
+		OpenedEditorTimes.Add(InInstance, EditorTime);
+	}
+
+	OpenedAssets.Add(Asset, InInstance);
+	OpenedEditors.Add(InInstance, Asset);
+
+	AssetOpenedInEditorEvent.Broadcast(Asset, InInstance);
+
+	SaveOpenAssetEditors(false);
+}
+
+
+void UAssetEditorSubsystem::NotifyAssetsOpened(const TArray< UObject* >& Assets, IAssetEditorInstance* InInstance)
+{
+	for (auto AssetIter = Assets.CreateConstIterator(); AssetIter; ++AssetIter)
+	{
+		NotifyAssetOpened(*AssetIter, InInstance);
+	}
+}
+
+
+void UAssetEditorSubsystem::NotifyAssetClosed(UObject* Asset, IAssetEditorInstance* InInstance)
+{
+	OpenedEditors.RemoveSingle(InInstance, Asset);
+	OpenedAssets.RemoveSingle(Asset, InInstance);
+
+	SaveOpenAssetEditors(false);
+}
+
+
+void UAssetEditorSubsystem::NotifyEditorClosed(IAssetEditorInstance* InInstance)
+{
+	// Remove all assets associated with the editor
+	TArray<UObject*> Assets;
+	OpenedEditors.MultiFind(InInstance, /*out*/ Assets);
+	for (int32 AssetIndex = 0; AssetIndex < Assets.Num(); ++AssetIndex)
+	{
+		OpenedAssets.Remove(Assets[AssetIndex], InInstance);
+	}
+
+	// Remove the editor itself
+	OpenedEditors.Remove(InInstance);
+	FOpenedEditorTime EditorTime = OpenedEditorTimes.FindAndRemoveChecked(InInstance);
+
+	// Record the editor open-close duration
+	FAssetEditorAnalyticInfo& AnalyticsForThisAsset = EditorUsageAnalytics.FindOrAdd(EditorTime.EditorName);
+	AnalyticsForThisAsset.SumDuration += FDateTime::UtcNow() - EditorTime.OpenedTime;
+	AnalyticsForThisAsset.NumTimesOpened++;
+
+	SaveOpenAssetEditors(false);
+}
+
+
+bool UAssetEditorSubsystem::CloseAllAssetEditors()
+{
+	bool bAllEditorsClosed = true;
+	for (TMultiMap<IAssetEditorInstance*, UObject*>::TIterator It(OpenedEditors); It; ++It)
+	{
+		IAssetEditorInstance* Editor = It.Key();
+		if (Editor != NULL)
+		{
+			if (!Editor->CloseWindow())
+			{
+				bAllEditorsClosed = false;
+			}
+		}
+	}
+
+	AssetEditorRequestCloseEvent.Broadcast(nullptr, EAssetEditorCloseReason::CloseAllAssetEditors);
+
+	return bAllEditorsClosed;
+}
+
+
+bool UAssetEditorSubsystem::OpenEditorForAsset(UObject* Asset, const EToolkitMode::Type ToolkitMode, TSharedPtr< IToolkitHost > OpenedFromLevelEditor, const bool bShowProgressWindow)
+{
+	check(Asset);
+	// @todo toolkit minor: When "Edit Here" happens in a different level editor from the one that an asset is already
+	//    being edited within, we should decide whether to disallow "Edit Here" in that case, or to close the old asset
+	//    editor and summon it in the new level editor, or to just foreground the old level editor (current behavior)
+
+	const bool bBringToFrontIfOpen = true;
+
+	// Don't open asset editors for cooked packages
+	if (UPackage* Package = Asset->GetOutermost())
+	{
+		if (Package->bIsCookedForEditor)
+		{
+			return false;
+		}
+	}
+
+	AssetEditorRequestOpenEvent.Broadcast(Asset);
+
+	if (FindEditorForAsset(Asset, bBringToFrontIfOpen) != nullptr)
+	{
+		// This asset is already open in an editor! (the call to FindEditorForAsset above will bring it to the front)
+		return true;
+	}
+	else
+	{
+		if (bShowProgressWindow)
+		{
+			GWarn->BeginSlowTask(LOCTEXT("OpenEditor", "Opening Editor..."), true);
+		}
+	}
+
+	UE_LOG(LogAssetEditorSubsystem, Log, TEXT("Opening Asset editor for %s"), *Asset->GetFullName());
+
+	FAssetToolsModule& AssetToolsModule = FModuleManager::LoadModuleChecked<FAssetToolsModule>(TEXT("AssetTools"));
+
+	TWeakPtr<IAssetTypeActions> AssetTypeActions = AssetToolsModule.Get().GetAssetTypeActionsForClass(Asset->GetClass());
+
+	EToolkitMode::Type ActualToolkitMode = ToolkitMode;
+	if (AssetTypeActions.IsValid())
+	{
+		if (AssetTypeActions.Pin()->ShouldForceWorldCentric())
+		{
+			// This asset type prefers a specific toolkit mode
+			ActualToolkitMode = EToolkitMode::WorldCentric;
+
+			if (!OpenedFromLevelEditor.IsValid())
+			{
+				// We don't have a level editor to spawn in world-centric mode, so we'll find one now
+				// @todo sequencer: We should eventually eliminate this code (incl include dependencies) or change it to not make assumptions about a single level editor
+				OpenedFromLevelEditor = FModuleManager::LoadModuleChecked< FLevelEditorModule >("LevelEditor").GetFirstLevelEditor();
+			}
+		}
+	}
+
+	if (ActualToolkitMode != EToolkitMode::WorldCentric && OpenedFromLevelEditor.IsValid())
+	{
+		// @todo toolkit minor: Kind of lame use of a static variable here to prime the new asset editor.  This was done to avoid refactoring a few dozen files for a very minor change.
+		FAssetEditorToolkit::SetPreviousWorldCentricToolkitHostForNewAssetEditor(OpenedFromLevelEditor.ToSharedRef());
+	}
+
+	// Disallow opening an asset editor for classes
+	bool bCanSummonSimpleAssetEditor = !Asset->IsA<UClass>();
+
+	if (AssetTypeActions.IsValid())
+	{
+		TArray<UObject*> AssetsToEdit;
+		AssetsToEdit.Add(Asset);
+
+		// Some assets (like UWorlds) may be destroyed and recreated as part of opening. To protect against this, keep the path to the asset and try to re-find it if it disappeared.
+		TWeakObjectPtr<UObject> WeakAsset = Asset;
+		const FString AssetPath = Asset->GetPathName();
+
+		AssetTypeActions.Pin()->OpenAssetEditor(AssetsToEdit, ActualToolkitMode == EToolkitMode::WorldCentric ? OpenedFromLevelEditor : TSharedPtr<IToolkitHost>());
+
+		// If the Asset was destroyed, attempt to find it if it was recreated
+		if (!WeakAsset.IsValid() && !AssetPath.IsEmpty())
+		{
+			Asset = FindObject<UObject>(nullptr, *AssetPath);
+		}
+
+		AssetEditorOpenedEvent.Broadcast(Asset);
+	}
+	else if (bCanSummonSimpleAssetEditor)
+	{
+		// No asset type actions for this asset. Just use a properties editor.
+		FSimpleAssetEditor::CreateEditor(ActualToolkitMode, ActualToolkitMode == EToolkitMode::WorldCentric ? OpenedFromLevelEditor : TSharedPtr<IToolkitHost>(), Asset);
+	}
+
+	if (bShowProgressWindow)
+	{
+		GWarn->EndSlowTask();
+	}
+	// Must check Asset here in addition to at the beginning of the function, because if the asset was destroyed and recreated it might not be found correctly
+	// Do not add to recently opened asset list if this is a level-associated asset like Level Blueprint or Built Data. Their naming is not compatible
+	if (Asset && Asset->IsAsset() && !Asset->IsA(UMapBuildDataRegistry::StaticClass()))
+	{
+		FString AssetPath = Asset->GetOuter()->GetPathName();
+		FContentBrowserModule& CBModule = FModuleManager::LoadModuleChecked<FContentBrowserModule>(TEXT("ContentBrowser"));
+		FMainMRUFavoritesList* RecentlyOpenedAssets = CBModule.GetRecentlyOpenedAssets();
+		if (RecentlyOpenedAssets && FPackageName::IsValidLongPackageName(AssetPath))
+		{
+			RecentlyOpenedAssets->AddMRUItem(AssetPath);
+		}
+	}
+	return true;
+}
+
+
+bool UAssetEditorSubsystem::OpenEditorForAssets(const TArray <UObject* >& Assets, const EToolkitMode::Type ToolkitMode, TSharedPtr< IToolkitHost > OpenedFromLevelEditor)
+{
+	if (Assets.Num() == 1)
+	{
+		return OpenEditorForAsset(Assets[0], ToolkitMode, OpenedFromLevelEditor);
+	}
+	else if (Assets.Num() > 0)
+	{
+		TArray<UObject*> SkipOpenAssets;
+		for (UObject* Asset : Assets)
+		{
+			// If any of the assets are already open or the package is cooked,
+			// remove them from the list of assets to open an editor for
+			UPackage* Package = Asset->GetOutermost();
+			if (FindEditorForAsset(Asset, true) != nullptr || (Package && Package->bIsCookedForEditor))
+			{
+				SkipOpenAssets.Add(Asset);
+			}
+		}
+
+		// Verify that all the assets are of the same class
+		bool bAssetClassesMatch = true;
+		UClass* AssetClass = Assets[0]->GetClass();
+		for (int32 i = 1; i < Assets.Num(); i++)
+		{
+			if (Assets[i]->GetClass() != AssetClass)
+			{
+				bAssetClassesMatch = false;
+				break;
+			}
+		}
+
+		// If the classes don't match or any of the selected assets are already open, just open each asset in its own editor.
+		if (bAssetClassesMatch && SkipOpenAssets.Num() == 0)
+		{
+			FAssetToolsModule& AssetToolsModule = FModuleManager::LoadModuleChecked<FAssetToolsModule>(TEXT("AssetTools"));
+			TWeakPtr<IAssetTypeActions> AssetTypeActions = AssetToolsModule.Get().GetAssetTypeActionsForClass(AssetClass);
+
+			if (AssetTypeActions.IsValid())
+			{
+				GWarn->BeginSlowTask(LOCTEXT("OpenEditors", "Opening Editor(s)..."), true);
+
+				// Determine the appropriate toolkit mode for the asset type
+				EToolkitMode::Type ActualToolkitMode = ToolkitMode;
+				if (AssetTypeActions.Pin()->ShouldForceWorldCentric())
+				{
+					// This asset type prefers a specific toolkit mode
+					ActualToolkitMode = EToolkitMode::WorldCentric;
+
+					if (!OpenedFromLevelEditor.IsValid())
+					{
+						// We don't have a level editor to spawn in world-centric mode, so we'll find one now
+						// @todo sequencer: We should eventually eliminate this code (incl include dependencies) or change it to not make assumptions about a single level editor
+						OpenedFromLevelEditor = FModuleManager::LoadModuleChecked< FLevelEditorModule >("LevelEditor").GetFirstLevelEditor();
+					}
+				}
+
+				if (ActualToolkitMode != EToolkitMode::WorldCentric && OpenedFromLevelEditor.IsValid())
+				{
+					// @todo toolkit minor: Kind of lame use of a static variable here to prime the new asset editor.  This was done to avoid refactoring a few dozen files for a very minor change.
+					FAssetEditorToolkit::SetPreviousWorldCentricToolkitHostForNewAssetEditor(OpenedFromLevelEditor.ToSharedRef());
+				}
+
+				// Some assets (like UWorlds) may be destroyed and recreated as part of opening. To protect against this, keep the path to each asset and try to re-find any if they disappear.
+				struct FLocalAssetInfo
+				{
+					TWeakObjectPtr<UObject> WeakAsset;
+					FString AssetPath;
+
+					FLocalAssetInfo(const TWeakObjectPtr<UObject>& InWeakAsset, const FString& InAssetPath)
+						: WeakAsset(InWeakAsset), AssetPath(InAssetPath) {}
+				};
+
+				TArray<FLocalAssetInfo> AssetInfoList;
+				AssetInfoList.Reserve(Assets.Num());
+				for (UObject* Asset : Assets)
+				{
+					AssetInfoList.Add(FLocalAssetInfo(Asset, Asset->GetPathName()));
+				}
+
+				// How to handle multiple assets is left up to the type actions (i.e. open a single shared editor or an editor for each)
+				AssetTypeActions.Pin()->OpenAssetEditor(Assets, ActualToolkitMode == EToolkitMode::WorldCentric ? OpenedFromLevelEditor : TSharedPtr<IToolkitHost>());
+
+				// If any assets were destroyed, attempt to find them if they were recreated
+				for (int32 i = 0; i < Assets.Num(); i++)
+				{
+					const FLocalAssetInfo& AssetInfo = AssetInfoList[i];
+					UObject* Asset = Assets[i];
+
+					if (!AssetInfo.WeakAsset.IsValid() && !AssetInfo.AssetPath.IsEmpty())
+					{
+						Asset = FindObject<UObject>(nullptr, *AssetInfo.AssetPath);
+					}
+				}
+
+				//@todo if needed, broadcast the event for every asset. It is possible, however, that a single shared editor was opened by the AssetTypeActions, not an editor for each asset.
+				/*AssetEditorOpenedEvent.Broadcast(Asset);*/
+
+				GWarn->EndSlowTask();
+			}
+		}
+		else
+		{
+			// Asset types don't match or some are already open, so just open individual editors for the unopened ones
+			for (UObject* Asset : Assets)
+			{
+				if (!SkipOpenAssets.Contains(Asset))
+				{
+					OpenEditorForAsset(Asset, ToolkitMode, OpenedFromLevelEditor);
+				}
+			}
+		}
+	}
+
+	return true;
+}
+
+void UAssetEditorSubsystem::HandleRequestOpenAssetMessage(const FAssetEditorRequestOpenAsset& Message, const TSharedRef<IMessageContext, ESPMode::ThreadSafe>& Context)
+{
+	OpenEditorForAsset(Message.AssetName);
+}
+
+void UAssetEditorSubsystem::OpenEditorForAsset(const FString& AssetPathName)
+{
+	// An asset needs loading
+	UPackage* Package = LoadPackage(NULL, *AssetPathName, LOAD_NoRedirects);
+
+	if (Package)
+	{
+		Package->FullyLoad();
+
+		FString AssetName = FPaths::GetBaseFilename(AssetPathName);
+		UObject* Object = FindObject<UObject>(Package, *AssetName);
+
+		if (Object != NULL)
+		{
+			OpenEditorForAsset(Object);
+		}
+	}
+}
+
+bool UAssetEditorSubsystem::HandleTicker(float DeltaTime)
+{
+	QUICK_SCOPE_CYCLE_COUNTER(STAT_UAssetEditorSubsystem_HandleTicker);
+
+	if (bRequestRestorePreviouslyOpenAssets)
+	{
+		RestorePreviouslyOpenAssets();
+		bRequestRestorePreviouslyOpenAssets = false;
+	}
+//	MessageEndpoint->ProcessInbox();
+
+	return true;
+}
+
+void UAssetEditorSubsystem::RequestRestorePreviouslyOpenAssets()
+{
+	// We defer the restore so that we guarantee that it happens once initialization is complete
+	bRequestRestorePreviouslyOpenAssets = true;
+}
+
+void UAssetEditorSubsystem::RestorePreviouslyOpenAssets()
+{
+	TArray<FString> OpenAssets;
+	GConfig->GetArray(TEXT("AssetEditorSubsystem"), TEXT("OpenAssetsAtExit"), OpenAssets, GEditorPerProjectIni);
+
+	bool bCleanShutdown = false;
+	GConfig->GetBool(TEXT("AssetEditorSubsystem"), TEXT("CleanShutdown"), bCleanShutdown, GEditorPerProjectIni);
+
+	SaveOpenAssetEditors(false);
+
+	if (OpenAssets.Num() > 0)
+	{
+		if (bCleanShutdown)
+		{
+			// Do we have permission to automatically re-open the assets, or should we ask?
+			const bool bAutoRestore = GetDefault<UEditorLoadingSavingSettings>()->bRestoreOpenAssetTabsOnRestart;
+
+			if (bAutoRestore)
+			{
+				// Pretend that we showed the notification and that the user clicked "Restore Now"
+				OpenEditorsForAssets(OpenAssets);
+			}
+			else
+			{
+				// Has this notification previously been suppressed by the user?
+				bool bSuppressNotification = false;
+				GConfig->GetBool(TEXT("AssetEditorSubsystem"), TEXT("SuppressRestorePreviouslyOpenAssetsNotification"), bSuppressNotification, GEditorPerProjectIni);
+
+				if (!bSuppressNotification)
+				{
+					// Ask the user; this doesn't block so will reopen the assets later
+					SpawnRestorePreviouslyOpenAssetsNotification(bCleanShutdown, OpenAssets);
+				}
+			}
+		}
+		else
+		{
+			// If we crashed, we always ask regardless of what the user previously said
+			SpawnRestorePreviouslyOpenAssetsNotification(bCleanShutdown, OpenAssets);
+		}
+	}
+}
+
+void UAssetEditorSubsystem::SpawnRestorePreviouslyOpenAssetsNotification(const bool bCleanShutdown, const TArray<FString>& AssetsToOpen)
+{
+	/** Utility functions for the notification which don't rely on the state from FAssetEditorManager */
+	struct Local
+	{
+		static ECheckBoxState GetDontAskAgainCheckBoxState()
+		{
+			bool bSuppressNotification = false;
+			GConfig->GetBool(TEXT("AssetEditorSubsystem"), TEXT("SuppressRestorePreviouslyOpenAssetsNotification"), bSuppressNotification, GEditorPerProjectIni);
+			return bSuppressNotification ? ECheckBoxState::Checked : ECheckBoxState::Unchecked;
+		}
+
+		static void OnDontAskAgainCheckBoxStateChanged(ECheckBoxState NewState)
+		{
+			const bool bSuppressNotification = (NewState == ECheckBoxState::Checked);
+			GConfig->SetBool(TEXT("AssetEditorSubsystem"), TEXT("SuppressRestorePreviouslyOpenAssetsNotification"), bSuppressNotification, GEditorPerProjectIni);
+		}
+	};
+
+	FText NotificationMessage = bCleanShutdown
+		? LOCTEXT("ReopenAssetEditorsAfterClose", "{0} asset {0}|plural(one=editor was,other=editors were) open when the editor was last closed. Would you like to re-open them?")
+		: LOCTEXT("ReopenAssetEditorsAfterCrash", "{0} asset {0}|plural(one=editor was,other=editors were) open when the editor quit unexpectedly. Would you like to re-open them?");
+	NotificationMessage = FText::Format(NotificationMessage, AssetsToOpen.Num());
+
+	FNotificationInfo Info = FNotificationInfo(NotificationMessage);
+
+	// Add the buttons
+	Info.ButtonDetails.Add(FNotificationButtonInfo(
+		LOCTEXT("ReopenAssetEditors_Confirm", "Open"),
+		FText(),
+		FSimpleDelegate::CreateUObject(this, &UAssetEditorSubsystem::OnConfirmRestorePreviouslyOpenAssets, AssetsToOpen),
+		SNotificationItem::CS_None
+	));
+	Info.ButtonDetails.Add(FNotificationButtonInfo(
+		LOCTEXT("ReopenAssetEditors_Cancel", "Cancel"),
+		FText(),
+		FSimpleDelegate::CreateUObject(this, &UAssetEditorSubsystem::OnCancelRestorePreviouslyOpenAssets),
+		SNotificationItem::CS_None
+	));
+
+	// We will let the notification expire automatically after 10 seconds
+	Info.bFireAndForget = false;
+	Info.ExpireDuration = 10.0f;
+
+	// We want the auto-save to be subtle
+	Info.bUseLargeFont = false;
+	Info.bUseThrobber = false;
+	Info.bUseSuccessFailIcons = false;
+
+	// Only let the user suppress the non-crash version
+	if (bCleanShutdown)
+	{
+		Info.CheckBoxState = TAttribute<ECheckBoxState>::Create(&Local::GetDontAskAgainCheckBoxState);
+		Info.CheckBoxStateChanged = FOnCheckStateChanged::CreateStatic(&Local::OnDontAskAgainCheckBoxStateChanged);
+		Info.CheckBoxText = NSLOCTEXT("ModalDialogs", "DefaultCheckBoxMessage", "Don't show this again");
+	}
+
+	// Close any existing notification
+	TSharedPtr<SNotificationItem> RestorePreviouslyOpenAssetsNotification = RestorePreviouslyOpenAssetsNotificationPtr.Pin();
+	if (RestorePreviouslyOpenAssetsNotification.IsValid())
+	{
+		RestorePreviouslyOpenAssetsNotification->ExpireAndFadeout();
+	}
+
+	RestorePreviouslyOpenAssetsNotificationPtr = FSlateNotificationManager::Get().AddNotification(Info);
+}
+
+void UAssetEditorSubsystem::OnConfirmRestorePreviouslyOpenAssets(TArray<FString> AssetsToOpen)
+{
+	// Close any existing notification
+	TSharedPtr<SNotificationItem> RestorePreviouslyOpenAssetsNotification = RestorePreviouslyOpenAssetsNotificationPtr.Pin();
+	if (RestorePreviouslyOpenAssetsNotification.IsValid())
+	{
+		RestorePreviouslyOpenAssetsNotification->SetExpireDuration(0.0f);
+		RestorePreviouslyOpenAssetsNotification->SetFadeOutDuration(0.5f);
+		RestorePreviouslyOpenAssetsNotification->ExpireAndFadeout();
+
+		// If the user suppressed the notification for future sessions, make sure this is reflected in their settings
+		// Note: We do that inside this if statement so that we only do it if we were showing a UI they could interact with
+		bool bSuppressNotification = false;
+		GConfig->GetBool(TEXT("AssetEditorSubsystem"), TEXT("SuppressRestorePreviouslyOpenAssetsNotification"), bSuppressNotification, GEditorPerProjectIni);
+		UEditorLoadingSavingSettings& Settings = *GetMutableDefault<UEditorLoadingSavingSettings>();
+		Settings.bRestoreOpenAssetTabsOnRestart = bSuppressNotification;
+		Settings.PostEditChange();
+
+		// we do this inside the condition so that it can only be done once. 
+		OpenEditorsForAssets(AssetsToOpen);
+
+	}
+}
+
+void UAssetEditorSubsystem::OnCancelRestorePreviouslyOpenAssets()
+{
+	// Close any existing notification
+	TSharedPtr<SNotificationItem> RestorePreviouslyOpenAssetsNotification = RestorePreviouslyOpenAssetsNotificationPtr.Pin();
+	if (RestorePreviouslyOpenAssetsNotification.IsValid())
+	{
+		RestorePreviouslyOpenAssetsNotification->SetExpireDuration(0.0f);
+		RestorePreviouslyOpenAssetsNotification->SetFadeOutDuration(0.5f);
+		RestorePreviouslyOpenAssetsNotification->ExpireAndFadeout();
+	}
+}
+
+void UAssetEditorSubsystem::SaveOpenAssetEditors(bool bOnShutdown)
+{
+	if (!bSavingOnShutdown)
+	{
+		TArray<FString> OpenAssets;
+
+		// Don't save a list of assets to restore if we are running under a debugger
+		if (!FPlatformMisc::IsDebuggerPresent())
+		{
+			for (const TPair<IAssetEditorInstance*, UObject*>& EditorPair : OpenedEditors)
+			{
+				IAssetEditorInstance* Editor = EditorPair.Key;
+				if (Editor != NULL)
+				{
+					UObject* EditedObject = EditorPair.Value;
+					if (EditedObject != NULL)
+					{
+						// only record assets that have a valid saved package
+						UPackage* Package = EditedObject->GetOutermost();
+						if (Package != NULL && Package->GetFileSize() != 0)
+						{
+							OpenAssets.Add(EditedObject->GetPathName());
+						}
+					}
+				}
+			}
+		}
+
+		GConfig->SetArray(TEXT("AssetEditorSubsystem"), TEXT("OpenAssetsAtExit"), OpenAssets, GEditorPerProjectIni);
+		GConfig->SetBool(TEXT("AssetEditorSubsystem"), TEXT("CleanShutdown"), bOnShutdown, GEditorPerProjectIni);
+		GConfig->Flush(false, GEditorPerProjectIni);
+	}
+}
+
+void UAssetEditorSubsystem::HandlePackageReloaded(const EPackageReloadPhase InPackageReloadPhase, FPackageReloadedEvent* InPackageReloadedEvent)
+{
+	static TArray<UObject*> PendingAssetsToOpen;
+
+	if (InPackageReloadPhase == EPackageReloadPhase::PrePackageFixup)
+	{
+		/** Call close for all old assets even if not open, so global callback will go off */
+		TArray<UObject*> ObjectsToClose;
+		const TMap<UObject*, UObject*>& RepointedMap = InPackageReloadedEvent->GetRepointedObjects();
+
+		for (const TPair<UObject*, UObject*> RepointPair : RepointedMap)
+		{
+			if (RepointPair.Key->IsAsset())
+			{
+				ObjectsToClose.Add(RepointPair.Key);
+			}
+		}
+
+		/** Look for replacement for assets that are open now so we can reopen */
+		for (TPair<UObject*, IAssetEditorInstance*>& AssetEditorPair : OpenedAssets)
+		{
+			UObject* NewAsset = nullptr;
+			if (InPackageReloadedEvent->GetRepointedObject(AssetEditorPair.Key, NewAsset))
+			{
+				if (NewAsset)
+				{
+					PendingAssetsToOpen.AddUnique(NewAsset);
+				}
+
+				ObjectsToClose.AddUnique(AssetEditorPair.Key);
+			}
+		}
+
+		int32 NumAssetEditorsClosed = 0;
+		for (UObject* OldAsset : ObjectsToClose)
+		{
+			NumAssetEditorsClosed += CloseAllEditorsForAsset(OldAsset);
+		}
+
+		if (NumAssetEditorsClosed > 0)
+		{
+			// Closing asset editors might have have left objects pending GC that still reference the asset we're about to reload
+			// Run a GC now to ensure those are cleaned up before the fix-up phase happens
+			CollectGarbage(GARBAGE_COLLECTION_KEEPFLAGS);
+		}
+	}
+
+	if (InPackageReloadPhase == EPackageReloadPhase::PostBatchPostGC)
+	{
+		for (UObject* NewAsset : PendingAssetsToOpen)
+		{
+			OpenEditorForAsset(NewAsset);
+		}
+		PendingAssetsToOpen.Reset();
+	}
+}
+
+void UAssetEditorSubsystem::OpenEditorsForAssets(const TArray<FString>& AssetsToOpen)
+{
+	for (const FString& AssetName : AssetsToOpen)
+	{
+		OpenEditorForAsset(AssetName);
+	}
+}
+
+void UAssetEditorSubsystem::OpenEditorsForAssets(const TArray<FName>& AssetsToOpen)
+{
+	for (const FName AssetName : AssetsToOpen)
+	{
+		OpenEditorForAsset(AssetName.ToString());
+	}
+}
+
+#undef LOCTEXT_NAMESPACE
+
diff --git a/Engine/Source/Editor/UnrealEd/Public/Subsystems/AssetEditorSubsystem.h b/Engine/Source/Editor/UnrealEd/Public/Subsystems/AssetEditorSubsystem.h
new file mode 100644
--- /dev/null
+++ b/Engine/Source/Editor/UnrealEd/Public/Subsystems/AssetEditorSubsystem.h
@@ -0,0 +1,221 @@
+// Copyright 1998-2019 Epic Games, Inc. All Rights Reserved.
+
+#pragma once
+
+#include "EditorSubsystem.h"
+
+#include "UObject/ObjectMacros.h"
+#include "UObject/UObjectGlobals.h"
+#include "Toolkits/IToolkit.h"
+
+#include "Templates/SubclassOf.h"
+#include "AssetEditorMessages.h"
+#include "IMessageContext.h"
+#include "MessageEndpoint.h"
+#include "Containers/Ticker.h"
+#include "AssetEditorSubsystem.generated.h"
+
+
+/**
+ * This class keeps track of a currently open asset editor; allowing it to be
+ * brought into focus, closed, etc..., without concern for how the editor was
+ * implemented.
+ */
+class UNREALED_API IAssetEditorInstance
+{
+public:
+
+	virtual FName GetEditorName() const = 0;
+	virtual void FocusWindow(UObject* ObjectToFocusOn = nullptr) = 0;
+	virtual bool CloseWindow() = 0;
+	virtual bool IsPrimaryEditor() const = 0;
+	virtual void InvokeTab(const struct FTabId& TabId) = 0;
+	virtual TSharedPtr<class FTabManager> GetAssociatedTabManager() = 0;
+	virtual double GetLastActivationTime() = 0;
+	virtual void RemoveEditingAsset(UObject* Asset) = 0;
+};
+
+/** The way that editors were requested to close */
+enum class EAssetEditorCloseReason : uint8
+{
+	CloseAllEditorsForAsset,
+	CloseOtherEditors,
+	RemoveAssetFromAllEditors,
+	CloseAllAssetEditors,
+};
+
+/**
+ * UAssetEditorSubsystem
+ */
+UCLASS()
+class UNREALED_API UAssetEditorSubsystem : public UEditorSubsystem
+{
+	GENERATED_BODY()
+
+public:
+	UAssetEditorSubsystem();
+
+	virtual void Deinitialize() override;
+
+	/** Opens an asset by path */
+	void OpenEditorForAsset(const FString& AssetPathName);
+
+	/**
+	 * Tries to open an editor for the specified asset.  Returns true if the asset is open in an editor.
+	 * If the file is already open in an editor, it will not create another editor window but instead bring it to front
+	 */
+	bool OpenEditorForAsset(UObject* Asset, const EToolkitMode::Type ToolkitMode = EToolkitMode::Standalone, TSharedPtr<IToolkitHost> OpenedFromLevelEditor = TSharedPtr<IToolkitHost>(), const bool bShowProgressWindow = true);
+
+	/**
+	 * Tries to open an editor for all of the specified assets.
+	 * If any of the assets are already open, it will not create a new editor for them.
+	 * If all assets are of the same type, the supporting AssetTypeAction (if it exists) is responsible for the details of how to handle opening multiple assets at once.
+	 */
+	bool OpenEditorForAssets(const TArray<UObject*>& Assets, const EToolkitMode::Type ToolkitMode = EToolkitMode::Standalone, TSharedPtr<IToolkitHost> OpenedFromLevelEditor = TSharedPtr<IToolkitHost>());
+
+	/** Opens editors for the supplied assets (via OpenEditorForAsset) */
+	void OpenEditorsForAssets(const TArray<FString>& AssetsToOpen);
+	void OpenEditorsForAssets(const TArray<FName>& AssetsToOpen);
+
+	/** Returns the primary editor if one is already open for the specified asset.
+	 * If there is one open and bFocusIfOpen is true, that editor will be brought to the foreground and focused if possible.
+	 */
+	IAssetEditorInstance* FindEditorForAsset(UObject* Asset, bool bFocusIfOpen);
+
+	/** Returns all editors currently opened for the specified asset */
+	TArray<IAssetEditorInstance*> FindEditorsForAsset(UObject* Asset);
+
+	/** Close all active editors for the supplied asset and return the number of asset editors that were closed */
+	int32 CloseAllEditorsForAsset(UObject* Asset);
+
+	/** Close any editor which is not this one */
+	void CloseOtherEditors(UObject* Asset, IAssetEditorInstance* OnlyEditor);
+
+	/** Remove given asset from all open editors */
+	void RemoveAssetFromAllEditors(UObject* Asset);
+
+	/** Event called when CloseAllEditorsForAsset/RemoveAssetFromAllEditors is called */
+	DECLARE_EVENT_TwoParams(UAssetEditorSubsystem, FAssetEditorRequestCloseEvent, UObject*, EAssetEditorCloseReason);
+	virtual FAssetEditorRequestCloseEvent& OnAssetEditorRequestClose() { return AssetEditorRequestCloseEvent; }
+
+	/** Get all assets currently being tracked with open editors */
+	TArray<UObject*> GetAllEditedAssets();
+
+	/** Notify the asset editor manager that an asset was opened */
+	void NotifyAssetOpened(UObject* Asset, IAssetEditorInstance* Instance);
+	void NotifyAssetsOpened(const TArray< UObject* >& Assets, IAssetEditorInstance* Instance);
+
+	/** Called when an asset has been opened in an editor */
+	DECLARE_EVENT_TwoParams(UAssetEditorSubsystem, FOnAssetOpenedInEditorEvent, UObject*, IAssetEditorInstance*);
+	virtual FOnAssetOpenedInEditorEvent& OnAssetOpenedInEditor() { return AssetOpenedInEditorEvent; }
+
+	/** Notify the asset editor manager that an asset editor is done editing an asset */
+	void NotifyAssetClosed(UObject* Asset, IAssetEditorInstance* Instance);
+
+	/** Notify the asset editor manager that an asset was closed */
+	void NotifyEditorClosed(IAssetEditorInstance* Instance);
+
+	/** Close all open asset editors */
+	bool CloseAllAssetEditors();
+
+	/** Called when an asset editor is requested to be opened */
+	DECLARE_EVENT_OneParam(UAssetEditorSubsystem, FAssetEditorRequestOpenEvent, UObject*);
+	virtual FAssetEditorRequestOpenEvent& OnAssetEditorRequestedOpen() { return AssetEditorRequestOpenEvent; }
+
+	/** Called when an asset editor is actually opened */
+	DECLARE_EVENT_OneParam(UAssetEditorSubsystem, FAssetEditorOpenEvent, UObject*);
+	FAssetEditorOpenEvent& OnAssetEditorOpened() { return AssetEditorOpenedEvent; }
+
+	/** Request notification to restore the assets that were previously open when the editor was last closed */
+	void RequestRestorePreviouslyOpenAssets();
+
+private:
+
+	/** Handles FAssetEditorRequestOpenAsset messages. */
+	void HandleRequestOpenAssetMessage(const FAssetEditorRequestOpenAsset& Message, const TSharedRef<IMessageContext, ESPMode::ThreadSafe>& Context);
+
+	/** Handles ticks from the ticker. */
+	bool HandleTicker(float DeltaTime);
+
+	/** Spawn a notification asking the user if they want to restore their previously open assets */
+	void SpawnRestorePreviouslyOpenAssetsNotification(const bool bCleanShutdown, const TArray<FString>& AssetsToOpen);
+
+	/** Handler for when the "Restore Now" button is clicked on the RestorePreviouslyOpenAssets notification */
+	void OnConfirmRestorePreviouslyOpenAssets(TArray<FString> AssetsToOpen);
+
+	/** Handler for when the "Don't Restore" button is clicked on the RestorePreviouslyOpenAssets notification */
+	void OnCancelRestorePreviouslyOpenAssets();
+
+	/** Saves a list of open asset editors so they can be restored on editor restart */
+	void SaveOpenAssetEditors(bool bOnShutdown);
+
+	/** Restore the assets that were previously open when the editor was last closed */
+	void RestorePreviouslyOpenAssets();
+
+	/** Handles a package being reloaded */
+	void HandlePackageReloaded(const EPackageReloadPhase InPackageReloadPhase, FPackageReloadedEvent* InPackageReloadedEvent);
+
+private:
+
+	/** struct used by OpenedEditorTimes map to store editor names and times */
+	struct FOpenedEditorTime
+	{
+		FName EditorName;
+		FDateTime OpenedTime;
+	};
+
+	/** struct used to track total time and # of invocations during an overall UnrealEd session */
+	struct FAssetEditorAnalyticInfo
+	{
+		FTimespan SumDuration;
+		int32 NumTimesOpened;
+
+		FAssetEditorAnalyticInfo()
+			: SumDuration(0)
+			, NumTimesOpened(0)
+		{
+		}
+	};
+
+	/** Holds the opened assets. */
+	TMultiMap<UObject*, IAssetEditorInstance*> OpenedAssets;
+
+	/** Holds the opened editors. */
+	TMultiMap<IAssetEditorInstance*, UObject*> OpenedEditors;
+
+	/** Holds the times that editors were opened. */
+	TMap<IAssetEditorInstance*, FOpenedEditorTime> OpenedEditorTimes;
+
+	/** Holds the cumulative time editors have been open by type. */
+	TMap<FName, FAssetEditorAnalyticInfo> EditorUsageAnalytics;
+
+private:
+
+	/** Holds the messaging endpoint. */
+	TSharedPtr<FMessageEndpoint, ESPMode::ThreadSafe> MessageEndpoint;
+
+	/** Holds a delegate to be invoked when the widget ticks. */
+	FTickerDelegate TickDelegate;
+
+	/** Call to request closing editors for an asset */
+	FAssetEditorRequestCloseEvent AssetEditorRequestCloseEvent;
+
+	/** Called when an asset has been opened in an editor */
+	FOnAssetOpenedInEditorEvent AssetOpenedInEditorEvent;
+
+	/** Multicast delegate executed when an asset editor is requested to be opened */
+	FAssetEditorRequestOpenEvent AssetEditorRequestOpenEvent;
+
+	/** Multicast delegate executed when an asset editor is actually opened */
+	FAssetEditorOpenEvent AssetEditorOpenedEvent;
+
+	/** Flag whether we are currently shutting down */
+	bool bSavingOnShutdown;
+
+	/** Flag whether there has been a request to notify whether to restore previously open assets */
+	bool bRequestRestorePreviouslyOpenAssets;
+
+	/** A pointer to the notification used by RestorePreviouslyOpenAssets */
+	TWeakPtr<SNotificationItem> RestorePreviouslyOpenAssetsNotificationPtr;
+
+};
