Enhacement request from JIRA UE-77383: Deprecation warnings removed from Enterprise.

diff --git a/Engine/Source/Editor/UnrealEd/Classes/Editor/EditorEngine.h b/Engine/Source/Editor/UnrealEd/Classes/Editor/EditorEngine.h
--- a/Engine/Source/Editor/UnrealEd/Classes/Editor/EditorEngine.h
+++ b/Engine/Source/Editor/UnrealEd/Classes/Editor/EditorEngine.h
@@ -725,7 +725,7 @@ private:
 public:
 
 	/** The "manager" of all the layers for the UWorld currently being edited */
-	UE_DEPRECATED(4.24, "ILayers and FLayers (GEditor->Layers) have been deprecated, use ULayersSubsystem (GEditor->GetEditorSubsystem<ULayersSubsystem>()) instead.")
+	UE_DEPRECATED(4.24, "ILayers and FLayers (GEditor->Layers) have been deprecated, use ULayersSubsystem (rather than ILayers) and GEditor->GetEditorSubsystem<ULayersSubsystem>() (rather than GEditor->Layers) instead.")
 	TSharedPtr< class ILayers >				Layers;
 
 	/** List of all viewport clients */
diff --git a/Engine/Source/Editor/UnrealEd/Public/Layers/ILayers.h b/Engine/Source/Editor/UnrealEd/Public/Layers/ILayers.h
--- a/Engine/Source/Editor/UnrealEd/Public/Layers/ILayers.h
+++ b/Engine/Source/Editor/UnrealEd/Public/Layers/ILayers.h
@@ -32,7 +32,7 @@ namespace ELayersAction
 	};
 }
 
-class UE_DEPRECATED(4.24, "ILayers and GEditor->Layers have been deprecated, use ULayersSubsystem and GEditor->GetEditorSubsystem<ULayersSubsystem>() instead.") ILayers
+class UE_DEPRECATED(4.24, "ILayers and GEditor->Layers have been deprecated, use ULayersSubsystem (rather than ILayers) and GEditor->GetEditorSubsystem<ULayersSubsystem>() (rather than GEditor->Layers) instead.") ILayers
 {
 
 public:
