Fix clang compile warnings
#rnx

diff --git a/Engine/Source/Editor/UnrealEd/Private/MaterialGraphSchema.cpp b/Engine/Source/Editor/UnrealEd/Private/MaterialGraphSchema.cpp
--- a/Engine/Source/Editor/UnrealEd/Private/MaterialGraphSchema.cpp
+++ b/Engine/Source/Editor/UnrealEd/Private/MaterialGraphSchema.cpp
@@ -410,7 +410,7 @@ void UMaterialGraphSchema::GetContextMenuActions(UToolMenu* Menu, UGraphNodeCont
 						LOCTEXT("SelectLinkedNodes", "Select Linked Nodes"),
 						LOCTEXT("SelectLinkedNodesTooltip", "Adds all input Nodes linked to this Pin to selection"),
 						FSlateIcon(),
-						FUIAction(FExecuteAction::CreateUObject(this, &UMaterialGraphSchema::SelectAllInputNodes, const_cast<UEdGraph*>(Context->Graph), const_cast<UEdGraphPin*>(InGraphPin)))
+						FUIAction(FExecuteAction::CreateUObject(const_cast<UMaterialGraphSchema*>(this), &UMaterialGraphSchema::SelectAllInputNodes, const_cast<UEdGraph*>(Context->Graph), const_cast<UEdGraphPin*>(InGraphPin)))
 						);
 				}
 
@@ -464,7 +464,7 @@ void UMaterialGraphSchema::GetContextMenuActions(UToolMenu* Menu, UGraphNodeCont
 									Label,
 									ToolTip,
 									FSlateIcon(),
-									FUIAction(FExecuteAction::CreateUObject(this, &UMaterialGraphSchema::OnConnectToFunctionOutput, const_cast< UEdGraphPin* >(InGraphPin), GraphNode->GetInputPin(0))));
+									FUIAction(FExecuteAction::CreateUObject(const_cast<UMaterialGraphSchema*>(this), &UMaterialGraphSchema::OnConnectToFunctionOutput, const_cast< UEdGraphPin* >(InGraphPin), GraphNode->GetInputPin(0))));
 							}
 						}
 					}
@@ -484,7 +484,7 @@ void UMaterialGraphSchema::GetContextMenuActions(UToolMenu* Menu, UGraphNodeCont
 								Label,
 								ToolTip,
 								FSlateIcon(),
-								FUIAction(FExecuteAction::CreateUObject(this, &UMaterialGraphSchema::OnConnectToMaterial, const_cast< UEdGraphPin* >(InGraphPin), Index)));
+								FUIAction(FExecuteAction::CreateUObject(const_cast<UMaterialGraphSchema*>(this), &UMaterialGraphSchema::OnConnectToMaterial, const_cast< UEdGraphPin* >(InGraphPin), Index)));
 						}
 					}
 				}
