PR #6129: Fix "Writiable" typo in SSettingsEditorCheckoutNotice.cpp and FileHelpers.cpp (Contributed by [REDACTED_PERSON])

diff --git a/Engine/Source/Developer/SharedSettingsWidgets/Private/SSettingsEditorCheckoutNotice.cpp b/Engine/Source/Developer/SharedSettingsWidgets/Private/SSettingsEditorCheckoutNotice.cpp
--- a/Engine/Source/Developer/SharedSettingsWidgets/Private/SSettingsEditorCheckoutNotice.cpp
+++ b/Engine/Source/Developer/SharedSettingsWidgets/Private/SSettingsEditorCheckoutNotice.cpp
@@ -149,10 +149,10 @@ namespace SettingsHelpers
 
 			if (ShowErrorInNotification)
 			{
-				FNotificationInfo MakeWritiableNotification(ErrorMessage);
-				MakeWritiableNotification.ExpireDuration = 3.0f;
+				FNotificationInfo MakeWritableNotification(ErrorMessage);
+				MakeWritableNotification.ExpireDuration = 3.0f;
 
-				FSlateNotificationManager::Get().AddNotification(MakeWritiableNotification);
+				FSlateNotificationManager::Get().AddNotification(MakeWritableNotification);
 			}
 		}
 
diff --git a/Engine/Source/Editor/UnrealEd/Private/FileHelpers.cpp b/Engine/Source/Editor/UnrealEd/Private/FileHelpers.cpp
--- a/Engine/Source/Editor/UnrealEd/Private/FileHelpers.cpp
+++ b/Engine/Source/Editor/UnrealEd/Private/FileHelpers.cpp
@@ -1594,7 +1594,7 @@ bool FEditorFileUtils::PromptToCheckoutPackages(bool bCheckDirty, const TArray<U
 		CheckoutPackagesDialogModule.AddButton(DRT_CheckOut, NSLOCTEXT("PackagesDialogModule", "Dlg_CheckOutButtonp", "Check Out Selected"), NSLOCTEXT("PackagesDialogModule", "Dlg_CheckOutTooltip", "Attempt to Check Out Checked Assets"), CheckOutSelectedDisabledAttrib );
 		
 		// Make writable button to make checked files writable
-		CheckoutPackagesDialogModule.AddButton(DRT_MakeWritable, NSLOCTEXT("PackagesDialogModule", "Dlg_MakeWritableButton", "Make Writable"), NSLOCTEXT("PackagesDialogModule", "Dlg_MakeWritableTooltip", "Makes selected files writiable on disk"));
+		CheckoutPackagesDialogModule.AddButton(DRT_MakeWritable, NSLOCTEXT("PackagesDialogModule", "Dlg_MakeWritableButton", "Make Writable"), NSLOCTEXT("PackagesDialogModule", "Dlg_MakeWritableTooltip", "Makes selected files writable on disk"));
 
 		// The cancel button should be different if we are prompting during a modify.
 		const FText CancelButtonText  = bPromptingAfterModify ? NSLOCTEXT("PackagesDialogModule", "Dlg_AskMeLater", "Ask Me Later") : NSLOCTEXT("PackagesDialogModule", "Dlg_Cancel", "Cancel");
