Additional refactors for deprecating FAssetEditorManager in favor of AssetEditorSubsystem

diff --git a/Engine/Plugins/Blendables/LightPropagationVolume/Source/LightPropagationVolumeEditor/LightPropagationVolumeEditor.Build.cs b/Engine/Plugins/Blendables/LightPropagationVolume/Source/LightPropagationVolumeEditor/LightPropagationVolumeEditor.Build.cs
--- a/Engine/Plugins/Blendables/LightPropagationVolume/Source/LightPropagationVolumeEditor/LightPropagationVolumeEditor.Build.cs
+++ b/Engine/Plugins/Blendables/LightPropagationVolume/Source/LightPropagationVolumeEditor/LightPropagationVolumeEditor.Build.cs
@@ -10,7 +10,7 @@ public class LightPropagationVolumeEditor : ModuleRules
 			new string[] {	
 				"Core",
 				"CoreUObject",
-                "UnrealEd",         // for FAssetEditorManager
+                "UnrealEd",         // for UAssetEditorSubsystem
 				"AssetTools",       // class FAssetTypeActions_Base
 				"LightPropagationVolumeRuntime",
 			}
diff --git a/Engine/Plugins/Runtime/SoundMod/Source/SoundModImporter/Private/AssetTypeActions_SoundMod.cpp b/Engine/Plugins/Runtime/SoundMod/Source/SoundModImporter/Private/AssetTypeActions_SoundMod.cpp
--- a/Engine/Plugins/Runtime/SoundMod/Source/SoundModImporter/Private/AssetTypeActions_SoundMod.cpp
+++ b/Engine/Plugins/Runtime/SoundMod/Source/SoundModImporter/Private/AssetTypeActions_SoundMod.cpp
@@ -45,7 +45,7 @@ bool FAssetTypeActions_SoundMod::CanExecutePlayCommand(TArray<TWeakObjectPtr<USo
 	return Objects.Num() == 1;
 }
 
-void FAssetTypeActions_SoundMod::AssetsActivated(const TArray<UObject*>& InObjects, EAssetTypeActivationMethod::Type ActivationType)
+bool FAssetTypeActions_SoundMod::AssetsActivatedOverride(const TArray<UObject*>& InObjects, EAssetTypeActivationMethod::Type ActivationType)
 {
 	if (ActivationType == EAssetTypeActivationMethod::Previewed)
 	{
@@ -79,11 +79,9 @@ void FAssetTypeActions_SoundMod::AssetsActivated(const TArray<UObject*>& InObjec
 			// Not already playing, play the target sound cue if it exists
 			PlaySound(TargetSound);
 		}
+		return true;
 	}
-	else
-	{
-		FAssetTypeActions_Base::AssetsActivated(InObjects, ActivationType);
-	}
+	return false;
 }
 
 void FAssetTypeActions_SoundMod::ExecutePlaySound(TArray<TWeakObjectPtr<USoundMod>> Objects)
diff --git a/Engine/Plugins/Runtime/SoundMod/Source/SoundModImporter/Private/AssetTypeActions_SoundMod.h b/Engine/Plugins/Runtime/SoundMod/Source/SoundModImporter/Private/AssetTypeActions_SoundMod.h
--- a/Engine/Plugins/Runtime/SoundMod/Source/SoundModImporter/Private/AssetTypeActions_SoundMod.h
+++ b/Engine/Plugins/Runtime/SoundMod/Source/SoundModImporter/Private/AssetTypeActions_SoundMod.h
@@ -16,7 +16,7 @@ public:
 	virtual UClass* GetSupportedClass() const override;
 	virtual bool HasActions(const TArray<UObject*>& InObjects) const override{ return true; }
 	virtual void GetActions(const TArray<UObject*>& InObjects, struct FToolMenuSection& Section) override;
-	virtual void AssetsActivated(const TArray<UObject*>& InObjects, EAssetTypeActivationMethod::Type ActivationType) override;
+	virtual bool AssetsActivatedOverride(const TArray<UObject*>& InObjects, EAssetTypeActivationMethod::Type ActivationType) override;
 	virtual uint32 GetCategories() override{ return EAssetTypeCategories::Sounds; }
 	virtual bool CanFilter() override { return false; }
 
diff --git a/Engine/Source/Developer/AssetTools/Private/AssetTools.cpp b/Engine/Source/Developer/AssetTools/Private/AssetTools.cpp
--- a/Engine/Source/Developer/AssetTools/Private/AssetTools.cpp
+++ b/Engine/Source/Developer/AssetTools/Private/AssetTools.cpp
@@ -129,7 +129,9 @@
 #include "SAdvancedCopyReportDialog.h"
 #include "AssetToolsSettings.h"
 #include "AssetVtConversion.h"
-
+#if WITH_EDITOR
+#include "Subsystems/AssetEditorSubsystem.h"
+#endif
 
 #define LOCTEXT_NAMESPACE "AssetTools"
 
@@ -2886,7 +2888,9 @@ void UAssetToolsImpl::FixupReferencers(const TArray<UObjectRedirector*>& Objects
 
 void UAssetToolsImpl::OpenEditorForAssets(const TArray<UObject*>& Assets)
 {
-	FAssetEditorManager::Get().OpenEditorForAssets(Assets);
+#if WITH_EDITOR
+	GEditor->GetEditorSubsystem<UAssetEditorSubsystem>()->OpenEditorForAssets(Assets);
+#endif
 }
 
 void UAssetToolsImpl::ConvertVirtualTextures(const TArray<UTexture2D *>& Textures, bool bConvertBackToNonVirtual, const TArray<UMaterial *>* RelatedMaterials /* = nullptr */) const
diff --git a/Engine/Source/Developer/AssetTools/Private/AssetTypeActions/AssetTypeActions_AnimBlueprint.cpp b/Engine/Source/Developer/AssetTools/Private/AssetTypeActions/AssetTypeActions_AnimBlueprint.cpp
--- a/Engine/Source/Developer/AssetTools/Private/AssetTypeActions/AssetTypeActions_AnimBlueprint.cpp
+++ b/Engine/Source/Developer/AssetTools/Private/AssetTypeActions/AssetTypeActions_AnimBlueprint.cpp
@@ -18,6 +18,10 @@
 #include "Styling/SlateIconFinder.h"
 #include "IAnimationBlueprintEditorModule.h"
 #include "Preferences/PersonaOptions.h"
+#if WITH_EDITOR
+#include "Subsystems/AssetEditorSubsystem.h"
+#include "Editor.h"
+#endif
 
 #define LOCTEXT_NAMESPACE "AssetTypeActions"
 
@@ -148,11 +152,13 @@ void FAssetTypeActions_AnimBlueprint::OpenAssetEditor( const TArray<UObject*>& I
 			else
 			{
 				const bool bBringToFrontIfOpen = true;
-				if (IAssetEditorInstance* EditorInstance = FAssetEditorManager::Get().FindEditorForAsset(AnimBlueprint, bBringToFrontIfOpen))
+#if WITH_EDITOR
+				if (IAssetEditorInstance* EditorInstance = GEditor->GetEditorSubsystem<UAssetEditorSubsystem>()->FindEditorForAsset(AnimBlueprint, bBringToFrontIfOpen))
 				{
 					EditorInstance->FocusWindow(AnimBlueprint);
 				}
 				else
+#endif
 				{
 					IAnimationBlueprintEditorModule& AnimationBlueprintEditorModule = FModuleManager::LoadModuleChecked<IAnimationBlueprintEditorModule>("AnimationBlueprintEditor");
 					AnimationBlueprintEditorModule.CreateAnimationBlueprintEditor(Mode, EditWithinLevelEditor, AnimBlueprint);
diff --git a/Engine/Source/Developer/AssetTools/Private/AssetTypeActions/AssetTypeActions_AnimationAsset.cpp b/Engine/Source/Developer/AssetTools/Private/AssetTypeActions/AssetTypeActions_AnimationAsset.cpp
--- a/Engine/Source/Developer/AssetTools/Private/AssetTypeActions/AssetTypeActions_AnimationAsset.cpp
+++ b/Engine/Source/Developer/AssetTools/Private/AssetTypeActions/AssetTypeActions_AnimationAsset.cpp
@@ -16,6 +16,10 @@
 #include "Framework/Application/SlateApplication.h"
 #include "Algo/Transform.h"
 #include "IPersonaToolkit.h"
+#if WITH_EDITOR
+#include "Subsystems/AssetEditorSubsystem.h"
+#include "Editor.h"
+#endif
 
 #define LOCTEXT_NAMESPACE "AssetTypeActions"
 
@@ -148,11 +152,13 @@ void FAssetTypeActions_AnimationAsset::OpenAnimAssetEditor(const TArray<UObject*
 		{
 			// First see if we already have it open
 			const bool bBringToFrontIfOpen = true;
-			if (IAssetEditorInstance* EditorInstance = FAssetEditorManager::Get().FindEditorForAsset(AnimAsset, bBringToFrontIfOpen))
+#if WITH_EDITOR
+			if (IAssetEditorInstance* EditorInstance = GEditor->GetEditorSubsystem<UAssetEditorSubsystem>()->FindEditorForAsset(AnimAsset, bBringToFrontIfOpen))
 			{
 				EditorInstance->FocusWindow();
 			}
 			else
+#endif
 			{
 				// See if we are trying to open a single asset. If we are, we re-use a compatible anim editor.
 				bool bSingleAsset = AnimAssets.Num() == 1;
@@ -160,7 +166,10 @@ void FAssetTypeActions_AnimationAsset::OpenAnimAssetEditor(const TArray<UObject*
 				if (bSingleAsset && !bForceNewEditor)
 				{
 					// See if there is an animation asset with the same skeleton already being edited
-					TArray<UObject*> AllEditedAssets = FAssetEditorManager::Get().GetAllEditedAssets();
+					TArray<UObject*> AllEditedAssets;
+#if WITH_EDITOR
+					AllEditedAssets = GEditor->GetEditorSubsystem<UAssetEditorSubsystem>()->GetAllEditedAssets();
+#endif
 					UAnimationAsset* CompatibleEditedAsset = nullptr;
 					for (UObject* EditedAsset : AllEditedAssets)
 					{
@@ -176,7 +185,10 @@ void FAssetTypeActions_AnimationAsset::OpenAnimAssetEditor(const TArray<UObject*
 					if(CompatibleEditedAsset)
 					{
 						// Find the anim editors that are doing it
-						TArray<IAssetEditorInstance*> AssetEditors = FAssetEditorManager::Get().FindEditorsForAsset(CompatibleEditedAsset);
+						TArray<IAssetEditorInstance*> AssetEditors;
+#if WITH_EDITOR
+						AssetEditors = GEditor->GetEditorSubsystem<UAssetEditorSubsystem>()->FindEditorsForAsset(CompatibleEditedAsset);
+#endif
 						for (IAssetEditorInstance* ExistingEditor : AssetEditors)
 						{
 							if (ExistingEditor->GetEditorName() == FName("AnimationEditor"))
@@ -269,6 +281,7 @@ void FAssetTypeActions_AnimationAsset::RetargetNonSkeletonAnimationHandler(USkel
 {
 	RetargetAnimationHandler(OldSkeleton, NewSkeleton, bRemapReferencedAssets, bAllowRemapToExisting, bConvertSpaces, NameRule, InAnimAssets);
 
+#if WITH_EDITOR
 	if(NewSkeleton)
 	{
 		for(auto Asset : InAnimAssets)
@@ -277,11 +290,12 @@ void FAssetTypeActions_AnimationAsset::RetargetNonSkeletonAnimationHandler(USkel
 			{
 				if (EditWithinLevelEditor.IsValid())
 				{
-					FAssetEditorManager::Get().OpenEditorForAsset(Asset.Get(), EToolkitMode::WorldCentric, EditWithinLevelEditor.Pin());
+
+					GEditor->GetEditorSubsystem<UAssetEditorSubsystem>()->OpenEditorForAsset(Asset.Get(), EToolkitMode::WorldCentric, EditWithinLevelEditor.Pin());
 				}
 				else
 				{
-					FAssetEditorManager::Get().OpenEditorForAsset(Asset.Get());
+					GEditor->GetEditorSubsystem<UAssetEditorSubsystem>()->OpenEditorForAsset(Asset.Get());
 				}
 			}
 		}
@@ -290,6 +304,8 @@ void FAssetTypeActions_AnimationAsset::RetargetNonSkeletonAnimationHandler(USkel
 	{
 		FMessageDialog::Open(EAppMsgType::Ok, LOCTEXT("FailedToLoadSkeletonlessAnimAsset", "The Anim Asset could not be loaded because it's skeleton is missing."));
 	}
+#endif
+
 }
 
 
diff --git a/Engine/Source/Developer/AssetTools/Private/AssetTypeActions/AssetTypeActions_Blueprint.cpp b/Engine/Source/Developer/AssetTools/Private/AssetTypeActions/AssetTypeActions_Blueprint.cpp
--- a/Engine/Source/Developer/AssetTools/Private/AssetTypeActions/AssetTypeActions_Blueprint.cpp
+++ b/Engine/Source/Developer/AssetTools/Private/AssetTypeActions/AssetTypeActions_Blueprint.cpp
@@ -16,6 +16,10 @@
 #include "Logging/MessageLog.h"
 #include "IContentBrowserSingleton.h"
 #include "ContentBrowserModule.h"
+#if WITH_EDITOR
+#include "Subsystems/AssetEditorSubsystem.h"
+#include "Editor.h"
+#endif
 
 #define LOCTEXT_NAMESPACE "AssetTypeActions"
 
@@ -105,14 +109,16 @@ void FAssetTypeActions_Blueprint::Merge(UObject* InObject)
 {
 	UBlueprint* AsBlueprint = CastChecked<UBlueprint>(InObject);
 	// Kludge to get the merge panel in the blueprint editor to show up:
-	bool Success = FAssetEditorManager::Get().OpenEditorForAsset(InObject);
+#if WITH_EDITOR
+	bool Success = GEditor->GetEditorSubsystem<UAssetEditorSubsystem>()->OpenEditorForAsset(InObject);
 	if( Success )
 	{
 		FBlueprintEditorModule& BlueprintEditorModule = FModuleManager::LoadModuleChecked<FBlueprintEditorModule>( "Kismet" );
 
-		FBlueprintEditor* BlueprintEditor = static_cast<FBlueprintEditor*>(FAssetEditorManager::Get().FindEditorForAsset(AsBlueprint, false));
+		FBlueprintEditor* BlueprintEditor = static_cast<FBlueprintEditor*>(GEditor->GetEditorSubsystem<UAssetEditorSubsystem>()->FindEditorForAsset(AsBlueprint, false));
 		BlueprintEditor->CreateMergeToolTab();
 	}
+#endif
 }
 
 void FAssetTypeActions_Blueprint::Merge(UObject* BaseAsset, UObject* RemoteAsset, UObject* LocalAsset, const FOnMergeResolved& ResolutionCallback)
@@ -120,14 +126,15 @@ void FAssetTypeActions_Blueprint::Merge(UObject* BaseAsset, UObject* RemoteAsset
 	UBlueprint* AsBlueprint = CastChecked<UBlueprint>(LocalAsset);
 	check(LocalAsset->GetClass() == BaseAsset->GetClass());
 	check(LocalAsset->GetClass() == RemoteAsset->GetClass());
-	
-	if (FAssetEditorManager::Get().OpenEditorForAsset(AsBlueprint))
+#if WITH_EDITOR
+	if (GEditor->GetEditorSubsystem<UAssetEditorSubsystem>()->OpenEditorForAsset(AsBlueprint))
 	{
 		FBlueprintEditorModule& BlueprintEditorModule = FModuleManager::LoadModuleChecked<FBlueprintEditorModule>("Kismet");
 
-		FBlueprintEditor* BlueprintEditor = static_cast<FBlueprintEditor*>(FAssetEditorManager::Get().FindEditorForAsset(AsBlueprint, /*bFocusIfOpen =*/false));
+		FBlueprintEditor* BlueprintEditor = static_cast<FBlueprintEditor*>(GEditor->GetEditorSubsystem<UAssetEditorSubsystem>()->FindEditorForAsset(AsBlueprint, /*bFocusIfOpen =*/false));
 		BlueprintEditor->CreateMergeToolTab(Cast<UBlueprint>(BaseAsset), Cast<UBlueprint>(RemoteAsset), ResolutionCallback);
 	}
+#endif
 }
 
 bool FAssetTypeActions_Blueprint::CanCreateNewDerivedBlueprint() const
diff --git a/Engine/Source/Developer/AssetTools/Private/AssetTypeActions/AssetTypeActions_ForceFeedbackEffect.cpp b/Engine/Source/Developer/AssetTools/Private/AssetTypeActions/AssetTypeActions_ForceFeedbackEffect.cpp
--- a/Engine/Source/Developer/AssetTools/Private/AssetTypeActions/AssetTypeActions_ForceFeedbackEffect.cpp
+++ b/Engine/Source/Developer/AssetTools/Private/AssetTypeActions/AssetTypeActions_ForceFeedbackEffect.cpp
@@ -87,7 +87,7 @@ bool FAssetTypeActions_ForceFeedbackEffect::CanExecutePlayCommand(TArray<TWeakOb
 	return Objects.Num() == 1;
 }
 
-void FAssetTypeActions_ForceFeedbackEffect::AssetsActivated( const TArray<UObject*>& InObjects, EAssetTypeActivationMethod::Type ActivationType )
+bool FAssetTypeActions_ForceFeedbackEffect::AssetsActivatedOverride( const TArray<UObject*>& InObjects, EAssetTypeActivationMethod::Type ActivationType )
 {
 	if (ActivationType == EAssetTypeActivationMethod::Previewed)
 	{
@@ -108,14 +108,11 @@ void FAssetTypeActions_ForceFeedbackEffect::AssetsActivated( const TArray<UObjec
 					ExecutePlayEffect(EffectList);
 				}
 
-				break;
+				return true;
 			}
 		}
 	}
-	else
-	{
-		FAssetTypeActions_Base::AssetsActivated(InObjects, ActivationType);
-	}
+	return false;
 }
 
 void FAssetTypeActions_ForceFeedbackEffect::ExecutePlayEffect(TArray<TWeakObjectPtr<UForceFeedbackEffect>> Objects)
diff --git a/Engine/Source/Developer/AssetTools/Private/AssetTypeActions/AssetTypeActions_ForceFeedbackEffect.h b/Engine/Source/Developer/AssetTools/Private/AssetTypeActions/AssetTypeActions_ForceFeedbackEffect.h
--- a/Engine/Source/Developer/AssetTools/Private/AssetTypeActions/AssetTypeActions_ForceFeedbackEffect.h
+++ b/Engine/Source/Developer/AssetTools/Private/AssetTypeActions/AssetTypeActions_ForceFeedbackEffect.h
@@ -32,7 +32,7 @@ public:
 	virtual uint32 GetCategories() override { return EAssetTypeCategories::Misc; }
 	virtual bool HasActions ( const TArray<UObject*>& InObjects ) const override { return true; }
 	virtual void GetActions(const TArray<UObject*>& InObjects, struct FToolMenuSection& Section) override;
-	virtual void AssetsActivated( const TArray<UObject*>& InObjects, EAssetTypeActivationMethod::Type ActivationType ) override;
+	virtual bool AssetsActivatedOverride( const TArray<UObject*>& InObjects, EAssetTypeActivationMethod::Type ActivationType ) override;
 	virtual TSharedPtr<SWidget> GetThumbnailOverlay(const FAssetData& AssetData) const override;
 
 	/** Return if the specified effect is playing*/
diff --git a/Engine/Source/Developer/AssetTools/Private/AssetTypeActions/AssetTypeActions_Redirector.cpp b/Engine/Source/Developer/AssetTools/Private/AssetTypeActions/AssetTypeActions_Redirector.cpp
--- a/Engine/Source/Developer/AssetTools/Private/AssetTypeActions/AssetTypeActions_Redirector.cpp
+++ b/Engine/Source/Developer/AssetTools/Private/AssetTypeActions/AssetTypeActions_Redirector.cpp
@@ -33,7 +33,7 @@ void FAssetTypeActions_Redirector::GetActions(const TArray<UObject*>& InObjects,
 		);
 }
 
-void FAssetTypeActions_Redirector::AssetsActivated( const TArray<UObject*>& InObjects, EAssetTypeActivationMethod::Type ActivationType )
+bool FAssetTypeActions_Redirector::AssetsActivatedOverride( const TArray<UObject*>& InObjects, EAssetTypeActivationMethod::Type ActivationType )
 {
 	if ( ActivationType == EAssetTypeActivationMethod::DoubleClicked || ActivationType == EAssetTypeActivationMethod::Opened )
 	{
@@ -51,12 +51,11 @@ void FAssetTypeActions_Redirector::AssetsActivated( const TArray<UObject*>& InOb
 		if ( Redirectors.Num() > 0 )
 		{
 			FindTargets(Redirectors);
+			return true;
 		}
 	}
-	else
-	{
-		FAssetTypeActions_Base::AssetsActivated(InObjects, ActivationType);
-	}
+
+	return false;
 }
 
 void FAssetTypeActions_Redirector::ExecuteFindTarget(TArray<TWeakObjectPtr<UObjectRedirector>> Objects)
diff --git a/Engine/Source/Developer/AssetTools/Private/AssetTypeActions/AssetTypeActions_Redirector.h b/Engine/Source/Developer/AssetTools/Private/AssetTypeActions/AssetTypeActions_Redirector.h
--- a/Engine/Source/Developer/AssetTools/Private/AssetTypeActions/AssetTypeActions_Redirector.h
+++ b/Engine/Source/Developer/AssetTools/Private/AssetTypeActions/AssetTypeActions_Redirector.h
@@ -16,7 +16,7 @@ public:
 	virtual UClass* GetSupportedClass() const override { return UObjectRedirector::StaticClass(); }
 	virtual bool HasActions ( const TArray<UObject*>& InObjects ) const override { return true; }
 	virtual void GetActions(const TArray<UObject*>& InObjects, struct FToolMenuSection& Section) override;
-	virtual void AssetsActivated( const TArray<UObject*>& InObjects, EAssetTypeActivationMethod::Type ActivationType ) override;
+	virtual bool AssetsActivatedOverride( const TArray<UObject*>& InObjects, EAssetTypeActivationMethod::Type ActivationType ) override;
 	virtual uint32 GetCategories() override { return EAssetTypeCategories::Misc; }
 	virtual bool CanLocalize() const override { return false; }
 
diff --git a/Engine/Source/Developer/AssetTools/Private/AssetTypeActions/AssetTypeActions_SkeletalMesh.cpp b/Engine/Source/Developer/AssetTools/Private/AssetTypeActions/AssetTypeActions_SkeletalMesh.cpp
--- a/Engine/Source/Developer/AssetTools/Private/AssetTypeActions/AssetTypeActions_SkeletalMesh.cpp
+++ b/Engine/Source/Developer/AssetTools/Private/AssetTypeActions/AssetTypeActions_SkeletalMesh.cpp
@@ -41,6 +41,11 @@
 #include "Widgets/Images/SImage.h"
 #include "Factories/FbxSkeletalMeshImportData.h"
 
+#if WITH_EDITOR
+#include "Subsystems/AssetEditorSubsystem.h"
+#include "Editor.h"
+#endif
+
 #define LOCTEXT_NAMESPACE "AssetTypeActions"
 
 /** 
@@ -541,11 +546,13 @@ void FAssetTypeActions_SkeletalMesh::OpenAssetEditor( const TArray<UObject*>& In
 			if ( Mesh->Skeleton != NULL )
 			{
 				const bool bBringToFrontIfOpen = true;
-				if (IAssetEditorInstance* EditorInstance = FAssetEditorManager::Get().FindEditorForAsset(Mesh, bBringToFrontIfOpen))
+#if WITH_EDITOR
+				if (IAssetEditorInstance* EditorInstance = GEditor->GetEditorSubsystem<UAssetEditorSubsystem>()->FindEditorForAsset(Mesh, bBringToFrontIfOpen))
 				{
 					EditorInstance->FocusWindow(Mesh);
 				}
 				else
+#endif
 				{
 					ISkeletalMeshEditorModule& SkeletalMeshEditorModule = FModuleManager::LoadModuleChecked<ISkeletalMeshEditorModule>("SkeletalMeshEditor");
 					SkeletalMeshEditorModule.CreateSkeletalMeshEditor(Mode, EditWithinLevelEditor, Mesh);
@@ -710,7 +717,9 @@ void FAssetTypeActions_SkeletalMesh::ExecuteNewPhysicsAsset(TArray<TWeakObjectPt
 	{
 		FContentBrowserModule& ContentBrowserModule = FModuleManager::LoadModuleChecked<FContentBrowserModule>("ContentBrowser");
 		ContentBrowserModule.Get().SyncBrowserToAssets(CreatedObjects);
-		FAssetEditorManager::Get().OpenEditorForAssets(CreatedObjects);
+#if WITH_EDITOR
+		GEditor->GetEditorSubsystem<UAssetEditorSubsystem>()->OpenEditorForAssets(CreatedObjects);
+#endif
 	}
 }
 
diff --git a/Engine/Source/Developer/AssetTools/Private/AssetTypeActions/AssetTypeActions_Skeleton.cpp b/Engine/Source/Developer/AssetTools/Private/AssetTypeActions/AssetTypeActions_Skeleton.cpp
--- a/Engine/Source/Developer/AssetTools/Private/AssetTypeActions/AssetTypeActions_Skeleton.cpp
+++ b/Engine/Source/Developer/AssetTools/Private/AssetTypeActions/AssetTypeActions_Skeleton.cpp
@@ -43,6 +43,10 @@
 #include "Preferences/PersonaOptions.h"
 #include "Algo/Transform.h"
 #include "PhysicsEngine/PhysicsAsset.h"
+#if WITH_EDITOR
+#include "Subsystems/AssetEditorSubsystem.h"
+#include "Editor.h"
+#endif
 
 #define LOCTEXT_NAMESPACE "AssetTypeActions"
 
@@ -504,11 +508,13 @@ void FAssetTypeActions_Skeleton::OpenAssetEditor( const TArray<UObject*>& InObje
 		if (Skeleton != NULL)
 		{
 			const bool bBringToFrontIfOpen = true;
-			if (IAssetEditorInstance* EditorInstance = FAssetEditorManager::Get().FindEditorForAsset(Skeleton, bBringToFrontIfOpen))
+#if WITH_EDITOR
+			if (IAssetEditorInstance* EditorInstance = GEditor->GetEditorSubsystem<UAssetEditorSubsystem>()->FindEditorForAsset(Skeleton, bBringToFrontIfOpen))
 			{
 				EditorInstance->FocusWindow(Skeleton);
 			}
 			else
+#endif
 			{
 				ISkeletonEditorModule& SkeletonEditorModule = FModuleManager::LoadModuleChecked<ISkeletonEditorModule>("SkeletonEditor");
 				SkeletonEditorModule.CreateSkeletonEditor(Mode, EditWithinLevelEditor, Skeleton);
@@ -599,7 +605,10 @@ void FAssetTypeActions_Skeleton::RetargetAnimationHandler(USkeleton* OldSkeleton
 }
 void FAssetTypeActions_Skeleton::ExecuteRetargetSkeleton(TArray<TWeakObjectPtr<USkeleton>> Skeletons)
 {
-	TArray<UObject*> AllEditedAssets = FAssetEditorManager::Get().GetAllEditedAssets();
+	TArray<UObject*> AllEditedAssets;
+#if WITH_EDITOR
+	AllEditedAssets = GEditor->GetEditorSubsystem<UAssetEditorSubsystem>()->GetAllEditedAssets();
+#endif
 
 	for (auto SkelIt = Skeletons.CreateConstIterator(); SkelIt; ++SkelIt)
 	{	
@@ -636,7 +645,9 @@ void FAssetTypeActions_Skeleton::ExecuteRetargetSkeleton(TArray<TWeakObjectPtr<U
 
 			if(bCloseAssetEditor)
 			{
-				FAssetEditorManager::Get().CloseAllEditorsForAsset(EditedAsset);
+#if WITH_EDITOR
+				GEditor->GetEditorSubsystem<UAssetEditorSubsystem>()->CloseAllEditorsForAsset(EditedAsset);
+#endif
 			}
 		}
 
diff --git a/Engine/Source/Developer/AssetTools/Public/AssetTypeActions_Base.h b/Engine/Source/Developer/AssetTools/Public/AssetTypeActions_Base.h
--- a/Engine/Source/Developer/AssetTools/Public/AssetTypeActions_Base.h
+++ b/Engine/Source/Developer/AssetTools/Public/AssetTypeActions_Base.h
@@ -46,21 +46,29 @@ public:
 	{
 		FSimpleAssetEditor::CreateEditor(EToolkitMode::Standalone, EditWithinLevelEditor, InObjects);
 	}
-	
+
+PRAGMA_DISABLE_DEPRECATION_WARNINGS
+UE_DEPRECATED(4.24, "Use AssetsActivatedOverride instead to provide any non-default behavior. Using AssetsActivatedOverride, you no longer need a call to FAssetTypeActions_Base::AssetsActivated.")
 	virtual void AssetsActivated( const TArray<UObject*>& InObjects, EAssetTypeActivationMethod::Type ActivationType ) override
 	{
-		if ( ActivationType == EAssetTypeActivationMethod::DoubleClicked || ActivationType == EAssetTypeActivationMethod::Opened )
+		if (ActivationType == EAssetTypeActivationMethod::DoubleClicked || ActivationType == EAssetTypeActivationMethod::Opened)
 		{
-			if ( InObjects.Num() == 1 )
+			if (InObjects.Num() == 1)
 			{
 				FAssetEditorManager::Get().OpenEditorForAsset(InObjects[0]);
 			}
-			else if ( InObjects.Num() > 1 )
+			else if (InObjects.Num() > 1)
 			{
 				FAssetEditorManager::Get().OpenEditorForAssets(InObjects);
 			}
 		}
 	}
+PRAGMA_ENABLE_DEPRECATION_WARNINGS
+
+	virtual bool AssetsActivatedOverride(const TArray<UObject*>& InObjects, EAssetTypeActivationMethod::Type ActivationType) override
+	{
+		return false;
+	}
 
 	virtual bool CanFilter() override
 	{
diff --git a/Engine/Source/Developer/AssetTools/Public/IAssetTools.h b/Engine/Source/Developer/AssetTools/Public/IAssetTools.h
--- a/Engine/Source/Developer/AssetTools/Public/IAssetTools.h
+++ b/Engine/Source/Developer/AssetTools/Public/IAssetTools.h
@@ -389,7 +389,7 @@ public:
 	virtual void InitAdvancedCopyFromCopyParams(FAdvancedCopyParams CopyParams) const = 0;
 
 	/** Opens editor for assets */
-	UFUNCTION(BlueprintCallable, Category = "Editor Scripting | Asset Tools")
+	UFUNCTION(BlueprintCallable, Category = "Editor Scripting | Asset Tools", meta = (DeprecatedFunction, DeprecationMessage = "Please use UAssetEditorSubsystem::OpenEditorForAssets instead."))
 	virtual void OpenEditorForAssets(const TArray<UObject*>& Assets) = 0;
 	
 	/** Converts the given UTexture2D to virtual textures or converts virtual textures back to standard textures and updates the related UMaterials 
diff --git a/Engine/Source/Developer/AssetTools/Public/IAssetTypeActions.h b/Engine/Source/Developer/AssetTools/Public/IAssetTypeActions.h
--- a/Engine/Source/Developer/AssetTools/Public/IAssetTypeActions.h
+++ b/Engine/Source/Developer/AssetTools/Public/IAssetTypeActions.h
@@ -66,6 +66,9 @@ public:
 	/** Performs asset type specific activation for the supplied assets. This happens when the user double clicks, presses enter, or presses space. */
 	virtual void AssetsActivated( const TArray<UObject*>& InObjects, EAssetTypeActivationMethod::Type ActivationType ) = 0;
 
+	/** Allows overriding asset activation to perform asset type specific activation for the supplied assets. This happens when the user double clicks, presses enter, or presses space. Return true if you have overridden the behavior. */
+	virtual bool AssetsActivatedOverride(const TArray<UObject*>& InObjects, EAssetTypeActivationMethod::Type ActivationType) = 0;
+
 	/** Returns true if this class can be used as a filter in the content browser */
 	virtual bool CanFilter() = 0;
 
diff --git a/Engine/Source/Developer/Merge/Private/SBlueprintMerge.cpp b/Engine/Source/Developer/Merge/Private/SBlueprintMerge.cpp
--- a/Engine/Source/Developer/Merge/Private/SBlueprintMerge.cpp
+++ b/Engine/Source/Developer/Merge/Private/SBlueprintMerge.cpp
@@ -28,6 +28,7 @@
 #include "Framework/Notifications/NotificationManager.h"
 #include "Widgets/Notifications/SNotificationList.h"
 #include "SMergeAssetPickerView.h"
+#include "Subsystems/AssetEditorSubsystem.h"
 
 #define LOCTEXT_NAMESPACE "SBlueprintMerge"
 
@@ -540,11 +541,14 @@ void SBlueprintMerge::ResolveMerge(UBlueprint* ResultantBlueprint)
 	UBlueprint* TargetBlueprint = GetTargetBlueprint();
 	if (ResultantBlueprint != TargetBlueprint)
 	{
-		if (TargetBlueprint)
+		if (GEditor)
 		{
-			FAssetEditorManager::Get().CloseAllEditorsForAsset(TargetBlueprint);
+			if (TargetBlueprint)
+			{
+				GEditor->GetEditorSubsystem<UAssetEditorSubsystem>()->CloseAllEditorsForAsset(TargetBlueprint);
+			}
+			GEditor->GetEditorSubsystem<UAssetEditorSubsystem>()->OpenEditorForAsset(ResultantBlueprint);
 		}
-		FAssetEditorManager::Get().OpenEditorForAsset(ResultantBlueprint);
 	}
 	
 	// should come before CloseMergeTool(), because CloseMergeTool() makes its
diff --git a/Engine/Source/Developer/MeshUtilities/Private/MeshUtilities.cpp b/Engine/Source/Developer/MeshUtilities/Private/MeshUtilities.cpp
--- a/Engine/Source/Developer/MeshUtilities/Private/MeshUtilities.cpp
+++ b/Engine/Source/Developer/MeshUtilities/Private/MeshUtilities.cpp
@@ -100,6 +100,7 @@
 #if WITH_EDITOR
 #include "Editor.h"
 #include "UnrealEdMisc.h"
+#include "Subsystems/AssetEditorSubsystem.h"
 #endif
 #include "MaterialBakingStructures.h"
 #include "IMaterialBakingModule.h"
@@ -111,6 +112,7 @@
 #include "IMeshReductionManagerModule.h"
 #include "MeshMergeModule.h"
 
+
 DEFINE_LOG_CATEGORY(LogMeshUtilities);
 /*------------------------------------------------------------------------------
 MeshUtilities module.
@@ -604,7 +606,7 @@ UStaticMesh* FMeshUtilities::ConvertMeshesToStaticMesh(const TArray<UMeshCompone
 				FNotificationInfo Info(FText::Format(LOCTEXT("SkeletalMeshConverted", "Successfully Converted Mesh"), FText::FromString(StaticMesh->GetName())));
 				Info.ExpireDuration = 8.0f;
 				Info.bUseLargeFont = false;
-				Info.Hyperlink = FSimpleDelegate::CreateLambda([=]() { FAssetEditorManager::Get().OpenEditorForAssets(TArray<UObject*>({ StaticMesh })); });
+				Info.Hyperlink = FSimpleDelegate::CreateLambda([=]() { GEditor->GetEditorSubsystem<UAssetEditorSubsystem>()->OpenEditorForAssets(TArray<UObject*>({ StaticMesh })); });
 				Info.HyperlinkText = FText::Format(LOCTEXT("OpenNewAnimationHyperlink", "Open {0}"), FText::FromString(StaticMesh->GetName()));
 				TSharedPtr<SNotificationItem> Notification = FSlateNotificationManager::Get().AddNotification(Info);
 				if ( Notification.IsValid() )
diff --git a/Engine/Source/Editor/AudioEditor/Private/AssetTypeActions/AssetTypeActions_SoundBase.cpp b/Engine/Source/Editor/AudioEditor/Private/AssetTypeActions/AssetTypeActions_SoundBase.cpp
--- a/Engine/Source/Editor/AudioEditor/Private/AssetTypeActions/AssetTypeActions_SoundBase.cpp
+++ b/Engine/Source/Editor/AudioEditor/Private/AssetTypeActions/AssetTypeActions_SoundBase.cpp
@@ -43,7 +43,7 @@ void FAssetTypeActions_SoundBase::GetActions(const TArray<UObject*>& InObjects,
 			)
 		);
 }
-
+PRAGMA_DISABLE_DEPRECATION_WARNINGS
 void FAssetTypeActions_SoundBase::AssetsActivated( const TArray<UObject*>& InObjects, EAssetTypeActivationMethod::Type ActivationType )
 {
 	if (ActivationType == EAssetTypeActivationMethod::Previewed)
@@ -84,6 +84,46 @@ void FAssetTypeActions_SoundBase::AssetsActivated( const TArray<UObject*>& InObj
 		FAssetTypeActions_Base::AssetsActivated(InObjects, ActivationType);
 	}
 }
+PRAGMA_ENABLE_DEPRECATION_WARNINGS
+
+bool FAssetTypeActions_SoundBase::AssetsActivatedOverride(const TArray<UObject*>& InObjects, EAssetTypeActivationMethod::Type ActivationType)
+{
+	if (ActivationType == EAssetTypeActivationMethod::Previewed)
+	{
+		USoundBase* TargetSound = NULL;
+
+		for (auto ObjIt = InObjects.CreateConstIterator(); ObjIt; ++ObjIt)
+		{
+			TargetSound = Cast<USoundBase>(*ObjIt);
+			if (TargetSound)
+			{
+				// Only target the first valid sound cue
+				break;
+			}
+		}
+
+		UAudioComponent* PreviewComp = GEditor->GetPreviewAudioComponent();
+		if (PreviewComp && PreviewComp->IsPlaying())
+		{
+			// Already previewing a sound, if it is the target cue then stop it, otherwise play the new one
+			if (!TargetSound || PreviewComp->Sound == TargetSound)
+			{
+				StopSound();
+			}
+			else
+			{
+				PlaySound(TargetSound);
+			}
+		}
+		else
+		{
+			// Not already playing, play the target sound cue if it exists
+			PlaySound(TargetSound);
+		}
+		return true;
+	}
+	return false;
+}
 
 void FAssetTypeActions_SoundBase::PlaySound(USoundBase* Sound) const
 {
diff --git a/Engine/Source/Editor/AudioEditor/Public/AssetTypeActions/AssetTypeActions_SoundBase.h b/Engine/Source/Editor/AudioEditor/Public/AssetTypeActions/AssetTypeActions_SoundBase.h
--- a/Engine/Source/Editor/AudioEditor/Public/AssetTypeActions/AssetTypeActions_SoundBase.h
+++ b/Engine/Source/Editor/AudioEditor/Public/AssetTypeActions/AssetTypeActions_SoundBase.h
@@ -17,7 +17,10 @@ public:
 	virtual UClass* GetSupportedClass() const override;
 	virtual bool HasActions ( const TArray<UObject*>& InObjects ) const override { return true; }
 	virtual void GetActions(const TArray<UObject*>& InObjects, struct FToolMenuSection& Section) override;
+PRAGMA_DISABLE_DEPRECATION_WARNINGS
 	virtual void AssetsActivated( const TArray<UObject*>& InObjects, EAssetTypeActivationMethod::Type ActivationType ) override;
+PRAGMA_ENABLE_DEPRECATION_WARNINGS
+	virtual bool AssetsActivatedOverride(const TArray<UObject*>& InObjects, EAssetTypeActivationMethod::Type ActivationType) override;
 	virtual uint32 GetCategories() override { return EAssetTypeCategories::Sounds; }
 	virtual bool CanFilter() override { return false; }
 	virtual TSharedPtr<SWidget> GetThumbnailOverlay(const FAssetData& AssetData) const override;
diff --git a/Engine/Source/Editor/ContentBrowser/Private/SContentBrowser.cpp b/Engine/Source/Editor/ContentBrowser/Private/SContentBrowser.cpp
--- a/Engine/Source/Editor/ContentBrowser/Private/SContentBrowser.cpp
+++ b/Engine/Source/Editor/ContentBrowser/Private/SContentBrowser.cpp
@@ -55,6 +55,7 @@
 #include "AddToProjectConfig.h"
 #include "GameProjectGenerationModule.h"
 #include "Toolkits/GlobalEditorCommonCommands.h"
+#include "Subsystems/AssetEditorSubsystem.h"
 
 #define LOCTEXT_NAMESPACE "ContentBrowser"
 
@@ -2194,7 +2195,20 @@ void SContentBrowser::OnAssetsActivated(const TArray<FAssetData>& ActivatedAsset
 		const TSharedRef<IAssetTypeActions>& TypeActions = TypeActionsIt.Key();
 		const TArray<UObject*>& ObjList = TypeActionsIt.Value();
 
-		TypeActions->AssetsActivated(ObjList, ActivationMethod);
+		if (!TypeActions->AssetsActivatedOverride(ObjList, ActivationMethod))
+		{
+			if (ActivationMethod == EAssetTypeActivationMethod::DoubleClicked || ActivationMethod == EAssetTypeActivationMethod::Opened)
+			{
+				if (ObjList.Num() == 1)
+				{
+					GEditor->GetEditorSubsystem<UAssetEditorSubsystem>()->OpenEditorForAsset(ObjList[0]);
+				}
+				else if (ObjList.Num() > 1)
+				{
+					GEditor->GetEditorSubsystem<UAssetEditorSubsystem>()->OpenEditorForAssets(ObjList);
+				}
+			}
+		}
 	}
 
 	// Finally, open a simple asset editor for all assets which do not have asset type actions if activating with enter or double click
diff --git a/Engine/Source/Editor/DataTableEditor/Private/DataTableEditor.cpp b/Engine/Source/Editor/DataTableEditor/Private/DataTableEditor.cpp
--- a/Engine/Source/Editor/DataTableEditor/Private/DataTableEditor.cpp
+++ b/Engine/Source/Editor/DataTableEditor/Private/DataTableEditor.cpp
@@ -40,6 +40,7 @@
 #include "Toolkits/AssetEditorManager.h"
 #include "DataTableRowUtlis.h"
 #include "Engine/DataTable.h"
+#include "Subsystems/AssetEditorSubsystem.h"
 
 
 
@@ -476,7 +477,7 @@ void FDataTableEditor::OnEditDataTableStructClicked()
 
 		if (ScriptStruct)
 		{
-			FAssetEditorManager::Get().OpenEditorForAsset(ScriptStruct->GetPathName());
+			GEditor->GetEditorSubsystem<UAssetEditorSubsystem>()->OpenEditorForAsset(ScriptStruct->GetPathName());
 			FSourceCodeNavigation::NavigateToStruct(ScriptStruct);
 		}
 	}
diff --git a/Engine/Source/Editor/Sequencer/Private/Sequencer.cpp b/Engine/Source/Editor/Sequencer/Private/Sequencer.cpp
--- a/Engine/Source/Editor/Sequencer/Private/Sequencer.cpp
+++ b/Engine/Source/Editor/Sequencer/Private/Sequencer.cpp
@@ -5059,7 +5059,7 @@ void FSequencer::SaveCurrentMovieSceneAs()
 	{
 		UAssetEditorSubsystem* AssetEditorSubsystem = GEditor->GetEditorSubsystem<UAssetEditorSubsystem>();
 		AssetEditorSubsystem->CloseAllEditorsForAsset(AssetsToSave[0]);
-		AssetEditorSubsystem->OpenEditorForAssets(SavedAssets, EToolkitMode::Standalone, MyToolkitHost.ToSharedRef());
+		AssetEditorSubsystem->OpenEditorForAssets_Advanced(SavedAssets, EToolkitMode::Standalone, MyToolkitHost.ToSharedRef());
 	}
 }
 
@@ -9245,7 +9245,7 @@ void FSequencer::DiscardChanges()
 	TArray<UObject*> AssetsToReopen;
 	AssetsToReopen.Add(ReloadedSequence);
 
-	AssetEditorSubsystem->OpenEditorForAssets(AssetsToReopen, EToolkitMode::Standalone, MyToolkitHost.ToSharedRef());
+	AssetEditorSubsystem->OpenEditorForAssets_Advanced(AssetsToReopen, EToolkitMode::Standalone, MyToolkitHost.ToSharedRef());
 }
 
 
diff --git a/Engine/Source/Editor/UnrealEd/Private/Subsystems/AssetEditorSubsystem.cpp b/Engine/Source/Editor/UnrealEd/Private/Subsystems/AssetEditorSubsystem.cpp
--- a/Engine/Source/Editor/UnrealEd/Private/Subsystems/AssetEditorSubsystem.cpp
+++ b/Engine/Source/Editor/UnrealEd/Private/Subsystems/AssetEditorSubsystem.cpp
@@ -358,7 +358,7 @@ bool UAssetEditorSubsystem::OpenEditorForAsset(UObject* Asset, const EToolkitMod
 }
 
 
-bool UAssetEditorSubsystem::OpenEditorForAssets(const TArray <UObject* >& Assets, const EToolkitMode::Type ToolkitMode, TSharedPtr< IToolkitHost > OpenedFromLevelEditor)
+bool UAssetEditorSubsystem::OpenEditorForAssets_Advanced(const TArray <UObject* >& Assets, const EToolkitMode::Type ToolkitMode, TSharedPtr< IToolkitHost > OpenedFromLevelEditor)
 {
 	if (Assets.Num() == 1)
 	{
@@ -475,6 +475,11 @@ bool UAssetEditorSubsystem::OpenEditorForAssets(const TArray <UObject* >& Assets
 	return true;
 }
 
+bool UAssetEditorSubsystem::OpenEditorForAssets(const TArray<UObject*>& Assets)
+{
+	return OpenEditorForAssets_Advanced(Assets, EToolkitMode::Standalone, TSharedPtr<IToolkitHost>());
+}
+
 void UAssetEditorSubsystem::HandleRequestOpenAssetMessage(const FAssetEditorRequestOpenAsset& Message, const TSharedRef<IMessageContext, ESPMode::ThreadSafe>& Context)
 {
 	OpenEditorForAsset(Message.AssetName);
diff --git a/Engine/Source/Editor/UnrealEd/Private/Toolkits/AssetEditorManager.cpp b/Engine/Source/Editor/UnrealEd/Private/Toolkits/AssetEditorManager.cpp
--- a/Engine/Source/Editor/UnrealEd/Private/Toolkits/AssetEditorManager.cpp
+++ b/Engine/Source/Editor/UnrealEd/Private/Toolkits/AssetEditorManager.cpp
@@ -42,6 +42,7 @@ FAssetEditorManager& FAssetEditorManager::Get()
 	return *Instance;
 }
 
+PRAGMA_DISABLE_DEPRECATION_WARNINGS
 FAssetEditorManager::FAssetEditorManager()
 	: bSavingOnShutdown(false)
 	, bRequestRestorePreviouslyOpenAssets(false)
@@ -843,5 +844,6 @@ void FAssetEditorManager::OpenEditorsForAssets(const TArray<FName>& AssetsToOpen
 		OpenEditorForAsset(AssetName.ToString());
 	}
 }
+PRAGMA_ENABLE_DEPRECATION_WARNINGS
 
 #undef LOCTEXT_NAMESPACE
diff --git a/Engine/Source/Editor/UnrealEd/Private/Toolkits/AssetEditorToolkit.cpp b/Engine/Source/Editor/UnrealEd/Private/Toolkits/AssetEditorToolkit.cpp
--- a/Engine/Source/Editor/UnrealEd/Private/Toolkits/AssetEditorToolkit.cpp
+++ b/Engine/Source/Editor/UnrealEd/Private/Toolkits/AssetEditorToolkit.cpp
@@ -605,7 +605,7 @@ void FAssetEditorToolkit::SaveAssetAs_Execute()
 		AssetEditorSubsystem->CloseAllEditorsForAsset(Object);
 		GEditor->GetEditorSubsystem<UAssetEditorSubsystem>()->NotifyAssetClosed(Object, this);
 	}
-	AssetEditorSubsystem->OpenEditorForAssets(ObjectsToReopen, ToolkitMode, MyToolkitHost.ToSharedRef());
+	AssetEditorSubsystem->OpenEditorForAssets_Advanced(ObjectsToReopen, ToolkitMode, MyToolkitHost.ToSharedRef());
 	// end hack
 }
 
@@ -740,7 +740,7 @@ void FAssetEditorToolkit::SwitchToStandaloneEditor_Execute( TWeakPtr< FAssetEdit
 
 	if( ObjectsToEdit.Num() > 0 )
 	{
-		ensure( GEditor->GetEditorSubsystem<UAssetEditorSubsystem>()->OpenEditorForAssets( ObjectsToEdit, EToolkitMode::Standalone, PreviousWorldCentricToolkitHost.ToSharedRef() ) );
+		ensure( GEditor->GetEditorSubsystem<UAssetEditorSubsystem>()->OpenEditorForAssets_Advanced( ObjectsToEdit, EToolkitMode::Standalone, PreviousWorldCentricToolkitHost.ToSharedRef() ) );
 	}
 }
 
@@ -788,7 +788,7 @@ void FAssetEditorToolkit::SwitchToWorldCentricEditor_Execute( TWeakPtr< FAssetEd
 
 	if( ObjectsToEdit.Num() > 0 )
 	{
-		ensure( GEditor->GetEditorSubsystem<UAssetEditorSubsystem>()->OpenEditorForAssets( ObjectsToEdit, EToolkitMode::WorldCentric, WorldCentricLevelEditor ) );
+		ensure( GEditor->GetEditorSubsystem<UAssetEditorSubsystem>()->OpenEditorForAssets_Advanced( ObjectsToEdit, EToolkitMode::WorldCentric, WorldCentricLevelEditor ) );
 	}
 }
 
diff --git a/Engine/Source/Editor/UnrealEd/Public/Subsystems/AssetEditorSubsystem.h b/Engine/Source/Editor/UnrealEd/Public/Subsystems/AssetEditorSubsystem.h
--- a/Engine/Source/Editor/UnrealEd/Public/Subsystems/AssetEditorSubsystem.h
+++ b/Engine/Source/Editor/UnrealEd/Public/Subsystems/AssetEditorSubsystem.h
@@ -71,7 +71,9 @@ public:
 	 * If any of the assets are already open, it will not create a new editor for them.
 	 * If all assets are of the same type, the supporting AssetTypeAction (if it exists) is responsible for the details of how to handle opening multiple assets at once.
 	 */
-	bool OpenEditorForAssets(const TArray<UObject*>& Assets, const EToolkitMode::Type ToolkitMode = EToolkitMode::Standalone, TSharedPtr<IToolkitHost> OpenedFromLevelEditor = TSharedPtr<IToolkitHost>());
+	UFUNCTION(BlueprintCallable, Category = "Editor Scripting | Asset Tools")
+	bool OpenEditorForAssets(const TArray<UObject*>& Assets);
+	bool OpenEditorForAssets_Advanced(const TArray<UObject*>& Assets, const EToolkitMode::Type ToolkitMode = EToolkitMode::Standalone, TSharedPtr<IToolkitHost> OpenedFromLevelEditor = TSharedPtr<IToolkitHost>());
 
 	/** Opens editors for the supplied assets (via OpenEditorForAsset) */
 	void OpenEditorsForAssets(const TArray<FString>& AssetsToOpen);
@@ -86,6 +88,7 @@ public:
 	TArray<IAssetEditorInstance*> FindEditorsForAsset(UObject* Asset);
 
 	/** Close all active editors for the supplied asset and return the number of asset editors that were closed */
+	UFUNCTION(BlueprintCallable, Category = "Editor Scripting | Asset Tools")
 	int32 CloseAllEditorsForAsset(UObject* Asset);
 
 	/** Close any editor which is not this one */
diff --git a/Engine/Source/Editor/UnrealEd/Public/Toolkits/AssetEditorManager.h b/Engine/Source/Editor/UnrealEd/Public/Toolkits/AssetEditorManager.h
--- a/Engine/Source/Editor/UnrealEd/Public/Toolkits/AssetEditorManager.h
+++ b/Engine/Source/Editor/UnrealEd/Public/Toolkits/AssetEditorManager.h
@@ -33,12 +33,14 @@ public:
 	void OnExit();
 
 	/** Opens an asset by path */
+	UE_DEPRECATED(4.24, "Use the matching function on AssetEditorSubsystem instead.")
 	void OpenEditorForAsset(const FString& AssetPathName);
 
 	/** 
 	 * Tries to open an editor for the specified asset.  Returns true if the asset is open in an editor.
 	 * If the file is already open in an editor, it will not create another editor window but instead bring it to front
 	 */
+	UE_DEPRECATED(4.24, "Use the matching function on AssetEditorSubsystem instead.")
 	bool OpenEditorForAsset(UObject* Asset, const EToolkitMode::Type ToolkitMode = EToolkitMode::Standalone, TSharedPtr<IToolkitHost> OpenedFromLevelEditor = TSharedPtr<IToolkitHost>(), const bool bShowProgressWindow = true);
 
 	/** 
@@ -46,48 +48,63 @@ public:
 	 * If any of the assets are already open, it will not create a new editor for them.
 	 * If all assets are of the same type, the supporting AssetTypeAction (if it exists) is responsible for the details of how to handle opening multiple assets at once.
 	 */
+	UE_DEPRECATED(4.24, "Use the matching function on AssetEditorSubsystem instead.")
 	bool OpenEditorForAssets(const TArray<UObject*>& Assets, const EToolkitMode::Type ToolkitMode = EToolkitMode::Standalone, TSharedPtr<IToolkitHost> OpenedFromLevelEditor = TSharedPtr<IToolkitHost>());
 
 	/** Opens editors for the supplied assets (via OpenEditorForAsset) */
+	UE_DEPRECATED(4.24, "Use the matching function on AssetEditorSubsystem instead.")
 	void OpenEditorsForAssets(const TArray<FString>& AssetsToOpen);
+	UE_DEPRECATED(4.24, "Use the matching function on AssetEditorSubsystem instead.")
 	void OpenEditorsForAssets(const TArray<FName>& AssetsToOpen);
 
 	/** Returns the primary editor if one is already open for the specified asset.
 	 * If there is one open and bFocusIfOpen is true, that editor will be brought to the foreground and focused if possible.
 	 */
+	UE_DEPRECATED(4.24, "Use the matching function on AssetEditorSubsystem instead.")
 	IAssetEditorInstance* FindEditorForAsset(UObject* Asset, bool bFocusIfOpen);
 
 	/** Returns all editors currently opened for the specified asset */
+	UE_DEPRECATED(4.24, "Use the matching function on AssetEditorSubsystem instead.")
 	TArray<IAssetEditorInstance*> FindEditorsForAsset(UObject* Asset);
 
 	/** Close all active editors for the supplied asset and return the number of asset editors that were closed */
+	UE_DEPRECATED(4.24, "Use the matching function on AssetEditorSubsystem instead.")
 	int32 CloseAllEditorsForAsset(UObject* Asset);
 
 	/** Close any editor which is not this one */
+	UE_DEPRECATED(4.24, "Use the matching function on AssetEditorSubsystem instead.")
 	void CloseOtherEditors(UObject* Asset, IAssetEditorInstance* OnlyEditor);
 
 	/** Remove given asset from all open editors */
+	UE_DEPRECATED(4.24, "Use the matching function on AssetEditorSubsystem instead.")
 	void RemoveAssetFromAllEditors(UObject* Asset);
 
 	/** Event called when CloseAllEditorsForAsset/RemoveAssetFromAllEditors is called */
 	DECLARE_EVENT_TwoParams(FAssetEditorManager, FAssetEditorRequestCloseEvent, UObject*, EAssetEditorCloseReason);
+	UE_DEPRECATED(4.24, "Use the matching function on AssetEditorSubsystem instead.")
 	virtual FAssetEditorRequestCloseEvent& OnAssetEditorRequestClose() { return AssetEditorRequestCloseEvent; }
 
 	/** Get all assets currently being tracked with open editors */
+	UE_DEPRECATED(4.24, "Use the matching function on AssetEditorSubsystem instead.")
 	TArray<UObject*> GetAllEditedAssets();
 
 	/** Notify the asset editor manager that an asset was opened */
+	UE_DEPRECATED(4.24, "Use the matching function on AssetEditorSubsystem instead.")
 	void NotifyAssetOpened(UObject* Asset, IAssetEditorInstance* Instance);
+	UE_DEPRECATED(4.24, "Use the matching function on AssetEditorSubsystem instead.")
 	void NotifyAssetsOpened( const TArray< UObject* >& Assets, IAssetEditorInstance* Instance);
 
 	/** Called when an asset has been opened in an editor */
 	DECLARE_EVENT_TwoParams(FAssetEditorManager, FOnAssetOpenedInEditorEvent, UObject*, IAssetEditorInstance*);
+	UE_DEPRECATED(4.24, "Use the matching function on AssetEditorSubsystem instead.")
 	virtual FOnAssetOpenedInEditorEvent& OnAssetOpenedInEditor() { return AssetOpenedInEditorEvent; }
 
 	/** Notify the asset editor manager that an asset editor is done editing an asset */
+	UE_DEPRECATED(4.24, "Use the matching function on AssetEditorSubsystem instead.")
 	void NotifyAssetClosed(UObject* Asset, IAssetEditorInstance* Instance);
 
 	/** Notify the asset editor manager that an asset was closed */
+	UE_DEPRECATED(4.24, "Use the matching function on AssetEditorSubsystem instead.")
 	void NotifyEditorClosed(IAssetEditorInstance* Instance);
 
 	// FGCObject interface
@@ -95,17 +112,21 @@ public:
 	virtual FString GetReferencerName() const override;
 
 	/** Close all open asset editors */
+	UE_DEPRECATED(4.24, "Use the matching function on AssetEditorSubsystem instead.")
 	bool CloseAllAssetEditors();
 	
 	/** Called when an asset editor is requested to be opened */
 	DECLARE_EVENT_OneParam( FAssetEditorManager, FAssetEditorRequestOpenEvent, UObject* );
+	UE_DEPRECATED(4.24, "Use the matching function on AssetEditorSubsystem instead.")
 	virtual FAssetEditorRequestOpenEvent& OnAssetEditorRequestedOpen() { return AssetEditorRequestOpenEvent; }
 
 	/** Called when an asset editor is actually opened */
 	DECLARE_EVENT_OneParam(FAssetEditorManager, FAssetEditorOpenEvent, UObject*);
+	UE_DEPRECATED(4.24, "Use the matching function on AssetEditorSubsystem instead.")
 	FAssetEditorOpenEvent& OnAssetEditorOpened() { return AssetEditorOpenedEvent; }
 
 	/** Request notification to restore the assets that were previously open when the editor was last closed */
+	UE_DEPRECATED(4.24, "Use the matching function on AssetEditorSubsystem instead.")
 	void RequestRestorePreviouslyOpenAssets();
 
 private:
