Fixing FBX reimport comparison when namespace is added

diff --git a/Engine/Source/Editor/UnrealEd/Private/Fbx/FbxFactory.cpp b/Engine/Source/Editor/UnrealEd/Private/Fbx/FbxFactory.cpp
--- a/Engine/Source/Editor/UnrealEd/Private/Fbx/FbxFactory.cpp
+++ b/Engine/Source/Editor/UnrealEd/Private/Fbx/FbxFactory.cpp
@@ -1223,13 +1223,13 @@ namespace ImportCompareHelper
 		}
 	}
 
-	void FillRecursivelySkeletonCompareData(const FbxNode *ParentNode, FSkeletonTreeNode& SkeletonTreeNode)
+	void FillRecursivelySkeletonCompareData(const FbxNode *ParentNode, FSkeletonTreeNode& SkeletonTreeNode, const UnFbx::FFbxImporter* FFbxImporter)
 	{
-		SkeletonTreeNode.JointName = FName(UTF8_TO_TCHAR(ParentNode->GetName()));
+		SkeletonTreeNode.JointName = FName(UTF8_TO_TCHAR(FFbxImporter->MakeName(ParentNode->GetName())));
 		for (int32 ChildIndex = 0; ChildIndex < ParentNode->GetChildCount(); ChildIndex++)
 		{
 			FSkeletonTreeNode& ChildNode = SkeletonTreeNode.Childrens.AddDefaulted_GetRef();
-			FillRecursivelySkeletonCompareData(ParentNode->GetChild(ChildIndex), ChildNode);
+			FillRecursivelySkeletonCompareData(ParentNode->GetChild(ChildIndex), ChildNode, FFbxImporter);
 		}
 	}
 
@@ -1281,7 +1281,7 @@ namespace ImportCompareHelper
 			}
 			
 			//Fill the Result skeleton data
-			FillRecursivelySkeletonCompareData(JointLinks[0], SkeletonCompareData.ResultAssetRoot);
+			FillRecursivelySkeletonCompareData(JointLinks[0], SkeletonCompareData.ResultAssetRoot, FFbxImporter);
 		}
 	}
 
diff --git a/Engine/Source/Editor/UnrealEd/Private/Fbx/FbxMainImport.cpp b/Engine/Source/Editor/UnrealEd/Private/Fbx/FbxMainImport.cpp
--- a/Engine/Source/Editor/UnrealEd/Private/Fbx/FbxMainImport.cpp
+++ b/Engine/Source/Editor/UnrealEd/Private/Fbx/FbxMainImport.cpp
@@ -1620,7 +1620,7 @@ bool FFbxImporter::ImportFromFile(const FString& Filename, const FString& Type,
 	return Result;
 }
 
-ANSICHAR* FFbxImporter::MakeName(const ANSICHAR* Name)
+ANSICHAR* FFbxImporter::MakeName(const ANSICHAR* Name) const
 {
 	const int SpecialChars[] = {'.', ',', '/', '`', '%'};
 
@@ -1656,7 +1656,7 @@ ANSICHAR* FFbxImporter::MakeName(const ANSICHAR* Name)
 	return TmpName;
 }
 
-FString FFbxImporter::MakeString(const ANSICHAR* Name)
+FString FFbxImporter::MakeString(const ANSICHAR* Name) const
 {
 	return FString(ANSI_TO_TCHAR(Name));
 }
diff --git a/Engine/Source/Editor/UnrealEd/Public/FbxImporter.h b/Engine/Source/Editor/UnrealEd/Public/FbxImporter.h
--- a/Engine/Source/Editor/UnrealEd/Public/FbxImporter.h
+++ b/Engine/Source/Editor/UnrealEd/Public/FbxImporter.h
@@ -947,8 +947,8 @@ public:
 	bool ImportCollisionModels(UStaticMesh* StaticMesh, const FbxString& NodeName);
 
 	//help
-	ANSICHAR* MakeName(const ANSICHAR* name);
-	FString MakeString(const ANSICHAR* Name);
+	ANSICHAR* MakeName(const ANSICHAR* name) const;
+	FString MakeString(const ANSICHAR* Name) const;
 	FName MakeNameForMesh(FString InName, FbxObject* FbxObject);
 
 	// meshes
