Added test properties for AllowedClasses/DisallowedClasses and TSubclassOf<> widgets.

diff --git a/Engine/Source/Editor/UnrealEd/Classes/Editor/PropertyEditorTestObject.h b/Engine/Source/Editor/UnrealEd/Classes/Editor/PropertyEditorTestObject.h
--- a/Engine/Source/Editor/UnrealEd/Classes/Editor/PropertyEditorTestObject.h
+++ b/Engine/Source/Editor/UnrealEd/Classes/Editor/PropertyEditorTestObject.h
@@ -7,6 +7,7 @@
 #include "UObject/Object.h"
 #include "UObject/ScriptInterface.h"
 #include "UObject/SoftObjectPath.h"
+#include "Blueprint/UserWidget.h"
 #include "Engine/BlendableInterface.h"
 #include "Engine/EngineTypes.h"
 
@@ -68,11 +69,19 @@ enum class EditColor : uint8
 	Cyan
 };
 
+USTRUCT()
+struct FPropertyEditTestTextStruct
+{
+	GENERATED_BODY()
+
+	UPROPERTY(EditAnywhere, Category=TextStruct)
+	FText NormalProperty;
+};
 
 USTRUCT()
 struct FPropertyEditorTestSubStruct
 {
-	GENERATED_USTRUCT_BODY()
+	GENERATED_BODY()
 
 	FPropertyEditorTestSubStruct()
 		: FirstProperty( 7897789 )
@@ -100,7 +109,7 @@ struct FPropertyEditorTestSubStruct
 USTRUCT()
 struct FPropertyEditorTestBasicStruct
 {
-	GENERATED_USTRUCT_BODY()
+	GENERATED_BODY()
 
 	FPropertyEditorTestBasicStruct()
 		: IntPropertyInsideAStruct( 0 )
@@ -191,6 +200,21 @@ class UPropertyEditorTestObject : public UObject
 	UPROPERTY(EditAnywhere, Category=BasicProperties)
 	UClass* ClassProperty;
 
+	UPROPERTY(EditAnywhere, Category=BasicProperties, meta=(AllowedClasses = "Texture2D"))
+	UClass* ClassPropertyWithAllowed;
+
+	UPROPERTY(EditAnywhere, Category = BasicProperties, meta=(DisallowedClasses = "Texture2D"))
+	UClass* ClassPropertyWithDisallowed;
+
+	UPROPERTY(EditAnywhere, Category = BasicProperties)
+	TSubclassOf<UTexture> SubclassOfTexture;
+
+	UPROPERTY(EditAnywhere, Category=BasicProperties, meta=(AllowedClasses = "Texture2D"))
+	TSubclassOf<UTexture> SubclassOfWithAllowed;
+
+	UPROPERTY(EditAnywhere, Category=BasicProperties, meta=(DisallowedClasses = "Texture2D"))
+	TSubclassOf<UTexture> SubclassOfWithDisallowed;
+
 	UPROPERTY(EditAnywhere, Category=BasicProperties)
 	FLinearColor LinearColorProperty;
 
@@ -347,6 +371,9 @@ class UPropertyEditorTestObject : public UObject
 	UPROPERTY(EditAnywhere, Category=StructTests)
 	TArray<FPropertyEditorTestBasicStruct> ArrayOfStructs;
 
+	UPROPERTY(EditAnywhere, Category=StructTests)
+	FPropertyEditTestTextStruct Struct;
+
 	UPROPERTY(EditAnywhere, Category=EditInlineProps)
 	UStaticMeshComponent* EditInlineNewStaticMeshComponent;
 
@@ -365,57 +392,63 @@ class UPropertyEditorTestObject : public UObject
 	UPROPERTY(EditAnywhere, Category=AssetPropertyTests)
 	AActor* OnlyActorsAllowed;
 
-	UPROPERTY(EditAnywhere, Category=TSetTests)
+	UPROPERTY(EditAnywhere, Category="TSet Tests")
 	TSet<int32> Int32Set;
 
-	UPROPERTY(EditAnywhere, Category=TSetTests)
+	UPROPERTY(EditAnywhere, Category="TSet Tests")
 	TSet<float> FloatSet;
 
-	UPROPERTY(EditAnywhere, Category=TSetTests)
+	UPROPERTY(EditAnywhere, Category="TSet Tests")
 	TSet<FString> StringSet;
 
-	UPROPERTY(EditAnywhere, Category=TSetTests)
+	UPROPERTY(EditAnywhere, Category="TSet Tests")
 	TSet<UObject*> ObjectSet;
 
-	UPROPERTY(EditAnywhere, Category=TSetTests)
+	UPROPERTY(EditAnywhere, Category="TSet Tests")
 	TSet<AActor*> ActorSet;
 
-	UPROPERTY(EditAnywhere, Category=TSetTests)
+	UPROPERTY(EditAnywhere, Category="TSet Tests")
 	TSet<EditColor> EditColorSet;
 
-	UPROPERTY(EditAnywhere, Category=TSetTests)
+	UPROPERTY(EditAnywhere, Category="TSet Tests")
 	TSet<FName> NameSet;
 
-	UPROPERTY(EditAnywhere, Category=TMapTests)
+	UPROPERTY(EditAnywhere, Category="TMap Tests")
 	TMap<int32, FString> Int32ToStringMap;
 
-	UPROPERTY(EditAnywhere, Category=TMapTests)
+	UPROPERTY(EditAnywhere, Category="TMap Tests")
 	TMap<FString, FLinearColor> StringToColorMap;
 
-	UPROPERTY(EditAnywhere, Category=TMapTests)
+	UPROPERTY(EditAnywhere, Category="TMap Tests")
 	TMap<int32, FPropertyEditorTestBasicStruct> Int32ToStructMap;
 
-	UPROPERTY(EditAnywhere, Category=TMapTests)
+	UPROPERTY(EditAnywhere, Category="TMap Tests")
 	TMap<FString, float> StringToFloatMap;
 
-	UPROPERTY(EditAnywhere, Category=TMapTests)
+	UPROPERTY(EditAnywhere, Category="TMap Tests")
 	TMap<FString, UObject*> StringToObjectMap;
 
-	UPROPERTY(EditAnywhere, Category=TMapTests)
+	UPROPERTY(EditAnywhere, Category="TMap Tests")
 	TMap<FString, AActor*> StringToActorMap;
 
-	UPROPERTY(EditAnywhere, Category=TMapTests)
+	UPROPERTY(EditAnywhere, Category="TMap Tests")
 	TMap<UObject*, int32> ObjectToInt32Map;
 
-	UPROPERTY(EditAnywhere, Category=TMapTests)
+	UPROPERTY(EditAnywhere, Category="TMap Tests")
 	TMap<UObject*, FLinearColor> ObjectToColorMap;
 
-	UPROPERTY(EditAnywhere, Category=TMapTests)
+	UPROPERTY(EditAnywhere, Category="TMap Tests")
 	TMap<int32, TEnumAsByte<PropertyEditorTestEnum> > IntToEnumMap;
 
-	UPROPERTY(EditAnywhere, Category=TMapTests)
+	UPROPERTY(EditAnywhere, Category="TMap Tests")
 	TMap<FName, FName> NameToNameMap;
 
+	UPROPERTY(EditAnywhere, Category = "TMap Tests")
+	TMap<int, FPropertyEditorTestBasicStruct> IntToCustomMap;
+
+	UPROPERTY(EditAnywhere, Category = "TMap Tests")
+	TMap<int, FPropertyEditorTestSubStruct> IntToSubStructMap;
+
 	UPROPERTY(EditAnywhere, Category=TSetStructTests)
 	TSet<FLinearColor> LinearColorSet;
 
