Added an option to do a partial save when an asset that has been modified cannot be checked out (eg. exclusively locked or not at head revision).

diff --git a/Engine/Source/Editor/UnrealEd/Private/FileHelpers.cpp b/Engine/Source/Editor/UnrealEd/Private/FileHelpers.cpp
--- a/Engine/Source/Editor/UnrealEd/Private/FileHelpers.cpp
+++ b/Engine/Source/Editor/UnrealEd/Private/FileHelpers.cpp
@@ -43,7 +43,6 @@
 #include "Framework/Application/SlateApplication.h"
 #include "Exporters/Exporter.h"
 
-
 #include "PackagesDialog.h"
 #include "Interfaces/IMainFrameModule.h"
 #include "IAssetTools.h"
@@ -53,7 +52,6 @@
 #include "Logging/TokenizedMessage.h"
 #include "Logging/MessageLog.h"
 
-
 #include "Dialogs/DlgPickPath.h"
 
 #include "IContentBrowserSingleton.h"
@@ -542,7 +540,7 @@ static bool SaveWorld(UWorld* World,
 	FString PackageName = Package->GetName();
 
 	FString	ExistingFilename;
-	FString		Path;
+	FString	Path;
 	FString	CleanFilename;
 
 	// Does a filename already exist for this package?
@@ -1590,6 +1588,12 @@ bool FEditorFileUtils::PromptToCheckoutPackages(bool bCheckDirty, const TArray<U
 		}
 		
 		// Prepare the buttons for the checkout dialog
+
+		if (OutPackagesNotNeedingCheckout != nullptr && OutPackagesNotNeedingCheckout->Num() > 0)
+		{
+			CheckoutPackagesDialogModule.AddButton(DRT_Save, NSLOCTEXT("PackagesDialogModule", "Dlg_SaveCheckedOutButton", "Save Checked Out"), NSLOCTEXT("PackagesDialogModule", "Dlg_SaveCheckedOutTooltip", "Save already checked out packages, but not the packages in this dialog."));
+		}
+
 		// The checkout button should be disabled if no packages can be checked out.
 		CheckoutPackagesDialogModule.AddButton(DRT_CheckOut, NSLOCTEXT("PackagesDialogModule", "Dlg_CheckOutButtonp", "Check Out Selected"), NSLOCTEXT("PackagesDialogModule", "Dlg_CheckOutTooltip", "Attempt to Check Out Checked Assets"), CheckOutSelectedDisabledAttrib );
 		
@@ -1684,9 +1688,14 @@ bool FEditorFileUtils::PromptToCheckoutPackages(bool bCheckDirty, const TArray<U
 
 				bPerformedOperation = true;
 			}
-			// Handle the case of the user canceling out of the dialog
-			else
+			else if (UserResponse == DRT_Save)
 			{
+				bResult = true;
+				bPerformedOperation = true;
+			}
+			else if (UserResponse == DRT_Cancel)
+			{
+				// Handle the case of the user canceling out of the dialog
 				bResult = false;
 				bPerformedOperation = true;
 			}
@@ -3701,8 +3710,8 @@ FEditorFileUtils::EPromptReturnCode FEditorFileUtils::PromptForCheckoutAndSave(
 	// Assemble list of packages to save
 	const TArray<UPackage*>& PackagesToSave = FilteredPackages;
 
-	// If there are any packages to save and the user didn't decline/cancel, then first prompt to check out any that are under source control, and then
-	// go ahead and save the specified packages
+	// If there are any packages to save and the user didn't decline/cancel, then first prompt to check out any that are under source control,
+	// and then go ahead and save the specified packages
 	if ( PackagesToSave.Num() > 0 && ReturnResponse == PR_Success )
 	{
 		TArray<UPackage*> FailedPackages;
@@ -3711,9 +3720,14 @@ FEditorFileUtils::EPromptReturnCode FEditorFileUtils::PromptForCheckoutAndSave(
 		TArray<UPackage*> PackagesNotNeedingCheckout;
 
 		// Prompt to check-out any packages under source control
-		const bool UserResponse = bAlreadyCheckedOut || FEditorFileUtils::PromptToCheckoutPackages( false, PackagesToSave, &PackagesCheckedOutOrMadeWritable, &PackagesNotNeedingCheckout );
+		bool bUserResponse = true;
+
+		if (!bAlreadyCheckedOut)
+		{
+			bUserResponse = FEditorFileUtils::PromptToCheckoutPackages(false, PackagesToSave, &PackagesCheckedOutOrMadeWritable, &PackagesNotNeedingCheckout);
+		}
 
-		if( UserResponse )
+		if( bUserResponse && (PackagesCheckedOutOrMadeWritable.Num() > 0 || PackagesNotNeedingCheckout.Num() > 0) )
 		{
 			TArray<UPackage*> FinalSaveList;
 			
@@ -3732,13 +3746,6 @@ FEditorFileUtils::EPromptReturnCode FEditorFileUtils::PromptForCheckoutAndSave(
 				ReturnResponse = InternalPromptForCheckoutAndSave(FinalSaveList, bUseDialog, FailedPackages);
 			}
 
-			if( UserResponse == false && PackagesNotNeedingCheckout.Num() > 0 )
-			{
-				// Return response should still be PR_Cancelled even if the user cancelled the source control dialog but there were writable packages we could save.
-				// This is in case the save is happing during editor exit. We don't want to shutdown the editor if some packages failed to save.
-				ReturnResponse = PR_Cancelled;
-			}
-
 			// Set the failure array to have the same contents as the local one.
 			// The local one is required so we can always display the error, even if an array is not provided.
 			if (OutFailedPackages)
