Fix Jira Bug UE-79782

UEditorEngine::Init was loading module EditorInteractiveToolsFramework which was not being built.

Fix was to include EditorInteractiveToolsFramwork as a dependency for UnrealEd

diff --git a/Engine/Source/Editor/UnrealEd/UnrealEd.Build.cs b/Engine/Source/Editor/UnrealEd/UnrealEd.Build.cs
--- a/Engine/Source/Editor/UnrealEd/UnrealEd.Build.cs
+++ b/Engine/Source/Editor/UnrealEd/UnrealEd.Build.cs
@@ -238,6 +238,7 @@ public class UnrealEd : ModuleRules
 				"Media",
 				"TimeManagementEditor",
 				"VirtualTexturingEditor",
+				"EditorInteractiveToolsFramework"
 			}
 		);
 
