PR #6168: Fix a PIE crash when there's an actor preview viewport open (Contributed by [REDACTED_PERSON])

diff --git a/Engine/Source/Editor/UnrealEd/Private/LevelEditorViewport.cpp b/Engine/Source/Editor/UnrealEd/Private/LevelEditorViewport.cpp
--- a/Engine/Source/Editor/UnrealEd/Private/LevelEditorViewport.cpp
+++ b/Engine/Source/Editor/UnrealEd/Private/LevelEditorViewport.cpp
@@ -4741,9 +4741,8 @@ void FLevelEditorViewportClient::CopyLayoutFromViewport( const FLevelEditorViewp
 UWorld* FLevelEditorViewportClient::ConditionalSetWorld()
 {
 	// Should set GWorld to the play world if we are simulating in the editor and not already in the play world (reentrant calls to this would cause the world to be the same)
-	if( bIsSimulateInEditorViewport && GEditor->PlayWorld != GWorld )
+	if( bIsSimulateInEditorViewport && GEditor->PlayWorld && GEditor->PlayWorld != GWorld )
 	{
-		check( GEditor->PlayWorld != NULL );
 		return SetPlayInEditorWorld( GEditor->PlayWorld );
 	}
 
