Fixed drag and drop spawning actors in wrong location of level viewport

diff --git a/Engine/Source/Editor/UnrealEd/Private/Editor/ActorPositioning.cpp b/Engine/Source/Editor/UnrealEd/Private/Editor/ActorPositioning.cpp
--- a/Engine/Source/Editor/UnrealEd/Private/Editor/ActorPositioning.cpp
+++ b/Engine/Source/Editor/UnrealEd/Private/Editor/ActorPositioning.cpp
@@ -143,13 +143,13 @@ FActorPositionTraceResult FActorPositioning::TraceWorldForPosition(const UWorld&
 	return Results;
 }
 
-FTransform FActorPositioning::GetCurrentViewportPlacementTransform(const AActor& Actor, bool bSnap)
+FTransform FActorPositioning::GetCurrentViewportPlacementTransform(const AActor& Actor, bool bSnap, const FViewportCursorLocation* InCursor)
 {
 	FVector Collision = Actor.GetPlacementExtent();
 	const UActorFactory* Factory = GEditor->FindActorFactoryForActorClass(Actor.GetClass());
 
 	// Get cursor origin and direction in world space.
-	FViewportCursorLocation CursorLocation = GCurrentLevelEditingViewportClient->GetCursorWorldLocationFromMousePos();
+	FViewportCursorLocation CursorLocation = InCursor ? *InCursor : GCurrentLevelEditingViewportClient->GetCursorWorldLocationFromMousePos();
 	const auto CursorPos = CursorLocation.GetCursorPos();
 
 	FTransform ActorTransform = FTransform::Identity;
diff --git a/Engine/Source/Editor/UnrealEd/Private/Editor/ActorPositioning.h b/Engine/Source/Editor/UnrealEd/Private/Editor/ActorPositioning.h
--- a/Engine/Source/Editor/UnrealEd/Private/Editor/ActorPositioning.h
+++ b/Engine/Source/Editor/UnrealEd/Private/Editor/ActorPositioning.h
@@ -139,10 +139,11 @@ struct FActorPositioning
 	 *
 	 *	@param	Actor			The actor to spawn
 	 *	@param	bSnap			Whether to perform snapping on the placement transform or not (defaults to true)
+	 *	@param	InCursor		Optional pre-calculated cursor location
 	 *
 	 *	@return	The expected transform for the specified actor
 	*/
-	static UNREALED_API FTransform GetCurrentViewportPlacementTransform(const AActor& Actor, bool bSnap = true);
+	static UNREALED_API FTransform GetCurrentViewportPlacementTransform(const AActor& Actor, bool bSnap = true, const FViewportCursorLocation* InCursor = nullptr);
 
 	/** Get a default actor position in front of the camera
 	 *
diff --git a/Engine/Source/Editor/UnrealEd/Private/LevelEditorViewport.cpp b/Engine/Source/Editor/UnrealEd/Private/LevelEditorViewport.cpp
--- a/Engine/Source/Editor/UnrealEd/Private/LevelEditorViewport.cpp
+++ b/Engine/Source/Editor/UnrealEd/Private/LevelEditorViewport.cpp
@@ -226,7 +226,7 @@ namespace LevelEditorViewportClientHelper
 	}
 }
 
-TArray<AActor*> FLevelEditorViewportClient::TryPlacingActorFromObject( ULevel* InLevel, UObject* ObjToUse, bool bSelectActors, EObjectFlags ObjectFlags, UActorFactory* FactoryToUse, const FName Name )
+TArray<AActor*> FLevelEditorViewportClient::TryPlacingActorFromObject( ULevel* InLevel, UObject* ObjToUse, bool bSelectActors, EObjectFlags ObjectFlags, UActorFactory* FactoryToUse, const FName Name, const FViewportCursorLocation* Cursor )
 {
 	TArray<AActor*> PlacedActors;
 
@@ -261,7 +261,7 @@ TArray<AActor*> FLevelEditorViewportClient::TryPlacingActorFromObject( ULevel* I
 		if ( PlacedActor == NULL && !ObjectClass->HasAnyClassFlags(CLASS_NotPlaceable | CLASS_Abstract) )
 		{
 			// If no actor factory was found or failed, add the actor directly.
-			const FTransform ActorTransform = FActorPositioning::GetCurrentViewportPlacementTransform(*ObjectClass->GetDefaultObject<AActor>());
+			const FTransform ActorTransform = FActorPositioning::GetCurrentViewportPlacementTransform(*ObjectClass->GetDefaultObject<AActor>(), /*bSnap=*/true, Cursor);
 			PlacedActor = GEditor->AddActor( InLevel, ObjectClass, ActorTransform, /*bSilent=*/false, ObjectFlags );
 		}
 
@@ -861,7 +861,7 @@ bool FLevelEditorViewportClient::DropObjectsOnBackground(FViewportCursorLocation
 		ensure( AssetObj );
 
 		// Attempt to create actors from the dropped object
-		TArray<AActor*> NewActors = TryPlacingActorFromObject(GetWorld()->GetCurrentLevel(), AssetObj, bSelectActors, ObjectFlags, FactoryToUse);
+		TArray<AActor*> NewActors = TryPlacingActorFromObject(GetWorld()->GetCurrentLevel(), AssetObj, bSelectActors, ObjectFlags, FactoryToUse, NAME_None, &Cursor);
 
 		if ( NewActors.Num() > 0 )
 		{
@@ -903,7 +903,7 @@ bool FLevelEditorViewportClient::DropObjectsOnActor(FViewportCursorLocation& Cur
 		if (!bAppliedToActor)
 		{
 			// Attempt to create actors from the dropped object
-			TArray<AActor*> NewActors = TryPlacingActorFromObject(GetWorld()->GetCurrentLevel(), DroppedObject, bSelectActors, ObjectFlags, FactoryToUse);
+			TArray<AActor*> NewActors = TryPlacingActorFromObject(GetWorld()->GetCurrentLevel(), DroppedObject, bSelectActors, ObjectFlags, FactoryToUse, NAME_None, &Cursor);
 
 			if ( NewActors.Num() > 0 )
 			{
@@ -949,7 +949,7 @@ bool FLevelEditorViewportClient::DropObjectsOnBSPSurface(FSceneView* View, FView
 		if (!bAppliedToActor)
 		{
 			// Attempt to create actors from the dropped object
-			TArray<AActor*> NewActors = TryPlacingActorFromObject(GetWorld()->GetCurrentLevel(), DroppedObject, bSelectActors, ObjectFlags, FactoryToUse);
+			TArray<AActor*> NewActors = TryPlacingActorFromObject(GetWorld()->GetCurrentLevel(), DroppedObject, bSelectActors, ObjectFlags, FactoryToUse, NAME_None, &Cursor);
 
 			if (NewActors.Num() > 0)
 			{
diff --git a/Engine/Source/Editor/UnrealEd/Public/LevelEditorViewport.h b/Engine/Source/Editor/UnrealEd/Public/LevelEditorViewport.h
--- a/Engine/Source/Editor/UnrealEd/Public/LevelEditorViewport.h
+++ b/Engine/Source/Editor/UnrealEd/Public/LevelEditorViewport.h
@@ -515,10 +515,11 @@ public:
 	 * @param	bSelectActors	If true, select the newly dropped actors (defaults: true)
 	 * @param	ObjectFlags		The flags to place on the actor when it is spawned
 	 * @param	FactoryToUse	The preferred actor factory to use (optional)
+	 * @param	Cursor			Optional pre-calculated cursor location
 	 *
 	 * @return	true if the object was successfully used to place an actor; false otherwise
 	 */
-	static TArray<AActor*> TryPlacingActorFromObject( ULevel* InLevel, UObject* ObjToUse, bool bSelectActors, EObjectFlags ObjectFlags, UActorFactory* FactoryToUse, const FName Name = NAME_None );
+	static TArray<AActor*> TryPlacingActorFromObject( ULevel* InLevel, UObject* ObjToUse, bool bSelectActors, EObjectFlags ObjectFlags, UActorFactory* FactoryToUse, const FName Name = NAME_None, const FViewportCursorLocation* Cursor = nullptr);
 
 	/** 
 	 * Returns true if creating a preview actor in the viewport. 
