Sequencer: Fix multiple export to fbx issues.

1. Fix skeletal animation export off by 1 frame - changed the anim adapter to advance by integer frames, rather than converting to float time and incrementing by the sampling rate.
2. Fix exporting an object in a shot - evaluation needs to happen in global time through RootToLocalTransform.Inverse(). As a result, also added an option to export in global time or in the sequence's local time. For example, if an object is animated in a shot that starts at global frame 150, the exported keys can either start at global frame 150 or local frame 0.
3. Step and export frames at the display rate of the sequence, rather than the hardcoded DEFAULT_SAMPLERATE.

diff --git a/Engine/Plugins/MovieScene/SequencerScripting/Source/SequencerScriptingEditor/Private/SequencerTools.cpp b/Engine/Plugins/MovieScene/SequencerScripting/Source/SequencerScriptingEditor/Private/SequencerTools.cpp
--- a/Engine/Plugins/MovieScene/SequencerScripting/Source/SequencerScriptingEditor/Private/SequencerTools.cpp
+++ b/Engine/Plugins/MovieScene/SequencerScripting/Source/SequencerScriptingEditor/Private/SequencerTools.cpp
@@ -127,7 +127,8 @@ bool USequencerToolsFunctionLibrary::ExportFBX(UWorld* World, ULevelSequence* Se
 	Player->Initialize(Sequence, World->PersistentLevel, Settings, CameraSettings);
 	Player->State.AssignSequence(MovieSceneSequenceID::Root, *Sequence, *Player);
 	FMovieSceneSequenceIDRef Template = MovieSceneSequenceID::Root;
-	bool bDidExport = MovieSceneToolHelpers::ExportFBX(World, MovieScene, Player, Bindings, NodeNameAdapter, Template, InFBXFileName);
+	FMovieSceneSequenceTransform RootToLocalTransform;
+	bool bDidExport = MovieSceneToolHelpers::ExportFBX(World, MovieScene, Player, Bindings, NodeNameAdapter, Template, InFBXFileName, RootToLocalTransform);
 	Exporter->SetExportOptionsOverride(nullptr);
 	return bDidExport;
 }
diff --git a/Engine/Source/Editor/MovieSceneTools/Private/MovieSceneToolHelpers.cpp b/Engine/Source/Editor/MovieSceneTools/Private/MovieSceneToolHelpers.cpp
--- a/Engine/Source/Editor/MovieSceneTools/Private/MovieSceneToolHelpers.cpp
+++ b/Engine/Source/Editor/MovieSceneTools/Private/MovieSceneToolHelpers.cpp
@@ -2090,7 +2090,7 @@ void ExportLevelMesh(UnFbx::FFbxExporter* Exporter, ULevel* Level, IMovieScenePl
 	Exporter->ExportLevelMesh(Level, !bSelectedOnly, ActorToExport, NodeNameAdapter, bSaveAnimSeq);
 }
 
-bool MovieSceneToolHelpers::ExportFBX(UWorld* World, UMovieScene* MovieScene, IMovieScenePlayer* Player, TArray<FGuid>& Bindings, INodeNameAdapter& NodeNameAdapter, FMovieSceneSequenceIDRef& Template, const FString& InFBXFileName)
+bool MovieSceneToolHelpers::ExportFBX(UWorld* World, UMovieScene* MovieScene, IMovieScenePlayer* Player, TArray<FGuid>& Bindings, INodeNameAdapter& NodeNameAdapter, FMovieSceneSequenceIDRef& Template, const FString& InFBXFileName, FMovieSceneSequenceTransform& RootToLocalTransform)
 {
 	UnFbx::FFbxExporter* Exporter = UnFbx::FFbxExporter::GetInstance();
 
@@ -2112,7 +2112,7 @@ bool MovieSceneToolHelpers::ExportFBX(UWorld* World, UMovieScene* MovieScene, IM
 		}
 	}
 
-	Exporter->ExportLevelSequence(MovieScene, Bindings, Player, Template);
+	Exporter->ExportLevelSequence(MovieScene, Bindings, Player, Template, RootToLocalTransform);
 
 	//Export all master tracks
 
@@ -2120,7 +2120,7 @@ bool MovieSceneToolHelpers::ExportFBX(UWorld* World, UMovieScene* MovieScene, IM
 	{
 		TArray<UMovieSceneTrack*> Tracks;
 		Tracks.Add(MasterTrack);
-		Exporter->ExportLevelSequenceTracks(MovieScene, Player, nullptr, nullptr, Tracks);
+		Exporter->ExportLevelSequenceTracks(MovieScene, Player, nullptr, nullptr, Tracks, RootToLocalTransform);
 	}
 	// Save to disk
 	Exporter->WriteToFile(*InFBXFileName);
diff --git a/Engine/Source/Editor/MovieSceneTools/Public/MovieSceneToolHelpers.h b/Engine/Source/Editor/MovieSceneTools/Public/MovieSceneToolHelpers.h
--- a/Engine/Source/Editor/MovieSceneTools/Public/MovieSceneToolHelpers.h
+++ b/Engine/Source/Editor/MovieSceneTools/Public/MovieSceneToolHelpers.h
@@ -26,6 +26,7 @@ struct FMovieSceneEvaluationTrack;
 class UMovieSceneUserImportFBXSettings;
 struct FMovieSceneFloatValue;
 class INodeNameAdapter;
+struct FMovieSceneSequenceTransform;
 template<typename ChannelType> struct TMovieSceneChannelData;
 
 namespace fbxsdk
@@ -231,9 +232,10 @@ public:
 	* @param NodeNameAdaptor Adaptor to look up actor names.
 	* @param InFBXFileName the fbx file name.
 	* @param Template Movie scene sequence id.
+	* @param RootToLocalTransform The root to local transform time.
 	* @return Whether the export was successful
 	*/
-	static bool ExportFBX(UWorld* World, UMovieScene* MovieScene, IMovieScenePlayer* Player, TArray<FGuid>& Bindings, INodeNameAdapter& NodeNameAdapter, FMovieSceneSequenceIDRef& Template,  const FString& InFBXFileName);
+	static bool ExportFBX(UWorld* World, UMovieScene* MovieScene, IMovieScenePlayer* Player, TArray<FGuid>& Bindings, INodeNameAdapter& NodeNameAdapter, FMovieSceneSequenceIDRef& Template,  const FString& InFBXFileName, FMovieSceneSequenceTransform& RootToLocalTransform);
 
 	/**
 	* Import FBX with dialog
diff --git a/Engine/Source/Editor/Sequencer/Private/Sequencer.cpp b/Engine/Source/Editor/Sequencer/Private/Sequencer.cpp
--- a/Engine/Source/Editor/Sequencer/Private/Sequencer.cpp
+++ b/Engine/Source/Editor/Sequencer/Private/Sequencer.cpp
@@ -9694,7 +9694,7 @@ void FSequencer::ExportFBXInternal(const FString& ExportFilename, TArray<FGuid>&
 			UWorld* World = Cast<UWorld>(GetPlaybackContext());
 			FMovieSceneSequenceIDRef Template = GetFocusedTemplateID();
 			UnFbx::FFbxExporter::FLevelSequenceNodeNameAdapter NodeNameAdapter(MovieScene, this, Template);
-			MovieSceneToolHelpers::ExportFBX(World,MovieScene, this, Bindings, NodeNameAdapter, Template, ExportFilename);
+			MovieSceneToolHelpers::ExportFBX(World,MovieScene, this, Bindings, NodeNameAdapter, Template, ExportFilename, RootToLocalTransform);
 		}
 	}
 }
diff --git a/Engine/Source/Editor/UnrealEd/Classes/Exporters/FbxExportOption.h b/Engine/Source/Editor/UnrealEd/Classes/Exporters/FbxExportOption.h
--- a/Engine/Source/Editor/UnrealEd/Classes/Exporters/FbxExportOption.h
+++ b/Engine/Source/Editor/UnrealEd/Classes/Exporters/FbxExportOption.h
@@ -63,6 +63,10 @@ public:
 	UPROPERTY(EditAnywhere, BlueprintReadWrite, config, category = Animation)
 	uint32 MapSkeletalMotionToRoot : 1;
 
+	/** If enabled, export sequencer animation in its local time, relative to its master sequence. */
+	UPROPERTY(EditAnywhere, BlueprintReadWrite, config, category = Animation)
+	uint32 bExportLocalTime : 1;
+
 	/* Set all the UProperty to the CDO value */
 	void ResetToDefault();
 
diff --git a/Engine/Source/Editor/UnrealEd/Private/Fbx/FbxAnimationExport.cpp b/Engine/Source/Editor/UnrealEd/Private/Fbx/FbxAnimationExport.cpp
--- a/Engine/Source/Editor/UnrealEd/Private/Fbx/FbxAnimationExport.cpp
+++ b/Engine/Source/Editor/UnrealEd/Private/Fbx/FbxAnimationExport.cpp
@@ -534,17 +534,14 @@ void FFbxExporter::ExportMatineeGroup(class AMatineeActor* MatineeActor, USkelet
 	FbxSkeletonRoots.Add(SkeletalMeshComponent, SkeletonRootNode);
 	BaseNode->AddChild(SkeletonRootNode);
 
+	static const float SamplingRate = 1.f / DEFAULT_SAMPLERATE;
+
 	FMatineeAnimTrackAdapter AnimTrackAdapter(MatineeActor);
-	ExportAnimTrack(AnimTrackAdapter, Owner, SkeletalMeshComponent);
+	ExportAnimTrack(AnimTrackAdapter, Owner, SkeletalMeshComponent, SamplingRate);
 }
 
-void FFbxExporter::ExportAnimTrack(IAnimTrackAdapter& AnimTrackAdapter, AActor* Actor, USkeletalMeshComponent* SkeletalMeshComponent)
+void FFbxExporter::ExportAnimTrack(IAnimTrackAdapter& AnimTrackAdapter, AActor* Actor, USkeletalMeshComponent* SkeletalMeshComponent, float SamplingRate)
 {
-	static const float SamplingRate = 1.f / DEFAULT_SAMPLERATE;
-
-	float AnimationStart = AnimTrackAdapter.GetAnimationStart();
-	float AnimationLength = AnimTrackAdapter.GetAnimationLength();
-	float AnimationEnd = AnimationStart + AnimationLength;
 	// show a status update every 1 second worth of samples
 	const float UpdateFrequency = 1.0f;
 	float NextUpdateTime = UpdateFrequency;
@@ -571,16 +568,23 @@ void FFbxExporter::ExportAnimTrack(IAnimTrackAdapter& AnimTrackAdapter, AActor*
 
 	FTransform InitialInvParentTransform;
 
-	float SampleTime;
-	for(SampleTime = AnimationStart; SampleTime <= AnimationEnd; SampleTime += SamplingRate)
+	int32 LocalStartFrame = AnimTrackAdapter.GetLocalStartFrame();
+	int32 StartFrame = AnimTrackAdapter.GetStartFrame();
+	int32 AnimationLength = AnimTrackAdapter.GetLength();
+	float FrameRate = AnimTrackAdapter.GetFrameRate();
+
+	for (int32 FrameCount = 0; FrameCount <= AnimationLength; ++FrameCount)
 	{
-		if (SampleTime == AnimationStart)
+		if (FrameCount == 0)
 		{
 			InitialInvParentTransform = Actor->GetRootComponent()->GetComponentTransform().Inverse();
 		}
 
+		int32 LocalFrame = LocalStartFrame + FrameCount;
+		float SampleTime = (StartFrame + FrameCount) / FrameRate;
+
 		// This will call UpdateSkelPose on the skeletal mesh component to move bones based on animations in the matinee group
-		AnimTrackAdapter.UpdateAnimation( SampleTime );
+		AnimTrackAdapter.UpdateAnimation(LocalFrame);
 
 		// Update space bases so new animation position has an effect.
 		// @todo - hack - this will be removed at some point
@@ -593,8 +597,8 @@ void FFbxExporter::ExportAnimTrack(IAnimTrackAdapter& AnimTrackAdapter, AActor*
 		SkeletalMeshComponent->MarkRenderTransformDirty();
 		SkeletalMeshComponent->MarkRenderDynamicDataDirty();
 
-		FbxTime ExportTime;
-		ExportTime.SetSecondDouble(SampleTime);
+		FbxTime ExportTime; 
+		ExportTime.SetSecondDouble(GetExportOptions()->bExportLocalTime ? LocalFrame / FrameRate : SampleTime);
 
 		NextUpdateTime -= SamplingRate;
 
diff --git a/Engine/Source/Editor/UnrealEd/Private/Fbx/FbxExportOption.cpp b/Engine/Source/Editor/UnrealEd/Private/Fbx/FbxExportOption.cpp
--- a/Engine/Source/Editor/UnrealEd/Private/Fbx/FbxExportOption.cpp
+++ b/Engine/Source/Editor/UnrealEd/Private/Fbx/FbxExportOption.cpp
@@ -22,6 +22,7 @@ UFbxExportOption::UFbxExportOption(const FObjectInitializer& ObjectInitializer)
 	bExportMorphTargets = true;
 	VertexColor = true;
 	MapSkeletalMotionToRoot = false;
+	bExportLocalTime = true;
 }
 
 void UFbxExportOption::ResetToDefault()
diff --git a/Engine/Source/Editor/UnrealEd/Private/Fbx/FbxMainExport.cpp b/Engine/Source/Editor/UnrealEd/Private/Fbx/FbxMainExport.cpp
--- a/Engine/Source/Editor/UnrealEd/Private/Fbx/FbxMainExport.cpp
+++ b/Engine/Source/Editor/UnrealEd/Private/Fbx/FbxMainExport.cpp
@@ -1433,20 +1433,21 @@ FFbxExporter::FMatineeAnimTrackAdapter::FMatineeAnimTrackAdapter( AMatineeActor*
 	MatineeActor = InMatineeActor;
 }
 
-float FFbxExporter::FMatineeAnimTrackAdapter::GetAnimationStart() const
+int32 FFbxExporter::FMatineeAnimTrackAdapter::GetLocalStartFrame() const
 {
-	return 0.f;
+	return 0;
 }
 
 
-float FFbxExporter::FMatineeAnimTrackAdapter::GetAnimationLength() const
+int32 FFbxExporter::FMatineeAnimTrackAdapter::GetLength() const
 {
-	return MatineeActor->MatineeData->InterpLength;
+	return MatineeActor->MatineeData->InterpLength * (1.f / GetFrameRate());
 }
 
 
-void FFbxExporter::FMatineeAnimTrackAdapter::UpdateAnimation( float Time )
+void FFbxExporter::FMatineeAnimTrackAdapter::UpdateAnimation( int32 LocalFrame )
 {
+	float Time = LocalFrame * GetFrameRate();
 	MatineeActor->UpdateInterp( Time, true );
 }
 
@@ -1502,8 +1503,10 @@ bool FFbxExporter::ExportMatinee(AMatineeActor* InMatineeActor)
 					USkeletalMeshComponent* SkeletalMeshComp = Cast<USkeletalMeshComponent>(Actor->GetComponentByClass(USkeletalMeshComponent::StaticClass()));
 					if(SkeletalMeshComp)
 					{
+						static const float SamplingRate = 1.f / DEFAULT_SAMPLERATE;
+
 						FMatineeAnimTrackAdapter AnimTrackAdapter(InMatineeActor);
-						ExportAnimTrack(AnimTrackAdapter, Actor, SkeletalMeshComp);
+						ExportAnimTrack(AnimTrackAdapter, Actor, SkeletalMeshComp, SamplingRate);
 					}
 				}
 			}
@@ -1551,36 +1554,62 @@ FString FFbxExporter::FLevelSequenceNodeNameAdapter::GetActorNodeName(const AAct
 	return NodeName;
 }
 
-FFbxExporter::FLevelSequenceAnimTrackAdapter::FLevelSequenceAnimTrackAdapter( IMovieScenePlayer* InMovieScenePlayer, UMovieScene* InMovieScene )
+FFbxExporter::FLevelSequenceAnimTrackAdapter::FLevelSequenceAnimTrackAdapter( IMovieScenePlayer* InMovieScenePlayer, UMovieScene* InMovieScene, const FMovieSceneSequenceTransform& InRootToLocalTransform)
 {
 	MovieScenePlayer = InMovieScenePlayer;
 	MovieScene = InMovieScene;
+	RootToLocalTransform = InRootToLocalTransform;
 }
 
 
-float FFbxExporter::FLevelSequenceAnimTrackAdapter::GetAnimationStart() const
+int32 FFbxExporter::FLevelSequenceAnimTrackAdapter::GetLocalStartFrame() const
 {
-	return MovieScene->GetPlaybackRange().GetLowerBoundValue() / MovieScene->GetTickResolution();
+	FFrameRate TickResolution = MovieScene->GetTickResolution();
+	FFrameRate DisplayRate = MovieScene->GetDisplayRate();
+	return FFrameRate::TransformTime(FFrameTime(MovieScene::DiscreteInclusiveLower(MovieScene->GetPlaybackRange())), TickResolution, DisplayRate).RoundToFrame().Value;
 }
 
-float FFbxExporter::FLevelSequenceAnimTrackAdapter::GetAnimationLength() const
+int32 FFbxExporter::FLevelSequenceAnimTrackAdapter::GetStartFrame() const
 {
-	return FFrameNumber(MovieScene::DiscreteSize(MovieScene->GetPlaybackRange())) / MovieScene->GetTickResolution();
+	FFrameRate TickResolution = MovieScene->GetTickResolution();
+	FFrameRate DisplayRate = MovieScene->GetDisplayRate();
+	return FFrameRate::TransformTime(FFrameTime(MovieScene::DiscreteInclusiveLower(MovieScene->GetPlaybackRange()) * RootToLocalTransform.Inverse()), TickResolution, DisplayRate).RoundToFrame().Value;
 }
 
-
-void FFbxExporter::FLevelSequenceAnimTrackAdapter::UpdateAnimation( float Time )
+int32 FFbxExporter::FLevelSequenceAnimTrackAdapter::GetLength() const
 {
-	// Convert the time in seconds to a frame time
-	FFrameTime FrameTime = Time * MovieScene->GetTickResolution();
-	FMovieSceneContext Context = FMovieSceneContext(FMovieSceneEvaluationRange(FrameTime, MovieScene->GetTickResolution()), MovieScenePlayer->GetPlaybackStatus()).SetHasJumped(true);
+	FFrameRate TickResolution = MovieScene->GetTickResolution();
+	FFrameRate DisplayRate = MovieScene->GetDisplayRate();
+	return FFrameRate::TransformTime(FFrameTime(MovieScene::DiscreteSize(MovieScene->GetPlaybackRange())), TickResolution, DisplayRate).RoundToFrame().Value;
+}
+
+
+void FFbxExporter::FLevelSequenceAnimTrackAdapter::UpdateAnimation( int32 LocalFrame )
+{	
+	FFrameRate TickResolution = MovieScene->GetTickResolution();
+	FFrameRate DisplayRate = MovieScene->GetDisplayRate();
+
+	FFrameTime LocalTime = FFrameRate::TransformTime(FFrameTime(LocalFrame), DisplayRate, TickResolution);
+	FFrameTime GlobalTime = LocalTime * RootToLocalTransform.Inverse();
+
+	FMovieSceneContext Context = FMovieSceneContext(FMovieSceneEvaluationRange(GlobalTime, TickResolution), MovieScenePlayer->GetPlaybackStatus()).SetHasJumped(true);
+
 	MovieScenePlayer->GetEvaluationTemplate().Evaluate( Context, *MovieScenePlayer );
 }
 
-bool FFbxExporter::ExportLevelSequenceTracks(UMovieScene* MovieScene, IMovieScenePlayer* MovieScenePlayer, FbxNode* FbxActor, AActor* Actor, const TArray<UMovieSceneTrack*>& Tracks)
+
+float FFbxExporter::FLevelSequenceAnimTrackAdapter::GetFrameRate() const
+{
+	return MovieScene->GetDisplayRate().AsDecimal();
+}
+
+
+bool FFbxExporter::ExportLevelSequenceTracks(UMovieScene* MovieScene, IMovieScenePlayer* MovieScenePlayer, FbxNode* FbxActor, AActor* Actor, const TArray<UMovieSceneTrack*>& Tracks, const FMovieSceneSequenceTransform& RootToLocalTransform)
 {
 	USkeletalMeshComponent* SkeletalMeshComp = Actor ? Cast<USkeletalMeshComponent>(Actor->GetComponentByClass(USkeletalMeshComponent::StaticClass())) : nullptr;
 
+	FFrameRate DisplayRate = MovieScene->GetDisplayRate();
+
 	const bool bSkip3DTransformTrack = SkeletalMeshComp && GetExportOptions()->MapSkeletalMotionToRoot;
 
 	// Look for the tracks that we currently support
@@ -1589,19 +1618,19 @@ bool FFbxExporter::ExportLevelSequenceTracks(UMovieScene* MovieScene, IMovieScen
 		if (Track->IsA(UMovieScene3DTransformTrack::StaticClass()) && !bSkip3DTransformTrack)
 		{
 			UMovieScene3DTransformTrack* TransformTrack = (UMovieScene3DTransformTrack*)Track;
-			ExportLevelSequence3DTransformTrack(FbxActor, *TransformTrack, Actor, MovieScene->GetPlaybackRange());
+			ExportLevelSequence3DTransformTrack(FbxActor, *TransformTrack, Actor, MovieScene->GetPlaybackRange(), RootToLocalTransform);
 		}
 		else if (Track->IsA(UMovieScenePropertyTrack::StaticClass()))
 		{
 			UMovieScenePropertyTrack* PropertyTrack = (UMovieScenePropertyTrack*)Track;
-			ExportLevelSequencePropertyTrack(FbxActor, *PropertyTrack);
+			ExportLevelSequencePropertyTrack(FbxActor, *PropertyTrack, MovieScene->GetPlaybackRange(), RootToLocalTransform);
 		}
 		else if (Track->IsA(UMovieSceneSkeletalAnimationTrack::StaticClass()))
 		{
 			if (SkeletalMeshComp)
 			{
-				FLevelSequenceAnimTrackAdapter AnimTrackAdapter(MovieScenePlayer, MovieScene);
-				ExportAnimTrack(AnimTrackAdapter, Actor, SkeletalMeshComp);
+				FLevelSequenceAnimTrackAdapter AnimTrackAdapter(MovieScenePlayer, MovieScene, RootToLocalTransform);
+				ExportAnimTrack(AnimTrackAdapter, Actor, SkeletalMeshComp, 1.0 / DisplayRate.AsDecimal());
 			}
 		}
 	}
@@ -1610,7 +1639,7 @@ bool FFbxExporter::ExportLevelSequenceTracks(UMovieScene* MovieScene, IMovieScen
 }
 
 
-bool FFbxExporter::ExportLevelSequence( UMovieScene* MovieScene, const TArray<FGuid>& Bindings, IMovieScenePlayer* MovieScenePlayer, FMovieSceneSequenceIDRef SequenceID )
+bool FFbxExporter::ExportLevelSequence( UMovieScene* MovieScene, const TArray<FGuid>& Bindings, IMovieScenePlayer* MovieScenePlayer, FMovieSceneSequenceIDRef SequenceID, const FMovieSceneSequenceTransform& RootToLocalTransform)
 {
 	if ( MovieScene == nullptr || MovieScenePlayer == nullptr )
 	{
@@ -1648,7 +1677,7 @@ bool FFbxExporter::ExportLevelSequence( UMovieScene* MovieScene, const TArray<FG
 				// now it should export everybody
 				if ( FbxActor )
 				{
-					ExportLevelSequenceTracks(MovieScene, MovieScenePlayer, FbxActor, Actor, MovieSceneBinding.GetTracks());
+					ExportLevelSequenceTracks(MovieScene, MovieScenePlayer, FbxActor, Actor, MovieSceneBinding.GetTracks(), RootToLocalTransform);
 					bAnyBindingsExported = true;
 				}
 			}
@@ -1657,7 +1686,7 @@ bool FFbxExporter::ExportLevelSequence( UMovieScene* MovieScene, const TArray<FG
 		// If no bindings exported, create a dummy actor to export tracks onto
 		if (!bAnyBindingsExported)
 		{
-			ExportLevelSequenceTracks(MovieScene, MovieScenePlayer, nullptr, nullptr, MovieSceneBinding.GetTracks());
+			ExportLevelSequenceTracks(MovieScene, MovieScenePlayer, nullptr, nullptr, MovieSceneBinding.GetTracks(), RootToLocalTransform);
 		}
 	}
 
@@ -2439,7 +2468,7 @@ void RichCurveInterpolationToFbxInterpolation(ERichCurveInterpMode InInterpolati
 	}
 }
 
-void FFbxExporter::ExportChannelToFbxCurve(FbxAnimCurve& InFbxCurve, const FMovieSceneFloatChannel& InChannel, FFrameRate TickResolution, ERichCurveValueMode ValueMode, bool bNegative)
+void FFbxExporter::ExportChannelToFbxCurve(FbxAnimCurve& InFbxCurve, const FMovieSceneFloatChannel& InChannel, FFrameRate TickResolution, ERichCurveValueMode ValueMode, bool bNegative, const FMovieSceneSequenceTransform& RootToLocalTransform)
 {
 	const float NegateFactor = bNegative ? -1.f : 1.f;
 	const float kOneThird = 1.f / 3.f;
@@ -2457,7 +2486,8 @@ void FFbxExporter::ExportChannelToFbxCurve(FbxAnimCurve& InFbxCurve, const FMovi
 
 		FbxTime FbxTime;
 		FbxAnimCurveKey FbxKey;
-		const double KeyTimeSeconds = KeyTime / TickResolution;
+		const double KeyTimeSeconds = GetExportOptions()->bExportLocalTime ? KeyTime / TickResolution : (KeyTime * RootToLocalTransform.Inverse()) / TickResolution;
+
 		FbxTime.SetSecondDouble(KeyTimeSeconds);
 
 		const int FbxKeyIndex = InFbxCurve.KeyAdd(FbxTime);
@@ -2513,7 +2543,7 @@ void FFbxExporter::ExportChannelToFbxCurve(FbxAnimCurve& InFbxCurve, const FMovi
 	InFbxCurve.KeyModifyEnd();
 }
 
-void FFbxExporter::ExportLevelSequence3DTransformTrack( FbxNode* FbxNode, UMovieScene3DTransformTrack& TransformTrack, AActor* Actor, const TRange<FFrameNumber>& InPlaybackRange )
+void FFbxExporter::ExportLevelSequence3DTransformTrack( FbxNode* FbxNode, UMovieScene3DTransformTrack& TransformTrack, AActor* Actor, const TRange<FFrameNumber>& InPlaybackRange, const FMovieSceneSequenceTransform& RootToLocalTransform )
 {
 	FbxAnimLayer* BaseLayer = AnimStack->GetMember<FbxAnimLayer>( 0 );
 
@@ -2537,6 +2567,7 @@ void FFbxExporter::ExportLevelSequence3DTransformTrack( FbxNode* FbxNode, UMovie
 	}
 
 	FFrameRate TickResolution = TransformTrack.GetTypedOuter<UMovieScene>()->GetTickResolution();
+	FFrameRate DisplayRate = TransformTrack.GetTypedOuter<UMovieScene>()->GetDisplayRate();
 
 	FbxAnimCurveNode* TranslationNode = FbxNode->LclTranslation.GetCurveNode( BaseLayer, true );
 	FbxAnimCurveNode* RotationNode = FbxNode->LclRotation.GetCurveNode( BaseLayer, true );
@@ -2557,24 +2588,24 @@ void FFbxExporter::ExportLevelSequence3DTransformTrack( FbxNode* FbxNode, UMovie
 	TArrayView<FMovieSceneFloatChannel*> FloatChannels = TransformSection->GetChannelProxy().GetChannels<FMovieSceneFloatChannel>();
 
 	// Translation
-	ExportChannelToFbxCurve(*FbxCurveTransX, *FloatChannels[0], TickResolution, ERichCurveValueMode::Default, false);
-	ExportChannelToFbxCurve(*FbxCurveTransY, *FloatChannels[1], TickResolution, ERichCurveValueMode::Default, true);
-	ExportChannelToFbxCurve(*FbxCurveTransZ, *FloatChannels[2], TickResolution, ERichCurveValueMode::Default, false);
+	ExportChannelToFbxCurve(*FbxCurveTransX, *FloatChannels[0], TickResolution, ERichCurveValueMode::Default, false, RootToLocalTransform);
+	ExportChannelToFbxCurve(*FbxCurveTransY, *FloatChannels[1], TickResolution, ERichCurveValueMode::Default, true, RootToLocalTransform);
+	ExportChannelToFbxCurve(*FbxCurveTransZ, *FloatChannels[2], TickResolution, ERichCurveValueMode::Default, false, RootToLocalTransform);
 
 	// Scale - don't generate scale keys for cameras
 	if (!bIsCameraActor)
 	{
-		ExportChannelToFbxCurve(*FbxCurveScaleX, *FloatChannels[6], TickResolution, ERichCurveValueMode::Default, false);
-		ExportChannelToFbxCurve(*FbxCurveScaleY, *FloatChannels[7], TickResolution, ERichCurveValueMode::Default, false);
-		ExportChannelToFbxCurve(*FbxCurveScaleZ, *FloatChannels[8], TickResolution, ERichCurveValueMode::Default, false);
+		ExportChannelToFbxCurve(*FbxCurveScaleX, *FloatChannels[6], TickResolution, ERichCurveValueMode::Default, false, RootToLocalTransform);
+		ExportChannelToFbxCurve(*FbxCurveScaleY, *FloatChannels[7], TickResolution, ERichCurveValueMode::Default, false, RootToLocalTransform);
+		ExportChannelToFbxCurve(*FbxCurveScaleZ, *FloatChannels[8], TickResolution, ERichCurveValueMode::Default, false, RootToLocalTransform);
 	}
 
 	// Rotation - bake rotation for cameras and lights
 	if (!bBakeRotations)
 	{
-		ExportChannelToFbxCurve(*FbxCurveRotX, *FloatChannels[3], TickResolution, ERichCurveValueMode::Default, false);
-		ExportChannelToFbxCurve(*FbxCurveRotY, *FloatChannels[4], TickResolution, ERichCurveValueMode::Default, true);
-		ExportChannelToFbxCurve(*FbxCurveRotZ, *FloatChannels[5], TickResolution, ERichCurveValueMode::Default, true);
+		ExportChannelToFbxCurve(*FbxCurveRotX, *FloatChannels[3], TickResolution, ERichCurveValueMode::Default, false, RootToLocalTransform);
+		ExportChannelToFbxCurve(*FbxCurveRotY, *FloatChannels[4], TickResolution, ERichCurveValueMode::Default, true, RootToLocalTransform);
+		ExportChannelToFbxCurve(*FbxCurveRotZ, *FloatChannels[5], TickResolution, ERichCurveValueMode::Default, true, RootToLocalTransform);
 	}
 	else
 	{
@@ -2594,29 +2625,30 @@ void FFbxExporter::ExportLevelSequence3DTransformTrack( FbxNode* FbxNode, UMovie
 		FbxCurveRotY->KeyModifyBegin();
 		FbxCurveRotZ->KeyModifyBegin();
 
-		FFrameNumber PlayStart      = MovieScene::DiscreteInclusiveLower(InPlaybackRange);
-		double       PlayDuration   = FFrameTime(MovieScene::DiscreteSize(InPlaybackRange)) / TickResolution;
-		int32        NumKeys        = PlayDuration * BakeTransformsFPS;
+		int32 LocalStartFrame = FFrameRate::TransformTime(FFrameTime(MovieScene::DiscreteInclusiveLower(InPlaybackRange)), TickResolution, DisplayRate).RoundToFrame().Value;
+		int32 StartFrame = FFrameRate::TransformTime(FFrameTime(MovieScene::DiscreteInclusiveLower(InPlaybackRange) * RootToLocalTransform.Inverse()), TickResolution, DisplayRate).RoundToFrame().Value;
+		int32 AnimationLength = FFrameRate::TransformTime(FFrameTime(FFrameNumber(MovieScene::DiscreteSize(InPlaybackRange))), TickResolution, DisplayRate).RoundToFrame().Value;
 
-		for ( int KeyIndex = 0; KeyIndex < NumKeys; KeyIndex++ )
+		for (int32 FrameCount = 0; FrameCount <= AnimationLength; ++FrameCount)
 		{
-			FFrameTime RelativeKeyTime = ( KeyIndex * PlayDuration / NumKeys ) * TickResolution;
-			FFrameTime KeyTime         = RelativeKeyTime + PlayStart;
+			int32 LocalFrame = LocalStartFrame + FrameCount;
+
+			FFrameTime LocalTime = FFrameRate::TransformTime(FFrameTime(LocalFrame), DisplayRate, TickResolution);
 
 			FVector Trans = FVector::ZeroVector;
-			FloatChannels[0]->Evaluate(KeyTime, Trans.X);
-			FloatChannels[1]->Evaluate(KeyTime, Trans.Y);
-			FloatChannels[2]->Evaluate(KeyTime, Trans.Z);
+			FloatChannels[0]->Evaluate(LocalTime, Trans.X);
+			FloatChannels[1]->Evaluate(LocalTime, Trans.Y);
+			FloatChannels[2]->Evaluate(LocalTime, Trans.Z);
 
 			FRotator Rotator;
-			FloatChannels[3]->Evaluate(KeyTime, Rotator.Roll);
-			FloatChannels[4]->Evaluate(KeyTime, Rotator.Pitch);
-			FloatChannels[5]->Evaluate(KeyTime, Rotator.Yaw);
+			FloatChannels[3]->Evaluate(LocalTime, Rotator.Roll);
+			FloatChannels[4]->Evaluate(LocalTime, Rotator.Pitch);
+			FloatChannels[5]->Evaluate(LocalTime, Rotator.Yaw);
 
 			FVector Scale;
-			FloatChannels[6]->Evaluate(KeyTime, Scale.X);
-			FloatChannels[7]->Evaluate(KeyTime, Scale.Y);
-			FloatChannels[8]->Evaluate(KeyTime, Scale.Z);
+			FloatChannels[6]->Evaluate(LocalTime, Scale.X);
+			FloatChannels[7]->Evaluate(LocalTime, Scale.Y);
+			FloatChannels[8]->Evaluate(LocalTime, Scale.Z);
 
 			FTransform RelativeTransform;
 			RelativeTransform.SetTranslation(Trans);
@@ -2630,7 +2662,7 @@ void FFbxExporter::ExportLevelSequence3DTransformTrack( FbxNode* FbxNode, UMovie
 			FbxVector4 KeyScale = Converter.ConvertToFbxScale(RelativeTransform.GetScale3D());
 
 			FbxTime FbxTime;
-			FbxTime.SetSecondDouble(KeyTime / TickResolution);
+			FbxTime.SetSecondDouble(GetExportOptions()->bExportLocalTime ? DisplayRate.AsSeconds(LocalFrame) : DisplayRate.AsSeconds(StartFrame + FrameCount));
 
 			FbxCurveRotX->KeySet(FbxCurveRotX->KeyAdd(FbxTime), FbxTime, KeyRot[0]);
 			FbxCurveRotY->KeySet(FbxCurveRotY->KeyAdd(FbxTime), FbxTime, KeyRot[1]);
@@ -2643,7 +2675,7 @@ void FFbxExporter::ExportLevelSequence3DTransformTrack( FbxNode* FbxNode, UMovie
 	}
 }
 
-void FFbxExporter::ExportLevelSequencePropertyTrack( FbxNode* FbxNode, UMovieScenePropertyTrack& PropertyTrack )
+void FFbxExporter::ExportLevelSequencePropertyTrack( FbxNode* FbxNode, UMovieScenePropertyTrack& PropertyTrack, const TRange<FFrameNumber>& InPlaybackRange, const FMovieSceneSequenceTransform& RootToLocalTransform)
 {
 	// TODO: Support more than one section?
 	UMovieSceneSection* Section = PropertyTrack.GetAllSections().Num() > 0 ? PropertyTrack.GetAllSections()[0] : nullptr;
@@ -2750,7 +2782,7 @@ void FFbxExporter::ExportLevelSequencePropertyTrack( FbxNode* FbxNode, UMovieSce
 			CurveNode->SetChannelValue<double>(0U, FloatChannel->GetDefault().Get(MAX_flt));
 			CurveNode->ConnectToChannel(AnimCurve, 0U);
 
-			ExportChannelToFbxCurve(*AnimCurve, *FloatChannel, TickResolution, IsFoV ? ERichCurveValueMode::Fov : ERichCurveValueMode::Default);
+			ExportChannelToFbxCurve(*AnimCurve, *FloatChannel, TickResolution, IsFoV ? ERichCurveValueMode::Fov : ERichCurveValueMode::Default, false, RootToLocalTransform);
 		}
 	}
 }
diff --git a/Engine/Source/Editor/UnrealEd/Private/FbxExporter.h b/Engine/Source/Editor/UnrealEd/Private/FbxExporter.h
--- a/Engine/Source/Editor/UnrealEd/Private/FbxExporter.h
+++ b/Engine/Source/Editor/UnrealEd/Private/FbxExporter.h
@@ -8,6 +8,7 @@
 #include "Engine/StaticMesh.h"
 #include "MatineeExporter.h"
 #include "MovieSceneSequenceID.h"
+#include "Evaluation/MovieSceneSequenceTransform.h"
 #include "MovieSceneFwd.h"
 #include "FbxImporter.h"
 #include "UObject/GCObject.h"
@@ -42,6 +43,7 @@ class UFbxExportOption;
 struct FAnimControlTrackKey;
 struct FExpressionInput;
 struct FMovieSceneFloatChannel;
+struct FMovieSceneSequenceTransform;
 
 namespace UnFbx
 {
@@ -131,14 +133,14 @@ public:
 	 *
 	 * @return	true, if successful
 	 */
-	bool ExportLevelSequence(UMovieScene* MovieScene, const TArray<FGuid>& InBindings, IMovieScenePlayer* MovieScenePlayer, FMovieSceneSequenceIDRef SequenceID);
+	bool ExportLevelSequence(UMovieScene* MovieScene, const TArray<FGuid>& InBindings, IMovieScenePlayer* MovieScenePlayer, FMovieSceneSequenceIDRef SequenceID, const FMovieSceneSequenceTransform& RootToLocalTransform);
 
 	/**
 	 * Exports the given level sequence track information into a FBX document.
 	 *
 	 * @return	true, if successful
 	 */
-	bool ExportLevelSequenceTracks(UMovieScene* MovieScene, IMovieScenePlayer* MovieScenePlayer, FbxNode* FbxActor, AActor* Actor, const TArray<UMovieSceneTrack*>& Tracks);
+	bool ExportLevelSequenceTracks(UMovieScene* MovieScene, IMovieScenePlayer* MovieScenePlayer, FbxNode* FbxActor, AActor* Actor, const TArray<UMovieSceneTrack*>& Tracks, const FMovieSceneSequenceTransform& RootToLocalTransform);
 
 	/**
 	 * Exports all the animation sequences part of a single Group in a Matinee sequence
@@ -264,11 +266,12 @@ private:
 	class IAnimTrackAdapter
 	{
 	public:
-		/** Gets the length of the animation track. */
-		virtual float GetAnimationStart() const = 0;
-		virtual float GetAnimationLength() const = 0;
-		/** Updates the runtime state of the animation track to the specified time. */
-		virtual void UpdateAnimation( float Time ) = 0;
+		virtual int32 GetLocalStartFrame() const = 0;
+		virtual int32 GetStartFrame() const { return GetLocalStartFrame(); }
+		virtual int32 GetLength() const = 0;
+		/** Updates the runtime state of the animation track to the specified frame. */
+		virtual void UpdateAnimation( int32 LocalFrame) = 0;
+		virtual float GetFrameRate() const { return 1.f / DEFAULT_SAMPLERATE; }
 	};
 
 	/** An anim track adapter for matinee. */
@@ -276,9 +279,9 @@ private:
 	{
 	public:
 		FMatineeAnimTrackAdapter( AMatineeActor* InMatineeActor );
-		virtual float GetAnimationStart() const override;
-		virtual float GetAnimationLength() const override;
-		virtual void UpdateAnimation( float Time ) override;
+		virtual int32 GetLocalStartFrame() const override;
+		virtual int32 GetLength() const override;
+		virtual void UpdateAnimation(int32 LocalFrame) override;
 
 	private:
 		AMatineeActor* MatineeActor;
@@ -288,20 +291,23 @@ private:
 	class FLevelSequenceAnimTrackAdapter : public FFbxExporter::IAnimTrackAdapter
 	{
 	public:
-		FLevelSequenceAnimTrackAdapter( IMovieScenePlayer* InMovieScenePlayer, UMovieScene* InMovieScene );
-		virtual float GetAnimationStart() const override;
-		virtual float GetAnimationLength() const override;
-		virtual void UpdateAnimation( float Time ) override;
+		FLevelSequenceAnimTrackAdapter( IMovieScenePlayer* InMovieScenePlayer, UMovieScene* InMovieScene, const FMovieSceneSequenceTransform& InRootToLocalTransform );
+		virtual int32 GetLocalStartFrame() const override;
+		virtual int32 GetStartFrame() const override;
+		virtual int32 GetLength() const override;
+		virtual void UpdateAnimation(int32 LocalFrame) override;
+		virtual float GetFrameRate() const override;
 
 	private:
 		IMovieScenePlayer* MovieScenePlayer;
 		UMovieScene* MovieScene;
+		FMovieSceneSequenceTransform RootToLocalTransform;
 	};
 
 	/**
 	* Export Anim Track of the given SkeletalMeshComponent
 	*/
-	void ExportAnimTrack( IAnimTrackAdapter& AnimTrackAdapter, AActor* Actor, USkeletalMeshComponent* SkeletalMeshComponent );
+	void ExportAnimTrack( IAnimTrackAdapter& AnimTrackAdapter, AActor* Actor, USkeletalMeshComponent* SkeletalMeshComponent, float SamplingRate );
 
 	void ExportModel(UModel* Model, FbxNode* Node, const char* Name);
 	
@@ -446,12 +452,12 @@ private:
 	/**
 	 * Exports a level sequence 3D transform track into the FBX animation stack.
 	 */
-	void ExportLevelSequence3DTransformTrack( FbxNode* FbxActor, UMovieScene3DTransformTrack& TransformTrack, AActor* Actor, const TRange<FFrameNumber>& InPlaybackRange );
+	void ExportLevelSequence3DTransformTrack( FbxNode* FbxActor, UMovieScene3DTransformTrack& TransformTrack, AActor* Actor, const TRange<FFrameNumber>& InPlaybackRange, const FMovieSceneSequenceTransform& RootToLocalTransform );
 
 	/** 
 	 * Exports a level sequence property track into the FBX animation stack. 
 	 */
-	void ExportLevelSequencePropertyTrack( FbxNode* FbxActor, UMovieScenePropertyTrack& PropertyTrack );
+	void ExportLevelSequencePropertyTrack( FbxNode* FbxActor, UMovieScenePropertyTrack& PropertyTrack, const TRange<FFrameNumber>& InPlaybackRange, const FMovieSceneSequenceTransform& RootToLocalTransform);
 
 	/** Defines value export modes for the EportRichCurveToFbxCurve method. */
 	enum class ERichCurveValueMode
@@ -463,7 +469,7 @@ private:
 	};
 
 	/** Exports a movie scene channel to an fbx animation curve. */
-	void ExportChannelToFbxCurve(FbxAnimCurve& InFbxCurve, const FMovieSceneFloatChannel& InChannel, FFrameRate TickResolution, ERichCurveValueMode ValueMode = ERichCurveValueMode::Default, bool bNegative = false);
+	void ExportChannelToFbxCurve(FbxAnimCurve& InFbxCurve, const FMovieSceneFloatChannel& InChannel, FFrameRate TickResolution, ERichCurveValueMode ValueMode = ERichCurveValueMode::Default, bool bNegative = false, const FMovieSceneSequenceTransform& RootToLocalTransform = FMovieSceneSequenceTransform());
 
 	/**
 	 * Finds the given actor in the already-exported list of structures
