UEdMode initial implementation and handling by FEdModeManager
#rnx

diff --git a/Engine/Config/BaseEngine.ini b/Engine/Config/BaseEngine.ini
--- a/Engine/Config/BaseEngine.ini
+++ b/Engine/Config/BaseEngine.ini
@@ -689,6 +689,12 @@ bUseStreamingPause=false
 
 ; 4.24
 +PropertyRedirects=(OldName="NavDataConfig.NavigationDataClassName", NewName="NavDataConfig.NavDataClass")
++ClassRedirects=(OldName="/Script/EditorInteractiveToolsFramework.StaticMeshComponentTargetFactory",NewName="/Script/UnrealEd.StaticMeshComponentTargetFactory")
++ClassRedirects=(OldName="/Script/EditorInteractiveToolsFramework.StandardToolModeCommands",NewName="/Script/UnrealEd.StandardToolModeCommands")
++EnumRedirects==(OldName="/Script/EditorInteractiveToolsFramework.StandardToolModeCommands",NewName="/Script/UnrealEd.StandardToolModeCommands")
++ClassRedirects=(OldName="/Script/EditorInteractiveToolsFramework.InteractiveToolCommands",NewName="/Script/UnrealEd.InteractiveToolCommands")
++ClassRedirects=(OldName="/Script/EditorInteractiveToolsFramework.EditorToolAssetAPI",NewName="/Script/UnrealEd.EditorToolAssetAPI")
++ClassRedirects=(OldName="/Script/EditorInteractiveToolsFramework.StaticMeshComponentTarget",NewName="/Script/UnrealEd.StaticMeshComponentTarget")
 
 +FunctionRedirects=(OldName="Actor.GetComponentsByClass", NewName="Actor.K2_GetComponentsByClass")
 
diff --git a/Engine/Plugins/Editor/MeshEditor/Source/MeshEditor/MeshEditorModule.cpp b/Engine/Plugins/Editor/MeshEditor/Source/MeshEditor/MeshEditorModule.cpp
--- a/Engine/Plugins/Editor/MeshEditor/Source/MeshEditor/MeshEditorModule.cpp
+++ b/Engine/Plugins/Editor/MeshEditor/Source/MeshEditor/MeshEditorModule.cpp
@@ -228,17 +228,9 @@ void FMeshEditorModule::OnMeshEditModeButtonClicked(EEditableMeshElementType InM
 	GLevelEditorModeTools().ActivateMode(GetEditorModeID());
 
 	// Find and disable any other 'visible' modes since we only ever allow one of those active at a time.
-	TArray<FEdMode*> ActiveModes;
-	GLevelEditorModeTools().GetActiveModes(ActiveModes);
-	for (FEdMode* Mode : ActiveModes)
-	{
-		if (Mode->GetID() != GetEditorModeID() && Mode->GetModeInfo().bVisible)
-		{
-			GLevelEditorModeTools().DeactivateMode(Mode->GetID());
-		}
-	}
+	GLevelEditorModeTools().DeactivateOtherVisibleModes(GetEditorModeID());
 
-	FMeshEditorMode* MeshEditorMode = static_cast<FMeshEditorMode*>( GLevelEditorModeTools().FindMode( GetEditorModeID() ) );
+	FMeshEditorMode* MeshEditorMode = static_cast<FMeshEditorMode*>( GLevelEditorModeTools().GetActiveMode( GetEditorModeID() ) );
 	if ( MeshEditorMode != nullptr)
 	{
 		IMeshEditorModeUIContract* MeshEditorModeUIContract = (IMeshEditorModeUIContract*)MeshEditorMode;
@@ -256,7 +248,7 @@ ECheckBoxState FMeshEditorModule::IsMeshEditModeButtonChecked(EEditableMeshEleme
 {
 	bool bMeshModeActive = false;
 
-	const FMeshEditorMode* MeshEditorMode = static_cast<FMeshEditorMode*>( GLevelEditorModeTools().FindMode( GetEditorModeID() ) );
+	const FMeshEditorMode* MeshEditorMode = static_cast<FMeshEditorMode*>( GLevelEditorModeTools().GetActiveMode( GetEditorModeID() ) );
 	if( MeshEditorMode != nullptr )
 	{
 		const IMeshEditorModeUIContract* MeshEditorModeUIContract = (const IMeshEditorModeUIContract*)MeshEditorMode;
diff --git a/Engine/Plugins/Experimental/GeometryCollectionPlugin/Source/GeometryCollectionEditor/Private/GeometryCollection/GeometryCollectionSelectRigidBodyEdMode.cpp b/Engine/Plugins/Experimental/GeometryCollectionPlugin/Source/GeometryCollectionEditor/Private/GeometryCollection/GeometryCollectionSelectRigidBodyEdMode.cpp
--- a/Engine/Plugins/Experimental/GeometryCollectionPlugin/Source/GeometryCollectionEditor/Private/GeometryCollection/GeometryCollectionSelectRigidBodyEdMode.cpp
+++ b/Engine/Plugins/Experimental/GeometryCollectionPlugin/Source/GeometryCollectionEditor/Private/GeometryCollection/GeometryCollectionSelectRigidBodyEdMode.cpp
@@ -60,7 +60,7 @@ void FGeometryCollectionSelectRigidBodyEdMode::ActivateMode(TSharedRef<IProperty
 
 		// Activate editor mode
 		GLevelEditorModeTools().ActivateMode(EditorModeID);
-		if (FEdMode* const EdMode = GLevelEditorModeTools().FindMode(EditorModeID))
+		if (FEdMode* const EdMode = GLevelEditorModeTools().GetActiveMode(EditorModeID))
 		{
 			// Set pointers
 			FGeometryCollectionSelectRigidBodyEdMode* const GeometryCollectionSelectRigidBodyEdMode = static_cast<FGeometryCollectionSelectRigidBodyEdMode*>(EdMode);
diff --git a/Engine/Source/Developer/AssetTools/Private/AssetTypeActions/AssetTypeActions_CameraAnim.cpp b/Engine/Source/Developer/AssetTools/Private/AssetTypeActions/AssetTypeActions_CameraAnim.cpp
--- a/Engine/Source/Developer/AssetTools/Private/AssetTypeActions/AssetTypeActions_CameraAnim.cpp
+++ b/Engine/Source/Developer/AssetTools/Private/AssetTypeActions/AssetTypeActions_CameraAnim.cpp
@@ -221,7 +221,7 @@ void FAssetTypeActions_CameraAnim::OpenAssetEditor( const TArray<UObject*>& InOb
 				GEditor->OpenMatinee(PreviewMatineeActor.Get());
 
 				// install our delegate so we can clean up when finished
-				OnMatineeEditorClosedDelegateHandle = FEditorDelegates::EditorModeExit.AddSP(this, &FAssetTypeActions_CameraAnim::OnMatineeEditorClosed);
+				OnMatineeEditorClosedDelegateHandle = FEditorDelegates::EditorModeIDExit.AddSP(this, &FAssetTypeActions_CameraAnim::OnMatineeEditorClosed);
 			}
 			else
 			{
@@ -231,10 +231,9 @@ void FAssetTypeActions_CameraAnim::OpenAssetEditor( const TArray<UObject*>& InOb
 	}
 }
 
-void FAssetTypeActions_CameraAnim::OnMatineeEditorClosed(FEdMode* InEditorMode)
+void FAssetTypeActions_CameraAnim::OnMatineeEditorClosed(const FEditorModeID& InEditorModeID)
 {
-	check(InEditorMode);
-	if(InEditorMode->GetID() == FBuiltinEditorModes::EM_InterpEdit)
+	if(InEditorModeID == FBuiltinEditorModes::EM_InterpEdit)
 	{
 		// clean up our preview actors if they are still present
 		if(PreviewCamera.IsValid())
@@ -256,7 +255,7 @@ void FAssetTypeActions_CameraAnim::OnMatineeEditorClosed(FEdMode* InEditorMode)
 		}
 
 		// remove our delegate
-		FEditorDelegates::EditorModeExit.Remove(OnMatineeEditorClosedDelegateHandle);
+		FEditorDelegates::EditorModeIDExit.Remove(OnMatineeEditorClosedDelegateHandle);
 	}
 }
 
diff --git a/Engine/Source/Developer/AssetTools/Private/AssetTypeActions/AssetTypeActions_CameraAnim.h b/Engine/Source/Developer/AssetTools/Private/AssetTypeActions/AssetTypeActions_CameraAnim.h
--- a/Engine/Source/Developer/AssetTools/Private/AssetTypeActions/AssetTypeActions_CameraAnim.h
+++ b/Engine/Source/Developer/AssetTools/Private/AssetTypeActions/AssetTypeActions_CameraAnim.h
@@ -10,6 +10,8 @@
 #include "Camera/CameraAnim.h"
 
 class AMatineeActorCameraAnim;
+class FName;
+typedef FName FEditorModeID;
 
 class FAssetTypeActions_CameraAnim : public FAssetTypeActions_Base
 {
@@ -30,7 +32,7 @@ private:
 	void CreateCameraActorForCameraAnim(class UCameraAnim* InCameraAnim);
 
 	/** Delegate fired when the editor mode is changed */
-	void OnMatineeEditorClosed(class FEdMode* InEditorMode);
+	void OnMatineeEditorClosed(const FEditorModeID& InEditorModeID);
 
 	/** Helper function to hookup preview pawn */
 	void CreatePreviewPawnForCameraAnim(class UCameraAnim* InCameraAnim);
diff --git a/Engine/Source/Editor/ClothPainter/Private/SClothPaintTab.cpp b/Engine/Source/Editor/ClothPainter/Private/SClothPaintTab.cpp
--- a/Engine/Source/Editor/ClothPainter/Private/SClothPaintTab.cpp
+++ b/Engine/Source/Editor/ClothPainter/Private/SClothPaintTab.cpp
@@ -147,7 +147,7 @@ void SClothPaintTab::UpdatePaintTools()
 		ISkeletalMeshEditor* SkeletalMeshEditor = static_cast<ISkeletalMeshEditor*>(HostingApp.Pin().Get());
 		SkeletalMeshEditor->GetAssetEditorModeManager()->ActivateMode(PaintModeID, true);
 
-		FClothingPaintEditMode* PaintMode = (FClothingPaintEditMode*)SkeletalMeshEditor->GetAssetEditorModeManager()->FindMode(PaintModeID);
+		FClothingPaintEditMode* PaintMode = (FClothingPaintEditMode*)SkeletalMeshEditor->GetAssetEditorModeManager()->GetActiveMode(PaintModeID);
 		if (PaintMode)
 		{
 			FClothPainter* ClothPainter = static_cast<FClothPainter*>(PaintMode->GetMeshPainter());
@@ -189,7 +189,7 @@ void SClothPaintTab::OnAssetSelectionChanged(TWeakObjectPtr<UClothingAsset> InAs
 	{
 		ISkeletalMeshEditor* SkeletalMeshEditor = static_cast<ISkeletalMeshEditor*>(HostingApp.Pin().Get());
 
-		FClothingPaintEditMode* PaintMode = (FClothingPaintEditMode*)SkeletalMeshEditor->GetAssetEditorModeManager()->FindMode(PaintModeID);
+		FClothingPaintEditMode* PaintMode = (FClothingPaintEditMode*)SkeletalMeshEditor->GetAssetEditorModeManager()->GetActiveMode(PaintModeID);
 		if(PaintMode)
 		{
 			FClothPainter* ClothPainter = static_cast<FClothPainter*>(PaintMode->GetMeshPainter());
diff --git a/Engine/Source/Editor/Experimental/EditorInteractiveToolsFramework/Private/EdModeInteractiveToolsContext.cpp b/Engine/Source/Editor/Experimental/EditorInteractiveToolsFramework/Private/EdModeInteractiveToolsContext.cpp
--- a/Engine/Source/Editor/Experimental/EditorInteractiveToolsFramework/Private/EdModeInteractiveToolsContext.cpp
+++ b/Engine/Source/Editor/Experimental/EditorInteractiveToolsFramework/Private/EdModeInteractiveToolsContext.cpp
@@ -14,8 +14,8 @@
 #include "Engine/StaticMesh.h"
 #include "Components/StaticMeshComponent.h"
 
-#include "EditorToolAssetAPI.h"
-#include "EditorComponentSourceFactory.h"
+#include "Tools/EditorToolAssetAPI.h"
+#include "Tools/EditorComponentSourceFactory.h"
 #include "InteractiveToolObjects.h"
 
 //#include "PhysicsEngine/BodySetup.h"
@@ -750,72 +750,53 @@ FRay UEdModeInteractiveToolsContext::GetRayFromMousePos(FEditorViewportClient* V
 
 bool UEdModeInteractiveToolsContext::CanStartTool(const FString& ToolTypeIdentifier) const
 {
-	return (ToolManager->HasActiveTool(EToolSide::Mouse) == false) &&
-		(ToolManager->CanActivateTool(EToolSide::Mouse, ToolTypeIdentifier) == true);
+	return UInteractiveToolsContext::CanStartTool(EToolSide::Mouse, ToolTypeIdentifier);
 }
 
 bool UEdModeInteractiveToolsContext::ActiveToolHasAccept() const
 {
-	return ToolManager->HasActiveTool(EToolSide::Mouse) &&
-		ToolManager->GetActiveTool(EToolSide::Mouse)->HasAccept();
+	return  UInteractiveToolsContext::ActiveToolHasAccept(EToolSide::Mouse);
 }
 
 bool UEdModeInteractiveToolsContext::CanAcceptActiveTool() const
 {
-	return ToolManager->CanAcceptActiveTool(EToolSide::Mouse);
+	return UInteractiveToolsContext::CanAcceptActiveTool(EToolSide::Mouse);
 }
 
 bool UEdModeInteractiveToolsContext::CanCancelActiveTool() const
 {
-	return ToolManager->CanCancelActiveTool(EToolSide::Mouse);
+	return UInteractiveToolsContext::CanCancelActiveTool(EToolSide::Mouse);
 }
 
 bool UEdModeInteractiveToolsContext::CanCompleteActiveTool() const
 {
-	return ToolManager->HasActiveTool(EToolSide::Mouse) && CanCancelActiveTool() == false;
+	return UInteractiveToolsContext::CanCompleteActiveTool(EToolSide::Mouse);
 }
 
 void UEdModeInteractiveToolsContext::StartTool(const FString& ToolTypeIdentifier)
 {
-	if (ToolManager->SelectActiveToolType(EToolSide::Mouse, ToolTypeIdentifier) == false)
+	if (UInteractiveToolsContext::StartTool(EToolSide::Mouse, ToolTypeIdentifier))
 	{
-		UE_LOG(LogTemp, Warning, TEXT("ToolManager: Unknown Tool Type %s"), *ToolTypeIdentifier);
-	}
-	else
-	{
-		ToolManager->ActivateTool(EToolSide::Mouse);
 		SaveEditorStateAndSetForTool();
 	}
 }
 
 void UEdModeInteractiveToolsContext::EndTool(EToolShutdownType ShutdownType)
 {
-	DeactivateActiveTool(EToolSide::Mouse, ShutdownType);
+	UInteractiveToolsContext::EndTool(EToolSide::Mouse, ShutdownType);
 }
 
 
 
 void UEdModeInteractiveToolsContext::DeactivateActiveTool(EToolSide WhichSide, EToolShutdownType ShutdownType)
 {
-	ToolManager->DeactivateTool(WhichSide, ShutdownType);
+	UInteractiveToolsContext::DeactivateActiveTool(WhichSide, ShutdownType);
 	RestoreEditorState();
 }
 
 void UEdModeInteractiveToolsContext::DeactivateAllActiveTools()
 {
-	if (ToolManager->HasActiveTool(EToolSide::Left))
-	{
-		EToolShutdownType ShutdownType = ToolManager->CanAcceptActiveTool(EToolSide::Left) ?
-			EToolShutdownType::Accept : EToolShutdownType::Cancel;
-		ToolManager->DeactivateTool(EToolSide::Left, ShutdownType);
-	}
-	if (ToolManager->HasActiveTool(EToolSide::Right))
-	{
-		EToolShutdownType ShutdownType = ToolManager->CanAcceptActiveTool(EToolSide::Right) ?
-			EToolShutdownType::Accept : EToolShutdownType::Cancel;
-		ToolManager->DeactivateTool(EToolSide::Right, ShutdownType);
-	}
-
+	UInteractiveToolsContext::DeactivateAllActiveTools();
 	RestoreEditorState();
 }
 
diff --git a/Engine/Source/Editor/Experimental/EditorInteractiveToolsFramework/Private/EditorComponentSourceFactory.cpp b/Engine/Source/Editor/Experimental/EditorInteractiveToolsFramework/Private/EditorComponentSourceFactory.cpp
deleted file mode 100644
diff --git a/Engine/Source/Editor/Experimental/EditorInteractiveToolsFramework/Private/EditorInteractiveToolsFrameworkModule.cpp b/Engine/Source/Editor/Experimental/EditorInteractiveToolsFramework/Private/EditorInteractiveToolsFrameworkModule.cpp
--- a/Engine/Source/Editor/Experimental/EditorInteractiveToolsFramework/Private/EditorInteractiveToolsFrameworkModule.cpp
+++ b/Engine/Source/Editor/Experimental/EditorInteractiveToolsFramework/Private/EditorInteractiveToolsFrameworkModule.cpp
@@ -2,7 +2,7 @@
 
 #include "EditorInteractiveToolsFrameworkModule.h"
 #include "ComponentSourceInterfaces.h"
-#include "EditorComponentSourceFactory.h"
+#include "Tools/EditorComponentSourceFactory.h"
 
 #define LOCTEXT_NAMESPACE "FEditorInteractiveToolsFrameworkModule"
 
diff --git a/Engine/Source/Editor/Experimental/EditorInteractiveToolsFramework/Private/EditorToolAssetAPI.cpp b/Engine/Source/Editor/Experimental/EditorInteractiveToolsFramework/Private/EditorToolAssetAPI.cpp
deleted file mode 100644
diff --git a/Engine/Source/Editor/Experimental/EditorInteractiveToolsFramework/Private/StandardToolModeCommands.cpp b/Engine/Source/Editor/Experimental/EditorInteractiveToolsFramework/Private/StandardToolModeCommands.cpp
deleted file mode 100644
diff --git a/Engine/Source/Editor/Experimental/EditorInteractiveToolsFramework/Public/EditorComponentSourceFactory.h b/Engine/Source/Editor/Experimental/EditorInteractiveToolsFramework/Public/EditorComponentSourceFactory.h
deleted file mode 100644
diff --git a/Engine/Source/Editor/Experimental/EditorInteractiveToolsFramework/Public/EditorToolAssetAPI.h b/Engine/Source/Editor/Experimental/EditorInteractiveToolsFramework/Public/EditorToolAssetAPI.h
deleted file mode 100644
diff --git a/Engine/Source/Editor/Experimental/EditorInteractiveToolsFramework/Public/InteractiveToolsCommands.h b/Engine/Source/Editor/Experimental/EditorInteractiveToolsFramework/Public/InteractiveToolsCommands.h
deleted file mode 100644
diff --git a/Engine/Source/Editor/Experimental/EditorInteractiveToolsFramework/Public/StandardToolModeCommands.h b/Engine/Source/Editor/Experimental/EditorInteractiveToolsFramework/Public/StandardToolModeCommands.h
deleted file mode 100644
diff --git a/Engine/Source/Editor/LandscapeEditor/Private/LandscapeEditorModule.cpp b/Engine/Source/Editor/LandscapeEditor/Private/LandscapeEditorModule.cpp
--- a/Engine/Source/Editor/LandscapeEditor/Private/LandscapeEditorModule.cpp
+++ b/Engine/Source/Editor/LandscapeEditor/Private/LandscapeEditorModule.cpp
@@ -234,7 +234,7 @@ public:
 
 	void OnPreSaveWorld(uint32 SaveFlags, class UWorld* World)
 	{
-		FEdModeLandscape* EdMode = (FEdModeLandscape*)GLevelEditorModeTools().FindMode(FBuiltinEditorModes::EM_Landscape);
+		FEdModeLandscape* EdMode = (FEdModeLandscape*)GLevelEditorModeTools().GetActiveMode(FBuiltinEditorModes::EM_Landscape);
 		if (EdMode)
 		{
 			EdMode->OnPreSaveWorld(SaveFlags, World);
diff --git a/Engine/Source/Editor/LevelEditor/Private/LevelEditorActions.cpp b/Engine/Source/Editor/LevelEditor/Private/LevelEditorActions.cpp
--- a/Engine/Source/Editor/LevelEditor/Private/LevelEditorActions.cpp
+++ b/Engine/Source/Editor/LevelEditor/Private/LevelEditorActions.cpp
@@ -1512,19 +1512,14 @@ void FLevelEditorActionCallbacks::ReplaceActorsFromClass_Clicked( UClass* ActorC
 
 bool FLevelEditorActionCallbacks::Duplicate_CanExecute()
 {
-	TArray<FEdMode*> ActiveModes; 
-	GLevelEditorModeTools().GetActiveModes( ActiveModes );
-	for( int32 ModeIndex = 0; ModeIndex < ActiveModes.Num(); ++ModeIndex )
+	const EEditAction::Type CanProcess = GLevelEditorModeTools().GetActionEditDuplicate();
+	if (CanProcess == EEditAction::Process)
 	{
-		const EEditAction::Type CanProcess = ActiveModes[ModeIndex]->GetActionEditDuplicate();
-		if (CanProcess == EEditAction::Process)
-		{
-			return true;
-		}
-		else if (CanProcess == EEditAction::Halt)
-		{
-			return false;
-		}
+		return true;
+	}
+	else if (CanProcess == EEditAction::Halt)
+	{
+		return false;
 	}
 	
 	// If we can copy, we can duplicate
@@ -1566,19 +1561,14 @@ bool FLevelEditorActionCallbacks::Duplicate_CanExecute()
 
 bool FLevelEditorActionCallbacks::Delete_CanExecute()
 {
-	TArray<FEdMode*> ActiveModes; 
-	GLevelEditorModeTools().GetActiveModes( ActiveModes );
-	for( int32 ModeIndex = 0; ModeIndex < ActiveModes.Num(); ++ModeIndex )
+	const EEditAction::Type CanProcess = GLevelEditorModeTools().GetActionEditDelete();
+	if (CanProcess == EEditAction::Process)
 	{
-		const EEditAction::Type CanProcess = ActiveModes[ModeIndex]->GetActionEditDelete();
-		if (CanProcess == EEditAction::Process)
-		{
-			return true;
-		}
-		else if (CanProcess == EEditAction::Halt)
-		{
-			return false;
-		}
+		return true;
+	}
+	else if (CanProcess == EEditAction::Halt)
+	{
+		return false;
 	}
 
 	bool bCanDelete = false;
@@ -1676,19 +1666,14 @@ bool FLevelEditorActionCallbacks::Rename_CanExecute()
 
 bool FLevelEditorActionCallbacks::Cut_CanExecute()
 {
-	TArray<FEdMode*> ActiveModes; 
-	GLevelEditorModeTools().GetActiveModes( ActiveModes );
-	for( int32 ModeIndex = 0; ModeIndex < ActiveModes.Num(); ++ModeIndex )
+	const EEditAction::Type CanProcess = GLevelEditorModeTools().GetActionEditCut();
+	if (CanProcess == EEditAction::Process)
 	{
-		const EEditAction::Type CanProcess = ActiveModes[ModeIndex]->GetActionEditCut();
-		if (CanProcess == EEditAction::Process)
-		{
-			return true;
-		}
-		else if (CanProcess == EEditAction::Halt)
-		{
-			return false;
-		}
+		return true;
+	}
+	else if (CanProcess == EEditAction::Halt)
+	{
+		return false;
 	}
 
 	bool bCanCut = false;
@@ -1731,19 +1716,14 @@ bool FLevelEditorActionCallbacks::Cut_CanExecute()
 
 bool FLevelEditorActionCallbacks::Copy_CanExecute()
 {
-	TArray<FEdMode*> ActiveModes; 
-	GLevelEditorModeTools().GetActiveModes( ActiveModes );
-	for( int32 ModeIndex = 0; ModeIndex < ActiveModes.Num(); ++ModeIndex )
+	const EEditAction::Type CanProcess = GLevelEditorModeTools().GetActionEditCopy();
+	if (CanProcess == EEditAction::Process)
 	{
-		const EEditAction::Type CanProcess = ActiveModes[ModeIndex]->GetActionEditCopy();
-		if (CanProcess == EEditAction::Process)
-		{
-			return true;
-		}
-		else if (CanProcess == EEditAction::Halt)
-		{
-			return false;
-		}
+		return true;
+	}
+	else if (CanProcess == EEditAction::Halt)
+	{
+		return false;
 	}
 
 	bool bCanCopy = false;
@@ -1784,19 +1764,14 @@ bool FLevelEditorActionCallbacks::Copy_CanExecute()
 
 bool FLevelEditorActionCallbacks::Paste_CanExecute()
 {
-	TArray<FEdMode*> ActiveModes; 
-	GLevelEditorModeTools().GetActiveModes( ActiveModes );
-	for( int32 ModeIndex = 0; ModeIndex < ActiveModes.Num(); ++ModeIndex )
+	const EEditAction::Type CanProcess = GLevelEditorModeTools().GetActionEditPaste();
+	if (CanProcess == EEditAction::Process)
 	{
-		const EEditAction::Type CanProcess = ActiveModes[ModeIndex]->GetActionEditPaste();
-		if (CanProcess == EEditAction::Process)
-		{
-			return true;
-		}
-		else if (CanProcess == EEditAction::Halt)
-		{
-			return false;
-		}
+		return true;
+	}
+	else if (CanProcess == EEditAction::Halt)
+	{
+		return false;
 	}
 
 	bool bCanPaste = false;
diff --git a/Engine/Source/Editor/LevelEditor/Private/SLevelEditor.cpp b/Engine/Source/Editor/LevelEditor/Private/SLevelEditor.cpp
--- a/Engine/Source/Editor/LevelEditor/Private/SLevelEditor.cpp
+++ b/Engine/Source/Editor/LevelEditor/Private/SLevelEditor.cpp
@@ -1357,37 +1357,12 @@ void SLevelEditor::ToggleEditorMode( FEditorModeID ModeID )
 	GLevelEditorModeTools().ActivateMode( ModeID );
 
 	// Find and disable any other 'visible' modes since we only ever allow one of those active at a time.
-	TArray<FEdMode*> ActiveModes;
-	GLevelEditorModeTools().GetActiveModes( ActiveModes );
-	for ( FEdMode* Mode : ActiveModes )
-	{
-		if ( Mode->GetID() != ModeID && Mode->GetModeInfo().bVisible )
-		{
-			GLevelEditorModeTools().DeactivateMode( Mode->GetID() );
-		}
-	}
+	GLevelEditorModeTools().DeactivateOtherVisibleModes(ModeID);
 
 	FLevelEditorModule& LevelEditorModule = FModuleManager::GetModuleChecked<FLevelEditorModule>( "LevelEditor" );
 	TSharedPtr<FTabManager> LevelEditorTabManager = LevelEditorModule.GetLevelEditorTabManager();
 
 	TSharedRef<SDockTab> ToolboxTab = LevelEditorTabManager->InvokeTab( FTabId("LevelEditorToolBox") );
-
-	//// If it's already active deactivate the mode
-	//if ( GLevelEditorModeTools().IsModeActive( ModeID ) )
-	//{
-	//	//GLevelEditorModeTools().DeactivateAllModes();
-	//	//GLevelEditorModeTools().DeactivateMode( ModeID );
-	//}
-	//else // Activate the mode and create the tab for it.
-	//{
-	//	GLevelEditorModeTools().DeactivateAllModes();
-	//	GLevelEditorModeTools().ActivateMode( ModeID );
-
-	//	//FLevelEditorModule& LevelEditorModule = FModuleManager::GetModuleChecked<FLevelEditorModule>( "LevelEditor" );
-	//	//TSharedPtr<FTabManager> LevelEditorTabManager = LevelEditorModule.GetLevelEditorTabManager();
-
-	//	//TSharedRef<SDockTab> ToolboxTab = LevelEditorTabManager->InvokeTab( GetEditorModeTabId( ModeID ) );
-	//}
 }
 
 bool SLevelEditor::IsModeActive( FEditorModeID ModeID )
@@ -1395,15 +1370,9 @@ bool SLevelEditor::IsModeActive( FEditorModeID ModeID )
 	// The level editor changes the default mode to placement
 	if ( ModeID == FBuiltinEditorModes::EM_Placement )
 	{
-		// Only return true if this is the *only* active mode
-		TArray<FEdMode*> ActiveModes;
-		GLevelEditorModeTools().GetActiveModes(ActiveModes);
-		for( FEdMode* Mode : ActiveModes )
+		if (!GLevelEditorModeTools().IsOnlyVisibleActiveMode(ModeID))
 		{
-			if( Mode->GetModeInfo().bVisible && Mode->GetID() != FBuiltinEditorModes::EM_Placement )
-			{
-				return false;
-			}
+			return false;
 		}
 	}
 	return GLevelEditorModeTools().IsModeActive( ModeID );
diff --git a/Engine/Source/Editor/LevelEditor/Private/SLevelEditorModeContent.cpp b/Engine/Source/Editor/LevelEditor/Private/SLevelEditorModeContent.cpp
--- a/Engine/Source/Editor/LevelEditor/Private/SLevelEditorModeContent.cpp
+++ b/Engine/Source/Editor/LevelEditor/Private/SLevelEditorModeContent.cpp
@@ -12,7 +12,7 @@
 
 SLevelEditorModeContent::~SLevelEditorModeContent()
 {
-	GLevelEditorModeTools().OnEditorModeChanged().RemoveAll( this );
+	GLevelEditorModeTools().OnEditorModeIDChanged().RemoveAll( this );
 	GetMutableDefault<UEditorPerProjectUserSettings>()->OnUserSettingChanged().RemoveAll( this );
 }
 
@@ -23,7 +23,7 @@ void SLevelEditorModeContent::Construct( const FArguments& InArgs, const TShared
 	EditorMode = InEditorMode;
 
 	InOwningDocTab->SetOnTabClosed( SDockTab::FOnTabClosedCallback::CreateSP(this, &SLevelEditorModeContent::HandleParentClosed ) );
-	GLevelEditorModeTools().OnEditorModeChanged().AddSP( this, &SLevelEditorModeContent::HandleEditorModeChanged );
+	GLevelEditorModeTools().OnEditorModeIDChanged().AddSP( this, &SLevelEditorModeContent::HandleEditorModeChanged );
 	GetMutableDefault<UEditorPerProjectUserSettings>()->OnUserSettingChanged().AddSP( this, &SLevelEditorModeContent::HandleUserSettingsChange );
 
 	ChildSlot
@@ -49,9 +49,9 @@ void SLevelEditorModeContent::Construct( const FArguments& InArgs, const TShared
 	UpdateModeToolBar();
 }
 
-void SLevelEditorModeContent::HandleEditorModeChanged( FEdMode* Mode, bool IsEnabled )
+void SLevelEditorModeContent::HandleEditorModeChanged(const FEditorModeID& EditorModeID, bool IsEnabled )
 {
-	if ( Mode == EditorMode && !IsEnabled )
+	if ( EditorModeID == EditorMode->GetID() && !IsEnabled )
 	{
 		DocTab.Pin()->SetOnTabClosed( SDockTab::FOnTabClosedCallback() );
 		DocTab.Pin()->RequestCloseTab();
diff --git a/Engine/Source/Editor/LevelEditor/Private/SLevelEditorModeContent.h b/Engine/Source/Editor/LevelEditor/Private/SLevelEditorModeContent.h
--- a/Engine/Source/Editor/LevelEditor/Private/SLevelEditorModeContent.h
+++ b/Engine/Source/Editor/LevelEditor/Private/SLevelEditorModeContent.h
@@ -12,6 +12,8 @@
 
 class FEdMode;
 class SBorder;
+class FName;
+typedef FName FEditorModeID;
 
 /**
  * Tools for the level editor                   
@@ -43,7 +45,7 @@ private:
 	 * @param Mode The mode that changed.
 	 * @param IsEnabled true if the mode is enabled, otherwise false.
 	 */
-	void HandleEditorModeChanged(FEdMode* Mode, bool IsEnabled);
+	void HandleEditorModeChanged(const FEditorModeID& EditorModeID, bool IsEnabled);
 
 	/** Gets the visibility for the SBorder showing toolbox editor mode inline content */
 	EVisibility GetInlineContentHolderVisibility() const;
diff --git a/Engine/Source/Editor/UnrealEd/Classes/Editor/UnrealEdEngine.h b/Engine/Source/Editor/UnrealEd/Classes/Editor/UnrealEdEngine.h
--- a/Engine/Source/Editor/UnrealEd/Classes/Editor/UnrealEdEngine.h
+++ b/Engine/Source/Editor/UnrealEd/Classes/Editor/UnrealEdEngine.h
@@ -27,6 +27,8 @@ class ITargetPlatform;
 class UPrimitiveComponent;
 class UTexture2D;
 class UUnrealEdOptions;
+class FName;
+typedef FName FEditorModeID;
 
 UENUM()
 enum EPackageNotifyState
@@ -828,7 +830,9 @@ public:
 	bool HandleDisasmScriptCommand( const TCHAR* Str, FOutputDevice& Ar );	
 
 	/** OnEditorModeChanged delegate which looks for Matinee editor closing */
+	UE_DEPRECATED(4.24, "Use UpdateEdModeOnMatineeClose() instead")
 	void OnMatineeEditorClosed( class FEdMode* Mode, bool IsEntering );
+	void UpdateEdModeOnMatineeClose(const FEditorModeID& EditorModeID, bool IsEntering);
 
 	bool IsComponentSelected(const UPrimitiveComponent* PrimComponent);
 
@@ -863,6 +867,8 @@ protected:
 
 	/** Handle to the registered OnMatineeEditorClosed delegate. */
 	FDelegateHandle OnMatineeEditorClosedDelegateHandle;
+	/** Handle to the registered UpdateEdModeOnMatineeClose delegate. */
+	FDelegateHandle UpdateEdModeOnMatineeCloseDelegateHandle;
 
 	/** Whether the pivot has been moved independently */
 	bool bPivotMovedIndependently;
diff --git a/Engine/Source/Editor/UnrealEd/Private/EdMode.cpp b/Engine/Source/Editor/UnrealEd/Private/EdMode.cpp
--- a/Engine/Source/Editor/UnrealEd/Private/EdMode.cpp
+++ b/Engine/Source/Editor/UnrealEd/Private/EdMode.cpp
@@ -589,16 +589,16 @@ void FEdMode::Enter()
 
 	bPendingDeletion = false;
 
-	FEditorDelegates::EditorModeEnter.Broadcast( this );
+	FEditorDelegates::EditorModeIDEnter.Broadcast( GetID() );
 	const bool bIsEnteringMode = true;
-	Owner->BroadcastEditorModeChanged( this, bIsEnteringMode );
+	Owner->BroadcastEditorModeIDChanged( GetID(), bIsEnteringMode );
 }
 
 void FEdMode::Exit()
 {
 	const bool bIsEnteringMode = false;
-	Owner->BroadcastEditorModeChanged( this, bIsEnteringMode );
-	FEditorDelegates::EditorModeExit.Broadcast( this );
+	Owner->BroadcastEditorModeIDChanged( GetID(), bIsEnteringMode );
+	FEditorDelegates::EditorModeIDExit.Broadcast(GetID());
 }
 
 void FEdMode::SetCurrentTool( EModeTools InID )
diff --git a/Engine/Source/Editor/UnrealEd/Private/Editor.cpp b/Engine/Source/Editor/UnrealEd/Private/Editor.cpp
--- a/Engine/Source/Editor/UnrealEd/Private/Editor.cpp
+++ b/Engine/Source/Editor/UnrealEd/Private/Editor.cpp
@@ -89,6 +89,8 @@ FSimpleMulticastDelegate								FEditorDelegates::LoadSelectedAssetsIfNeeded;
 FSimpleMulticastDelegate								FEditorDelegates::DisplayLoadErrors;
 FEditorDelegates::FOnEditorModeTransitioned				FEditorDelegates::EditorModeEnter;
 FEditorDelegates::FOnEditorModeTransitioned				FEditorDelegates::EditorModeExit;
+FEditorDelegates::FOnEditorModeIDTransitioned			FEditorDelegates::EditorModeIDEnter;
+FEditorDelegates::FOnEditorModeIDTransitioned			FEditorDelegates::EditorModeIDExit;
 FEditorDelegates::FOnPIEEvent							FEditorDelegates::PreBeginPIE;
 FEditorDelegates::FOnPIEEvent							FEditorDelegates::BeginPIE;
 FEditorDelegates::FOnPIEEvent							FEditorDelegates::PrePIEEnded;
diff --git a/Engine/Source/Editor/UnrealEd/Private/EditorActor.cpp b/Engine/Source/Editor/UnrealEd/Private/EditorActor.cpp
--- a/Engine/Source/Editor/UnrealEd/Private/EditorActor.cpp
+++ b/Engine/Source/Editor/UnrealEd/Private/EditorActor.cpp
@@ -620,14 +620,7 @@ void UUnrealEdEngine::edactDuplicateSelected( ULevel* InLevel, bool bOffsetLocat
 			PostDuplicateSelection.Add(Actor);
 		}
 
-		TArray<FEdMode*> ActiveModes;
-		GLevelEditorModeTools().GetActiveModes(ActiveModes);
-
-		for (int32 ModeIndex = 0; ModeIndex < ActiveModes.Num(); ++ModeIndex)
-		{
-			// Tell the tools about the duplication
-			ActiveModes[ModeIndex]->ActorsDuplicatedNotify(PreDuplicateSelection, PostDuplicateSelection, bOffsetLocations);
-		}
+		GLevelEditorModeTools().ActorsDuplicatedNotify(PreDuplicateSelection, PostDuplicateSelection, bOffsetLocations);
 	}
 }
 
diff --git a/Engine/Source/Editor/UnrealEd/Private/EditorCsg.cpp b/Engine/Source/Editor/UnrealEd/Private/EditorCsg.cpp
--- a/Engine/Source/Editor/UnrealEd/Private/EditorCsg.cpp
+++ b/Engine/Source/Editor/UnrealEd/Private/EditorCsg.cpp
@@ -1320,9 +1320,6 @@ void UEditorEngine::MapBrushGet(UWorld* InWorld)
 
 void UEditorEngine::mapBrushPut()
 {
-	TArray<FEdMode*> ActiveModes; 
-	GLevelEditorModeTools().GetActiveModes( ActiveModes );
-
 	for ( FSelectionIterator It( GEditor->GetSelectedActorIterator() ) ; It ; ++It )
 	{
 		AActor* Actor = static_cast<AActor*>( *It );
@@ -1342,10 +1339,7 @@ void UEditorEngine::mapBrushPut()
 
 			WorldBrush->ReregisterAllComponents();
 
-			for( int32 ModeIndex = 0; ModeIndex < ActiveModes.Num(); ++ModeIndex )
-			{
-				ActiveModes[ModeIndex]->UpdateInternalData();
-			}
+			GLevelEditorModeTools().UpdateInternalData();
 		}
 	}
 }
diff --git a/Engine/Source/Editor/UnrealEd/Private/EditorHook.cpp b/Engine/Source/Editor/UnrealEd/Private/EditorHook.cpp
--- a/Engine/Source/Editor/UnrealEd/Private/EditorHook.cpp
+++ b/Engine/Source/Editor/UnrealEd/Private/EditorHook.cpp
@@ -33,13 +33,7 @@ void UUnrealEdEngine::NotifyPreChange(UProperty* PropertyAboutToChange)
 void UUnrealEdEngine::NotifyPostChange(const FPropertyChangedEvent& PropertyChangedEvent, UProperty* PropertyThatChanged)
 {
 	// Notify all active modes of actor property changes.
-	TArray<FEdMode*> ActiveModes;
-	GLevelEditorModeTools().GetActiveModes(ActiveModes);
-
-	for (int32 ModeIndex = 0; ModeIndex < ActiveModes.Num(); ++ModeIndex)
-	{
-		ActiveModes[ModeIndex]->ActorPropChangeNotify();
-	}
+	GLevelEditorModeTools().ActorPropChangeNotify();
 }
 
 void UUnrealEdEngine::UpdateFloatingPropertyWindows(bool bForceRefresh)
diff --git a/Engine/Source/Editor/UnrealEd/Private/EditorModeManager.cpp b/Engine/Source/Editor/UnrealEd/Private/EditorModeManager.cpp
--- a/Engine/Source/Editor/UnrealEd/Private/EditorModeManager.cpp
+++ b/Engine/Source/Editor/UnrealEd/Private/EditorModeManager.cpp
@@ -28,6 +28,7 @@
 #include "Framework/Commands/UICommandList.h"
 #include "Framework/MultiBox/MultiBoxBuilder.h"
 #include "Toolkits/BaseToolkit.h"
+#include "Tools/UEdMode.h"
 
 /*------------------------------------------------------------------------------
 	FEditorModeTools.
@@ -159,6 +160,363 @@ bool FEditorModeTools::SelectionHasSceneComponent() const
 	return bSelectionHasSceneComponent;
 }
 
+bool FEditorModeTools::IsSelectionAllowed(AActor* InActor, const bool bInSelected) const
+{
+	bool bSelectionAllowed = ((ActiveModes.Num() == 0) && (ActiveScriptableModes.Num() == 0));
+	for (int32 ModeIndex = 0; ModeIndex < ActiveModes.Num(); ++ModeIndex)
+	{
+		bSelectionAllowed |= ActiveModes[ModeIndex]->IsSelectionAllowed(InActor, bInSelected);
+	}
+	for (int32 ModeIndex = 0; ModeIndex < ActiveScriptableModes.Num(); ++ModeIndex)
+	{
+		bSelectionAllowed |= ActiveScriptableModes[ModeIndex]->IsSelectionAllowed(InActor, bInSelected);
+	}
+
+	return bSelectionAllowed;
+}
+
+bool FEditorModeTools::IsSelectionHandled(AActor* InActor, const bool bInSelected) const
+{
+	bool bSelectionHandled = false;
+	for (int32 ModeIndex = 0; ModeIndex < ActiveModes.Num(); ++ModeIndex)
+	{
+		bSelectionHandled |= ActiveModes[ModeIndex]->Select(InActor, bInSelected);
+	}
+	for (int32 ModeIndex = 0; ModeIndex < ActiveScriptableModes.Num(); ++ModeIndex)
+	{
+		bSelectionHandled |= ActiveScriptableModes[ModeIndex]->Select(InActor, bInSelected);
+	}
+
+	return bSelectionHandled;
+}
+
+bool FEditorModeTools::ProcessEditDuplicate()
+{
+	bool bHandled = false;
+	for (int32 ModeIndex = 0; ModeIndex < ActiveModes.Num(); ++ModeIndex)
+	{
+		bHandled |= ActiveModes[ModeIndex]->ProcessEditDuplicate();
+	}
+	for (int32 ModeIndex = 0; ModeIndex < ActiveScriptableModes.Num(); ++ModeIndex)
+	{
+		bHandled |= ActiveScriptableModes[ModeIndex]->ProcessEditDuplicate();
+	}
+	return bHandled;
+}
+
+bool FEditorModeTools::ProcessEditDelete()
+{
+	bool bHandled = false;
+	for (int32 ModeIndex = 0; ModeIndex < ActiveModes.Num(); ++ModeIndex)
+	{
+		bHandled |= ActiveModes[ModeIndex]->ProcessEditDelete();
+	}
+	for (int32 ModeIndex = 0; ModeIndex < ActiveScriptableModes.Num(); ++ModeIndex)
+	{
+		bHandled |= ActiveScriptableModes[ModeIndex]->ProcessEditDelete();
+	}
+	return bHandled;
+}
+
+bool FEditorModeTools::ProcessEditCut()
+{
+	for (int32 ModeIndex = 0; ModeIndex < ActiveModes.Num(); ++ModeIndex)
+	{
+		if (ActiveModes[ModeIndex]->ProcessEditCut())
+		{
+			return true;
+		}
+	}
+	for (int32 ModeIndex = 0; ModeIndex < ActiveScriptableModes.Num(); ++ModeIndex)
+	{
+		if (ActiveScriptableModes[ModeIndex]->ProcessEditCut())
+		{
+			return true;
+		}
+	}
+	return false;
+}
+
+bool FEditorModeTools::ProcessEditCopy()
+{
+	for (int32 ModeIndex = 0; ModeIndex < ActiveModes.Num(); ++ModeIndex)
+	{
+		if (ActiveModes[ModeIndex]->ProcessEditCopy())
+		{
+			return true;
+		}
+	}
+	for (int32 ModeIndex = 0; ModeIndex < ActiveScriptableModes.Num(); ++ModeIndex)
+	{
+		if (ActiveScriptableModes[ModeIndex]->ProcessEditCopy())
+		{
+			return true;
+		}
+	}
+	return false;
+}
+
+bool FEditorModeTools::ProcessEditPaste()
+{
+	for (int32 ModeIndex = 0; ModeIndex < ActiveModes.Num(); ++ModeIndex)
+	{
+		if (ActiveModes[ModeIndex]->ProcessEditPaste())
+		{
+			return true;
+		}
+	}
+	for (int32 ModeIndex = 0; ModeIndex < ActiveScriptableModes.Num(); ++ModeIndex)
+	{
+		if (ActiveScriptableModes[ModeIndex]->ProcessEditPaste())
+		{
+			return true;
+		}
+	}
+	return false;
+}
+
+EEditAction::Type FEditorModeTools::GetActionEditDuplicate()
+{
+	for (int32 ModeIndex = 0; ModeIndex < ActiveModes.Num(); ++ModeIndex)
+	{
+		const EEditAction::Type CanProcess = ActiveModes[ModeIndex]->GetActionEditDuplicate();
+		if (CanProcess == EEditAction::Process || CanProcess == EEditAction::Halt)
+		{
+			return CanProcess;
+		}
+	}
+	for (int32 ModeIndex = 0; ModeIndex < ActiveScriptableModes.Num(); ++ModeIndex)
+	{
+		const EEditAction::Type CanProcess = ActiveScriptableModes[ModeIndex]->GetActionEditDuplicate();
+		if (CanProcess == EEditAction::Process || CanProcess == EEditAction::Halt)
+		{
+			return CanProcess;
+		}
+	}
+	return EEditAction::Skip;
+}
+
+EEditAction::Type FEditorModeTools::GetActionEditDelete()
+{
+	for (int32 ModeIndex = 0; ModeIndex < ActiveModes.Num(); ++ModeIndex)
+	{
+		const EEditAction::Type CanProcess = ActiveModes[ModeIndex]->GetActionEditDelete();
+		if (CanProcess == EEditAction::Process || CanProcess == EEditAction::Halt)
+		{
+			return CanProcess;
+		}
+	}
+	for (int32 ModeIndex = 0; ModeIndex < ActiveScriptableModes.Num(); ++ModeIndex)
+	{
+		const EEditAction::Type CanProcess = ActiveScriptableModes[ModeIndex]->GetActionEditDelete();
+		if (CanProcess == EEditAction::Process || CanProcess == EEditAction::Halt)
+		{
+			return CanProcess;
+		}
+	}
+	return EEditAction::Skip;
+}
+
+EEditAction::Type FEditorModeTools::GetActionEditCut()
+{
+	for (int32 ModeIndex = 0; ModeIndex < ActiveModes.Num(); ++ModeIndex)
+	{
+		const EEditAction::Type CanProcess = ActiveModes[ModeIndex]->GetActionEditCut();
+		if (CanProcess == EEditAction::Process || CanProcess == EEditAction::Halt)
+		{
+			return CanProcess;
+		}
+	}
+	for (int32 ModeIndex = 0; ModeIndex < ActiveScriptableModes.Num(); ++ModeIndex)
+	{
+		const EEditAction::Type CanProcess = ActiveScriptableModes[ModeIndex]->GetActionEditCut();
+		if (CanProcess == EEditAction::Process || CanProcess == EEditAction::Halt)
+		{
+			return CanProcess;
+		}
+	}
+	return EEditAction::Skip;
+}
+
+EEditAction::Type FEditorModeTools::GetActionEditCopy()
+{
+	for (int32 ModeIndex = 0; ModeIndex < ActiveModes.Num(); ++ModeIndex)
+	{
+		const EEditAction::Type CanProcess = ActiveModes[ModeIndex]->GetActionEditCopy();
+		if (CanProcess == EEditAction::Process || CanProcess == EEditAction::Halt)
+		{
+			return CanProcess;
+		}
+	}
+	for (int32 ModeIndex = 0; ModeIndex < ActiveScriptableModes.Num(); ++ModeIndex)
+	{
+		const EEditAction::Type CanProcess = ActiveScriptableModes[ModeIndex]->GetActionEditCopy();
+		if (CanProcess == EEditAction::Process || CanProcess == EEditAction::Halt)
+		{
+			return CanProcess;
+		}
+	}
+	return EEditAction::Skip;
+}
+
+EEditAction::Type FEditorModeTools::GetActionEditPaste()
+{
+	for (int32 ModeIndex = 0; ModeIndex < ActiveModes.Num(); ++ModeIndex)
+	{
+		const EEditAction::Type CanProcess = ActiveModes[ModeIndex]->GetActionEditPaste();
+		if (CanProcess == EEditAction::Process || CanProcess == EEditAction::Halt)
+		{
+			return CanProcess;
+		}
+	}
+	for (int32 ModeIndex = 0; ModeIndex < ActiveScriptableModes.Num(); ++ModeIndex)
+	{
+		const EEditAction::Type CanProcess = ActiveScriptableModes[ModeIndex]->GetActionEditPaste();
+		if (CanProcess == EEditAction::Process || CanProcess == EEditAction::Halt)
+		{
+			return CanProcess;
+		}
+	}
+	return EEditAction::Skip;
+}
+
+void FEditorModeTools::DeactivateOtherVisibleModes(FEditorModeID InMode)
+{
+	for (int32 ModeIndex = 0; ModeIndex < ActiveModes.Num(); ModeIndex++)
+	{
+		if (ActiveModes[ModeIndex]->GetID() != InMode && ActiveModes[ModeIndex]->GetModeInfo().bVisible)
+		{
+			DeactivateMode(ActiveModes[ModeIndex]->GetID());
+		}
+	}
+
+	for (int32 ModeIndex = 0; ModeIndex < ActiveScriptableModes.Num(); ModeIndex++)
+	{
+		if (ActiveScriptableModes[ModeIndex]->GetID() != InMode && ActiveScriptableModes[ModeIndex]->GetModeInfo().bVisible)
+		{
+			DeactivateMode(ActiveScriptableModes[ModeIndex]->GetID());
+		}
+	}
+}
+
+bool FEditorModeTools::IsSnapRotationEnabled() const
+{
+	bool bSnapRotationEnabled = false;
+	for (int32 ModeIndex = 0; ModeIndex < ActiveModes.Num(); ++ModeIndex)
+	{
+		bSnapRotationEnabled |= ActiveModes[ModeIndex]->IsSnapRotationEnabled();
+	}
+	for (int32 ModeIndex = 0; ModeIndex < ActiveScriptableModes.Num(); ++ModeIndex)
+	{
+		bSnapRotationEnabled |= ActiveScriptableModes[ModeIndex]->IsSnapRotationEnabled();
+	}
+	return bSnapRotationEnabled;
+}
+
+bool FEditorModeTools::SnapRotatorToGridOverride(FRotator& InRotation) const
+{
+	for (int32 ModeIndex = 0; ModeIndex < ActiveModes.Num(); ++ModeIndex)
+	{
+		if (ActiveModes[ModeIndex]->SnapRotatorToGridOverride(InRotation))
+		{
+			return true;
+		}
+	}
+	for (int32 ModeIndex = 0; ModeIndex < ActiveScriptableModes.Num(); ++ModeIndex)
+	{
+		if (ActiveScriptableModes[ModeIndex]->SnapRotatorToGridOverride(InRotation))
+		{
+			return true;
+		}
+	}
+	return false;
+}
+
+void FEditorModeTools::ActorsDuplicatedNotify(TArray<AActor*>& InPreDuplicateSelection, TArray<AActor*>& InPostDuplicateSelection, const bool bOffsetLocations)
+{
+	for (int32 ModeIndex = 0; ModeIndex < ActiveModes.Num(); ++ModeIndex)
+	{
+		// Tell the tools about the duplication
+		ActiveModes[ModeIndex]->ActorsDuplicatedNotify(InPreDuplicateSelection, InPostDuplicateSelection, bOffsetLocations);
+	}
+
+	for (int32 ModeIndex = 0; ModeIndex < ActiveScriptableModes.Num(); ++ModeIndex)
+	{
+		// Tell the tools about the duplication
+		ActiveScriptableModes[ModeIndex]->ActorsDuplicatedNotify(InPreDuplicateSelection, InPostDuplicateSelection, bOffsetLocations);
+	}
+}
+
+void FEditorModeTools::ActorMoveNotify()
+{
+	for (int32 ModeIndex = 0; ModeIndex < ActiveModes.Num(); ++ModeIndex)
+	{
+		// Also notify the current editing modes if they are interested.
+		ActiveModes[ModeIndex]->ActorMoveNotify();
+	}
+	for (int32 ModeIndex = 0; ModeIndex < ActiveScriptableModes.Num(); ++ModeIndex)
+	{
+		// Also notify the current editing modes if they are interested.
+		ActiveScriptableModes[ModeIndex]->ActorMoveNotify();
+	}
+}
+
+void FEditorModeTools::ActorSelectionChangeNotify()
+{
+	for (int32 ModeIndex = 0; ModeIndex < ActiveModes.Num(); ++ModeIndex)
+	{
+		ActiveModes[ModeIndex]->ActorSelectionChangeNotify();
+	}
+	for (int32 ModeIndex = 0; ModeIndex < ActiveScriptableModes.Num(); ++ModeIndex)
+	{
+		ActiveScriptableModes[ModeIndex]->ActorSelectionChangeNotify();
+	}
+}
+
+void FEditorModeTools::ActorPropChangeNotify()
+{
+	for (int32 ModeIndex = 0; ModeIndex < ActiveModes.Num(); ++ModeIndex)
+	{
+		ActiveModes[ModeIndex]->ActorPropChangeNotify();
+	}
+	for (int32 ModeIndex = 0; ModeIndex < ActiveScriptableModes.Num(); ++ModeIndex)
+	{
+		ActiveScriptableModes[ModeIndex]->ActorPropChangeNotify();
+	}
+}
+
+void FEditorModeTools::UpdateInternalData()
+{
+	for (int32 ModeIndex = 0; ModeIndex < ActiveModes.Num(); ++ModeIndex)
+	{
+		ActiveModes[ModeIndex]->UpdateInternalData();
+	}
+	for (int32 ModeIndex = 0; ModeIndex < ActiveScriptableModes.Num(); ++ModeIndex)
+	{
+		ActiveScriptableModes[ModeIndex]->UpdateInternalData();
+	}
+}
+
+bool FEditorModeTools::IsOnlyVisibleActiveMode(FEditorModeID InMode) const
+{
+	// Only return true if this is the *only* active mode
+	for (TSharedPtr<FEdMode> Mode : ActiveModes)
+	{
+		if (Mode->GetModeInfo().bVisible && Mode->GetID() != InMode)
+		{
+			return false;
+		}
+	}
+	for (UEdMode* Mode : ActiveScriptableModes)
+	{
+		if (Mode->GetModeInfo().bVisible && Mode->GetID() != InMode)
+		{
+			return false;
+		}
+	}
+	return true;
+}
+
 void FEditorModeTools::UpdateModeWidgetLocation()
 {
 	if (!bIsTracking)
@@ -188,6 +546,7 @@ void FEditorModeTools::UpdateModeWidgetLocation()
 						}
 
 					}
+					// lbarnes: need to handle UEdMode gizmos here
 				}
 				else
 				{
@@ -343,22 +702,44 @@ void FEditorModeTools::DeactivateModeAtIndex(int32 InIndex)
 {
 	check( InIndex >= 0 && InIndex < ActiveModes.Num() );
 
-	auto& Mode = ActiveModes[InIndex];
+	TSharedPtr<FEdMode> Mode = ActiveModes[InIndex];
 
 	Mode->Exit();
 
 	// Remove the toolbar widget
 	ActiveToolBarRows.RemoveAll(
 		[&Mode](FEdModeToolbarRow& Row)
-		{ 
-			return Row.Mode == Mode;
-		}
+	{
+		return Row.ModeID == Mode->GetID();
+	}
+	);
+
+	RebuildModeToolBar();
+
+	RecycledModes.Add(Mode->GetID(), Mode);
+	ActiveModes.RemoveAt(InIndex);
+}
+
+void FEditorModeTools::DeactivateScriptableModeAtIndex(int32 InIndex)
+{
+	check(InIndex >= 0 && InIndex < ActiveScriptableModes.Num());
+
+	UEdMode* Mode = ActiveScriptableModes[InIndex];
+
+	Mode->Exit();
+
+	// Remove the toolbar widget
+	ActiveToolBarRows.RemoveAll(
+		[&Mode](FEdModeToolbarRow& Row)
+	{
+		return Row.ModeID == Mode->GetID();
+	}
 	);
 
 	RebuildModeToolBar();
 
-	RecycledModes.Add( Mode->GetID(), Mode );
-	ActiveModes.RemoveAt( InIndex );
+	RecycledScriptableModes.Add(Mode->GetID(), Mode);
+	ActiveScriptableModes.RemoveAt(InIndex);
 }
 
 void FEditorModeTools::RebuildModeToolBar()
@@ -422,7 +803,19 @@ void FEditorModeTools::DeactivateMode( FEditorModeID InID )
 		}
 	}
 
-	if( ActiveModes.Num() == 0 )
+	// Find the mode from the ID and exit it.
+	for (int32 Index = ActiveScriptableModes.Num() - 1; Index >= 0; --Index)
+	{
+		auto& Mode = ActiveScriptableModes[Index];
+		if (Mode->GetID() == InID)
+		{
+			DeactivateScriptableModeAtIndex(Index);
+			break;
+		}
+	}
+
+
+	if( (ActiveModes.Num() == 0) && (ActiveScriptableModes.Num() == 0))
 	{
 		// Ensure the default mode is active if there are no active modes.
 		ActivateDefaultMode();
@@ -435,6 +828,16 @@ void FEditorModeTools::DeactivateAllModes()
 	{
 		DeactivateModeAtIndex(Index);
 	}
+
+	for (int32 Index = ActiveScriptableModes.Num() - 1; Index >= 0; --Index)
+	{
+		DeactivateScriptableModeAtIndex(Index);
+	}
+}
+
+FEdMode* FEditorModeTools::FindMode(FEditorModeID InID)
+{
+	return GetActiveMode(InID);
 }
 
 void FEditorModeTools::DestroyMode( FEditorModeID InID )
@@ -451,7 +854,20 @@ void FEditorModeTools::DestroyMode( FEditorModeID InID )
 		}
 	}
 
+	// Find the mode from the ID and exit it.
+	for (int32 Index = ActiveScriptableModes.Num() - 1; Index >= 0; --Index)
+	{
+		auto& Mode = ActiveScriptableModes[Index];
+		if (Mode->GetID() == InID)
+		{
+			// Deactivate and destroy
+			DeactivateScriptableModeAtIndex(Index);
+			break;
+		}
+	}
+
 	RecycledModes.Remove(InID);
+	RecycledScriptableModes.Remove(InID);
 }
 
 TSharedRef<SDockTab> FEditorModeTools::MakeModeToolbarTab()
@@ -516,59 +932,142 @@ void FEditorModeTools::ActivateMode(FEditorModeID InID, bool bToggle)
 		return;
 	}
 
-	// Recycle a mode or factory a new one
-	TSharedPtr<FEdMode> Mode = RecycledModes.FindRef(InID);
-
-	if (Mode.IsValid())
+	const TSharedRef<IEditorModeFactory>* ModeFactory = FEditorModeRegistry::Get().GetFactoryMap().Find(InID);
+	if (ModeFactory)
 	{
-		RecycledModes.Remove(InID);
-	}
-	else
-	{
-		Mode = FEditorModeRegistry::Get().CreateMode(InID, *this);
-	}
+		if (!(*ModeFactory)->ForScriptableMode())
+		{
+			// Recycle a mode or factory a new one
+			TSharedPtr<FEdMode> Mode = RecycledModes.FindRef(InID);
 
-	if (!Mode.IsValid())
-	{
-		UE_LOG(LogEditorModes, Log, TEXT("FEditorModeTools::ActivateMode : Couldn't find mode '%s'."), *InID.ToString());
-		// Just return and leave the mode list unmodified
-		return;
-	}
+			if (Mode.IsValid())
+			{
+				RecycledModes.Remove(InID);
+			}
+			else
+			{
+				Mode = FEditorModeRegistry::Get().CreateMode(InID, *this);
+			}
+
+			if (!Mode.IsValid())
+			{
+				UE_LOG(LogEditorModes, Log, TEXT("FEditorModeTools::ActivateMode : Couldn't find mode '%s'."), *InID.ToString());
+				// Just return and leave the mode list unmodified
+				return;
+			}
+
+			// Remove anything that isn't compatible with this mode
+			for (int32 ModeIndex = ActiveModes.Num() - 1; ModeIndex >= 0; --ModeIndex)
+			{
+				const bool bModesAreCompatible = Mode->IsCompatibleWith(ActiveModes[ModeIndex]->GetID()) || ActiveModes[ModeIndex]->IsCompatibleWith(Mode->GetID());
+				if (!bModesAreCompatible)
+				{
+					DeactivateModeAtIndex(ModeIndex);
+				}
+			}
+			// Remove anything that isn't compatible with this mode
+			for (int32 ModeIndex = ActiveScriptableModes.Num() - 1; ModeIndex >= 0; --ModeIndex)
+			{
+				const bool bModesAreCompatible = Mode->IsCompatibleWith(ActiveScriptableModes[ModeIndex]->GetID()) || ActiveScriptableModes[ModeIndex]->IsCompatibleWith(Mode->GetID());
+				if (!bModesAreCompatible)
+				{
+					DeactivateScriptableModeAtIndex(ModeIndex);
+				}
+			}
+
+			ActiveModes.Add(Mode);
+			// Enter the new mode
+			Mode->Enter();
+
+			// Ask the mode to build the toolbar.
+			TSharedPtr<FUICommandList> CommandList;
+			const TSharedPtr<FModeToolkit> Toolkit = Mode->GetToolkit();
+			if (Toolkit.IsValid())
+			{
+				CommandList = Toolkit->GetToolkitCommands();
+			}
+
+			FToolBarBuilder ModeToolbarBuilder(CommandList, FMultiBoxCustomization(Mode->GetModeInfo().ToolbarCustomizationName), TSharedPtr<FExtender>(), Orient_Horizontal, false);
+			Mode->BuildModeToolbar(ModeToolbarBuilder);
+
+			if (ModeToolbarBuilder.GetMultiBox()->GetBlocks().Num())
+			{
+				TSharedRef<SWidget> ToolbarWidget = ModeToolbarBuilder.MakeWidget();
+				ActiveToolBarRows.Emplace(Mode->GetID(), ToolbarWidget);
+
+				SpawnOrUpdateModeToolbar();
+			}
 
-	// Remove anything that isn't compatible with this mode
-	for (int32 ModeIndex = ActiveModes.Num() - 1; ModeIndex >= 0; --ModeIndex)
-	{
-		const bool bModesAreCompatible = Mode->IsCompatibleWith(ActiveModes[ModeIndex]->GetID()) || ActiveModes[ModeIndex]->IsCompatibleWith(Mode->GetID());
-		if (!bModesAreCompatible)
-		{
-			DeactivateModeAtIndex(ModeIndex);
 		}
-	}
+		else
+		{
+			// Recycle a mode or factory a new one
+			UEdMode* Mode = RecycledScriptableModes.FindRef(InID);
 
-	ActiveModes.Add(Mode);
+			if (Mode)
+			{
+				RecycledScriptableModes.Remove(InID);
+			}
+			else
+			{
+				Mode = FEditorModeRegistry::Get().CreateScriptableMode(InID, *this);
+			}
 
-	// Enter the new mode
-	Mode->Enter();
+			if (!Mode)
+			{
+				UE_LOG(LogEditorModes, Log, TEXT("FEditorModeTools::ActivateMode : Couldn't find mode '%s'."), *InID.ToString());
+				// Just return and leave the mode list unmodified
+				return;
+			}
 
-	// Ask the mode to build the toolbar.
-	TSharedPtr<FUICommandList> CommandList;
-	const TSharedPtr<FModeToolkit> Toolkit = Mode->GetToolkit();
-	if (Toolkit.IsValid())
-	{
-		CommandList = Toolkit->GetToolkitCommands();
-	}
+			// Remove anything that isn't compatible with this mode
+			for (int32 ModeIndex = ActiveModes.Num() - 1; ModeIndex >= 0; --ModeIndex)
+			{
+				const bool bModesAreCompatible = Mode->IsCompatibleWith(ActiveModes[ModeIndex]->GetID()) || ActiveModes[ModeIndex]->IsCompatibleWith(Mode->GetID());
+				if (!bModesAreCompatible)
+				{
+					DeactivateModeAtIndex(ModeIndex);
+				}
+			}
+			// Remove anything that isn't compatible with this mode
+			for (int32 ModeIndex = ActiveScriptableModes.Num() - 1; ModeIndex >= 0; --ModeIndex)
+			{
+				const bool bModesAreCompatible = Mode->IsCompatibleWith(ActiveScriptableModes[ModeIndex]->GetID()) || ActiveScriptableModes[ModeIndex]->IsCompatibleWith(Mode->GetID());
+				if (!bModesAreCompatible)
+				{
+					DeactivateScriptableModeAtIndex(ModeIndex);
+				}
+			}
 
-	FToolBarBuilder ModeToolbarBuilder(CommandList, FMultiBoxCustomization(Mode->GetModeInfo().ToolbarCustomizationName), TSharedPtr<FExtender>(), Orient_Horizontal, false);
-	Mode->BuildModeToolbar(ModeToolbarBuilder);
+			ActiveScriptableModes.Add(Mode);
+			// Enter the new mode
+			Mode->Enter();
 
-	if (ModeToolbarBuilder.GetMultiBox()->GetBlocks().Num())
-	{
-		TSharedRef<SWidget> ToolbarWidget = ModeToolbarBuilder.MakeWidget();
-		ActiveToolBarRows.Emplace(Mode, ToolbarWidget);
+			// Ask the mode to build the toolbar.
+			TSharedPtr<FUICommandList> CommandList;
+			const TSharedPtr<FModeToolkit> Toolkit = Mode->GetToolkit();
+			if (Toolkit.IsValid())
+			{
+				CommandList = Toolkit->GetToolkitCommands();
+			}
+
+			FToolBarBuilder ModeToolbarBuilder(CommandList, FMultiBoxCustomization(Mode->GetModeInfo().ToolbarCustomizationName), TSharedPtr<FExtender>(), Orient_Horizontal, false);
+			Mode->BuildModeToolbar(ModeToolbarBuilder);
+
+			if (ModeToolbarBuilder.GetMultiBox()->GetBlocks().Num())
+			{
+				TSharedRef<SWidget> ToolbarWidget = ModeToolbarBuilder.MakeWidget();
+				ActiveToolBarRows.Emplace(Mode->GetID(), ToolbarWidget);
 
-		SpawnOrUpdateModeToolbar();
+				SpawnOrUpdateModeToolbar();
+			}
+		}
 	}
 
+
+
+
+
 	// Update the editor UI
 	FEditorSupportDelegates::UpdateUI.Broadcast();
 }
@@ -593,13 +1092,13 @@ bool FEditorModeTools::EnsureNotInMode(FEditorModeID ModeID, const FText& ErrorM
 	return bInASafeMode;
 }
 
-FEdMode* FEditorModeTools::FindMode( FEditorModeID InID )
+UEdMode* FEditorModeTools::GetActiveScriptableMode(FEditorModeID InID)
 {
-	for( auto& Mode : ActiveModes )
+	for (auto& Mode : ActiveScriptableModes)
 	{
-		if( Mode->GetID() == InID )
+		if (Mode->GetID() == InID)
 		{
-			return Mode.Get();
+			return Mode;
 		}
 	}
 
@@ -702,6 +1201,11 @@ bool FEditorModeTools::StartTracking(FEditorViewportClient* InViewportClient, FV
 		const TSharedPtr<FEdMode>& Mode = ActiveModes[ ModeIndex ];
 		bTransactionHandled |= Mode->StartTracking(InViewportClient, InViewport);
 	}
+	for (int32 ModeIndex = 0; ModeIndex < ActiveScriptableModes.Num(); ++ModeIndex)
+	{
+		UEdMode* Mode = ActiveScriptableModes[ModeIndex];
+		bTransactionHandled |= Mode->StartTracking(InViewportClient, InViewportClient->Viewport);
+	}
 
 	return bTransactionHandled;
 }
@@ -717,6 +1221,11 @@ bool FEditorModeTools::EndTracking(FEditorViewportClient* InViewportClient, FVie
 		const TSharedPtr<FEdMode>& Mode = ActiveModes[ ModeIndex ];
 		bTransactionHandled |= Mode->EndTracking(InViewportClient, InViewportClient->Viewport);
 	}
+	for (int32 ModeIndex = 0; ModeIndex < ActiveScriptableModes.Num(); ++ModeIndex)
+	{
+		UEdMode* Mode = ActiveScriptableModes[ModeIndex];
+		bTransactionHandled |= Mode->EndTracking(InViewportClient, InViewportClient->Viewport);
+	}
 
 	CachedLocation = PivotLocation;	// Clear the pivot location
 	
@@ -741,6 +1250,11 @@ void FEditorModeTools::MapChangeNotify()
 		const TSharedPtr<FEdMode>& Mode = ActiveModes[ ModeIndex ];
 		Mode->MapChangeNotify();
 	}
+	for (int32 ModeIndex = 0; ModeIndex < ActiveScriptableModes.Num(); ++ModeIndex)
+	{
+		UEdMode* Mode = ActiveScriptableModes[ModeIndex];
+		Mode->MapChangeNotify();
+	}
 }
 
 
@@ -847,7 +1361,12 @@ bool FEditorModeTools::ShouldDrawBrushWireframe( AActor* InActor ) const
 		bShouldDraw |= Mode->ShouldDrawBrushWireframe( InActor );
 	}
 
-	if( ActiveModes.Num() == 0 )
+	for (const auto& Mode : ActiveScriptableModes)
+	{
+		bShouldDraw |= Mode->ShouldDrawBrushWireframe(InActor);
+	}
+
+	if(( ActiveModes.Num() == 0 ) && (ActiveScriptableModes.Num() == 0))
 	{
 		// We can get into a state where there are no active modes at editor startup if the builder brush is created before the default mode is activated.
 		// Ensure we can see the builder brush when no modes are active.
@@ -875,18 +1394,36 @@ void FEditorModeTools::Tick( FEditorViewportClient* ViewportClient, float DeltaT
 		}
 	}
 	
-	if (ActiveModes.Num() == 0)
+	// Remove anything pending destruction
+	for (int32 Index = ActiveScriptableModes.Num() - 1; Index >= 0; --Index)
+	{
+		if (ActiveScriptableModes[Index]->IsPendingDeletion())
+		{
+			DeactivateScriptableModeAtIndex(Index);
+		}
+	}
+
+
+	if ((ActiveModes.Num() == 0) && ActiveScriptableModes.Num() == 0)
 	{
 		// Ensure the default mode is active if there are no active modes.
 		ActivateDefaultMode();
 	}
 
+
 	for( int32 ModeIndex = 0; ModeIndex < ActiveModes.Num(); ++ModeIndex )
 	{
 		const TSharedPtr<FEdMode>& Mode = ActiveModes[ ModeIndex ];
 		Mode->Tick( ViewportClient, DeltaTime );
 	}
 
+
+	for (int32 ModeIndex = 0; ModeIndex < ActiveScriptableModes.Num(); ++ModeIndex)
+	{
+		UEdMode* Mode = ActiveScriptableModes[ModeIndex];
+		Mode->Tick(ViewportClient, DeltaTime);
+	}
+
 	UpdateModeWidgetLocation();
 }
 
@@ -899,6 +1436,12 @@ bool FEditorModeTools::InputDelta( FEditorViewportClient* InViewportClient,FView
 		const TSharedPtr<FEdMode>& Mode = ActiveModes[ ModeIndex ];
 		bHandled |= Mode->InputDelta( InViewportClient, InViewport, InDrag, InRot, InScale );
 	}
+
+	for (int32 ModeIndex = 0; ModeIndex < ActiveScriptableModes.Num(); ++ModeIndex)
+	{
+		UEdMode* Mode = ActiveScriptableModes[ModeIndex];
+		bHandled |= Mode->InputDelta(InViewportClient, InViewport, InDrag, InRot, InScale);
+	}
 	return bHandled;
 }
 
@@ -911,6 +1454,11 @@ bool FEditorModeTools::CapturedMouseMove( FEditorViewportClient* InViewportClien
 		const TSharedPtr<FEdMode>& Mode = ActiveModes[ ModeIndex ];
 		bHandled |= Mode->CapturedMouseMove( InViewportClient, InViewport, InMouseX, InMouseY );
 	}
+	for (int32 ModeIndex = 0; ModeIndex < ActiveScriptableModes.Num(); ++ModeIndex)
+	{
+		UEdMode* Mode = ActiveScriptableModes[ModeIndex];
+		bHandled |= Mode->CapturedMouseMove(InViewportClient, InViewport, InMouseX, InMouseY);
+	}
 	return bHandled;
 }
 
@@ -923,6 +1471,11 @@ bool FEditorModeTools::ProcessCapturedMouseMoves( FEditorViewportClient* InViewp
 		const TSharedPtr<FEdMode>& Mode = ActiveModes[ ModeIndex ];
 		bHandled |= Mode->ProcessCapturedMouseMoves( InViewportClient, InViewport, CapturedMouseMoves );
 	}
+	for (int32 ModeIndex = 0; ModeIndex < ActiveScriptableModes.Num(); ++ModeIndex)
+	{
+		UEdMode* Mode = ActiveScriptableModes[ModeIndex];
+		bHandled |= Mode->ProcessCapturedMouseMoves(InViewportClient, InViewport, CapturedMouseMoves);
+	}
 	return bHandled;
 }
 
@@ -935,6 +1488,11 @@ bool FEditorModeTools::InputKey(FEditorViewportClient* InViewportClient, FViewpo
 		const TSharedPtr<FEdMode>& Mode = ActiveModes[ ModeIndex ];
 		bHandled |= Mode->InputKey( InViewportClient, Viewport, Key, Event );
 	}
+	for (int32 ModeIndex = 0; ModeIndex < ActiveScriptableModes.Num(); ++ModeIndex)
+	{
+		UEdMode* Mode = ActiveScriptableModes[ModeIndex];
+		bHandled |= Mode->InputKey(InViewportClient, Viewport, Key, Event);
+	}
 	return bHandled;
 }
 
@@ -947,6 +1505,11 @@ bool FEditorModeTools::InputAxis(FEditorViewportClient* InViewportClient, FViewp
 		const TSharedPtr<FEdMode>& Mode = ActiveModes[ ModeIndex ];
 		bHandled |= Mode->InputAxis( InViewportClient, Viewport, ControllerId, Key, Delta, DeltaTime );
 	}
+	for (int32 ModeIndex = 0; ModeIndex < ActiveScriptableModes.Num(); ++ModeIndex)
+	{
+		UEdMode* Mode = ActiveScriptableModes[ModeIndex];
+		bHandled |= Mode->InputAxis(InViewportClient, Viewport, ControllerId, Key, Delta, DeltaTime);
+	}
 	return bHandled;
 }
 
@@ -961,6 +1524,14 @@ bool FEditorModeTools::GetPivotForOrbit( FVector& Pivot ) const
 			return true;
 		}
 	}
+	for (int32 ModeIndex = 0; ModeIndex < ActiveScriptableModes.Num(); ++ModeIndex)
+	{
+		UEdMode* Mode = ActiveScriptableModes[ModeIndex];
+		if (Mode->GetPivotForOrbit(Pivot))
+		{
+			return true;
+		}
+	}
 	return false;
 }
 
@@ -972,6 +1543,11 @@ bool FEditorModeTools::MouseEnter( FEditorViewportClient* InViewportClient, FVie
 		const TSharedPtr<FEdMode>& Mode = ActiveModes[ ModeIndex ];
 		bHandled |= Mode->MouseEnter( InViewportClient, Viewport, X, Y );
 	}
+	for (int32 ModeIndex = 0; ModeIndex < ActiveScriptableModes.Num(); ++ModeIndex)
+	{
+		UEdMode* Mode = ActiveScriptableModes[ModeIndex];
+		bHandled |= Mode->MouseEnter(InViewportClient, Viewport, X, Y);
+	}
 	return bHandled;
 }
 
@@ -983,6 +1559,11 @@ bool FEditorModeTools::MouseLeave( FEditorViewportClient* InViewportClient, FVie
 		const TSharedPtr<FEdMode>& Mode = ActiveModes[ ModeIndex ];
 		bHandled |= Mode->MouseLeave( InViewportClient, Viewport );
 	}
+	for (int32 ModeIndex = 0; ModeIndex < ActiveScriptableModes.Num(); ++ModeIndex)
+	{
+		UEdMode* Mode = ActiveScriptableModes[ModeIndex];
+		bHandled |= Mode->MouseLeave(InViewportClient, Viewport);
+	}
 	return bHandled;
 }
 
@@ -995,6 +1576,11 @@ bool FEditorModeTools::MouseMove( FEditorViewportClient* InViewportClient, FView
 		const TSharedPtr<FEdMode>& Mode = ActiveModes[ ModeIndex ];
 		bHandled |= Mode->MouseMove( InViewportClient, Viewport, X, Y );
 	}
+	for (int32 ModeIndex = 0; ModeIndex < ActiveScriptableModes.Num(); ++ModeIndex)
+	{
+		UEdMode* Mode = ActiveScriptableModes[ModeIndex];
+		bHandled |= Mode->MouseMove(InViewportClient, Viewport, X, Y);
+	}
 	return bHandled;
 }
 
@@ -1006,6 +1592,11 @@ bool FEditorModeTools::ReceivedFocus( FEditorViewportClient* InViewportClient, F
 		const TSharedPtr<FEdMode>& Mode = ActiveModes[ ModeIndex ];
 		bHandled |= Mode->ReceivedFocus( InViewportClient, Viewport );
 	}
+	for (int32 ModeIndex = 0; ModeIndex < ActiveScriptableModes.Num(); ++ModeIndex)
+	{
+		UEdMode* Mode = ActiveScriptableModes[ModeIndex];
+		bHandled |= Mode->ReceivedFocus(InViewportClient, Viewport);
+	}
 	return bHandled;
 }
 
@@ -1017,6 +1608,11 @@ bool FEditorModeTools::LostFocus( FEditorViewportClient* InViewportClient, FView
 		const TSharedPtr<FEdMode>& Mode = ActiveModes[ ModeIndex ];
 		bHandled |= Mode->LostFocus( InViewportClient, Viewport );
 	}
+	for (int32 ModeIndex = 0; ModeIndex < ActiveScriptableModes.Num(); ++ModeIndex)
+	{
+		UEdMode* Mode = ActiveScriptableModes[ModeIndex];
+		bHandled |= Mode->LostFocus(InViewportClient, Viewport);
+	}
 	return bHandled;
 }
 
@@ -1027,6 +1623,10 @@ void FEditorModeTools::DrawActiveModes( const FSceneView* InView, FPrimitiveDraw
 	{
 		Mode->Draw( InView, PDI );
 	}
+	for (const auto& Mode : ActiveScriptableModes)
+	{
+		Mode->Draw(InView, PDI);
+	}
 }
 
 /** Renders all active modes */
@@ -1036,6 +1636,10 @@ void FEditorModeTools::Render( const FSceneView* InView, FViewport* Viewport, FP
 	{
 		Mode->Render( InView, Viewport, PDI );
 	}
+	for (const auto& Mode : ActiveScriptableModes)
+	{
+		Mode->Render(InView, Viewport, PDI);
+	}
 }
 
 /** Draws the HUD for all active modes */
@@ -1045,6 +1649,10 @@ void FEditorModeTools::DrawHUD( FEditorViewportClient* InViewportClient,FViewpor
 	{
 		Mode->DrawHUD( InViewportClient, Viewport, View, Canvas );
 	}
+	for (const auto& Mode : ActiveScriptableModes)
+	{
+		Mode->DrawHUD(InViewportClient, Viewport, View, Canvas);
+	}
 }
 
 /** Calls PostUndo on all active modes */
@@ -1057,6 +1665,12 @@ void FEditorModeTools::PostUndo(bool bSuccess)
 			const TSharedPtr<FEdMode>& Mode = ActiveModes[ ModeIndex ];
 			Mode->PostUndo();
 		}
+
+		for (int32 ModeIndex = 0; ModeIndex < ActiveScriptableModes.Num(); ++ModeIndex)
+		{
+			UEdMode* Mode = ActiveScriptableModes[ModeIndex];
+			Mode->PostUndo();
+		}
 	}
 }
 void FEditorModeTools::PostRedo(bool bSuccess)
@@ -1082,6 +1696,10 @@ bool FEditorModeTools::DisallowMouseDeltaTracking() const
 	{
 		bDisallow |= Mode->DisallowMouseDeltaTracking();
 	}
+	for (const auto& Mode : ActiveScriptableModes)
+	{
+		bDisallow |= Mode->DisallowMouseDeltaTracking();
+	}
 	return bDisallow;
 }
 
@@ -1092,6 +1710,10 @@ bool FEditorModeTools::GetCursor(EMouseCursor::Type& OutCursor) const
 	{
 		bHandled |= Mode->GetCursor(OutCursor);
 	}
+	for (const auto& Mode : ActiveScriptableModes)
+	{
+		bHandled |= Mode->GetCursor(OutCursor);
+	}
 	return bHandled;
 }
 
@@ -1102,6 +1724,10 @@ bool FEditorModeTools::GetOverrideCursorVisibility(bool& bWantsOverride, bool& b
 	{
 		bHandled |= Mode->GetOverrideCursorVisibility(bWantsOverride, bHardwareCursorVisible, bSoftwareCursorVisible);
 	}
+	for (const auto& Mode : ActiveScriptableModes)
+	{
+		bHandled |= Mode->GetOverrideCursorVisibility(bWantsOverride, bHardwareCursorVisible, bSoftwareCursorVisible);
+	}
 	return bHandled;
 }
 
@@ -1112,6 +1738,10 @@ bool FEditorModeTools::PreConvertMouseMovement(FEditorViewportClient* InViewport
 	{
 		bHandled |= Mode->PreConvertMouseMovement(InViewportClient);
 	}
+	for (const auto& Mode : ActiveScriptableModes)
+	{
+		bHandled |= Mode->PreConvertMouseMovement(InViewportClient);
+	}
 	return bHandled;
 }
 
@@ -1122,9 +1752,24 @@ bool FEditorModeTools::PostConvertMouseMovement(FEditorViewportClient* InViewpor
 	{
 		bHandled |= Mode->PostConvertMouseMovement(InViewportClient);
 	}
+	for (const auto& Mode : ActiveScriptableModes)
+	{
+		bHandled |= Mode->PostConvertMouseMovement(InViewportClient);
+	}
 	return bHandled;
 }
 
+bool FEditorModeTools::GetShowWidget() const
+{
+	bool bDrawModeSupportsWidgetDrawing = false;
+	// Check to see of any active modes support widget drawing
+	for (int32 ModeIndex = 0; ModeIndex < ActiveModes.Num(); ++ModeIndex)
+	{
+		bDrawModeSupportsWidgetDrawing |= ActiveModes[ModeIndex]->ShouldDrawWidget();
+	}
+	return bDrawModeSupportsWidgetDrawing && bShowWidget;
+}
+
 /**
  * Used to cycle widget modes
  */
@@ -1301,6 +1946,8 @@ void FEditorModeTools::AddReferencedObjects( FReferenceCollector& Collector )
 	{
 		ActiveModes[x]->AddReferencedObjects( Collector );
 	}
+
+	Collector.AddReferencedObjects(ActiveScriptableModes);
 }
 
 FEdMode* FEditorModeTools::GetActiveMode( FEditorModeID InID )
@@ -1376,6 +2023,5 @@ bool FEditorModeTools::CanCycleWidgetMode() const
 			return true;
 		}
 	}
-
 	return false;
 }
diff --git a/Engine/Source/Editor/UnrealEd/Private/EditorModeRegistry.cpp b/Engine/Source/Editor/UnrealEd/Private/EditorModeRegistry.cpp
--- a/Engine/Source/Editor/UnrealEd/Private/EditorModeRegistry.cpp
+++ b/Engine/Source/Editor/UnrealEd/Private/EditorModeRegistry.cpp
@@ -17,6 +17,7 @@
 #include "Editor/TextureAlignMode/Public/TextureAlignEdMode.h"
 #include "Editor/FoliageEdit/Public/FoliageEditModule.h"
 #include "Editor/VirtualTexturingEditor/Public/VirtualTexturingEditorModule.h"
+#include "Tools/UEdMode.h"
 
 FEditorModeInfo::FEditorModeInfo()
 	: ID(NAME_None)
@@ -134,6 +135,29 @@ TSharedPtr<FEdMode> FEditorModeRegistry::CreateMode(FEditorModeID ModeID, FEdito
 	return nullptr;
 }
 
+UEdMode* FEditorModeRegistry::CreateScriptableMode(FEditorModeID ModeID, FEditorModeTools& Owner)
+{
+	const TSharedRef<IEditorModeFactory>* ModeFactory = ModeFactories.Find(ModeID);
+	if (ModeFactory)
+	{
+		UEdMode* Instance = (*ModeFactory)->CreateScriptableMode();
+
+		// Assign the mode info from the factory before we initialize
+		Instance->Info = (*ModeFactory)->GetModeInfo();
+		Instance->Owner = &Owner;
+
+		// This binding ensures the mode is destroyed if the type is unregistered
+		OnModeUnregistered().AddUObject(Instance, &UEdMode::OnModeUnregistered);
+
+		Instance->Initialize();
+
+		return Instance;
+	}
+
+	return nullptr;
+}
+
+
 void FEditorModeRegistry::RegisterMode(FEditorModeID ModeID, TSharedRef<IEditorModeFactory> Factory)
 {
 	check(ModeID != FBuiltinEditorModes::EM_None);
diff --git a/Engine/Source/Editor/UnrealEd/Private/EditorSelectUtils.cpp b/Engine/Source/Editor/UnrealEd/Private/EditorSelectUtils.cpp
--- a/Engine/Source/Editor/UnrealEd/Private/EditorSelectUtils.cpp
+++ b/Engine/Source/Editor/UnrealEd/Private/EditorSelectUtils.cpp
@@ -391,12 +391,7 @@ void UUnrealEdEngine::NoteSelectionChange(bool bNotify)
 	// Clear active editing visualizer on selection change
 	ComponentVisManager.ClearActiveComponentVis();
 
-	TArray<FEdMode*> ActiveModes;
-	GLevelEditorModeTools().GetActiveModes( ActiveModes );
-	for( int32 ModeIndex = 0; ModeIndex < ActiveModes.Num(); ++ModeIndex )
-	{
-		ActiveModes[ModeIndex]->ActorSelectionChangeNotify();
-	}
+	GLevelEditorModeTools().ActorSelectionChangeNotify();
 
 	const bool bComponentSelectionChanged = GetSelectedComponentCount() > 0;
 	if (bNotify)
@@ -529,16 +524,8 @@ bool UUnrealEdEngine::CanSelectActor(AActor* Actor, bool bInSelected, bool bSele
 		return false;
 	}
 
-	// Allow active modes to determine whether the selection is allowed. If there are no active modes, allow selection anyway.
-	TArray<FEdMode*> ActiveModes;
-	GLevelEditorModeTools().GetActiveModes( ActiveModes );
-	bool bSelectionAllowed = (ActiveModes.Num() == 0);
-	for( int32 ModeIndex = 0; ModeIndex < ActiveModes.Num(); ++ModeIndex )
-	{
-		bSelectionAllowed |= ActiveModes[ModeIndex]->IsSelectionAllowed( Actor, bInSelected );
-	}
-
-	return bSelectionAllowed;
+	// Allow active modes to determine whether the selection is allowed. 
+	return GLevelEditorModeTools().IsSelectionAllowed(Actor, bInSelected);
 }
 
 void UUnrealEdEngine::SelectActor(AActor* Actor, bool bInSelected, bool bNotify, bool bSelectEvenIfHidden, bool bForceRefresh)
@@ -549,14 +536,8 @@ void UUnrealEdEngine::SelectActor(AActor* Actor, bool bInSelected, bool bNotify,
 		return;
 	}
 
-	bool bSelectionHandled = false;
+	bool bSelectionHandled = GLevelEditorModeTools().IsSelectionHandled(Actor, bInSelected);
 
-	TArray<FEdMode*> ActiveModes;
-	GLevelEditorModeTools().GetActiveModes( ActiveModes );
-	for( int32 ModeIndex = 0; ModeIndex < ActiveModes.Num(); ++ModeIndex )
-	{
-		bSelectionHandled |= ActiveModes[ModeIndex]->Select( Actor, bInSelected );
-	}
 
 	// Select the actor and update its internals.
 	if( !bSelectionHandled )
diff --git a/Engine/Source/Editor/UnrealEd/Private/EditorServer.cpp b/Engine/Source/Editor/UnrealEd/Private/EditorServer.cpp
--- a/Engine/Source/Editor/UnrealEd/Private/EditorServer.cpp
+++ b/Engine/Source/Editor/UnrealEd/Private/EditorServer.cpp
@@ -4992,13 +4992,7 @@ bool UEditorEngine::SnapObjectTo( FActorOrComponent Object, const bool InAlign,
 			RebuildAlteredBSP();
 		}
 
-		TArray<FEdMode*> ActiveModes; 
-		GCurrentLevelEditingViewportClient->GetModeTools()->GetActiveModes(ActiveModes);
-		for( int32 ModeIndex = 0; ModeIndex < ActiveModes.Num(); ++ModeIndex )
-		{
-			// Notify active modes
-			ActiveModes[ModeIndex]->ActorMoveNotify();
-		}
+		GCurrentLevelEditingViewportClient->GetModeTools()->ActorMoveNotify();
 
 		return true;
 	}
diff --git a/Engine/Source/Editor/UnrealEd/Private/EditorViewportClient.cpp b/Engine/Source/Editor/UnrealEd/Private/EditorViewportClient.cpp
--- a/Engine/Source/Editor/UnrealEd/Private/EditorViewportClient.cpp
+++ b/Engine/Source/Editor/UnrealEd/Private/EditorViewportClient.cpp
@@ -426,7 +426,7 @@ FEditorViewportClient::FEditorViewportClient(FEditorModeTools* InModeTools, FPre
 
 	SetViewMode(IsPerspective() ? PerspViewModeIndex : OrthoViewModeIndex);
 
-	ModeTools->OnEditorModeChanged().AddRaw(this, &FEditorViewportClient::OnEditorModeChanged);
+	ModeTools->OnEditorModeIDChanged().AddRaw(this, &FEditorViewportClient::OnEditorModeIDChanged);
 
 	FCoreDelegates::StatCheckEnabled.AddRaw(this, &FEditorViewportClient::HandleViewportStatCheckEnabled);
 	FCoreDelegates::StatEnabled.AddRaw(this, &FEditorViewportClient::HandleViewportStatEnabled);
@@ -446,7 +446,7 @@ FEditorViewportClient::~FEditorViewportClient()
 		ModeTools->DeactivateAllModes(); // this also activates the default mode
 	}
 
-	ModeTools->OnEditorModeChanged().RemoveAll(this);
+	ModeTools->OnEditorModeIDChanged().RemoveAll(this);
 
 	delete Widget;
 	delete MouseDeltaTracker;
@@ -563,6 +563,14 @@ void FEditorViewportClient::OnEditorModeChanged(FEdMode* EditorMode, bool bIsEnt
 	}
 }
 
+void FEditorViewportClient::OnEditorModeIDChanged(const FEditorModeID& EditorModeID, bool bIsEntering)
+{
+	if (Viewport)
+	{
+		RequestInvalidateHitProxy(Viewport);
+	}
+}
+
 float FEditorViewportClient::GetOrthoUnitsPerPixel(const FViewport* InViewport) const
 {
 	const float SizeX = InViewport->GetSizeXY().X;
diff --git a/Engine/Source/Editor/UnrealEd/Private/LevelEditorViewport.cpp b/Engine/Source/Editor/UnrealEd/Private/LevelEditorViewport.cpp
--- a/Engine/Source/Editor/UnrealEd/Private/LevelEditorViewport.cpp
+++ b/Engine/Source/Editor/UnrealEd/Private/LevelEditorViewport.cpp
@@ -2538,13 +2538,7 @@ bool FLevelEditorViewportClient::InputWidgetDelta(FViewport* InViewport, EAxisLi
 					bApplyCameraSpeedScaleByDistance = true;
 				}
 
-				TArray<FEdMode*> ActiveModes; 
-				ModeTools->GetActiveModes(ActiveModes);
-
-				for( int32 ModeIndex = 0; ModeIndex < ActiveModes.Num(); ++ModeIndex )
-				{
-					ActiveModes[ModeIndex]->UpdateInternalData();
-				}
+				ModeTools->UpdateInternalData();
 			}
 
 			bHandled = true;
@@ -3068,13 +3062,7 @@ void FLevelEditorViewportClient::TrackingStopped()
 		GEditor->DisableDeltaModification(false);
 	}
 
-	TArray<FEdMode*> ActiveModes; 
-	ModeTools->GetActiveModes(ActiveModes);
-	for( int32 ModeIndex = 0; ModeIndex < ActiveModes.Num(); ++ModeIndex )
-	{
-		// Also notify the current editing modes if they are interested.
-		ActiveModes[ModeIndex]->ActorMoveNotify();
-	}
+	ModeTools->ActorMoveNotify();
 
 	if( bDidAnythingActuallyChange )
 	{
diff --git a/Engine/Source/Editor/UnrealEd/Private/SnappingUtils.cpp b/Engine/Source/Editor/UnrealEd/Private/SnappingUtils.cpp
--- a/Engine/Source/Editor/UnrealEd/Private/SnappingUtils.cpp
+++ b/Engine/Source/Editor/UnrealEd/Private/SnappingUtils.cpp
@@ -115,14 +115,7 @@ bool FEditorViewportSnapping::IsSnapToGridEnabled()
 bool FEditorViewportSnapping::IsSnapRotationEnabled()
 {	
 	// Ask Current Editor Mode if Rotation Snap is enabled
-	bool bSnapEnabled = false;
-	TArray<FEdMode*> ActiveModes; 
-	GLevelEditorModeTools().GetActiveModes( ActiveModes );
-	for( int32 ModeIndex = 0; ModeIndex < ActiveModes.Num(); ++ModeIndex )
-	{
-		bSnapEnabled |= ActiveModes[ModeIndex]->IsSnapRotationEnabled();
-	}
-	return bSnapEnabled;
+	return GLevelEditorModeTools().IsSnapRotationEnabled();
 }
 
 bool FEditorViewportSnapping::IsSnapScaleEnabled()
@@ -303,14 +296,9 @@ void FEditorViewportSnapping::SnapRotatorToGrid(FRotator& Rotation)
 {
 	if( IsSnapRotationEnabled() )
 	{		
-		TArray<FEdMode*> ActiveModes; 
-		GLevelEditorModeTools().GetActiveModes( ActiveModes );
-		for( int32 ModeIndex = 0; ModeIndex < ActiveModes.Num(); ++ModeIndex )
+		if (GLevelEditorModeTools().SnapRotatorToGridOverride(Rotation))
 		{
-			if( ActiveModes[ModeIndex]->SnapRotatorToGridOverride( Rotation ) == true )
-			{
-				return;
-			}
+			return;
 		}
 
 		Rotation = Rotation.GridSnap( GEditor->GetRotGridSize() );		
diff --git a/Engine/Source/Editor/UnrealEd/Private/Toolkits/BaseToolkit.cpp b/Engine/Source/Editor/UnrealEd/Private/Toolkits/BaseToolkit.cpp
--- a/Engine/Source/Editor/UnrealEd/Private/Toolkits/BaseToolkit.cpp
+++ b/Engine/Source/Editor/UnrealEd/Private/Toolkits/BaseToolkit.cpp
@@ -4,6 +4,13 @@
 #include "Toolkits/ToolkitManager.h"
 #include "Widgets/Docking/SDockableTab.h"
 #include "Widgets/Docking/SDockTabStack.h"
+#include "PropertyEditorModule.h"
+#include "Modules/ModuleManager.h"
+#include "IDetailsView.h"
+#include "InteractiveToolManager.h"
+#include "Editor.h"
+#include "InteractiveTool.h"
+#include "Tools/UEdMode.h"
 
 #define LOCTEXT_NAMESPACE "BaseToolkit"
 
@@ -209,6 +216,34 @@ void FModeToolkit::Init(const TSharedPtr< class IToolkitHost >& InitToolkitHost)
 	ToolkitHost = InitToolkitHost;
 
 	FToolkitManager::Get().RegisterNewToolkit( SharedThis( this ) );
+
+	FPropertyEditorModule& PropertyEditorModule = FModuleManager::GetModuleChecked<FPropertyEditorModule>("PropertyEditor");
+
+	FDetailsViewArgs DetailsViewArgs(
+		/*bUpdateFromSelection=*/ false,
+		/*bLockable=*/ false,
+		/*bAllowSearch=*/ false,
+		FDetailsViewArgs::HideNameArea,
+		/*bHideSelectionTip=*/ true,
+		/*InNotifyHook=*/ nullptr,
+		/*InSearchInitialKeyFocus=*/ false,
+		/*InViewIdentifier=*/ NAME_None);
+	DetailsViewArgs.DefaultsOnlyVisibility = EEditDefaultsOnlyNodeVisibility::Automatic;
+	DetailsViewArgs.bShowOptions = false;
+	DetailsViewArgs.bAllowMultipleTopLevelObjects = true;
+
+	DetailsView = PropertyEditorModule.CreateDetailView(DetailsViewArgs);
+}
+
+
+FName FModeToolkit::GetToolkitFName() const
+{
+	return FName("EditorModeToolkit");
+}
+
+FText FModeToolkit::GetBaseToolkitName() const
+{
+	return NSLOCTEXT("EditorModeToolkit", "DisplayName", "EditorMode Tool");
 }
 
 FString FModeToolkit::GetWorldCentricTabPrefix() const
@@ -230,3 +265,68 @@ FLinearColor FModeToolkit::GetWorldCentricTabColorScale() const
 {
 	return FLinearColor();
 }
+
+
+
+bool FModeToolkit::CanStartTool(const FString& ToolTypeIdentifier)
+{
+	UInteractiveToolManager* Manager = GetScriptableEditorMode()->GetToolManager();
+
+	return (Manager->HasActiveTool(EToolSide::Left) == false) &&
+		(Manager->CanActivateTool(EToolSide::Left, ToolTypeIdentifier) == true);
+}
+
+bool FModeToolkit::CanAcceptActiveTool()
+{
+	return GetScriptableEditorMode()->GetToolManager()->CanAcceptActiveTool(EToolSide::Left);
+}
+
+bool FModeToolkit::CanCancelActiveTool()
+{
+	return GetScriptableEditorMode()->GetToolManager()->CanCancelActiveTool(EToolSide::Left);
+}
+
+bool FModeToolkit::CanCompleteActiveTool()
+{
+	return GetScriptableEditorMode()->GetToolManager()->HasActiveTool(EToolSide::Left) && CanCancelActiveTool() == false;
+}
+
+
+FReply FModeToolkit::StartTool(const FString& ToolTypeIdentifier)
+{
+	if (GetScriptableEditorMode()->GetToolManager()->SelectActiveToolType(EToolSide::Left, ToolTypeIdentifier) == false)
+	{
+		UE_LOG(LogTemp, Warning, TEXT("ToolManager: Unknown Tool Type %s"), *ToolTypeIdentifier);
+	}
+	else
+	{
+		UE_LOG(LogTemp, Warning, TEXT("ToolManager: Starting Tool Type %s"), *ToolTypeIdentifier);
+		GetScriptableEditorMode()->GetToolManager()->ActivateTool(EToolSide::Left);
+
+		// Update properties panel
+		UInteractiveTool* CurTool = GetScriptableEditorMode()->GetToolManager()->GetActiveTool(EToolSide::Left);
+		DetailsView->SetObjects(CurTool->GetToolProperties());
+	}
+	return FReply::Handled();
+}
+
+FReply FModeToolkit::EndTool(EToolShutdownType ShutdownType)
+{
+	UE_LOG(LogTemp, Warning, TEXT("ENDING TOOL"));
+
+	GetScriptableEditorMode()->GetToolManager()->DeactivateTool(EToolSide::Left, ShutdownType);
+
+	DetailsView->SetObject(nullptr);
+
+	return FReply::Handled();
+}
+
+class FEdMode* FModeToolkit::GetEditorMode() const
+{
+	return nullptr; 
+}
+
+UEdMode* FModeToolkit::GetScriptableEditorMode() const
+{
+	return nullptr;
+}
\ No newline at end of file
diff --git a/Engine/Source/Editor/UnrealEd/Private/Tools/EditorComponentSourceFactory.cpp b/Engine/Source/Editor/UnrealEd/Private/Tools/EditorComponentSourceFactory.cpp
new file mode 100644
--- /dev/null
+++ b/Engine/Source/Editor/UnrealEd/Private/Tools/EditorComponentSourceFactory.cpp
@@ -0,0 +1,56 @@
+// Copyright 1998-2019 Epic Games, Inc. All Rights Reserved.
+
+#include "Tools/EditorComponentSourceFactory.h"
+
+#include "Engine/StaticMesh.h"
+#include "Components/PrimitiveComponent.h"
+#include "Components/StaticMeshComponent.h"
+#include "PhysicsEngine/BodySetup.h"
+
+
+FMeshDescription *
+FStaticMeshComponentTarget::GetMesh() {
+	return Cast<UStaticMeshComponent>(Component)->GetStaticMesh()->GetMeshDescription(LODIndex);
+}
+
+void
+FStaticMeshComponentTarget::CommitMesh( const FCommitter& Committer )
+{
+	//bool bSaved = Component->Modify();
+	//check(bSaved);
+	UStaticMesh* StaticMesh = Cast<UStaticMeshComponent>(Component)->GetStaticMesh();
+
+	// make sure transactional flag is on
+	StaticMesh->SetFlags(RF_Transactional);
+
+	bool bSavedToTransactionBuffer = StaticMesh->Modify();
+	check(bSavedToTransactionBuffer);
+	FMeshDescription* MeshDescription = StaticMesh->GetMeshDescription(LODIndex);
+
+	Committer(MeshDescription);
+
+	StaticMesh->CommitMeshDescription(LODIndex);
+	StaticMesh->PostEditChange();
+
+	// this rebuilds physics, but it doesn't undo!
+	Component->RecreatePhysicsState();
+
+	//Component->PostEditChange();
+}
+
+bool
+FStaticMeshComponentTargetFactory::CanBuild(UActorComponent* Component)
+{
+	return Component != nullptr;
+}
+
+TUniquePtr< FPrimitiveComponentTarget >
+FStaticMeshComponentTargetFactory::Build(UPrimitiveComponent* Component)
+{
+	UStaticMeshComponent* StaticMeshComponent = Cast<UStaticMeshComponent>(Component);
+	if (StaticMeshComponent != nullptr)
+	{
+		return TUniquePtr< FPrimitiveComponentTarget > { new FStaticMeshComponentTarget{Component} };
+	}
+	return {};
+}
diff --git a/Engine/Source/Editor/UnrealEd/Private/Tools/EditorToolAssetAPI.cpp b/Engine/Source/Editor/UnrealEd/Private/Tools/EditorToolAssetAPI.cpp
new file mode 100644
--- /dev/null
+++ b/Engine/Source/Editor/UnrealEd/Private/Tools/EditorToolAssetAPI.cpp
@@ -0,0 +1,93 @@
+// Copyright 1998-2019 Epic Games, Inc. All Rights Reserved.
+
+#include "Tools/EditorToolAssetAPI.h"
+#include "CoreMinimal.h"
+#include "GameFramework/Actor.h"
+
+
+#include "AssetRegistryModule.h"
+#include "FileHelpers.h"
+#include "PackageTools.h"
+#include "ObjectTools.h"
+
+// for content-browser things
+#include "ContentBrowserModule.h"
+#include "IContentBrowserSingleton.h"
+
+
+FString FEditorToolAssetAPI::GetActiveAssetFolderPath()
+{
+	check(false);  	// [RMS] this only works if we have an asset selected!
+
+	IContentBrowserSingleton& ContentBrowser = FModuleManager::LoadModuleChecked<FContentBrowserModule>("ContentBrowser").Get();
+
+	//ContentBrowser.CreatePathPicker()
+
+	TArray<FAssetData> SelectedAssets;
+	ContentBrowser.GetSelectedAssets(SelectedAssets);
+	return SelectedAssets[0].PackagePath.ToString();
+}
+
+
+
+FString FEditorToolAssetAPI::MakePackageName(const FString& AssetName, const FString& FolderPath)
+{
+	return FolderPath + TEXT("/") + AssetName;
+}
+
+
+FString FEditorToolAssetAPI::MakeUniqueAssetName(const FString& AssetName, const FString& FolderPath)
+{
+	FString NewPackageName = MakePackageName(AssetName, FolderPath);
+	if (FPackageName::DoesPackageExist(NewPackageName) == false)
+	{
+		return AssetName;
+	}
+
+	int Counter = 1;
+	FString UniqueName = AssetName;
+	while ( FPackageName::DoesPackageExist(NewPackageName) )
+	{
+		UniqueName = AssetName + FString::Printf(TEXT("_%d"), Counter++);
+		NewPackageName = MakePackageName(UniqueName, FolderPath);
+	}
+	return UniqueName;
+}
+
+
+UPackage* FEditorToolAssetAPI::CreateNewPackage(const FString& AssetName, const FString& FolderPath)
+{
+	FString NewPackageName = MakePackageName(AssetName, FolderPath);
+	NewPackageName = UPackageTools::SanitizePackageName(NewPackageName);
+	UPackage* AssetPackage = CreatePackage(nullptr, *NewPackageName);
+	return AssetPackage;
+}
+
+
+
+void FEditorToolAssetAPI::InteractiveSaveGeneratedAsset(UObject* Asset, UPackage* AssetPackage)
+{
+	Asset->MarkPackageDirty();
+
+	FAssetRegistryModule::AssetCreated(Asset);
+
+	TArray<UPackage*> PackagesToSave;
+	PackagesToSave.Add(AssetPackage);
+	bool bCheckDirty = true;
+	bool bPromptToSave = true;
+	FEditorFileUtils::PromptForCheckoutAndSave(PackagesToSave, bCheckDirty, bPromptToSave);
+}
+
+
+void FEditorToolAssetAPI::AutoSaveGeneratedAsset(UObject* Asset, UPackage* AssetPackage)
+{
+	Asset->MarkPackageDirty();
+
+	FAssetRegistryModule::AssetCreated(Asset);
+
+	TArray<UPackage*> PackagesToSave;
+	PackagesToSave.Add(AssetPackage);
+	bool bCheckDirty = true;
+	bool bPromptToSave = false;
+	FEditorFileUtils::PromptForCheckoutAndSave(PackagesToSave, bCheckDirty, bPromptToSave);
+}
\ No newline at end of file
diff --git a/Engine/Source/Editor/UnrealEd/Private/Tools/StandardToolModeCommands.cpp b/Engine/Source/Editor/UnrealEd/Private/Tools/StandardToolModeCommands.cpp
new file mode 100644
--- /dev/null
+++ b/Engine/Source/Editor/UnrealEd/Private/Tools/StandardToolModeCommands.cpp
@@ -0,0 +1,39 @@
+// Copyright 1998-2019 Epic Games, Inc. All Rights Reserved.
+
+#include "Tools/StandardToolModeCommands.h"
+#include "EditorStyleSet.h"
+
+#define LOCTEXT_NAMESPACE "StandardToolModeCommands"
+
+
+FStandardToolModeCommands::FStandardToolModeCommands() :
+	TCommands<FStandardToolModeCommands>(
+		"StandardToolCommands", // Context name for fast lookup
+		NSLOCTEXT("Contexts", "StandardToolCommands", "Standard Tool Commands"), // Localized context name for displaying
+		NAME_None, // Parent
+		FEditorStyle::GetStyleSetName() // Icon Style Set
+		)
+{}
+
+
+void FStandardToolModeCommands::RegisterCommands()
+{
+	TSharedPtr<FUICommandInfo> IncreaseBrushSize;
+	UI_COMMAND(IncreaseBrushSize, "Increase Brush Size", "Increases the size of the brush", EUserInterfaceActionType::Button, FInputChord(EKeys::RightBracket));
+	Commands.Add(EStandardToolModeCommands::IncreaseBrushSize, IncreaseBrushSize);
+
+	TSharedPtr<FUICommandInfo> DecreaseBrushSize;
+	UI_COMMAND(DecreaseBrushSize, "Decrease Brush Size", "Decreases the size of the brush", EUserInterfaceActionType::Button, FInputChord(EKeys::LeftBracket));
+	Commands.Add(EStandardToolModeCommands::DecreaseBrushSize, DecreaseBrushSize);
+}
+
+
+TSharedPtr<FUICommandInfo> FStandardToolModeCommands::FindStandardCommand(EStandardToolModeCommands Command) const
+{
+	const TSharedPtr<FUICommandInfo>* Found = Commands.Find(Command);
+	checkf(Found != nullptr, TEXT("FStandardToolModeCommands::FindStandardCommand: standard command %d was not found!"), Command);
+	return (Found != nullptr) ? *Found : nullptr;
+}
+
+
+#undef LOCTEXT_NAMESPACE
diff --git a/Engine/Source/Editor/UnrealEd/Private/Tools/UEdMode.cpp b/Engine/Source/Editor/UnrealEd/Private/Tools/UEdMode.cpp
new file mode 100644
--- /dev/null
+++ b/Engine/Source/Editor/UnrealEd/Private/Tools/UEdMode.cpp
@@ -0,0 +1,978 @@
+// Copyright 1998-2019 Epic Games, Inc. All Rights Reserved.
+
+#include "Tools/UEdMode.h"
+#include "EditorModeTools.h"
+#include "EditorViewportClient.h"
+#include "Framework/Application/SlateApplication.h"
+#include "CanvasItem.h"
+#include "Engine/StaticMeshActor.h"
+#include "Engine/Selection.h"
+#include "EngineUtils.h"
+#include "EditorModeManager.h"
+#include "EditorModes.h"
+#include "StaticMeshResources.h"
+#include "Toolkits/BaseToolkit.h"
+#include "InteractiveToolsContext.h"
+#include "CanvasTypes.h"
+#include "ScopedTransaction.h"
+#include "Tools/EditorToolAssetAPI.h"
+#include "Editor.h"
+#include "Toolkits/ToolkitManager.h"
+#include "Materials/Material.h"
+#include "Materials/MaterialInterface.h"
+#include "LevelEditorViewport.h"
+#include "InteractiveToolObjects.h"
+
+
+class FEdModeQueriesImpl : public IToolsContextQueriesAPI
+{
+public:
+	UInteractiveToolsContext* ToolsContext;
+	UEdMode* EdMode;
+	FViewCameraState CachedViewState;
+
+	FEdModeQueriesImpl(UInteractiveToolsContext* Context, UEdMode* InEdMode)
+	{
+		ToolsContext = Context;
+		EdMode = InEdMode;
+	}
+
+	void CacheCurrentViewState(FEditorViewportClient* ViewportClient)
+	{
+		FViewportCameraTransform ViewTransform = ViewportClient->GetViewTransform();
+		CachedViewState.Position = ViewTransform.GetLocation();
+		CachedViewState.Orientation = ViewTransform.GetRotation().Quaternion();
+		CachedViewState.bIsOrthographic = ViewportClient->IsOrtho();
+		CachedViewState.bIsVR = false;
+	}
+
+	virtual void GetCurrentSelectionState(FToolBuilderState& StateOut) const override
+	{
+		StateOut.ToolManager = ToolsContext->ToolManager;
+		StateOut.GizmoManager = ToolsContext->GizmoManager;
+		StateOut.World = GEditor->GetWorld();
+		GEditor->GetSelectedActors()->GetSelectedObjects(StateOut.SelectedActors);
+		GEditor->GetSelectedComponents()->GetSelectedObjects(StateOut.SelectedComponents);
+	}
+
+
+	virtual void GetCurrentViewState(FViewCameraState& StateOut) const override
+	{
+		StateOut = CachedViewState;
+	}
+
+
+	virtual bool ExecuteSceneSnapQuery(const FSceneSnapQueryRequest& Request, TArray<FSceneSnapQueryResult>& Results) const override
+	{
+		if (Request.RequestType != ESceneSnapQueryType::Position)
+		{
+			return false;		// not supported yet
+		}
+
+		int FoundResultCount = 0;
+
+		//
+		// Run a snap query by casting ray into the world.
+		// If a hit is found, we look up what triangle was hit, and then test its vertices and edges
+		// 
+
+		// cast ray into world
+		FVector RayStart = CachedViewState.Position;
+		FVector RayDirection = Request.Position - RayStart; RayDirection.Normalize();
+		FVector RayEnd = RayStart + 9999999 * RayDirection;
+		FCollisionObjectQueryParams ObjectQueryParams(FCollisionObjectQueryParams::AllObjects);
+		FCollisionQueryParams QueryParams = FCollisionQueryParams::DefaultQueryParam;
+		QueryParams.bTraceComplex = true;
+		QueryParams.bReturnFaceIndex = true;
+		FHitResult HitResult;
+		bool bHitWorld = GEditor->GetWorld()->LineTraceSingleByObjectType(HitResult, RayStart, RayEnd, ObjectQueryParams, QueryParams);
+		if (bHitWorld && HitResult.FaceIndex >= 0)
+		{
+			float VisualAngle = OpeningAngleDeg(Request.Position, HitResult.ImpactPoint, RayStart);
+			//UE_LOG(LogTemp, Warning, TEXT("[HIT] visualangle %f faceindex %d"), VisualAngle, HitResult.FaceIndex);
+			if (VisualAngle < Request.VisualAngleThresholdDegrees)
+			{
+				UPrimitiveComponent* Component = HitResult.Component.Get();
+				if (Cast<UStaticMeshComponent>(Component) != nullptr)
+				{
+					// HitResult.FaceIndex is apparently an index into the TriMeshCollisionData, not sure how
+					// to directly access it. Calling GetPhysicsTriMeshData is expensive!
+					//UBodySetup* BodySetup = Cast<UStaticMeshComponent>(Component)->GetBodySetup();
+					//UObject* CDPObj = BodySetup->GetOuter();
+					//IInterface_CollisionDataProvider* CDP = Cast<IInterface_CollisionDataProvider>(CDPObj);
+					//FTriMeshCollisionData TriMesh;
+					//CDP->GetPhysicsTriMeshData(&TriMesh, true);
+					//FTriIndices Triangle = TriMesh.Indices[HitResult.FaceIndex];
+					//FVector Positions[3] = { TriMesh.Vertices[Triangle.v0], TriMesh.Vertices[Triangle.v1], TriMesh.Vertices[Triangle.v2] };
+
+					// physics collision data is created from StaticMesh RenderData
+					// so use HitResult.FaceIndex to extract triangle from the LOD0 mesh
+					// (note: this may be incorrect if there are multiple sections...in that case I think we have to 
+					//  first find section whose accumulated index range would contain .FaceIndexX)
+					UStaticMesh* StaticMesh = Cast<UStaticMeshComponent>(Component)->GetStaticMesh();
+					FStaticMeshLODResources& LOD = StaticMesh->RenderData->LODResources[0];
+					FIndexArrayView Indices = LOD.IndexBuffer.GetArrayView();
+					int32 TriIdx = 3 * HitResult.FaceIndex;
+					FVector Positions[3];
+					Positions[0] = LOD.VertexBuffers.PositionVertexBuffer.VertexPosition(Indices[TriIdx]);
+					Positions[1] = LOD.VertexBuffers.PositionVertexBuffer.VertexPosition(Indices[TriIdx + 1]);
+					Positions[2] = LOD.VertexBuffers.PositionVertexBuffer.VertexPosition(Indices[TriIdx + 2]);
+
+					// transform to world space
+					FTransform ComponentTransform = Component->GetComponentTransform();
+					Positions[0] = ComponentTransform.TransformPosition(Positions[0]);
+					Positions[1] = ComponentTransform.TransformPosition(Positions[1]);
+					Positions[2] = ComponentTransform.TransformPosition(Positions[2]);
+
+					FSceneSnapQueryResult SnapResult;
+					SnapResult.TriVertices[0] = Positions[0];
+					SnapResult.TriVertices[1] = Positions[1];
+					SnapResult.TriVertices[2] = Positions[2];
+
+					// try snapping to vertices
+					float SmallestAngle = Request.VisualAngleThresholdDegrees;
+					if ((Request.TargetTypes & ESceneSnapQueryTargetType::MeshVertex) != ESceneSnapQueryTargetType::None)
+					{
+						for (int j = 0; j < 3; ++j)
+						{
+							VisualAngle = OpeningAngleDeg(Request.Position, Positions[j], RayStart);
+							if (VisualAngle < SmallestAngle)
+							{
+								SmallestAngle = VisualAngle;
+								SnapResult.Position = Positions[j];
+								SnapResult.TargetType = ESceneSnapQueryTargetType::MeshVertex;
+								SnapResult.TriSnapIndex = j;
+							}
+						}
+					}
+
+					// try snapping to nearest points on edges
+					if (((Request.TargetTypes & ESceneSnapQueryTargetType::MeshEdge) != ESceneSnapQueryTargetType::None) &&
+						(SnapResult.TargetType != ESceneSnapQueryTargetType::MeshVertex))
+					{
+						for (int j = 0; j < 3; ++j)
+						{
+							FVector EdgeNearestPt = NearestSegmentPt(Positions[j], Positions[(j + 1) % 3], Request.Position);
+							VisualAngle = OpeningAngleDeg(Request.Position, EdgeNearestPt, RayStart);
+							if (VisualAngle < SmallestAngle)
+							{
+								SmallestAngle = VisualAngle;
+								SnapResult.Position = EdgeNearestPt;
+								SnapResult.TargetType = ESceneSnapQueryTargetType::MeshEdge;
+								SnapResult.TriSnapIndex = j;
+							}
+						}
+					}
+
+					// if we found a valid snap, return it
+					if (SmallestAngle < Request.VisualAngleThresholdDegrees)
+					{
+						SnapResult.TargetActor = HitResult.Actor.Get();
+						SnapResult.TargetComponent = HitResult.Component.Get();
+						Results.Add(SnapResult);
+						FoundResultCount++;
+					}
+				}
+			}
+
+		}
+
+		return (FoundResultCount > 0);
+	}
+
+
+	//@ todo this are mirrored from GeometryProcessing, which is still experimental...replace w/ direct calls once GP component is standardized
+	static float OpeningAngleDeg(FVector A, FVector B, const FVector& P)
+	{
+		A -= P;
+		A.Normalize();
+		B -= P;
+		B.Normalize();
+		float Dot = FMath::Clamp(FVector::DotProduct(A, B), -1.0f, 1.0f);
+		return acos(Dot) * (180.0f / 3.141592653589f);
+	}
+	static FVector NearestSegmentPt(FVector A, FVector B, const FVector& P)
+	{
+		FVector Direction = (B - A);
+		float Length = Direction.Size();
+		Direction /= Length;
+		float t = FVector::DotProduct((P - A), Direction);
+		if (t >= Length)
+		{
+			return B;
+		}
+		if (t <= 0)
+		{
+			return A;
+		}
+		return A + t * Direction;
+	}
+
+
+
+	virtual UMaterialInterface* GetStandardMaterial(EStandardToolContextMaterials MaterialType) const
+	{
+		if (MaterialType == EStandardToolContextMaterials::VertexColorMaterial)
+		{
+			return EdMode->StandardVertexColorMaterial;
+		}
+		check(false);
+		return nullptr;
+	}
+
+};
+
+
+
+
+class FEdModeTransactionImpl : public IToolsContextTransactionsAPI
+{
+public:
+	UInteractiveToolsContext* ToolsContext;
+	UEdMode* EdMode;
+
+	FEdModeTransactionImpl(UInteractiveToolsContext* Context, UEdMode* InEdMode)
+	{
+		ToolsContext = Context;
+		EdMode = InEdMode;
+	}
+
+
+	virtual void PostMessage(const TCHAR* Message, EToolMessageLevel Level) override
+	{
+		UE_LOG(LogTemp, Warning, TEXT("[ToolsContext] %s"), Message);
+	}
+
+
+	virtual void PostInvalidation() override
+	{
+		EdMode->PostInvalidation();
+	}
+
+	virtual void BeginUndoTransaction(const FText& Description) override
+	{
+		GEditor->BeginTransaction(Description);
+	}
+
+	virtual void EndUndoTransaction() override
+	{
+		GEditor->EndTransaction();
+	}
+
+	virtual void AppendChange(UObject* TargetObject, TUniquePtr<FChange> Change, const FText& Description) override
+	{
+		FScopedTransaction Transaction(Description);
+		check(GUndo != nullptr);
+		GUndo->StoreUndo(TargetObject, MoveTemp(Change));
+		// end transaction
+	}
+
+
+	virtual bool RequestSelectionChange(const FSelectedOjectsChangeList& SelectionChange) override
+	{
+		checkf(SelectionChange.Components.Num() == 0, TEXT("FEdModeToolsContextTransactionImpl::RequestSelectionChange - Component selection not supported yet"));
+
+		if (SelectionChange.ModificationType == ESelectedObjectsModificationType::Clear)
+		{
+			GEditor->SelectNone(true, true, false);
+			return true;
+		}
+
+		if (SelectionChange.ModificationType == ESelectedObjectsModificationType::Replace)
+		{
+			GEditor->SelectNone(false, true, false);
+		}
+
+		bool bAdd = (SelectionChange.ModificationType != ESelectedObjectsModificationType::Remove);
+		int NumActors = SelectionChange.Actors.Num();
+		for (int k = 0; k < NumActors; ++k)
+		{
+			GEditor->SelectActor(SelectionChange.Actors[k], bAdd, false, true, false);
+		}
+
+		GEditor->NoteSelectionChange(true);
+		return true;
+	}
+
+};
+
+
+
+/** Hit proxy used for editable properties */
+struct HPropertyWidgetProxyTools : public HHitProxy
+{
+	DECLARE_HIT_PROXY();
+
+	/** Name of property this is the widget for */
+	FString	PropertyName;
+
+	/** If the property is an array property, the index into that array that this widget is for */
+	int32	PropertyIndex;
+
+	/** This property is a transform */
+	bool	bPropertyIsTransform;
+
+	HPropertyWidgetProxyTools(FString InPropertyName, int32 InPropertyIndex, bool bInPropertyIsTransform)
+		: HHitProxy(HPP_Foreground)
+		, PropertyName(InPropertyName)
+		, PropertyIndex(InPropertyIndex)
+		, bPropertyIsTransform(bInPropertyIsTransform)
+	{}
+
+	/** Show cursor as cross when over this handle */
+	virtual EMouseCursor::Type GetMouseCursor() override
+	{
+		return EMouseCursor::Crosshairs;
+	}
+};
+
+IMPLEMENT_HIT_PROXY(HPropertyWidgetProxyTools, HHitProxy);
+
+
+
+//////////////////////////////////
+// UEdMode
+
+UEdMode::UEdMode()
+	: bPendingDeletion(false)
+	, Owner(nullptr)
+{
+	bDrawKillZ = true;
+	ToolsContext = nullptr;
+}
+
+
+void UEdMode::OnModeUnregistered(FEditorModeID ModeID)
+{
+	if (ModeID == Info.ID)
+	{
+		// This should be synonymous with "delete this"
+		Owner->DestroyMode(ModeID);
+	}
+}
+
+void UEdMode::TerminateActiveToolsOnPIEStart()
+{
+	DeactivateAllActiveTools();
+}
+void UEdMode::TerminateActiveToolsOnSaveWorld()
+{
+	DeactivateAllActiveTools();
+}
+
+FRay UEdMode::GetRayFromMousePos(FEditorViewportClient* ViewportClient, FViewport* Viewport, int MouseX, int MouseY)
+{
+	FSceneViewFamilyContext ViewFamily(FSceneViewFamily::ConstructionValues(
+		ViewportClient->Viewport,
+		ViewportClient->GetScene(),
+		ViewportClient->EngineShowFlags)
+		.SetRealtimeUpdate(ViewportClient->IsRealtime()));
+	FSceneView* View = ViewportClient->CalcSceneView(&ViewFamily);
+	FViewportCursorLocation MouseViewportRay(View, (FEditorViewportClient*)Viewport->GetClient(), MouseX, MouseY);
+
+	return FRay(MouseViewportRay.GetOrigin(), MouseViewportRay.GetDirection(), true);
+}
+
+
+
+bool UEdMode::MouseEnter(FEditorViewportClient* ViewportClient, FViewport* Viewport, int32 x, int32 y)
+{
+	CurrentMouseState.Mouse.Position2D = FVector2D(x, y);
+	CurrentMouseState.Mouse.WorldRay = GetRayFromMousePos(ViewportClient, Viewport, x, y);
+
+	return false;
+}
+
+bool UEdMode::MouseLeave(FEditorViewportClient* ViewportClient, FViewport* Viewport)
+{
+	return false;
+}
+
+bool UEdMode::MouseMove(FEditorViewportClient* ViewportClient, FViewport* Viewport, int32 x, int32 y)
+{
+	CurrentMouseState.Mouse.Position2D = FVector2D(x, y);
+	CurrentMouseState.Mouse.WorldRay = GetRayFromMousePos(ViewportClient, Viewport, x, y);
+	FInputDeviceState InputState = CurrentMouseState;
+	InputState.InputDevice = EInputDevices::Mouse;
+
+	InputState.SetModifierKeyStates(
+		ViewportClient->IsShiftPressed(), ViewportClient->IsAltPressed(),
+		ViewportClient->IsCtrlPressed(), ViewportClient->IsCmdPressed());
+
+	if (ToolsContext->InputRouter->HasActiveMouseCapture())
+	{
+		// This state occurs if InputBehavior did not release capture on mouse release. 
+		// UMultiClickSequenceInputBehavior does this, eg for multi-click draw-polygon sequences.
+		// It's not ideal though and maybe would be better done via multiple captures + hover...?
+		ToolsContext->InputRouter->PostInputEvent(InputState);
+	}
+	else
+	{
+		ToolsContext->InputRouter->PostHoverInputEvent(InputState);
+	}
+
+	return false;
+}
+
+bool UEdMode::ReceivedFocus(FEditorViewportClient* ViewportClient, FViewport* Viewport)
+{
+	return false;
+}
+
+bool UEdMode::LostFocus(FEditorViewportClient* ViewportClient, FViewport* Viewport)
+{
+	return false;
+}
+
+bool UEdMode::CapturedMouseMove(FEditorViewportClient* InViewportClient, FViewport* InViewport, int32 InMouseX, int32 InMouseY)
+{
+	// if alt is down we will not allow client to see this event
+	if (InViewportClient->IsAltPressed())
+	{
+		return false;
+	}
+
+	FVector2D OldPosition = CurrentMouseState.Mouse.Position2D;
+	CurrentMouseState.Mouse.Position2D = FVector2D(InMouseX, InMouseY);
+	CurrentMouseState.Mouse.WorldRay = GetRayFromMousePos(InViewportClient, InViewport, InMouseX, InMouseY);
+
+	if (ToolsContext->InputRouter->HasActiveMouseCapture())
+	{
+		FInputDeviceState InputState = CurrentMouseState;
+		InputState.InputDevice = EInputDevices::Mouse;
+		InputState.SetModifierKeyStates(
+			InViewportClient->IsShiftPressed(), InViewportClient->IsAltPressed(),
+			InViewportClient->IsCtrlPressed(), InViewportClient->IsCmdPressed());
+		InputState.Mouse.Delta2D = CurrentMouseState.Mouse.Position2D - OldPosition;
+		ToolsContext->InputRouter->PostInputEvent(InputState);
+		return true;
+	}
+
+	return false;
+}
+
+bool UEdMode::InputKey(FEditorViewportClient* ViewportClient, FViewport* Viewport, FKey Key, EInputEvent Event)
+{
+	
+	bool bHandled = false;
+
+	// escape key cancels current tool
+	if (Key == EKeys::Escape && Event == IE_Released)
+	{
+		if (ToolsContext->ToolManager->HasAnyActiveTool())
+		{
+			if (ToolsContext->ToolManager->HasActiveTool(EToolSide::Mouse))
+			{
+				ToolsContext->DeactivateActiveTool(EToolSide::Mouse, EToolShutdownType::Cancel);
+				RestoreEditorState();
+			}
+			return true;
+		}
+	}
+
+	// enter key accepts current tool, or ends tool if it does not have accept state
+	if (Key == EKeys::Enter && Event == IE_Released && ToolsContext->ToolManager->HasAnyActiveTool())
+	{
+		if (ToolsContext->ToolManager->HasActiveTool(EToolSide::Mouse))
+		{
+			if (ToolsContext->ToolManager->GetActiveTool(EToolSide::Mouse)->HasAccept())
+			{
+				if (ToolsContext->ToolManager->CanAcceptActiveTool(EToolSide::Mouse))
+				{
+					ToolsContext->DeactivateActiveTool(EToolSide::Mouse, EToolShutdownType::Accept);
+					RestoreEditorState();
+					return true;
+				}
+			}
+			else
+			{
+				ToolsContext->DeactivateActiveTool(EToolSide::Mouse, EToolShutdownType::Completed);
+				RestoreEditorState();
+				return true;
+			}
+		}
+	}
+
+	// if alt is down we do not process mouse event
+	if (ViewportClient->IsAltPressed())
+	{
+		return false;
+	}
+
+	if (Event == IE_Pressed || Event == IE_Released)
+	{
+		if (Key.IsMouseButton())
+		{
+			bool bIsLeftMouse = (Key == EKeys::LeftMouseButton);
+			bool bIsMiddleMouse = (Key == EKeys::MiddleMouseButton);
+			bool bIsRightMouse = (Key == EKeys::RightMouseButton);
+
+			if (bIsLeftMouse || bIsMiddleMouse || bIsRightMouse)
+			{
+
+				// early-out here if we are going to do camera manipulation
+				if (ViewportClient->IsAltPressed())
+				{
+					return bHandled;
+				}
+
+				FInputDeviceState InputState = CurrentMouseState;
+				InputState.InputDevice = EInputDevices::Mouse;
+				InputState.SetModifierKeyStates(
+					ViewportClient->IsShiftPressed(), ViewportClient->IsAltPressed(),
+					ViewportClient->IsCtrlPressed(), ViewportClient->IsCmdPressed());
+
+				if (bIsLeftMouse)
+				{
+					InputState.Mouse.Left.SetStates(
+						(Event == IE_Pressed), (Event == IE_Pressed), (Event == IE_Released));
+					CurrentMouseState.Mouse.Left.bDown = (Event == IE_Pressed);
+				}
+				else if (bIsMiddleMouse)
+				{
+					InputState.Mouse.Middle.SetStates(
+						(Event == IE_Pressed), (Event == IE_Pressed), (Event == IE_Released));
+					CurrentMouseState.Mouse.Middle.bDown = (Event == IE_Pressed);
+				}
+				else
+				{
+					InputState.Mouse.Right.SetStates(
+						(Event == IE_Pressed), (Event == IE_Pressed), (Event == IE_Released));
+					CurrentMouseState.Mouse.Right.bDown = (Event == IE_Pressed);
+				}
+
+				ToolsContext->InputRouter->PostInputEvent(InputState);
+
+				if (ToolsContext->InputRouter->HasActiveMouseCapture())
+				{
+					// what is this about? MeshPaintMode has it...
+					ViewportClient->bLockFlightCamera = true;
+					bHandled = true;   // indicate that we handled this event,
+									   // which will disable camera movement/etc ?
+				}
+				else
+				{
+					//ViewportClient->bLockFlightCamera = false;
+				}
+
+			}
+		}
+		else if (Key.IsGamepadKey())
+		{
+			// not supported yet
+		}
+		else if (Key.IsTouch())
+		{
+			// not supported yet
+		}
+		else if (Key.IsFloatAxis() || Key.IsVectorAxis())
+		{
+			// not supported yet
+		}
+		else    // is this definitely a keyboard key?
+		{
+			FInputDeviceState InputState;
+			InputState.InputDevice = EInputDevices::Keyboard;
+			InputState.SetModifierKeyStates(
+				ViewportClient->IsShiftPressed(), ViewportClient->IsAltPressed(),
+				ViewportClient->IsCtrlPressed(), ViewportClient->IsCmdPressed());
+			InputState.Keyboard.ActiveKey.Button = Key;
+			bool bPressed = (Event == IE_Pressed);
+			InputState.Keyboard.ActiveKey.SetStates(bPressed, bPressed, !bPressed);
+			ToolsContext->InputRouter->PostInputEvent(InputState);
+		}
+
+	}
+
+	if(!bHandled)
+	{
+		// Next pass input to the mode toolkit
+		if (Toolkit.IsValid() && ((Event == IE_Pressed) || (Event == IE_Repeat)))
+		{
+			if (Toolkit->GetToolkitCommands()->ProcessCommandBindings(Key, FSlateApplication::Get().GetModifierKeys(), (Event == IE_Repeat)))
+			{
+				return true;
+			}
+		}
+
+		// Finally, pass input up to selected actors if not in a tool mode
+		TArray<AActor*> SelectedActors;
+		Owner->GetSelectedActors()->GetSelectedObjects<AActor>(SelectedActors);
+
+		for (TArray<AActor*>::TIterator It(SelectedActors); It; ++It)
+		{
+			// Tell the object we've had a key press
+			(*It)->EditorKeyPressed(Key, Event);
+		}
+	}
+
+	return false;
+}
+
+bool UEdMode::InputAxis(FEditorViewportClient* InViewportClient, FViewport* Viewport, int32 ControllerId, FKey Key, float Delta, float DeltaTime)
+{
+	return false;
+}
+
+bool UEdMode::InputDelta(FEditorViewportClient* InViewportClient, FViewport* InViewport, FVector& InDrag, FRotator& InRot, FVector& InScale)
+{
+	return false;
+}
+
+void UEdMode::SelectNone()
+{
+	GEditor->SelectNone(true, true);
+}
+
+bool UEdMode::ProcessEditDelete()
+{
+	if (ToolsContext && ToolsContext->ToolManager->HasAnyActiveTool() == false)
+	{
+		return false;
+	}
+
+	bool bSkipDelete = false;
+
+	// Test if any of the selected actors are AInternalToolFrameworkActor
+	// subclasses. In this case we do not want to allow them to be deleted,
+	// as generally this will cause problems for the Tool.
+	USelection* SelectedActors = GEditor->GetSelectedActors();
+	for (int i = 0; i < SelectedActors->Num(); ++i)
+	{
+		UObject* SelectedActor = SelectedActors->GetSelectedObject(i);
+		if (Cast<AInternalToolFrameworkActor>(SelectedActor) != nullptr)
+		{
+			bSkipDelete = true;
+		}
+	}
+
+	return bSkipDelete;
+}
+
+void UEdMode::Tick(FEditorViewportClient* ViewportClient, float DeltaTime)
+{
+	// give ToolsContext a chance to tick
+	if (ToolsContext != nullptr)
+	{
+		ToolsContext->ToolManager->Tick(DeltaTime);
+		ToolsContext->GizmoManager->Tick(DeltaTime);
+
+		if (bInvalidationPending)
+		{
+			ViewportClient->Invalidate();
+			bInvalidationPending = false;
+		}
+
+		// save this view
+		// Check against GCurrentLevelEditingViewportClient is temporary and should be removed in future.
+		// Current issue is that this ::Tick() is called *per viewport*, so once for each view in a 4-up view.
+		if (ViewportClient == GCurrentLevelEditingViewportClient)
+		{
+			((FEdModeQueriesImpl*)QueriesAPI)->CacheCurrentViewState(ViewportClient);
+		}
+	}
+}
+
+void UEdMode::ActorSelectionChangeNotify()
+{
+}
+
+bool UEdMode::HandleClick(FEditorViewportClient* InViewportClient, HHitProxy *HitProxy, const FViewportClick &Click)
+{
+	return false;
+}
+
+void UEdMode::Enter()
+{
+	// Update components for selected actors, in case the mode we just exited
+	// was hijacking selection events selection and not updating components.
+	for (FSelectionIterator It(*Owner->GetSelectedActors()); It; ++It)
+	{
+		AActor* SelectedActor = CastChecked<AActor>(*It);
+		SelectedActor->MarkComponentsRenderStateDirty();
+	}
+
+	bPendingDeletion = false;
+
+	if (!Toolkit.IsValid())
+	{
+		Toolkit = MakeShareable(new FModeToolkit);
+		Toolkit->Init(Owner->GetToolkitHost());
+	}
+
+	// initialize the adapter that attaches the ToolsContext to this FEdMode
+	ToolsContext = NewObject<UInteractiveToolsContext>(GetTransientPackage(), TEXT("ToolsContext"), RF_Transient);
+	TransactionAPI = new FEdModeTransactionImpl(ToolsContext, this);
+	QueriesAPI = new FEdModeQueriesImpl(ToolsContext, this);
+	AssetAPI = new FEditorToolAssetAPI();
+	ToolsContext->Initialize(QueriesAPI, TransactionAPI);
+
+	// enable auto invalidation in Editor, because invalidating for all hover and capture events is unpleasant
+	ToolsContext->InputRouter->bAutoInvalidateOnHover = true;
+	ToolsContext->InputRouter->bAutoInvalidateOnCapture = true;
+
+
+	// set up standard materials
+	StandardVertexColorMaterial = LoadObject<UMaterial>(nullptr, TEXT("/Game/Materials/VertexColor"));
+
+	BeginPIEDelegateHandle = FEditorDelegates::BeginPIE.AddLambda([this](bool bSimulating)
+	{
+		TerminateActiveToolsOnPIEStart();
+	});
+	PreSaveWorldDelegateHandle = FEditorDelegates::PreSaveWorld.AddLambda([this](uint32 SaveFlags, UWorld* World)
+	{
+		TerminateActiveToolsOnSaveWorld();
+	});
+	bInvalidationPending = false;
+
+	FEditorDelegates::EditorModeIDEnter.Broadcast(GetID());
+	const bool bIsEnteringMode = true;
+	Owner->BroadcastEditorModeIDChanged(GetID(), bIsEnteringMode);
+}
+
+void UEdMode::Exit()
+{
+	FEditorDelegates::BeginPIE.Remove(BeginPIEDelegateHandle);
+	FEditorDelegates::PreSaveWorld.Remove(PreSaveWorldDelegateHandle);
+
+	// auto-accept any in-progress tools
+	DeactivateAllActiveTools();
+
+	if (QueriesAPI != nullptr)
+	{
+		delete QueriesAPI;
+		QueriesAPI = nullptr;
+	}
+
+	if (TransactionAPI != nullptr)
+	{
+		delete TransactionAPI;
+		TransactionAPI = nullptr;
+	}
+
+	if (AssetAPI != nullptr)
+	{
+		delete AssetAPI;
+		AssetAPI = nullptr;
+	}
+
+	if (Toolkit.IsValid())
+	{
+		FToolkitManager::Get().CloseToolkit(Toolkit.ToSharedRef());
+		Toolkit.Reset();
+	}
+
+	const bool bIsEnteringMode = false;
+	Owner->BroadcastEditorModeIDChanged(GetID(), bIsEnteringMode);
+	FEditorDelegates::EditorModeIDExit.Broadcast(GetID());
+	if (ToolsContext != nullptr)
+	{
+		ToolsContext->Shutdown();
+		ToolsContext = nullptr;
+	}
+}
+
+
+class FTempRenderContext : public IToolsContextRenderAPI
+{
+public:
+	FPrimitiveDrawInterface* PDI;
+
+	virtual FPrimitiveDrawInterface* GetPrimitiveDrawInterface() override
+	{
+		return PDI;
+	}
+};
+
+void UEdMode::Render(const FSceneView* View, FViewport* Viewport, FPrimitiveDrawInterface* PDI)
+{
+	// give ToolsContext a chance to render
+	if (ToolsContext != nullptr)
+	{
+		FTempRenderContext RenderContext;
+		RenderContext.PDI = PDI;
+		ToolsContext->ToolManager->Render(&RenderContext);
+		ToolsContext->GizmoManager->Render(&RenderContext);
+	}
+}
+
+void UEdMode::DrawHUD(FEditorViewportClient* ViewportClient, FViewport* Viewport, const FSceneView* View, FCanvas* Canvas)
+{
+	// Render the drag tool.
+	ViewportClient->RenderDragTool(View, Canvas);
+
+	if (ViewportClient->IsPerspective() && GetDefault<ULevelEditorViewportSettings>()->bHighlightWithBrackets)
+	{
+		DrawBrackets(ViewportClient, Viewport, View, Canvas);
+	}
+
+	// If this viewport doesn't show mode widgets, leave.
+	if (!(ViewportClient->EngineShowFlags.ModeWidgets))
+	{
+		return;
+	}
+
+	// Clear Hit proxies
+	const bool bIsHitTesting = Canvas->IsHitTesting();
+	if (!bIsHitTesting)
+	{
+		Canvas->SetHitProxy(NULL);
+	}
+
+	// Draw vertices for selected BSP brushes and static meshes if the large vertices show flag is set.
+	if (!ViewportClient->bDrawVertices)
+	{
+		return;
+	}
+
+	const bool bLargeVertices = View->Family->EngineShowFlags.LargeVertices;
+	const bool bShowBrushes = View->Family->EngineShowFlags.Brushes;
+	const bool bShowBSP = View->Family->EngineShowFlags.BSP;
+	const bool bShowBuilderBrush = View->Family->EngineShowFlags.BuilderBrush != 0;
+
+	UTexture2D* VertexTexture = GetVertexTexture();
+	const float TextureSizeX = VertexTexture->GetSizeX() * (bLargeVertices ? 1.0f : 0.5f);
+	const float TextureSizeY = VertexTexture->GetSizeY() * (bLargeVertices ? 1.0f : 0.5f);
+
+	// Temporaries.
+	TArray<FVector> Vertices;
+
+	for (FSelectionIterator It(*Owner->GetSelectedActors()); It; ++It)
+	{
+		AActor* SelectedActor = static_cast<AActor*>(*It);
+		checkSlow(SelectedActor->IsA(AActor::StaticClass()));
+
+		if (bLargeVertices)
+		{
+			FCanvasItemTestbed::bTestState = !FCanvasItemTestbed::bTestState;
+
+			// Static mesh vertices
+			AStaticMeshActor* Actor = Cast<AStaticMeshActor>(SelectedActor);
+			if (Actor && Actor->GetStaticMeshComponent() && Actor->GetStaticMeshComponent()->GetStaticMesh()
+				&& Actor->GetStaticMeshComponent()->GetStaticMesh()->RenderData)
+			{
+				FTransform ActorToWorld = Actor->ActorToWorld();
+				Vertices.Empty();
+				const FPositionVertexBuffer& VertexBuffer = Actor->GetStaticMeshComponent()->GetStaticMesh()->RenderData->LODResources[0].VertexBuffers.PositionVertexBuffer;
+				for (uint32 i = 0; i < VertexBuffer.GetNumVertices(); i++)
+				{
+					Vertices.AddUnique(ActorToWorld.TransformPosition(VertexBuffer.VertexPosition(i)));
+				}
+
+				const float InvDpiScale = 1.0f / Canvas->GetDPIScale();
+
+				FCanvasTileItem TileItem(FVector2D(0.0f, 0.0f), FVector2D(0.0f, 0.0f), FLinearColor::White);
+				TileItem.BlendMode = SE_BLEND_Translucent;
+				for (int32 VertexIndex = 0; VertexIndex < Vertices.Num(); ++VertexIndex)
+				{
+					const FVector& Vertex = Vertices[VertexIndex];
+					FVector2D PixelLocation;
+					if (View->ScreenToPixel(View->WorldToScreen(Vertex), PixelLocation))
+					{
+						PixelLocation *= InvDpiScale;
+
+						const bool bOutside =
+							PixelLocation.X < 0.0f || PixelLocation.X > View->UnscaledViewRect.Width()*InvDpiScale ||
+							PixelLocation.Y < 0.0f || PixelLocation.Y > View->UnscaledViewRect.Height()*InvDpiScale;
+						if (!bOutside)
+						{
+							const float X = PixelLocation.X - (TextureSizeX / 2);
+							const float Y = PixelLocation.Y - (TextureSizeY / 2);
+							if (bIsHitTesting)
+							{
+								Canvas->SetHitProxy(new HStaticMeshVert(Actor, Vertex));
+							}
+							TileItem.Texture = VertexTexture->Resource;
+
+							TileItem.Size = FVector2D(TextureSizeX, TextureSizeY);
+							Canvas->DrawItem(TileItem, FVector2D(X, Y));
+							if (bIsHitTesting)
+							{
+								Canvas->SetHitProxy(NULL);
+							}
+						}
+					}
+				}
+			}
+		}
+	}
+}
+
+void UEdMode::DrawBrackets(FEditorViewportClient* ViewportClient, FViewport* Viewport, const FSceneView* View, FCanvas* Canvas)
+{
+	USelection& SelectedActors = *Owner->GetSelectedActors();
+	for (int32 CurSelectedActorIndex = 0; CurSelectedActorIndex < SelectedActors.Num(); ++CurSelectedActorIndex)
+	{
+		AActor* SelectedActor = Cast<AActor>(SelectedActors.GetSelectedObject(CurSelectedActorIndex));
+		if (SelectedActor != NULL)
+		{
+			// Draw a bracket for selected "paintable" static mesh actors
+			const bool bIsValidActor = (Cast< AStaticMeshActor >(SelectedActor) != NULL);
+
+			const FLinearColor SelectedActorBoxColor(0.6f, 0.6f, 1.0f);
+			const bool bDrawBracket = bIsValidActor;
+			ViewportClient->DrawActorScreenSpaceBoundingBox(Canvas, View, Viewport, SelectedActor, SelectedActorBoxColor, bDrawBracket);
+		}
+	}
+}
+
+bool UEdMode::UsesToolkits() const
+{
+	return true;
+}
+
+UWorld* UEdMode::GetWorld() const
+{
+	return Owner->GetWorld();
+}
+
+class FEditorModeTools* UEdMode::GetModeManager() const
+{
+	return Owner;
+}
+
+bool UEdMode::StartTracking(FEditorViewportClient* InViewportClient, FViewport* InViewport)
+{
+	return false; 
+}
+
+bool UEdMode::EndTracking(FEditorViewportClient* InViewportClient, FViewport* InViewport)
+{
+	return true; 
+}
+
+AActor* UEdMode::GetFirstSelectedActorInstance() const
+{
+	return Owner->GetSelectedActors()->GetTop<AActor>();
+}
+
+
+void UEdMode::DeactivateAllActiveTools()
+{
+	ToolsContext->DeactivateAllActiveTools();
+	RestoreEditorState();
+}
+
+UInteractiveToolManager* UEdMode::GetToolManager() const
+{
+	return ToolsContext->ToolManager;
+}
+
+bool UEdMode::IsSnapRotationEnabled()
+{
+	return GetDefault<ULevelEditorViewportSettings>()->RotGridEnabled;
+}
+
+void UEdMode::PostInvalidation()
+{
+	bInvalidationPending = true;
+}
+
+void UEdMode::RestoreEditorState()
+{
+	if (bHaveSavedEditorState && !GIsRequestingExit)
+	{
+		GCurrentLevelEditingViewportClient->EngineShowFlags.SetAntiAliasing(bSavedAntiAliasingState);
+		bHaveSavedEditorState = false;
+	}
+}
diff --git a/Engine/Source/Editor/UnrealEd/Private/Tools/UEditorModes.cpp b/Engine/Source/Editor/UnrealEd/Private/Tools/UEditorModes.cpp
new file mode 100644
--- /dev/null
+++ b/Engine/Source/Editor/UnrealEd/Private/Tools/UEditorModes.cpp
@@ -0,0 +1,5 @@
+// Copyright 1998-2019 Epic Games, Inc. All Rights Reserved.
+
+#include "EditorModes.h"
+#include "Tools/UEdMode.h"
+
diff --git a/Engine/Source/Editor/UnrealEd/Private/UnrealEdEngine.cpp b/Engine/Source/Editor/UnrealEd/Private/UnrealEdEngine.cpp
--- a/Engine/Source/Editor/UnrealEd/Private/UnrealEdEngine.cpp
+++ b/Engine/Source/Editor/UnrealEd/Private/UnrealEdEngine.cpp
@@ -713,7 +713,7 @@ void UUnrealEdEngine::OnPostWindowsMessage(FViewport* Viewport, uint32 Message)
 void UUnrealEdEngine::OnOpenMatinee()
 {
 	// Register a delegate to pickup when Matinee is closed.
-	OnMatineeEditorClosedDelegateHandle = GLevelEditorModeTools().OnEditorModeChanged().AddUObject( this, &UUnrealEdEngine::OnMatineeEditorClosed );
+	UpdateEdModeOnMatineeCloseDelegateHandle = GLevelEditorModeTools().OnEditorModeIDChanged().AddUObject( this, &UUnrealEdEngine::UpdateEdModeOnMatineeClose );
 }
 
 bool UUnrealEdEngine::IsAutosaving() const
@@ -1437,7 +1437,7 @@ EWriteDisallowedWarningState UUnrealEdEngine::GetWarningStateForWritePermission(
 	return WarningState;
 }
 
-
+PRAGMA_DISABLE_DEPRECATION_WARNINGS
 void UUnrealEdEngine::OnMatineeEditorClosed( FEdMode* Mode, bool IsEntering )
 {
 	// if we are closing the Matinee editor
@@ -1453,3 +1453,20 @@ void UUnrealEdEngine::OnMatineeEditorClosed( FEdMode* Mode, bool IsEntering )
 		GLevelEditorModeTools().OnEditorModeChanged().Remove( OnMatineeEditorClosedDelegateHandle );
 	}	
 }
+PRAGMA_ENABLE_DEPRECATION_WARNINGS
+
+void UUnrealEdEngine::UpdateEdModeOnMatineeClose(const FEditorModeID& EditorModeID, bool IsEntering)
+{
+	// if we are closing the Matinee editor
+	if (!IsEntering && EditorModeID == FBuiltinEditorModes::EM_InterpEdit)
+	{
+		// set the autosave timer to save soon
+		if (PackageAutoSaver)
+		{
+			PackageAutoSaver->ForceMinimumTimeTillAutoSave();
+		}
+
+		// Remove this delegate. 
+		GLevelEditorModeTools().OnEditorModeIDChanged().Remove(UpdateEdModeOnMatineeCloseDelegateHandle);
+	}
+}
diff --git a/Engine/Source/Editor/UnrealEd/Private/UnrealEdSrv.cpp b/Engine/Source/Editor/UnrealEd/Private/UnrealEdSrv.cpp
--- a/Engine/Source/Editor/UnrealEd/Private/UnrealEdSrv.cpp
+++ b/Engine/Source/Editor/UnrealEd/Private/UnrealEdSrv.cpp
@@ -1567,14 +1567,9 @@ bool UUnrealEdEngine::Exec_Edit( UWorld* InWorld, const TCHAR* Str, FOutputDevic
 
 	if( FParse::Command(&Str,TEXT("CUT")) )
 	{
-		TArray<FEdMode*> ActiveModes;
-		GLevelEditorModeTools().GetActiveModes(ActiveModes);
-		for (int32 ModeIndex = 0; ModeIndex < ActiveModes.Num(); ++ModeIndex)
+		if (GLevelEditorModeTools().ProcessEditCut())
 		{
-			if (ActiveModes[ModeIndex]->ProcessEditCut())
-			{
-				return true;
-			}
+			return true;
 		}
 
 		if (bComponentsSelected)
@@ -1595,14 +1590,9 @@ bool UUnrealEdEngine::Exec_Edit( UWorld* InWorld, const TCHAR* Str, FOutputDevic
 	}
 	else if( FParse::Command(&Str,TEXT("COPY")) )
 	{
-		TArray<FEdMode*> ActiveModes;
-		GLevelEditorModeTools().GetActiveModes(ActiveModes);
-		for (int32 ModeIndex = 0; ModeIndex < ActiveModes.Num(); ++ModeIndex)
+		if (GLevelEditorModeTools().ProcessEditCopy())
 		{
-			if (ActiveModes[ModeIndex]->ProcessEditCopy())
-			{
-				return true;
-			}
+			return true;
 		}
 
 		if (bComponentsSelected)
@@ -1618,14 +1608,9 @@ bool UUnrealEdEngine::Exec_Edit( UWorld* InWorld, const TCHAR* Str, FOutputDevic
 	}
 	else if( FParse::Command(&Str,TEXT("PASTE")) )
 	{
-		TArray<FEdMode*> ActiveModes;
-		GLevelEditorModeTools().GetActiveModes(ActiveModes);
-		for (int32 ModeIndex = 0; ModeIndex < ActiveModes.Num(); ++ModeIndex)
+		if (GLevelEditorModeTools().ProcessEditPaste())
 		{
-			if (ActiveModes[ModeIndex]->ProcessEditPaste())
-			{
-				return true;
-			}
+			return true;
 		}
 
 		if (bComponentsSelected)
@@ -2558,13 +2543,9 @@ bool UUnrealEdEngine::Exec_Actor( UWorld* InWorld, const TCHAR* Str, FOutputDevi
 	}
 	else if( FParse::Command(&Str,TEXT("DELETE")) )		// ACTOR SELECT DELETE
 	{
+	
 		bool bHandled = false;
-		TArray<FEdMode*> ActiveModes; 
-		GLevelEditorModeTools().GetActiveModes( ActiveModes );
-		for( int32 ModeIndex = 0; ModeIndex < ActiveModes.Num(); ++ModeIndex )
-		{
-			bHandled |= ActiveModes[ModeIndex]->ProcessEditDelete();
-		}
+		bHandled |= GLevelEditorModeTools().ProcessEditDelete();
 
 		// if not specially handled by the current editing mode,
 		if (!bHandled)
@@ -2743,12 +2724,7 @@ bool UUnrealEdEngine::Exec_Actor( UWorld* InWorld, const TCHAR* Str, FOutputDevi
 	else if( FParse::Command(&Str,TEXT("DUPLICATE")) )
 	{
 		bool bHandled = false;
-		TArray<FEdMode*> ActiveModes; 
-		GLevelEditorModeTools().GetActiveModes( ActiveModes );
-		for( int32 ModeIndex = 0; ModeIndex < ActiveModes.Num(); ++ModeIndex )
-		{
-			bHandled |= ActiveModes[ModeIndex]->ProcessEditDuplicate();
-		}
+		bHandled |= GLevelEditorModeTools().ProcessEditDuplicate();
 
 		// if not specially handled by the current editing mode,
 		if (!bHandled)
@@ -3032,7 +3008,7 @@ bool UUnrealEdEngine::Exec_Mode( const TCHAR* Str, FOutputDevice& Ar )
 	if ( EditorMode == FBuiltinEditorModes::EM_None )
 	{
 		FString CommandToken = FParse::Token(Str, false);
-		FEdMode* FoundMode = GLevelEditorModeTools().FindMode( FName( *CommandToken ) );
+		FEdMode* FoundMode = GLevelEditorModeTools().GetActiveMode( FName( *CommandToken ) );
 
 		if ( FoundMode != NULL )
 		{
diff --git a/Engine/Source/Editor/UnrealEd/Private/UnrealWidget.cpp b/Engine/Source/Editor/UnrealEd/Private/UnrealWidget.cpp
--- a/Engine/Source/Editor/UnrealEd/Private/UnrealWidget.cpp
+++ b/Engine/Source/Editor/UnrealEd/Private/UnrealWidget.cpp
@@ -254,28 +254,10 @@ void FWidget::Render( const FSceneView* View,FPrimitiveDrawInterface* PDI, FEdit
 {
 	check(ViewportClient);
 
-	TArray<FEdMode*> ActiveModes;
-	if( EditorModeTools )
-	{
-		EditorModeTools->GetActiveModes( ActiveModes );
-	}
 
 	//reset HUD text
 	HUDString.Empty();
 
-	bool bDrawModeSupportsWidgetDrawing = bDefaultVisibility;
-
-	if( EditorModeTools && ActiveModes.Num() > 0)
-	{
-		bDrawModeSupportsWidgetDrawing = false;
-		// Check to see of any active modes support widget drawing
-		for( int32 ModeIndex = 0; ModeIndex < ActiveModes.Num(); ++ModeIndex )
-		{
-			bDrawModeSupportsWidgetDrawing |= ActiveModes[ModeIndex]->ShouldDrawWidget();
-		}
-	}
-
-
 	const bool bShowFlagsSupportsWidgetDrawing = View->Family->EngineShowFlags.ModeWidgets;
 	const bool bEditorModeToolsSupportsWidgetDrawing = EditorModeTools ? EditorModeTools->GetShowWidget() : true;
 	bool bDrawWidget;
@@ -284,7 +266,7 @@ void FWidget::Render( const FSceneView* View,FPrimitiveDrawInterface* PDI, FEdit
 	// editor object movement, we need to still run through the Render routine even though widget drawing may be
 	// disabled.  So we keep a flag that is used to determine whether or not to actually render anything.  This way
 	// we can still update the widget axis' based on the Context's transform matrices, even though drawing is disabled.
-	if(bDrawModeSupportsWidgetDrawing && bShowFlagsSupportsWidgetDrawing && bEditorModeToolsSupportsWidgetDrawing)
+	if(bDefaultVisibility && bShowFlagsSupportsWidgetDrawing && bEditorModeToolsSupportsWidgetDrawing)
 	{
 		bDrawWidget = true;
 
diff --git a/Engine/Source/Editor/UnrealEd/Public/EdMode.h b/Engine/Source/Editor/UnrealEd/Public/EdMode.h
--- a/Engine/Source/Editor/UnrealEd/Public/EdMode.h
+++ b/Engine/Source/Editor/UnrealEd/Public/EdMode.h
@@ -9,6 +9,7 @@
 #include "EditorComponents.h"
 #include "EngineGlobals.h"
 #include "EditorModeRegistry.h"
+#include "Tools/UEdMode.h"
 
 class FCanvas;
 class FEditorModeTools;
@@ -30,20 +31,6 @@ class FModeTool;
 class FEditorViewportClient;
 struct FViewportClick;
 
-/** Outcomes when determining whether it's possible to perform an action on the edit modes*/
-namespace EEditAction
-{
-	enum Type
-	{
-		/** Can't process this action */
-		Skip		= 0,
-		/** Can process this action */
-		Process,
-		/** Stop evaluating other modes (early out) */
-		Halt,
-	};
-};
-
 /**
  * Base class for all editor modes.
  */
diff --git a/Engine/Source/Editor/UnrealEd/Public/Editor.h b/Engine/Source/Editor/UnrealEd/Public/Editor.h
--- a/Engine/Source/Editor/UnrealEd/Public/Editor.h
+++ b/Engine/Source/Editor/UnrealEd/Public/Editor.h
@@ -98,6 +98,8 @@ struct UNREALED_API FEditorDelegates
 	DECLARE_MULTICAST_DELEGATE_TwoParams(FOnMapOpened, const FString& /* Filename */, bool /*bAsTemplate*/);
 	/** Delegate used for entering or exiting an editor mode */
 	DECLARE_MULTICAST_DELEGATE_OneParam(FOnEditorModeTransitioned, FEdMode* /*Mode*/);
+	/** Delegate used for entering or exiting an editor mode */
+	DECLARE_MULTICAST_DELEGATE_OneParam(FOnEditorModeIDTransitioned, const FEditorModeID& /*Mode*/);
 	/** delegate type to determine if a user requests can delete certain assets. */
 	DECLARE_MULTICAST_DELEGATE_TwoParams(FOnAssetsCanDelete, const TArray<UObject*>& /*InObjectToDelete*/, FCanDeleteAssetResult& /*OutCanDelete*/);
 	/** delegate type for when a user requests to delete certain assets... DOES NOT mean the asset(s) will be deleted (the user could cancel) */
@@ -166,9 +168,15 @@ struct UNREALED_API FEditorDelegates
 	/** Called when load errors are about to be displayed */
 	static FSimpleMulticastDelegate DisplayLoadErrors;
 	/** Called when an editor mode is being entered */
+	UE_DEPRECATED(4.24, "Use EditorModeIDEnter instead")
 	static FOnEditorModeTransitioned EditorModeEnter;
 	/** Called when an editor mode is being exited */
+	UE_DEPRECATED(4.24, "Use EditorModeIDExit instead")
 	static FOnEditorModeTransitioned EditorModeExit;
+	/** Called when an editor mode ID is being entered */
+	static FOnEditorModeIDTransitioned EditorModeIDEnter;
+	/** Called when an editor mode ID is being exited */
+	static FOnEditorModeIDTransitioned EditorModeIDExit;
 	/** Sent when a PIE session is beginning (before we decide if PIE can run - allows clients to avoid blocking PIE) */
 	static FOnPIEEvent PreBeginPIE;
 	/** Sent when a PIE session is beginning (but hasn't actually started yet) */
diff --git a/Engine/Source/Editor/UnrealEd/Public/EditorModeManager.h b/Engine/Source/Editor/UnrealEd/Public/EditorModeManager.h
--- a/Engine/Source/Editor/UnrealEd/Public/EditorModeManager.h
+++ b/Engine/Source/Editor/UnrealEd/Public/EditorModeManager.h
@@ -8,6 +8,7 @@
 #include "UnrealWidget.h"
 #include "Editor.h"
 #include "EditorUndoClient.h"
+#include "EdMode.h"
 
 class FCanvas;
 class FEditorViewportClient;
@@ -20,6 +21,7 @@ class IToolkitHost;
 class USelection;
 struct FConvexVolume;
 struct FViewportClick;
+class UEdMode;
 
 /**
  * A helper class to store the state of the various editor modes.
@@ -95,6 +97,8 @@ protected:
 	
 	/** Deactivates the editor mode at the specified index */
 	void DeactivateModeAtIndex( int32 InIndex );
+	/** Deactivates the editor mode at the specified index */
+	void DeactivateScriptableModeAtIndex(int32 InIndex);
 		
 private:
 	void RebuildModeToolBar();
@@ -113,8 +117,11 @@ public:
 	/** 
 	 * Returns the editor mode specified by the passed in ID
 	 */
+	UE_DEPRECATED(4.24, "Use GetActiveMode instead.")
 	FEdMode* FindMode( FEditorModeID InID );
 
+	UEdMode* GetActiveScriptableMode(FEditorModeID InID);
+
 	/**
 	 * Returns true if the current mode is not the specified ModeID.  Also optionally warns the user.
 	 *
@@ -162,10 +169,11 @@ public:
 	/** 
 	 * Returns an array of all active modes
 	 */
+	UE_DEPRECATED(4.24, "All access to modes now needs to go through the mode manager")
 	void GetActiveModes( TArray<FEdMode*>& OutActiveModes );
 
 	void SetShowWidget( bool InShowWidget )	{ bShowWidget = InShowWidget; }
-	bool GetShowWidget() const				{ return bShowWidget; }
+	bool GetShowWidget() const;
 
 	/** Cycle the widget mode, forwarding queries to modes */
 	void CycleWidgetMode (void);
@@ -419,8 +427,14 @@ public:
 	 * First parameter:  The editor mode that was changed
 	 * Second parameter:  True if entering the mode, or false if exiting the mode
 	 */
+PRAGMA_DISABLE_DEPRECATION_WARNINGS
 	DECLARE_EVENT_TwoParams( FEditorModeTools, FEditorModeChangedEvent, FEdMode*, bool );
+UE_DEPRECATED(4.24, "Use OnEditorModeIDChanged() instead.")
 	FEditorModeChangedEvent& OnEditorModeChanged() { return EditorModeChangedEvent; }
+PRAGMA_ENABLE_DEPRECATION_WARNINGS
+
+	DECLARE_EVENT_TwoParams(FEditorModeTools, FEditorModeIDChangedEvent, const FEditorModeID&, bool);
+	FEditorModeIDChangedEvent& OnEditorModeIDChanged() { return EditorModeIDChangedEvent; }
 
 	/** delegate type for triggering when widget mode changed */
 	DECLARE_EVENT_OneParam( FEditorModeTools, FWidgetModeChangedEvent, FWidget::EWidgetMode );
@@ -430,8 +444,12 @@ public:
 	void BroadcastWidgetModeChanged(FWidget::EWidgetMode InWidgetMode) { WidgetModeChangedEvent.Broadcast(InWidgetMode); }
 
 	/**	Broadcasts the EditorModeChanged event */
+	UE_DEPRECATED(4.24, "Use BroadcastEditorModeIDChanged() instead.")
 	void BroadcastEditorModeChanged(FEdMode* Mode, bool IsEnteringMode) { EditorModeChangedEvent.Broadcast(Mode, IsEnteringMode); }
 
+	/**	Broadcasts the EditorModeIDChanged event */
+	void BroadcastEditorModeIDChanged(const FEditorModeID& ModeID, bool IsEnteringMode) { EditorModeIDChangedEvent.Broadcast(ModeID, IsEnteringMode); }
+
 	/**
 	 * Returns the current CoordSystem
 	 * 
@@ -502,6 +520,32 @@ public:
 	 * Whether or not the current selection has a scene component selected
  	 */
 	bool SelectionHasSceneComponent() const;
+
+	bool IsSelectionAllowed(AActor* InActor, const bool bInSelected) const;
+
+	bool IsSelectionHandled(AActor* InActor, const bool bInSelected) const;
+
+	bool ProcessEditDuplicate();
+	bool ProcessEditDelete();
+	bool ProcessEditCut();
+	bool ProcessEditCopy();
+	bool ProcessEditPaste();
+	EEditAction::Type  GetActionEditDuplicate();
+	EEditAction::Type  GetActionEditDelete();
+	EEditAction::Type  GetActionEditCut();
+	EEditAction::Type  GetActionEditCopy();
+	EEditAction::Type GetActionEditPaste();
+
+	void DeactivateOtherVisibleModes(FEditorModeID InMode);
+	bool IsSnapRotationEnabled() const;
+	bool SnapRotatorToGridOverride(FRotator& InRotation) const;
+	void ActorsDuplicatedNotify(TArray<AActor*>& InPreDuplicateSelection, TArray<AActor*>& InPostDuplicateSelection, const bool bOffsetLocations);
+	void ActorMoveNotify();
+	void ActorSelectionChangeNotify();
+	void ActorPropChangeNotify();
+	void UpdateInternalData();
+	bool IsOnlyVisibleActiveMode(FEditorModeID InMode) const;
+
 protected:
 	/** 
 	 * Delegate handlers
@@ -515,12 +559,18 @@ protected:
 	/** A list of active editor modes. */
 	TArray< TSharedPtr<FEdMode> > ActiveModes;
 
+	/** A list of active editor modes. */
+	TArray< UEdMode* > ActiveScriptableModes;
+
 	/** The host of the toolkits created by these modes */
 	TWeakPtr<IToolkitHost> ToolkitHost;
 
 	/** A list of previously active editor modes that we will potentially recycle */
 	TMap< FEditorModeID, TSharedPtr<FEdMode> > RecycledModes;
 
+	/** A list of previously active editor modes that we will potentially recycle */
+	TMap< FEditorModeID, UEdMode* > RecycledScriptableModes;
+
 	/** The mode that the editor viewport widget is in. */
 	FWidget::EWidgetMode WidgetMode;
 
@@ -543,11 +593,11 @@ protected:
 private:
 	struct FEdModeToolbarRow
 	{
-		FEdModeToolbarRow(TSharedPtr<FEdMode>& InMode, TSharedRef<SWidget>& InToolbarWidget)
-			: Mode(InMode)
+		FEdModeToolbarRow(FName InModeID, TSharedRef<SWidget>& InToolbarWidget)
+			: ModeID(InModeID)
 			, ToolbarWidget(InToolbarWidget)
 		{}
-		TSharedPtr<FEdMode> Mode;
+		FName ModeID;
 		TSharedPtr<SWidget> ToolbarWidget;
 	};
 
@@ -560,6 +610,9 @@ private:
 	/** Multicast delegate that is broadcast when a mode is entered or exited */
 	FEditorModeChangedEvent EditorModeChangedEvent;
 
+	/** Multicast delegate that is broadcast when a mode is entered or exited */
+	FEditorModeIDChangedEvent EditorModeIDChangedEvent;
+
 	/** Multicast delegate that is broadcast when a widget mode is changed */
 	FWidgetModeChangedEvent WidgetModeChangedEvent;
 
diff --git a/Engine/Source/Editor/UnrealEd/Public/EditorModeRegistry.h b/Engine/Source/Editor/UnrealEd/Public/EditorModeRegistry.h
--- a/Engine/Source/Editor/UnrealEd/Public/EditorModeRegistry.h
+++ b/Engine/Source/Editor/UnrealEd/Public/EditorModeRegistry.h
@@ -7,13 +7,13 @@
 #include "Editor.h"
 
 class FEditorModeTools;
-class FEdMode;
 
 // Required forward declarations
 class FEdMode;
+class UEdMode;
 
 DECLARE_DELEGATE_RetVal(TSharedRef<FEdMode>, FEditorModeFactoryCallback);
-
+DECLARE_DELEGATE_RetVal(UEdMode*, FEditorModeFactoryDelegate);
 /**
  *	Class responsible for maintaining a list of registered editor mode types.
  *
@@ -90,6 +90,10 @@ struct IEditorModeFactory : public TSharedFromThis<IEditorModeFactory>
 	 * Create a new instance of our mode
 	 */
 	virtual TSharedRef<FEdMode> CreateMode() const = 0;
+
+	virtual UEdMode* CreateScriptableMode() const = 0;
+
+	virtual bool ForScriptableMode() const = 0;
 };
 
 struct FEditorModeFactory : IEditorModeFactory
@@ -112,6 +116,39 @@ struct FEditorModeFactory : IEditorModeFactory
 	 * Create a new instance of our mode
 	 */
 	virtual TSharedRef<FEdMode> CreateMode() const override { return FactoryCallback.Execute(); }
+
+	virtual UEdMode* CreateScriptableMode() const override { return nullptr; }
+
+	virtual bool ForScriptableMode() const override { return false; }
+};
+
+class UEditorModeFactory : public UObject, public IEditorModeFactory
+{
+	UEditorModeFactory(FEditorModeInfo InModeInfo) : ModeInfo(InModeInfo) {}
+	virtual ~UEditorModeFactory() {}
+
+	/** Information pertaining to this factory's mode */
+	FEditorModeInfo ModeInfo;
+
+	/** Callback used to create an instance of this mode type */
+	FEditorModeFactoryCallback FactoryCallback;
+
+	/** Callback used to create an instance of this mode type */
+	FEditorModeFactoryDelegate FactoryDelegate;
+
+	/**
+	 * Gets the information pertaining to the mode type that this factory creates
+	 */
+	virtual FEditorModeInfo GetModeInfo() const override { return ModeInfo; }
+
+	virtual TSharedRef<FEdMode> CreateMode() const override { return FactoryCallback.Execute(); }
+
+	/**
+	* Create a new instance of our mode
+	*/
+	virtual UEdMode* CreateScriptableMode() const override { return FactoryDelegate.Execute(); }
+
+	virtual bool ForScriptableMode() const override { return true; }
 };
 
 /**
@@ -202,6 +239,11 @@ public:
 	 */
 	TSharedPtr<FEdMode> CreateMode(FEditorModeID ModeID, FEditorModeTools& Owner);
 
+	/**
+	 * Create a new instance of the mode registered under the specified ID
+	 */
+	UEdMode* CreateScriptableMode(FEditorModeID ModeID, FEditorModeTools& Owner);
+
 	/**
 	 * Const access to the internal factory map
 	 */
@@ -213,9 +255,6 @@ private:
 	/** A map of editor mode IDs to factory callbacks */
 	FactoryMap ModeFactories;
 	
-	/** A list of all modes created */
-	TArray<TWeakPtr<FEdMode>> CreatedModes;
-
 	/** Event that is triggered whenever a mode is registered or unregistered */
 	FRegisteredModesChangedEvent RegisteredModesChanged;
 
diff --git a/Engine/Source/Editor/UnrealEd/Public/EditorViewportClient.h b/Engine/Source/Editor/UnrealEd/Public/EditorViewportClient.h
--- a/Engine/Source/Editor/UnrealEd/Public/EditorViewportClient.h
+++ b/Engine/Source/Editor/UnrealEd/Public/EditorViewportClient.h
@@ -17,6 +17,7 @@
 #include "SceneManagement.h"
 #include "EditorComponents.h"
 #include "Framework/Commands/Commands.h"
+#include "Editor.h"
 
 struct FAssetData;
 class FCachedJoystickState;
@@ -458,7 +459,9 @@ public:
 	void TakeHighResScreenShot();
 
 	/** Called when an editor mode has been (de)activated */
+	UE_DEPRECATED(4.24, "Use OnEditorModeIDChanged() instead.")
 	void OnEditorModeChanged(FEdMode* EditorMode, bool bIsEntering);
+	void OnEditorModeIDChanged(const FEditorModeID& EditorModeID, bool bIsEntering);
 
 	/** FViewElementDrawer interface */
 	virtual void Draw(const FSceneView* View,FPrimitiveDrawInterface* PDI) override;
diff --git a/Engine/Source/Editor/UnrealEd/Public/PhysicsManipulationMode.h b/Engine/Source/Editor/UnrealEd/Public/PhysicsManipulationMode.h
--- a/Engine/Source/Editor/UnrealEd/Public/PhysicsManipulationMode.h
+++ b/Engine/Source/Editor/UnrealEd/Public/PhysicsManipulationMode.h
@@ -8,6 +8,7 @@
 class FEditorModeTools;
 class FEditorViewportClient;
 class FViewport;
+class UEdMode;
 
 class FPhysicsManipulationEdModeFactory : public IEditorModeFactory
 {
@@ -19,6 +20,10 @@ class FPhysicsManipulationEdModeFactory : public IEditorModeFactory
 
 	/** Create a new instance of our mode */
 	virtual TSharedRef<FEdMode> CreateMode() const override;
+
+	virtual UEdMode* CreateScriptableMode() const override { return nullptr; }
+
+	virtual bool ForScriptableMode() const override { return false; }
 };
 
 /** 
diff --git a/Engine/Source/Editor/UnrealEd/Public/Toolkits/BaseToolkit.h b/Engine/Source/Editor/UnrealEd/Public/Toolkits/BaseToolkit.h
--- a/Engine/Source/Editor/UnrealEd/Public/Toolkits/BaseToolkit.h
+++ b/Engine/Source/Editor/UnrealEd/Public/Toolkits/BaseToolkit.h
@@ -9,7 +9,11 @@
 #include "Framework/Commands/UICommandList.h"
 #include "Toolkits/IToolkitHost.h"
 
+class UInteractiveTool;
 class SDockableTab;
+class UEdMode;
+class IDetailsView;
+enum class EToolShutdownType;
 
 /**
  * Base class for all toolkits (abstract).
@@ -89,8 +93,10 @@ protected:
 };
 
 
+
 /**
- * Base class for all editor mode toolkits.
+ * This FModeToolkit just creates a basic UI panel that allows various InteractiveTools to
+ * be initialized, and a DetailsView used to show properties of the active Tool.
  */
 class UNREALED_API FModeToolkit
 	: public FBaseToolkit
@@ -111,10 +117,30 @@ public:
 public:
 
 	// IToolkit interface
+	virtual FName GetToolkitFName() const override;
+	virtual FText GetBaseToolkitName() const override;
 	virtual FText GetToolkitName() const override { return GetBaseToolkitName(); }
 	virtual FText GetToolkitToolTipText() const override { return GetBaseToolkitName(); }
 	virtual FString GetWorldCentricTabPrefix() const override;
 	virtual bool IsAssetEditor() const override;
 	virtual const TArray<UObject*>* GetObjectsCurrentlyBeingEdited() const override;
 	virtual FLinearColor GetWorldCentricTabColorScale() const override;
+	virtual class FEdMode* GetEditorMode() const override;
+
+	virtual UEdMode* GetScriptableEditorMode() const;
+
+private:
+
+	TSharedPtr<SWidget> ToolkitWidget;
+	TSharedPtr<IDetailsView> DetailsView;
+
+	// these functions just forward calls to the ToolsContext / ToolManager
+
+	bool CanStartTool(const FString& ToolTypeIdentifier);
+	bool CanAcceptActiveTool();
+	bool CanCancelActiveTool();
+	bool CanCompleteActiveTool();
+
+	FReply StartTool(const FString& ToolTypeIdentifier);
+	FReply EndTool(EToolShutdownType ShutdownType);
 };
diff --git a/Engine/Source/Editor/UnrealEd/Public/Tools/EditorComponentSourceFactory.h b/Engine/Source/Editor/UnrealEd/Public/Tools/EditorComponentSourceFactory.h
new file mode 100644
--- /dev/null
+++ b/Engine/Source/Editor/UnrealEd/Public/Tools/EditorComponentSourceFactory.h
@@ -0,0 +1,25 @@
+// Copyright 1998-2019 Epic Games, Inc. All Rights Reserved.
+
+#pragma once
+
+#include "CoreMinimal.h"
+#include "ComponentSourceInterfaces.h"
+#include "Components/PrimitiveComponent.h"
+#include "Components/StaticMeshComponent.h"
+
+class UNREALED_API FStaticMeshComponentTargetFactory : public FComponentTargetFactory
+{
+public:
+	bool CanBuild( UActorComponent* Candidate ) override;
+	TUniquePtr<FPrimitiveComponentTarget> Build( UPrimitiveComponent* PrimitiveComponent ) override;
+};
+
+class UNREALED_API FStaticMeshComponentTarget : public FPrimitiveComponentTarget
+{
+public:
+	FStaticMeshComponentTarget( UPrimitiveComponent* Component )
+		: FPrimitiveComponentTarget( Cast<UStaticMeshComponent>(Component) ){}
+	FMeshDescription* GetMesh() override;
+	void CommitMesh( const FCommitter& ) override;
+	static const int LODIndex{0};
+};
diff --git a/Engine/Source/Editor/UnrealEd/Public/Tools/EditorToolAssetAPI.h b/Engine/Source/Editor/UnrealEd/Public/Tools/EditorToolAssetAPI.h
new file mode 100644
--- /dev/null
+++ b/Engine/Source/Editor/UnrealEd/Public/Tools/EditorToolAssetAPI.h
@@ -0,0 +1,28 @@
+// Copyright 1998-2019 Epic Games, Inc. All Rights Reserved.
+
+#pragma once
+
+#include "CoreMinimal.h"
+#include "ToolContextInterfaces.h"
+
+/**
+ * Implementation of ToolsContext Asset management API that is suitable for use
+ * inside the Editor (eg inside an FEdMode)
+ */
+class UNREALED_API FEditorToolAssetAPI : public IToolsContextAssetAPI
+{
+public:
+	virtual FString GetActiveAssetFolderPath() override;
+
+	virtual FString MakePackageName(const FString& AssetName, const FString& FolderPath) override;
+
+	virtual FString MakeUniqueAssetName(const FString& AssetName, const FString& FolderPath) override;
+
+	// FolderPath should neither start with, nor end with, path separators
+	// Generated path will be /Game/FolderPath/AssetName
+	virtual UPackage* CreateNewPackage(const FString& AssetName, const FString& FolderPath) override;
+
+	virtual void InteractiveSaveGeneratedAsset(UObject* Asset, UPackage* AssetPackage) override;
+
+	virtual void AutoSaveGeneratedAsset(UObject* Asset, UPackage* AssetPackage) override;
+};
\ No newline at end of file
diff --git a/Engine/Source/Editor/UnrealEd/Public/Tools/InteractiveToolsCommands.h b/Engine/Source/Editor/UnrealEd/Public/Tools/InteractiveToolsCommands.h
new file mode 100644
--- /dev/null
+++ b/Engine/Source/Editor/UnrealEd/Public/Tools/InteractiveToolsCommands.h
@@ -0,0 +1,296 @@
+// Copyright 1998-2019 Epic Games, Inc. All Rights Reserved.
+
+#pragma once
+
+#include "CoreMinimal.h"
+#include "Framework/Commands/Commands.h"
+#include "Framework/Commands/UICommandList.h"
+#include "InteractiveTool.h"
+#include "StandardToolModeCommands.h"
+
+
+
+/**
+ * TInteractiveToolCommands is a base class that handles connecting up Tool Actions 
+ * (ie the FInteractiveToolAction provided by a UInteractiveTool) to the UnrealEditor
+ * Command system, which allows for remappable hotkeys, etc
+ * 
+ * Usage is as follows:
+ *    - in your EdMode Module, create a subclass-instance of this, say FMyToolCommands and call FMyToolCommands::Register() in your ::StartupModule() function
+ *        - subclass must implement ::GetToolDefaultObjectList(), here you just add GetMutableDefault<MyToolX> to the input list for all your Tools
+ *    - add a member TSharedPtr<FUICommandList> UICommandList; to your EdMode impl
+ *    - when you start a new Tool, call FMyToolCommands::Get().BindCommandsForCurrentTool(UICommandList, NewTool)
+ *    - when you end a Tool, call FMyToolCommands::Get().UnbindActiveCommands()
+ *    - in your EdModeImpl::InputKey override, call UICommandList->ProcessCommandBindings()
+ */
+template<typename CommandContextType>
+class TInteractiveToolCommands : public TCommands<CommandContextType>
+{
+
+public:
+	/**
+	 * Initialize commands. Collects possible Actions from the available UInteractiveTools
+	 * (provided by GetToolDefaultObjectList) and then registers a FUICommandInfo for each.
+	 * This needs to be called in ::StartupModule for your EdMode Module
+	 */
+	virtual void RegisterCommands() override;
+
+	/**
+	 * Bind any of the registered UICommands to the given Tool, if they are compatible.
+	 */
+	virtual void BindCommandsForCurrentTool(TSharedPtr<FUICommandList> UICommandList, UInteractiveTool* Tool) const;
+
+	/**
+	 * Unbind all of the currently-registered commands
+	 */
+	virtual void UnbindActiveCommands(TSharedPtr<FUICommandList> UICommandList) const;
+
+
+	//
+	// Interface that subclasses need to implement
+	// 
+
+
+	/**
+	 * RegisterCommands() needs actual UInteractiveTool instances for all the Tools that want to
+	 * provide Actions which will be connected to hotkeys. We can do this based on the CDO's for
+	 * each tool, returned by GetMutableDefault<UYourInteractiveToolType>(). The Tool CDOs are
+	 * not owned by a ToolManager and we will only call .GetActionSet() on them.
+	 */
+	virtual void GetToolDefaultObjectList(TArray<UInteractiveTool*>& ToolCDOs) = 0;
+
+
+protected:
+
+	/** Forwarding constructor */
+	TInteractiveToolCommands(const FName InContextName, const FText& InContextDesc, const FName InContextParent, const FName InStyleSetName)
+		: TCommands<CommandContextType>(InContextName, InContextDesc, InContextParent, InStyleSetName)
+	{
+	}
+
+	/**
+	 * Query FStandardToolModeCommands to find an existing command/hotkey for this standard tool action
+	 */
+	virtual TSharedPtr<FUICommandInfo> FindStandardCommand(EStandardToolActions ToolActionID);
+
+	/**
+	 * List of pairs of known Tool Actions and their associated UICommands.
+	 * This list is used at bind-time to connect up the already-registered UICommands to the active Tool
+	 */
+	struct FToolActionCommand
+	{
+		FInteractiveToolAction ToolAction;
+		TSharedPtr<FUICommandInfo> UICommand;
+	};
+	TArray<FToolActionCommand> ActionCommands;
+
+
+	//
+	// Support for sharing common commands between Tools where the commands
+	// are not shared across Tool Modes (and hence not in FStandardToolModeCommands)
+	// 
+
+	/** List FToolActionCommands for standard Tool Actions */
+	TMap<EStandardToolActions, FToolActionCommand> SharedActionCommands;
+
+	/**
+	 * Find or Create a UICommand for a standard Tool Action, that will be shared across Tools
+	 */
+	bool FindOrCreateSharedCommand(EStandardToolActions ActionID, FToolActionCommand& FoundOut);
+
+	/**
+	 * Create a suitable FInteractiveToolAction for the standard Tool Action, ie with suitable
+	 * command name and description strings and hotkeys. This info will be used to create the shared UICommand.
+	 */
+	virtual bool IntializeStandardToolAction(EStandardToolActions ActionID, FInteractiveToolAction& ToolActionOut);
+
+	/**
+	 * Utility function that registeres a Tool Aciton as a UICommand
+	 */
+	void RegisterUIToolCommand(const FInteractiveToolAction& ToolAction, TSharedPtr<FUICommandInfo>& UICommandInfo);
+};
+
+
+
+
+
+
+template<typename CommandContextType>
+void TInteractiveToolCommands<CommandContextType>::RegisterCommands()
+{
+	// make sure standard commands are registered
+	FStandardToolModeCommands::Register();
+
+	// get list of all tools used by command set
+	TArray<UInteractiveTool*> AllToolsCDOs;
+	GetToolDefaultObjectList(AllToolsCDOs);
+
+	// collect all the actions in these tools
+	TArray<FInteractiveToolAction> ToolActions;
+	for (UInteractiveTool* Tool : AllToolsCDOs)
+	{
+		Tool->GetActionSet()->CollectActions(ToolActions);
+	}
+
+	// register all the commands in all the actions
+	int NumActions = ToolActions.Num();
+	ActionCommands.SetNum(NumActions);
+	for (int k = 0; k < NumActions; ++k)
+	{
+		const FInteractiveToolAction& ToolAction = ToolActions[k];
+		ActionCommands[k].ToolAction = ToolAction;
+
+		bool bRegistered = false;
+		if ((int32)ToolAction.ActionID < (int32)EStandardToolActions::BaseClientDefinedActionID)
+		{
+			TSharedPtr<FUICommandInfo> FoundStandard = FindStandardCommand((EStandardToolActions)ToolAction.ActionID);
+			if (FoundStandard.IsValid())
+			{
+				ActionCommands[k].UICommand = FoundStandard;
+				bRegistered = true;
+			}
+			else
+			{
+				FToolActionCommand StandardActionCommand;
+				bool bFoundRemap = FindOrCreateSharedCommand((EStandardToolActions)ToolAction.ActionID, StandardActionCommand);
+				if (bFoundRemap)
+				{
+					ActionCommands[k].UICommand = StandardActionCommand.UICommand;
+					bRegistered = true;
+				}
+			}
+		}
+
+		if (bRegistered == false)
+		{
+			RegisterUIToolCommand(ActionCommands[k].ToolAction, ActionCommands[k].UICommand);
+		}
+	}
+}
+
+
+template<typename CommandContextType>
+void TInteractiveToolCommands<CommandContextType>::BindCommandsForCurrentTool(TSharedPtr<FUICommandList> UICommandList, UInteractiveTool* Tool) const
+{
+	UClass* ToolClassType = Tool->GetClass();
+
+	int NumActionCommands = ActionCommands.Num();
+	for (int32 k = 0; k < NumActionCommands; ++k)
+	{
+		const FToolActionCommand& ActionCommand = ActionCommands[k];
+		if (ActionCommand.ToolAction.ClassType == ToolClassType)
+		{
+			UICommandList->MapAction(
+				ActionCommand.UICommand,
+				FExecuteAction::CreateUObject(Tool, &UInteractiveTool::ExecuteAction, ActionCommand.ToolAction.ActionID));
+			//FCanExecuteAction::CreateRaw(this, &FEdModeFoliage::CurrentToolUsesBrush));
+		}
+	}
+}
+
+
+template<typename CommandContextType>
+void TInteractiveToolCommands<CommandContextType>::UnbindActiveCommands(TSharedPtr<FUICommandList> UICommandList) const
+{
+	// @todo would be more efficient if we kept track of which commands were mapped.
+	// However currently this function must be const because TCommands::Get() returns a const
+	for (const FToolActionCommand& ActionCommand : ActionCommands)
+	{
+		if (UICommandList->IsActionMapped(ActionCommand.UICommand))
+		{
+			UICommandList->UnmapAction(ActionCommand.UICommand);
+		}
+	}
+}
+
+
+
+
+
+
+//
+// Support for re-using known commands from FStandardToolModeCommands
+//
+
+template<typename CommandContextType>
+TSharedPtr<FUICommandInfo> TInteractiveToolCommands<CommandContextType>::FindStandardCommand(EStandardToolActions ToolActionID)
+{
+	const FStandardToolModeCommands& StdCommands = FStandardToolModeCommands::Get();
+	switch (ToolActionID)
+	{
+	case EStandardToolActions::IncreaseBrushSize:
+		return StdCommands.FindStandardCommand(EStandardToolModeCommands::IncreaseBrushSize);
+	case EStandardToolActions::DecreaseBrushSize:
+		return StdCommands.FindStandardCommand(EStandardToolModeCommands::DecreaseBrushSize);
+	default:
+		return nullptr;
+	}
+}
+
+
+
+
+
+
+//
+// Support for sharing UI commands between multiple Tools, where the
+// shared command is not at the shared-between-modes level
+// 
+
+template<typename CommandContextType>
+bool TInteractiveToolCommands<CommandContextType>::FindOrCreateSharedCommand(EStandardToolActions ActionID, FToolActionCommand& FoundOut)
+{
+	if (SharedActionCommands.Contains(ActionID))
+	{
+		FoundOut = SharedActionCommands[ActionID];
+		return true;
+	}
+
+	FInteractiveToolAction NewAction;
+	bool bIsKnownAction = IntializeStandardToolAction(ActionID, NewAction);
+
+	if (bIsKnownAction)
+	{
+		check(NewAction.ActionID == (int32)ActionID);
+		FToolActionCommand NewActionCommand;
+		NewActionCommand.ToolAction = NewAction;
+		RegisterUIToolCommand(NewActionCommand.ToolAction, NewActionCommand.UICommand);
+		FoundOut = SharedActionCommands.Add(ActionID, NewActionCommand);
+		return true;
+	}
+	return false;
+}
+
+
+template<typename CommandContextType>
+bool TInteractiveToolCommands<CommandContextType>::IntializeStandardToolAction(EStandardToolActions ActionID, FInteractiveToolAction& ToolActionOut)
+{
+	// base class does not have any standard Tool actions. Subclasses can override
+	// this function, populating ToolActionOut and returning true, if the ActionID should be shared between multiple Tools
+	
+	return false;
+}
+
+
+
+
+template<typename CommandContextType>
+void TInteractiveToolCommands<CommandContextType>::RegisterUIToolCommand(const FInteractiveToolAction& ToolAction, TSharedPtr<FUICommandInfo>& UICommandInfo)
+{
+	const FString DotString = FString(TEXT(".")) + ToolAction.ActionName;
+
+	FUICommandInfo::MakeCommandInfo(
+		this->AsShared(),
+		UICommandInfo,
+		*ToolAction.ActionName,
+		ToolAction.ShortName,
+		ToolAction.Description,
+		FSlateIcon(this->GetStyleSetName(), ISlateStyle::Join(this->GetContextName(), TCHAR_TO_ANSI(*DotString))),
+		EUserInterfaceActionType::Button,
+		FInputChord(ToolAction.DefaultModifiers, ToolAction.DefaultKey)
+		);
+}
+
+
+
+
diff --git a/Engine/Source/Editor/UnrealEd/Public/Tools/StandardToolModeCommands.h b/Engine/Source/Editor/UnrealEd/Public/Tools/StandardToolModeCommands.h
new file mode 100644
--- /dev/null
+++ b/Engine/Source/Editor/UnrealEd/Public/Tools/StandardToolModeCommands.h
@@ -0,0 +1,54 @@
+// Copyright 1998-2019 Epic Games, Inc. All Rights Reserved.
+
+#pragma once
+
+#include "CoreMinimal.h"
+#include "Framework/Commands/Commands.h"
+#include "Framework/Commands/UICommandInfo.h"
+
+
+/**
+ * Integer identifiers for "Standard" Tool commands, that many EditorModes / AssetEditors / etc may share.
+ * This allows a single hotkey binding to be used across many contexts.
+ */
+enum class EStandardToolModeCommands
+{
+	IncreaseBrushSize = 100,
+	DecreaseBrushSize = 101
+};
+
+
+
+/**
+ * FStandardToolModeCommands provides standard commands for Tools. This allows
+ * for "sharing" common hotkey bindings between multiple EditorModes
+ * 
+ * The set of standard commands is defined by EStandardToolModeCommands.
+ * You must call FStandardToolModeCommands::Register() on module startup to register these commands
+ * (note that they may already be registered by another module, but you should call this to be safe).
+ * 
+ * Then, when you want to Bind a standard command in your FUICommandList,
+ * call FStandardToolModeCommands::Get().FindStandardCommand() to get the registered UICommandInfo.
+ * 
+ */
+class UNREALED_API FStandardToolModeCommands : public TCommands<FStandardToolModeCommands>
+{
+public:
+	FStandardToolModeCommands();
+
+	/**
+	 * Registers the set of standard commands. Call on module startup.
+	 */
+	virtual void RegisterCommands() override;
+
+	/**
+	 * Look up the UICommandInfo for a standard command
+	 */
+	TSharedPtr<FUICommandInfo> FindStandardCommand(EStandardToolModeCommands Command) const;
+
+protected:
+	TMap<EStandardToolModeCommands, TSharedPtr<FUICommandInfo>> Commands;
+
+};
+
+
diff --git a/Engine/Source/Editor/UnrealEd/Public/Tools/UEdMode.h b/Engine/Source/Editor/UnrealEd/Public/Tools/UEdMode.h
new file mode 100644
--- /dev/null
+++ b/Engine/Source/Editor/UnrealEd/Public/Tools/UEdMode.h
@@ -0,0 +1,307 @@
+// Copyright 1998-2019 Epic Games, Inc. All Rights Reserved.
+
+#pragma once
+#include "CoreMinimal.h"
+#include "InputCoreTypes.h"
+#include "UnrealWidget.h"
+#include "EditorComponents.h"
+#include "EngineGlobals.h"
+#include "EditorModeRegistry.h"
+#include "Templates/SharedPointer.h"
+#include "UObject/Object.h"
+#include "InteractiveToolManager.h"
+#include "Math/Ray.h"
+#include "UEdMode.generated.h"
+
+
+class FCanvas;
+class FEditorModeTools;
+class FEditorViewportClient;
+class FModeTool;
+class FModeToolkit;
+class FPrimitiveDrawInterface;
+class FSceneView;
+class FViewport;
+class UTexture2D;
+struct FConvexVolume;
+struct FViewportClick;
+class UInteractiveToolsContext;
+
+enum EModeTools : int8;
+class FEditorViewportClient;
+class HHitProxy;
+struct FViewportClick;
+class FModeTool;
+class FEditorViewportClient;
+struct FViewportClick;
+
+
+/** Outcomes when determining whether it's possible to perform an action on the edit modes*/
+namespace EEditAction
+{
+	enum Type
+	{
+		/** Can't process this action */
+		Skip = 0,
+		/** Can process this action */
+		Process,
+		/** Stop evaluating other modes (early out) */
+		Halt,
+	};
+};
+
+/**
+ * Base class for all editor modes.
+ */
+UCLASS()
+class UNREALED_API UEdMode : public UObject, public FEditorCommonDrawHelper
+{
+	GENERATED_BODY()
+
+public:
+	/** Friends so it can access mode's internals on construction */
+	friend class FEditorModeRegistry;
+
+	UEdMode();
+
+	virtual void Initialize() {};
+
+	virtual bool MouseEnter(FEditorViewportClient* ViewportClient, FViewport* Viewport, int32 x, int32 y);
+
+	virtual bool MouseLeave(FEditorViewportClient* ViewportClient, FViewport* Viewport);
+
+	virtual bool MouseMove(FEditorViewportClient* ViewportClient, FViewport* Viewport, int32 x, int32 y);
+
+	virtual bool ReceivedFocus(FEditorViewportClient* ViewportClient, FViewport* Viewport);
+
+	virtual bool LostFocus(FEditorViewportClient* ViewportClient, FViewport* Viewport);
+
+	/**
+	 * Called when the mouse is moved while a window input capture is in effect
+	 *
+	 * @param	InViewportClient	Level editor viewport client that captured the mouse input
+	 * @param	InViewport			Viewport that captured the mouse input
+	 * @param	InMouseX			New mouse cursor X coordinate
+	 * @param	InMouseY			New mouse cursor Y coordinate
+	 *
+	 * @return	true if input was handled
+	 */
+	virtual bool CapturedMouseMove(FEditorViewportClient* InViewportClient, FViewport* InViewport, int32 InMouseX, int32 InMouseY);
+
+	/** Process all captured mouse moves that occurred during the current frame */
+	virtual bool ProcessCapturedMouseMoves(FEditorViewportClient* InViewportClient, FViewport* InViewport, const TArrayView<FIntPoint>& CapturedMouseMoves) { return false; }
+
+	virtual bool InputKey(FEditorViewportClient* ViewportClient, FViewport* Viewport, FKey Key, EInputEvent Event);
+	virtual bool InputAxis(FEditorViewportClient* InViewportClient, FViewport* Viewport, int32 ControllerId, FKey Key, float Delta, float DeltaTime);
+	virtual bool InputDelta(FEditorViewportClient* InViewportClient, FViewport* InViewport, FVector& InDrag, FRotator& InRot, FVector& InScale);
+	virtual bool StartTracking(FEditorViewportClient* InViewportClient, FViewport* InViewport);
+	virtual bool EndTracking(FEditorViewportClient* InViewportClient, FViewport* InViewport);
+	// Added for handling EDIT Command...
+	virtual EEditAction::Type GetActionEditDuplicate() { return EEditAction::Skip; }
+	virtual EEditAction::Type GetActionEditDelete() { return EEditAction::Skip; }
+	virtual EEditAction::Type GetActionEditCut() { return EEditAction::Skip; }
+	virtual EEditAction::Type GetActionEditCopy() { return EEditAction::Skip; }
+	virtual EEditAction::Type GetActionEditPaste() { return EEditAction::Skip; }
+	virtual bool ProcessEditDuplicate() { return false; }
+	virtual bool ProcessEditDelete();
+	virtual bool ProcessEditCut() { return false; }
+	virtual bool ProcessEditCopy() { return false; }
+	virtual bool ProcessEditPaste() { return false; }
+
+	virtual void Tick(FEditorViewportClient* ViewportClient, float DeltaTime);
+
+	virtual bool IsCompatibleWith(FEditorModeID OtherModeID) const { return false; }
+	virtual void ActorMoveNotify() {}
+	virtual void ActorsDuplicatedNotify(TArray<AActor*>& PreDuplicateSelection, TArray<AActor*>& PostDuplicateSelection, bool bOffsetLocations) {}
+	virtual void ActorSelectionChangeNotify();
+	virtual void ActorPropChangeNotify() {}
+	virtual void MapChangeNotify() {}
+
+	/** If the Edmode is handling its own mouse deltas, it can disable the MouseDeltaTacker */
+	virtual bool DisallowMouseDeltaTracking() const { return false; }
+
+	/**
+	 * Lets each mode/tool specify a pivot point around which the camera should orbit
+	 * @param	OutPivot	The custom pivot point returned by the mode/tool
+	 * @return	true if a custom pivot point was specified, false otherwise.
+	 */
+	virtual bool GetPivotForOrbit(FVector& OutPivot) const { return false; }
+
+	/**
+	 * Get a cursor to override the default with, if any.
+	 * @return true if the cursor was overridden.
+	 */
+	virtual bool GetCursor(EMouseCursor::Type& OutCursor) const { return false; }
+
+	/** Get override cursor visibility settings */
+	virtual bool GetOverrideCursorVisibility(bool& bWantsOverride, bool& bHardwareCursorVisible, bool bSoftwareCursorVisible) const { return false; }
+
+	/** Called before mouse movement is converted to drag/rot */
+	virtual bool PreConvertMouseMovement(FEditorViewportClient* InViewportClient) { return false; }
+
+	/** Called after mouse movement is converted to drag/rot */
+	virtual bool PostConvertMouseMovement(FEditorViewportClient* InViewportClient) { return false; }
+
+	virtual bool ShouldDrawBrushWireframe(AActor* InActor) const { return true; }
+
+	/** If Rotation Snap should be enabled for this mode*/
+	virtual bool IsSnapRotationEnabled();
+
+	void PostInvalidation();
+	void RestoreEditorState();
+	/** If this mode should override the snap rotation
+	* @param	Rotation		The Rotation Override
+	*
+	* @return					True if you have overridden the value
+	*/
+	virtual bool SnapRotatorToGridOverride(FRotator& Rotation) { return false; };
+
+
+	virtual void UpdateInternalData() {}
+
+	virtual void Enter();
+	virtual void Exit();
+	virtual UTexture2D* GetVertexTexture() { return GEngine->DefaultBSPVertexTexture; }
+
+	virtual void PostUndo() {}
+
+	virtual void SelectNone();
+	virtual void SelectionChanged() {}
+
+	virtual bool HandleClick(FEditorViewportClient* InViewportClient, HHitProxy* HitProxy, const FViewportClick& Click);
+
+	/**
+	 * Allows an editor mode to override the bounding box used to focus the viewport on a selection
+	 *
+	 * @param Actor			The selected actor that is being considered for focus
+	 * @param PrimitiveComponent	The component in the actor being considered for focus
+	 * @param InOutBox		The box that should be computed for the actor and component
+	 * @return bool			true if the mode overrides the box and populated InOutBox, false if it did not populate InOutBox
+	 */
+	virtual bool ComputeBoundingBoxForViewportFocus(AActor* Actor, UPrimitiveComponent* PrimitiveComponent, FBox& InOutBox) const { return false; }
+
+	/** Handling SelectActor */
+	virtual bool Select(AActor* InActor, bool bInSelected) { return 0; }
+
+	/** Check to see if an actor can be selected in this mode - no side effects */
+	virtual bool IsSelectionAllowed(AActor* InActor, bool bInSelection) const { return true; }
+
+	/** @return True if this mode allows the viewport to use a drag tool */
+	virtual bool AllowsViewportDragTool() const { return true; }
+
+	/** Returns the editor mode identifier. */
+	FEditorModeID GetID() const { return Info.ID; }
+
+	/** Returns the editor mode information. */
+	const FEditorModeInfo& GetModeInfo() const { return Info; }
+
+	/** @name Rendering */
+	//@{
+	/** Draws translucent polygons on brushes and volumes. */
+	virtual void Render(const FSceneView* View, FViewport* Viewport, FPrimitiveDrawInterface* PDI);
+	//void DrawGridSection(int32 ViewportLocX,int32 ViewportGridY,FVector* A,FVector* B,float* AX,float* BX,int32 Axis,int32 AlphaCase,FSceneView* View,FPrimitiveDrawInterface* PDI);
+
+	/** Overlays the editor hud (brushes, drag tools, static mesh vertices, etc*. */
+	virtual void DrawHUD(FEditorViewportClient* ViewportClient, FViewport* Viewport, const FSceneView* View, FCanvas* Canvas);
+	//@}
+
+	/**
+	 * Called when the mode wants to draw brackets around selected objects
+	 */
+	virtual void DrawBrackets(FEditorViewportClient* ViewportClient, FViewport* Viewport, const FSceneView* View, FCanvas* Canvas);
+
+	/** True if this mode uses a toolkit mode (eventually they all should) */
+	virtual bool UsesToolkits() const;
+
+	/** Gets the toolkit created by this mode */
+	TSharedPtr<FModeToolkit> GetToolkit() { return Toolkit; }
+
+	/** Returns the world this toolkit is editing */
+	UWorld* GetWorld() const;
+
+	/** Returns the owning mode manager for this mode */
+	class FEditorModeTools* GetModeManager() const;
+
+	/**
+	 * Called when the editor mode should rebuild its toolbar
+	 *
+	 * @param ToolbarBuilder	The builder which should be used to add toolbar widgets
+	 */
+	virtual void BuildModeToolbar(class FToolBarBuilder& ToolbarBuilder) {}
+
+public:
+
+	/** Request that this mode be deleted at the next convenient opportunity (FEditorModeTools::Tick) */
+	void RequestDeletion() { bPendingDeletion = true; }
+
+	/** returns true if this mode is to be deleted at the next convenient opportunity (FEditorModeTools::Tick) */
+	bool IsPendingDeletion() const { return bPendingDeletion; }
+
+private:
+	/** true if this mode is pending removal from its owner */
+	bool bPendingDeletion;
+
+	/** Called whenever a mode type is unregistered */
+	void OnModeUnregistered(FEditorModeID ModeID);
+
+	// default behavior is to accept active tool
+	virtual void TerminateActiveToolsOnPIEStart();
+
+	// default behavior is to accept active tool
+	virtual void TerminateActiveToolsOnSaveWorld();
+
+	FRay GetRayFromMousePos(FEditorViewportClient* ViewportClient, FViewport* Viewport, int MouseX, int MouseY);
+protected:
+
+	/** Information pertaining to this mode. Assigned by FEditorModeRegistry. */
+	FEditorModeInfo Info;
+
+	/** Editor Mode Toolkit that is associated with this toolkit mode */
+	TSharedPtr<class FModeToolkit> Toolkit;
+
+	/** Pointer back to the mode tools that we are registered with */
+	FEditorModeTools* Owner;
+
+	// Property Widgets
+public:
+
+	UPROPERTY()
+	UMaterialInterface* StandardVertexColorMaterial;
+
+protected:
+	/**
+	 * Returns the first selected Actor, or NULL if there is no selection.
+	 */
+	AActor* GetFirstSelectedActorInstance() const;
+
+	void DeactivateAllActiveTools();
+
+	bool bHaveSavedEditorState;
+	bool bSavedAntiAliasingState;
+
+	/** Input event instance used to keep track of various button states, etc, that we cannot directly query on-demand */
+	FInputDeviceState CurrentMouseState;
+
+public:
+
+	/**
+	 * @return active ToolManager
+	 */
+	virtual UInteractiveToolManager* GetToolManager() const;
+
+protected:
+
+	UInteractiveToolsContext* ToolsContext;
+
+	bool bInvalidationPending;
+	IToolsContextQueriesAPI* QueriesAPI;
+	IToolsContextTransactionsAPI* TransactionAPI;
+	IToolsContextAssetAPI* AssetAPI;
+
+	// called when PIE is about to start, shuts down active tools
+	FDelegateHandle BeginPIEDelegateHandle;
+	// called before a Save starts. This currently shuts down active tools.
+	FDelegateHandle PreSaveWorldDelegateHandle;
+
+};
+
diff --git a/Engine/Source/Editor/UnrealEd/UnrealEd.Build.cs b/Engine/Source/Editor/UnrealEd/UnrealEd.Build.cs
--- a/Engine/Source/Editor/UnrealEd/UnrealEd.Build.cs
+++ b/Engine/Source/Editor/UnrealEd/UnrealEd.Build.cs
@@ -102,6 +102,7 @@ public class UnrealEd : ModuleRules
                 "MeshBuilder",
                 "MaterialShaderQualitySettings",
                 "EditorSubsystem",
+                "InteractiveToolsFramework"
             }
 		);
 
@@ -172,7 +173,7 @@ public class UnrealEd : ModuleRules
 				"TimeManagement",
                 "LandscapeEditorUtilities",
                 "DerivedDataCache",
-				"ToolMenus",
+				"ToolMenus"
             }
 		);
 
diff --git a/Engine/Source/Editor/VREditor/UI/VREditorUISystem.cpp b/Engine/Source/Editor/VREditor/UI/VREditorUISystem.cpp
--- a/Engine/Source/Editor/VREditor/UI/VREditorUISystem.cpp
+++ b/Engine/Source/Editor/VREditor/UI/VREditorUISystem.cpp
@@ -213,7 +213,7 @@ void UVREditorUISystem::Init(UVREditorMode* InVRMode)
 	FVREditorActionCallbacks::UpdateSelectingCandidateActorsText(VRMode);
 	FVREditorActionCallbacks::SelectingCandidateActorsText = FVREditorActionCallbacks::GetSelectingCandidateActorsText();
 
-	GLevelEditorModeTools().OnEditorModeChanged().AddUObject(this, &UVREditorUISystem::HandleEditorModeChanged);
+	GLevelEditorModeTools().OnEditorModeIDChanged().AddUObject(this, &UVREditorUISystem::HandleEditorModeChanged);
 }
 
 void UVREditorUISystem::UpdateInteractors( )
@@ -247,7 +247,7 @@ void UVREditorUISystem::Shutdown()
 		VRMode->OnToggleDebugMode().RemoveAll(this);
 	}
 
-	GLevelEditorModeTools().OnEditorModeChanged().RemoveAll(this);
+	GLevelEditorModeTools().OnEditorModeIDChanged().RemoveAll(this);
 
 	// Unbind the color picker creation & destruction overrides
 	SColorPicker::OnColorPickerNonModalCreateOverride.Unbind();
@@ -2260,19 +2260,20 @@ void UVREditorUISystem::SequencerRadialMenuGenerator(FMenuBuilder& MenuBuilder,
 	);
 }
 
-void UVREditorUISystem::HandleEditorModeChanged(FEdMode* Mode, bool IsEnabled)
+void UVREditorUISystem::HandleEditorModeChanged(const FEditorModeID& EditorModeID, bool IsEnabled)
 {
 	if (IsEnabled == true &&
-		(Mode->GetID() == FBuiltinEditorModes::EM_Foliage ||
-		 Mode->GetID() == FBuiltinEditorModes::EM_Landscape ||
-		 Mode->GetID() == FBuiltinEditorModes::EM_MeshPaint))
+		(EditorModeID == FBuiltinEditorModes::EM_Foliage ||
+		 EditorModeID == FBuiltinEditorModes::EM_Landscape ||
+		 EditorModeID == FBuiltinEditorModes::EM_MeshPaint))
 	{
 		AVREditorFloatingUI* Panel = GetPanel(ModesPanelID);
 		if (Panel != nullptr && UIInteractor != nullptr && !Panel->IsUIVisible())
 		{
 			const bool bShouldShow = true;
 			const bool bSpawnInFront = true;
-			const bool bDragFromOpen = ShouldPreviewPanel();
+			bool bDragFromOpen = false;
+			bDragFromOpen = ShouldPreviewPanel();
 			ShowEditorUIPanel(Panel, UIInteractor, bShouldShow, bSpawnInFront, bDragFromOpen);
 		}
 	}
diff --git a/Engine/Source/Editor/VREditor/UI/VREditorUISystem.h b/Engine/Source/Editor/VREditor/UI/VREditorUISystem.h
--- a/Engine/Source/Editor/VREditor/UI/VREditorUISystem.h
+++ b/Engine/Source/Editor/VREditor/UI/VREditorUISystem.h
@@ -27,6 +27,8 @@ class FMenuBuilder;
 class FUICommandList;
 class UVREditorWidgetComponent;
 class UWidgetComponent;
+class FName;
+typedef FName FEditorModeID;
 
 typedef FName VREditorPanelID;
 
@@ -313,7 +315,7 @@ protected:
 	* @param Mode The mode that changed.
 	* @param IsEnabled true if the mode is enabled, otherwise false.
 	*/
-	void HandleEditorModeChanged(class FEdMode* Mode, bool IsEnabled);
+	void HandleEditorModeChanged(const FEditorModeID& EditorModeID, bool IsEnabled);
 
 	/** Reset function that puts you back in placement mode, closes all UIs, etc. */
 	void ResetAll();
diff --git a/Engine/Source/Editor/VREditor/VREditorActions.cpp b/Engine/Source/Editor/VREditor/VREditorActions.cpp
--- a/Engine/Source/Editor/VREditor/VREditorActions.cpp
+++ b/Engine/Source/Editor/VREditor/VREditorActions.cpp
@@ -606,15 +606,7 @@ void FVREditorActionCallbacks::ChangeEditorModes(FEditorModeID InMode)
 	GLevelEditorModeTools().ActivateMode(InMode);
 
 	// Find and disable any other 'visible' modes since we only ever allow one of those active at a time.
-	TArray<FEdMode*> ActiveModes;
-	GLevelEditorModeTools().GetActiveModes(ActiveModes);
-	for (FEdMode* Mode : ActiveModes)
-	{
-		if (Mode->GetID() != InMode && Mode->GetModeInfo().bVisible)
-		{
-			GLevelEditorModeTools().DeactivateMode(Mode->GetID());
-		}
-	}
+	GLevelEditorModeTools().DeactivateOtherVisibleModes(InMode);
 }
 
 ECheckBoxState FVREditorActionCallbacks::EditorModeActive(FEditorModeID InMode)
diff --git a/Engine/Source/Runtime/Experimental/InteractiveToolsFramework/Private/InteractiveGizmoManager.cpp b/Engine/Source/Runtime/Experimental/InteractiveToolsFramework/Private/InteractiveGizmoManager.cpp
--- a/Engine/Source/Runtime/Experimental/InteractiveToolsFramework/Private/InteractiveGizmoManager.cpp
+++ b/Engine/Source/Runtime/Experimental/InteractiveToolsFramework/Private/InteractiveGizmoManager.cpp
@@ -2,6 +2,7 @@
 
 
 #include "InteractiveGizmoManager.h"
+#include "InteractiveToolsContext.h"
 
 
 UInteractiveGizmoManager::UInteractiveGizmoManager()
@@ -191,7 +192,6 @@ void UInteractiveGizmoManager::Render(IToolsContextRenderAPI* RenderAPI)
 
 }
 
-
 void UInteractiveGizmoManager::PostMessage(const TCHAR* Message, EToolMessageLevel Level)
 {
 	TransactionsAPI->PostMessage(Message, Level);
diff --git a/Engine/Source/Runtime/Experimental/InteractiveToolsFramework/Private/InteractiveToolManager.cpp b/Engine/Source/Runtime/Experimental/InteractiveToolsFramework/Private/InteractiveToolManager.cpp
--- a/Engine/Source/Runtime/Experimental/InteractiveToolsFramework/Private/InteractiveToolManager.cpp
+++ b/Engine/Source/Runtime/Experimental/InteractiveToolsFramework/Private/InteractiveToolManager.cpp
@@ -2,6 +2,7 @@
 
 
 #include "InteractiveToolManager.h"
+#include "InteractiveToolsContext.h"
 
 
 UInteractiveToolManager::UInteractiveToolManager()
@@ -226,6 +227,12 @@ void UInteractiveToolManager::Render(IToolsContextRenderAPI* RenderAPI)
 }
 
 
+
+UInteractiveGizmoManager* UInteractiveToolManager::GetPairedGizmoManager()
+{
+	return Cast<UInteractiveToolsContext>(GetOuter())->GizmoManager;
+}
+
 void UInteractiveToolManager::PostMessage(const TCHAR* Message, EToolMessageLevel Level)
 {
 	TransactionsAPI->PostMessage(Message, Level);
diff --git a/Engine/Source/Runtime/Experimental/InteractiveToolsFramework/Private/InteractiveToolsContext.cpp b/Engine/Source/Runtime/Experimental/InteractiveToolsFramework/Private/InteractiveToolsContext.cpp
--- a/Engine/Source/Runtime/Experimental/InteractiveToolsFramework/Private/InteractiveToolsContext.cpp
+++ b/Engine/Source/Runtime/Experimental/InteractiveToolsFramework/Private/InteractiveToolsContext.cpp
@@ -36,3 +36,71 @@ void UInteractiveToolsContext::Shutdown()
 	ToolManager->Shutdown();
 	ToolManager = nullptr;
 }
+
+void UInteractiveToolsContext::DeactivateActiveTool(EToolSide WhichSide, EToolShutdownType ShutdownType)
+{
+	ToolManager->DeactivateTool(WhichSide, ShutdownType);
+}
+
+void UInteractiveToolsContext::DeactivateAllActiveTools()
+{
+	if (ToolManager->HasActiveTool(EToolSide::Left))
+	{
+		EToolShutdownType ShutdownType = ToolManager->CanAcceptActiveTool(EToolSide::Left) ?
+			EToolShutdownType::Accept : EToolShutdownType::Cancel;
+		ToolManager->DeactivateTool(EToolSide::Left, ShutdownType);
+	}
+	if (ToolManager->HasActiveTool(EToolSide::Right))
+	{
+		EToolShutdownType ShutdownType = ToolManager->CanAcceptActiveTool(EToolSide::Right) ?
+			EToolShutdownType::Accept : EToolShutdownType::Cancel;
+		ToolManager->DeactivateTool(EToolSide::Right, ShutdownType);
+	}
+}
+
+
+bool UInteractiveToolsContext::CanStartTool(EToolSide WhichSide, const FString& ToolTypeIdentifier) const
+{
+	return (ToolManager->HasActiveTool(WhichSide) == false) &&
+		(ToolManager->CanActivateTool(WhichSide, ToolTypeIdentifier) == true);
+}
+
+bool UInteractiveToolsContext::ActiveToolHasAccept(EToolSide WhichSide) const
+{
+	return ToolManager->HasActiveTool(WhichSide) &&
+		ToolManager->GetActiveTool(WhichSide)->HasAccept();
+}
+
+bool UInteractiveToolsContext::CanAcceptActiveTool(EToolSide WhichSide) const
+{
+	return ToolManager->CanAcceptActiveTool(WhichSide);
+}
+
+bool UInteractiveToolsContext::CanCancelActiveTool(EToolSide WhichSide) const
+{
+	return ToolManager->CanCancelActiveTool(WhichSide);
+}
+
+bool UInteractiveToolsContext::CanCompleteActiveTool(EToolSide WhichSide) const
+{
+	return ToolManager->HasActiveTool(WhichSide) && CanCancelActiveTool(WhichSide) == false;
+}
+
+bool UInteractiveToolsContext::StartTool(EToolSide WhichSide, const FString& ToolTypeIdentifier)
+{
+	if (ToolManager->SelectActiveToolType(WhichSide, ToolTypeIdentifier) == false)
+	{
+		UE_LOG(LogTemp, Warning, TEXT("ToolManager: Unknown Tool Type %s"), *ToolTypeIdentifier);
+		return false;
+	}
+	else
+	{
+		ToolManager->ActivateTool(WhichSide);
+		return true;
+	}
+}
+
+void UInteractiveToolsContext::EndTool(EToolSide WhichSide, EToolShutdownType ShutdownType)
+{
+	DeactivateActiveTool(WhichSide, ShutdownType);
+}
diff --git a/Engine/Source/Runtime/Experimental/InteractiveToolsFramework/Public/InteractiveToolManager.h b/Engine/Source/Runtime/Experimental/InteractiveToolsFramework/Public/InteractiveToolManager.h
--- a/Engine/Source/Runtime/Experimental/InteractiveToolsFramework/Public/InteractiveToolManager.h
+++ b/Engine/Source/Runtime/Experimental/InteractiveToolsFramework/Public/InteractiveToolManager.h
@@ -192,6 +192,7 @@ public:
 	/** @return current IToolsContextQueriesAPI */
 	virtual IToolsContextQueriesAPI* GetContextQueriesAPI() { return QueriesAPI; }
 
+	UInteractiveGizmoManager* GetPairedGizmoManager();
 
 
 public:
diff --git a/Engine/Source/Runtime/Experimental/InteractiveToolsFramework/Public/InteractiveToolsContext.h b/Engine/Source/Runtime/Experimental/InteractiveToolsFramework/Public/InteractiveToolsContext.h
--- a/Engine/Source/Runtime/Experimental/InteractiveToolsFramework/Public/InteractiveToolsContext.h
+++ b/Engine/Source/Runtime/Experimental/InteractiveToolsFramework/Public/InteractiveToolsContext.h
@@ -30,6 +30,16 @@ public:
 	/** Shutdown Context by destroying InputRouter and ToolManager */
 	virtual void Shutdown();
 
+	virtual void DeactivateActiveTool(EToolSide WhichSide, EToolShutdownType ShutdownType);
+	virtual void DeactivateAllActiveTools();
+
+	bool CanStartTool(EToolSide WhichSide, const FString& ToolTypeIdentifier) const;
+	bool ActiveToolHasAccept(EToolSide WhichSide) const;
+	bool CanAcceptActiveTool(EToolSide WhichSide) const;
+	bool CanCancelActiveTool(EToolSide WhichSide) const;
+	bool CanCompleteActiveTool(EToolSide WhichSide) const;
+	bool StartTool(EToolSide WhichSide, const FString& ToolTypeIdentifier);
+	void EndTool(EToolSide WhichSide, EToolShutdownType ShutdownType);
 public:
 	/** current UInputRouter for this Context */
 	UPROPERTY()
diff --git a/Engine/Source/Runtime/NavigationSystem/Private/NavigationSystem.cpp b/Engine/Source/Runtime/NavigationSystem/Private/NavigationSystem.cpp
--- a/Engine/Source/Runtime/NavigationSystem/Private/NavigationSystem.cpp
+++ b/Engine/Source/Runtime/NavigationSystem/Private/NavigationSystem.cpp
@@ -424,8 +424,8 @@ UNavigationSystemV1::UNavigationSystemV1(const FObjectInitializer& ObjectInitial
 #if WITH_EDITOR
 	if (GIsEditor && HasAnyFlags(RF_ClassDefaultObject) == false)
 	{
-		FEditorDelegates::EditorModeEnter.AddUObject(this, &UNavigationSystemV1::OnEditorModeChanged, true);
-		FEditorDelegates::EditorModeExit.AddUObject(this, &UNavigationSystemV1::OnEditorModeChanged, false);
+		FEditorDelegates::EditorModeIDEnter.AddUObject(this, &UNavigationSystemV1::OnEditorModeIDChanged, true);
+		FEditorDelegates::EditorModeIDExit.AddUObject(this, &UNavigationSystemV1::OnEditorModeIDChanged, false);
 	}
 #endif // WITH_EDITOR
 }
@@ -438,8 +438,8 @@ UNavigationSystemV1::~UNavigationSystemV1()
 #if WITH_EDITOR
 	if (GIsEditor)
 	{
-		FEditorDelegates::EditorModeEnter.RemoveAll(this);
-		FEditorDelegates::EditorModeExit.RemoveAll(this);
+		FEditorDelegates::EditorModeIDEnter.RemoveAll(this);
+		FEditorDelegates::EditorModeIDExit.RemoveAll(this);
 	}
 #endif // WITH_EDITOR
 
@@ -2763,6 +2763,25 @@ void UNavigationSystemV1::OnEditorModeChanged(FEdMode* Mode, bool IsEntering)
 		}
 	}
 }
+
+void UNavigationSystemV1::OnEditorModeIDChanged(const FEditorModeID& ModeID, bool IsEntering)
+{
+	if (IsEntering == false && ModeID == FBuiltinEditorModes::EM_Geometry)
+	{
+		// check if any of modified brushes belongs to an ANavMeshBoundsVolume
+		FEdMode* Mode = GLevelEditorModeTools().GetActiveMode(ModeID);
+		FEdModeGeometry* GeometryMode = (FEdModeGeometry*)Mode;
+		for (auto GeomObjectIt = GeometryMode->GeomObjectItor(); GeomObjectIt; GeomObjectIt++)
+		{
+			ANavMeshBoundsVolume* Volume = Cast<ANavMeshBoundsVolume>((*GeomObjectIt)->GetActualBrush());
+			if (Volume)
+			{
+				OnNavigationBoundsUpdated(Volume);
+			}
+		}
+	}
+}
+
 #endif
 
 void UNavigationSystemV1::OnNavigationBoundsUpdated(ANavMeshBoundsVolume* NavVolume)
diff --git a/Engine/Source/Runtime/NavigationSystem/Public/NavigationSystem.h b/Engine/Source/Runtime/NavigationSystem/Public/NavigationSystem.h
--- a/Engine/Source/Runtime/NavigationSystem/Public/NavigationSystem.h
+++ b/Engine/Source/Runtime/NavigationSystem/Public/NavigationSystem.h
@@ -790,7 +790,9 @@ public:
 
 	void UpdateLevelCollision(ULevel* InLevel);
 
+	UE_DEPRECATED(4.24, "Use OnEditorModeIDChanged instead")
 	virtual void OnEditorModeChanged(FEdMode* Mode, bool IsEntering);
+	virtual void OnEditorModeIDChanged(const FEditorModeID& ModeID, bool IsEntering);
 #endif // WITH_EDITOR
 
 	FORCEINLINE bool IsSetUpForLazyGeometryExporting() const { return bGenerateNavigationOnlyAroundNavigationInvokers; }
