Fix build error caused by 8614137

diff --git a/Engine/Source/Editor/UnrealEd/Private/FbxExporter.h b/Engine/Source/Editor/UnrealEd/Private/FbxExporter.h
--- a/Engine/Source/Editor/UnrealEd/Private/FbxExporter.h
+++ b/Engine/Source/Editor/UnrealEd/Private/FbxExporter.h
@@ -12,6 +12,7 @@
 #include "MovieSceneFwd.h"
 #include "FbxImporter.h"
 #include "UObject/GCObject.h"
+#include "Animation/AnimTypes.h"
 
 class ABrush;
 class ACameraActor;
