Add the concept of being "Expired" to FChange custom changes that are emebedded inside Transactions.
If all FChange's in a Transaction are Expired, then Undo/Redo skips over that change.
Codepaths not using FChange are unaffected, and default FChange behavior is to never be expired, so this path should only be taken for new code.
#rnx

diff --git a/Engine/Source/Editor/UnrealEd/Classes/Editor/Transactor.h b/Engine/Source/Editor/UnrealEd/Classes/Editor/Transactor.h
--- a/Engine/Source/Editor/UnrealEd/Classes/Editor/Transactor.h
+++ b/Engine/Source/Editor/UnrealEd/Classes/Editor/Transactor.h
@@ -270,6 +270,9 @@ protected:
 		/** @return true if the record has a delta change or a custom change */
 		bool HasChanges() const;
 
+		/** @return true if the record has a custom change and the change has expired */
+		bool HasExpired() const;
+
 		/** Transfers data from an array. */
 		class FReader : public FArchiveUObject
 		{
@@ -598,6 +601,9 @@ public:
 	/** @return True if this record contains a reference to a pie object */
 	virtual bool ContainsPieObjects() const override;
 
+	/** @return true if the transaction contains custom changes and all the changes have expired */
+	virtual bool HasExpired() const;
+
 
 	/** Returns a unique string to serve as a type ID for the FTranscationBase-derived type. */
 	virtual const TCHAR* GetTransactionType() const
diff --git a/Engine/Source/Editor/UnrealEd/Private/EditorTransaction.cpp b/Engine/Source/Editor/UnrealEd/Private/EditorTransaction.cpp
--- a/Engine/Source/Editor/UnrealEd/Private/EditorTransaction.cpp
+++ b/Engine/Source/Editor/UnrealEd/Private/EditorTransaction.cpp
@@ -554,6 +554,23 @@ bool FTransaction::ContainsPieObjects() const
 	return false;
 }
 
+bool FTransaction::HasExpired() const
+{
+	if (Records.Num() == 0)		// only return true if we definitely have expired changes
+	{
+		return false;
+	}
+	for (const FObjectRecord& Record : Records)
+	{
+		if (Record.HasExpired() == false)
+		{
+			return false;
+		}
+	}
+	return true;
+}
+
+
 bool FTransaction::IsObjectTransacting(const UObject* Object) const
 {
 	// This function is meaningless when called outside of a transaction context. Without this
@@ -727,6 +744,15 @@ bool FTransaction::FObjectRecord::HasChanges() const
 	return DeltaChange.HasChanged() || CustomChange;
 }
 
+bool FTransaction::FObjectRecord::HasExpired() const
+{
+	if (CustomChange && CustomChange->HasExpired(Object.Get()) == true)
+	{
+		return true;
+	}
+	return false;
+}
+
 void FTransaction::AddReferencedObjects( FReferenceCollector& Collector )
 {
 	for( FObjectRecord& ObjectRecord : Records )
@@ -1428,21 +1454,33 @@ bool UTransBuffer::Undo(bool bCanRedo)
 
 	// Apply the undo changes.
 	GIsTransacting = true;
+
+	// custom changes (FChange) can be applied to temporary objects that require undo/redo for some time,
+	// but we want to skip over these changes later (eg in the context of a Tool that is used for a while and
+	// then closed). In this case the Transaction is "expired" and we continue to Undo until we find a 
+	// non-Expired Transaction. 
+	bool bDoneTransacting = false;
+	do
 	{
 		FTransaction& Transaction = UndoBuffer[ UndoBuffer.Num() - ++UndoCount ].Get();
-		UE_LOG(LogEditorTransaction, Log,  TEXT("Undo %s"), *Transaction.GetTitle().ToString() );
-		CurrentTransaction = &Transaction;
-		CurrentTransaction->BeginOperation();
+		if (Transaction.HasExpired() == false)
+		{
+			UE_LOG(LogEditorTransaction, Log, TEXT("Undo %s"), *Transaction.GetTitle().ToString());
+			CurrentTransaction = &Transaction;
+			CurrentTransaction->BeginOperation();
+
+			const FTransactionContext TransactionContext = CurrentTransaction->GetContext();
+			TransactionStateChangedDelegate.Broadcast(TransactionContext, ETransactionStateEventType::UndoRedoStarted);
+			BeforeRedoUndoDelegate.Broadcast(TransactionContext);
+			Transaction.Apply();
+			UndoDelegate.Broadcast(TransactionContext, true);
+			TransactionStateChangedDelegate.Broadcast(TransactionContext, ETransactionStateEventType::UndoRedoFinalized);
 
-		const FTransactionContext TransactionContext = CurrentTransaction->GetContext();
-		TransactionStateChangedDelegate.Broadcast(TransactionContext, ETransactionStateEventType::UndoRedoStarted);
-		BeforeRedoUndoDelegate.Broadcast(TransactionContext);
-		Transaction.Apply();
-		UndoDelegate.Broadcast(TransactionContext, true);
-		TransactionStateChangedDelegate.Broadcast(TransactionContext, ETransactionStateEventType::UndoRedoFinalized);
+			CurrentTransaction->EndOperation();
+			CurrentTransaction = nullptr;
 
-		CurrentTransaction->EndOperation();
-		CurrentTransaction = nullptr;
+			bDoneTransacting = true;
+		}
 
 		if (!bCanRedo)
 		{
@@ -1451,9 +1489,18 @@ bool UTransBuffer::Undo(bool bCanRedo)
 
 			UndoBufferChangedDelegate.Broadcast();
 		}
-	}
+	} 
+	while (bDoneTransacting == false && CanUndo());
+
 	GIsTransacting = false;
 
+	// if all transactions were expired, reproduce the !CanUndo() branch at the top of the function
+	if (bDoneTransacting == false)
+	{
+		UndoDelegate.Broadcast(FTransactionContext(), false);
+		return false;
+	}
+
 	CheckState();
 
 	return true;
@@ -1472,24 +1519,42 @@ bool UTransBuffer::Redo()
 
 	// Apply the redo changes.
 	GIsTransacting = true;
+
+	// Skip over Expired transactions (see comments in ::Undo)
+	bool bDoneTransacting = false;
+	do
 	{
 		FTransaction& Transaction = UndoBuffer[ UndoBuffer.Num() - UndoCount-- ].Get();
-		UE_LOG(LogEditorTransaction, Log,  TEXT("Redo %s"), *Transaction.GetTitle().ToString() );
-		CurrentTransaction = &Transaction;
-		CurrentTransaction->BeginOperation();
+		if (Transaction.HasExpired() == false)
+		{
+			UE_LOG(LogEditorTransaction, Log, TEXT("Redo %s"), *Transaction.GetTitle().ToString());
+			CurrentTransaction = &Transaction;
+			CurrentTransaction->BeginOperation();
 
-		const FTransactionContext TransactionContext = CurrentTransaction->GetContext();
-		TransactionStateChangedDelegate.Broadcast(TransactionContext, ETransactionStateEventType::UndoRedoStarted);
-		BeforeRedoUndoDelegate.Broadcast(TransactionContext);
-		Transaction.Apply();
-		RedoDelegate.Broadcast(TransactionContext, true);
-		TransactionStateChangedDelegate.Broadcast(TransactionContext, ETransactionStateEventType::UndoRedoFinalized);
+			const FTransactionContext TransactionContext = CurrentTransaction->GetContext();
+			TransactionStateChangedDelegate.Broadcast(TransactionContext, ETransactionStateEventType::UndoRedoStarted);
+			BeforeRedoUndoDelegate.Broadcast(TransactionContext);
+			Transaction.Apply();
+			RedoDelegate.Broadcast(TransactionContext, true);
+			TransactionStateChangedDelegate.Broadcast(TransactionContext, ETransactionStateEventType::UndoRedoFinalized);
+
+			CurrentTransaction->EndOperation();
+			CurrentTransaction = nullptr;
+
+			bDoneTransacting = true;
+		}
+	} 
+	while (bDoneTransacting == false && CanRedo());
 
-		CurrentTransaction->EndOperation();
-		CurrentTransaction = nullptr;
-	}
 	GIsTransacting = false;
 
+	// if all transactions were expired, reproduce the !CanRedo() branch at the top of the function
+	if (bDoneTransacting == false)
+	{
+		RedoDelegate.Broadcast(FTransactionContext(), false);
+		return false;
+	}
+
 	CheckState();
 
 	return true;
diff --git a/Engine/Source/Runtime/Core/Public/Misc/Change.h b/Engine/Source/Runtime/Core/Public/Misc/Change.h
--- a/Engine/Source/Runtime/Core/Public/Misc/Change.h
+++ b/Engine/Source/Runtime/Core/Public/Misc/Change.h
@@ -41,6 +41,9 @@ public:
 	/** Reverts change to the object */
 	virtual void Revert( UObject* Object ) = 0;
 
+	/** @return true if this Change has Expired, ie it will no longer have any effect and could be skipped by undo/redo */
+	virtual bool HasExpired( UObject* Object ) const { return false; }
+
 	/** Describes this change (for debugging) */
 	virtual FString ToString() const = 0;
 
