static mesh component factory should only think it can build things that it can cast to static mesh

diff --git a/Engine/Source/Editor/UnrealEd/Private/Tools/EditorComponentSourceFactory.cpp b/Engine/Source/Editor/UnrealEd/Private/Tools/EditorComponentSourceFactory.cpp
--- a/Engine/Source/Editor/UnrealEd/Private/Tools/EditorComponentSourceFactory.cpp
+++ b/Engine/Source/Editor/UnrealEd/Private/Tools/EditorComponentSourceFactory.cpp
@@ -41,7 +41,8 @@ FStaticMeshComponentTarget::CommitMesh( const FCommitter& Committer )
 bool
 FStaticMeshComponentTargetFactory::CanBuild(UActorComponent* Component)
 {
-	return Component != nullptr;
+	UStaticMeshComponent* StaticMeshComponent = Cast<UStaticMeshComponent>(Component);
+	return StaticMeshComponent != nullptr;
 }
 
 TUniquePtr< FPrimitiveComponentTarget >
