Removing unneeded override keyword

diff --git a/Engine/Source/Editor/UnrealEd/Private/Tools/UEdMode.cpp b/Engine/Source/Editor/UnrealEd/Private/Tools/UEdMode.cpp
--- a/Engine/Source/Editor/UnrealEd/Private/Tools/UEdMode.cpp
+++ b/Engine/Source/Editor/UnrealEd/Private/Tools/UEdMode.cpp
@@ -238,7 +238,7 @@ public:
 	}
 
 
-	virtual void PostMessage(const TCHAR* Message, EToolMessageLevel Level) override
+	virtual void PostMessage(const TCHAR* Message, EToolMessageLevel Level)
 	{
 		UE_LOG(LogTemp, Warning, TEXT("[ToolsContext] %s"), Message);
 	}
