PR #6220: Only call DrawComponentVisualizersHUD if not in Game View (Contributed by [REDACTED_PERSON])

diff --git a/Engine/Source/Editor/UnrealEd/Private/LevelEditorViewport.cpp b/Engine/Source/Editor/UnrealEd/Private/LevelEditorViewport.cpp
--- a/Engine/Source/Editor/UnrealEd/Private/LevelEditorViewport.cpp
+++ b/Engine/Source/Editor/UnrealEd/Private/LevelEditorViewport.cpp
@@ -4635,7 +4635,7 @@ void FLevelEditorViewportClient::SetupViewForRendering( FSceneViewFamily& ViewFa
 void FLevelEditorViewportClient::DrawCanvas( FViewport& InViewport, FSceneView& View, FCanvas& Canvas )
 {
 	// HUD for components visualizers
-	if (GUnrealEd != NULL)
+	if (GUnrealEd != NULL && !IsInGameView())
 	{
 		GUnrealEd->DrawComponentVisualizersHUD(&InViewport, &View, &Canvas);
 	}
