Sequencer: Additional changes to allow for exporting spawnables.

- Spawnables need to be set to have master sequence ownership when exporting from a subsequence
- Spawnable actor can be respawned during the export, so keep a map from the guid to the fbx node

diff --git a/Engine/Source/Editor/MovieSceneTools/Private/MovieSceneToolHelpers.cpp b/Engine/Source/Editor/MovieSceneTools/Private/MovieSceneToolHelpers.cpp
--- a/Engine/Source/Editor/MovieSceneTools/Private/MovieSceneToolHelpers.cpp
+++ b/Engine/Source/Editor/MovieSceneTools/Private/MovieSceneToolHelpers.cpp
@@ -2112,7 +2112,7 @@ bool MovieSceneToolHelpers::ExportFBX(UWorld* World, UMovieScene* MovieScene, IM
 		}
 	}
 
-	Exporter->ExportLevelSequence(MovieScene, Bindings, Player, Template, RootToLocalTransform);
+	Exporter->ExportLevelSequence(MovieScene, Bindings, Player, NodeNameAdapter, Template, RootToLocalTransform);
 
 	//Export all master tracks
 
diff --git a/Engine/Source/Editor/Sequencer/Private/Sequencer.cpp b/Engine/Source/Editor/Sequencer/Private/Sequencer.cpp
--- a/Engine/Source/Editor/Sequencer/Private/Sequencer.cpp
+++ b/Engine/Source/Editor/Sequencer/Private/Sequencer.cpp
@@ -9716,38 +9716,74 @@ void FSequencer::ExportFBXInternal(const FString& ExportFilename, TArray<FGuid>&
 			FMovieSceneSequenceIDRef Template = GetFocusedTemplateID();
 			UnFbx::FFbxExporter::FLevelSequenceNodeNameAdapter NodeNameAdapter(MovieScene, this, Template);
 
-			bool bWasTransacted = false;
+			// Helper to make spawnables persist throughout the export process and then restore properly afterwards
+			struct FSpawnableRestoreState
 			{
-				FScopedTransaction ExportFBXTransaction(NSLOCTEXT("Sequencer", "ExportFBX", "Export FBX"));
-
-				// Temporarily spawn all spawnables so that they can be exported. This is undone after export below with UndoTransaction()
-				for (int32 SpawnableIndex = 0; SpawnableIndex < MovieScene->GetSpawnableCount(); ++SpawnableIndex)
+				FSpawnableRestoreState(FSequencer& InSequencer) :
+					bWasChanged(false)
 				{
-					FMovieSceneSpawnable& Spawnable = MovieScene->GetSpawnable(SpawnableIndex);
+					WeakMovieScene = InSequencer.GetFocusedMovieSceneSequence()->GetMovieScene();
+
+					for (int32 SpawnableIndex = 0; SpawnableIndex < WeakMovieScene->GetSpawnableCount(); ++SpawnableIndex)
+					{
+						FMovieSceneSpawnable& Spawnable = WeakMovieScene->GetSpawnable(SpawnableIndex);
+
+						UMovieSceneSpawnTrack* SpawnTrack = WeakMovieScene->FindTrack<UMovieSceneSpawnTrack>(Spawnable.GetGuid());
+
+						if (SpawnTrack)
+						{
+							bWasChanged = true;
 
-					UMovieSceneSpawnTrack* SpawnTrack = MovieScene->FindTrack<UMovieSceneSpawnTrack>(Spawnable.GetGuid());
+							// Spawnable could be in a subscene, so temporarily override it to persist throughout
+							SpawnOwnershipMap.Add(Spawnable.GetGuid(), Spawnable.GetSpawnOwnership());
+							Spawnable.SetSpawnOwnership(ESpawnOwnership::MasterSequence);
 
-					if (SpawnTrack)
+							// Spawnable could have animated spawned state, so temporarily override it to spawn infinitely
+							UMovieSceneSpawnSection* SpawnSection = Cast<UMovieSceneSpawnSection>(SpawnTrack->CreateNewSection());
+							SpawnSection->Modify();
+							SpawnSection->GetChannel().Reset();
+							SpawnSection->GetChannel().SetDefault(true);
+						}
+					}
+
+					if (bWasChanged)
 					{
-						bWasTransacted = true;
+						// Evaluate at the beginning of the subscene time to ensure that spawnables are created before export
+						InSequencer.SetLocalTimeDirectly(MovieScene::DiscreteInclusiveLower(InSequencer.GetTimeBounds()));
+					}
+				}
 
-						SpawnTrack->Modify();
-						UMovieSceneSpawnSection* SpawnSection = Cast<UMovieSceneSpawnSection>(SpawnTrack->CreateNewSection());
-						SpawnTrack->RemoveAllAnimationData();
-						SpawnSection->GetChannel().SetDefault(true);
-						SpawnTrack->AddSection(*SpawnSection);
+				~FSpawnableRestoreState()
+				{
+					if (!bWasChanged || !WeakMovieScene.IsValid())
+					{
+						return;
+					}
+
+					// Restore spawnable owners
+					for (int32 SpawnableIndex = 0; SpawnableIndex < WeakMovieScene->GetSpawnableCount(); ++SpawnableIndex)
+					{
+						FMovieSceneSpawnable& Spawnable = WeakMovieScene->GetSpawnable(SpawnableIndex);
+						Spawnable.SetSpawnOwnership(SpawnOwnershipMap[Spawnable.GetGuid()]);
 					}
+
+					// Restore modified spawned sections
+					GEditor->UndoTransaction(false);
 				}
 
-				ForceEvaluate();
+				bool bWasChanged;
+				TMap<FGuid, ESpawnOwnership> SpawnOwnershipMap;
+				TWeakObjectPtr<UMovieScene> WeakMovieScene;
+			};
 
-				MovieSceneToolHelpers::ExportFBX(World, MovieScene, this, Bindings, NodeNameAdapter, Template, ExportFilename, RootToLocalTransform);
-			}
+			FScopedTransaction ExportFBXTransaction(NSLOCTEXT("Sequencer", "ExportFBX", "Export FBX"));
 
-			if (bWasTransacted)
 			{
-				GEditor->UndoTransaction(false);
+				FSpawnableRestoreState SpawnableRestoreState(*this);
+
+				MovieSceneToolHelpers::ExportFBX(World, MovieScene, this, Bindings, NodeNameAdapter, Template, ExportFilename, RootToLocalTransform);
 			}
+
 			ForceEvaluate();
 		}
 	}
diff --git a/Engine/Source/Editor/UnrealEd/Private/Fbx/FbxAnimationExport.cpp b/Engine/Source/Editor/UnrealEd/Private/Fbx/FbxAnimationExport.cpp
--- a/Engine/Source/Editor/UnrealEd/Private/Fbx/FbxAnimationExport.cpp
+++ b/Engine/Source/Editor/UnrealEd/Private/Fbx/FbxAnimationExport.cpp
@@ -610,6 +610,11 @@ void FFbxExporter::ExportAnimTrack(IAnimTrackAdapter& AnimTrackAdapter, AActor*
 
 		TArray<FTransform> LocalBoneTransforms = SkeletalMeshComponent->GetBoneSpaceTransforms();
 
+		if (LocalBoneTransforms.Num() == 0)
+		{
+			continue;
+		}
+
 		// Add the animation data to the bone nodes
 		for(int32 BoneIndex = 0; BoneIndex < BoneNodes.Num(); ++BoneIndex)
 		{
diff --git a/Engine/Source/Editor/UnrealEd/Private/Fbx/FbxMainExport.cpp b/Engine/Source/Editor/UnrealEd/Private/Fbx/FbxMainExport.cpp
--- a/Engine/Source/Editor/UnrealEd/Private/Fbx/FbxMainExport.cpp
+++ b/Engine/Source/Editor/UnrealEd/Private/Fbx/FbxMainExport.cpp
@@ -1554,6 +1554,32 @@ FString FFbxExporter::FLevelSequenceNodeNameAdapter::GetActorNodeName(const AAct
 	return NodeName;
 }
 
+void FFbxExporter::FLevelSequenceNodeNameAdapter::AddFbxNode(UObject* InObject, FbxNode* FbxNode)
+{
+	FGuid ObjectGuid = MovieScenePlayer->FindObjectId(*InObject, SequenceID);
+	if (ObjectGuid.IsValid())
+	{
+		if (!GuidToFbxNodeMap.Contains(ObjectGuid))
+		{
+			GuidToFbxNodeMap.Add(ObjectGuid);
+		}
+		GuidToFbxNodeMap[ObjectGuid] = FbxNode;
+	}
+}
+
+FbxNode* FFbxExporter::FLevelSequenceNodeNameAdapter::GetFbxNode(UObject* InObject)
+{
+	FGuid ObjectGuid = MovieScenePlayer->FindObjectId(*InObject, SequenceID);
+	if (ObjectGuid.IsValid())
+	{
+		if (GuidToFbxNodeMap.Contains(ObjectGuid))
+		{
+			return GuidToFbxNodeMap[ObjectGuid];
+		}
+	}
+	return nullptr;
+}
+
 FFbxExporter::FLevelSequenceAnimTrackAdapter::FLevelSequenceAnimTrackAdapter( IMovieScenePlayer* InMovieScenePlayer, UMovieScene* InMovieScene, const FMovieSceneSequenceTransform& InRootToLocalTransform)
 {
 	MovieScenePlayer = InMovieScenePlayer;
@@ -1613,6 +1639,7 @@ bool FFbxExporter::ExportLevelSequenceTracks(UMovieScene* MovieScene, IMovieScen
 	const bool bSkip3DTransformTrack = SkeletalMeshComp && GetExportOptions()->MapSkeletalMotionToRoot;
 
 	// Look for the tracks that we currently support
+	bool bExportedAnimTrack = false; // Only export the anim track once since the evaluation is already blended
 	for (UMovieSceneTrack* Track : Tracks)
 	{
 		if (Track->IsA(UMovieScene3DTransformTrack::StaticClass()) && !bSkip3DTransformTrack)
@@ -1625,10 +1652,11 @@ bool FFbxExporter::ExportLevelSequenceTracks(UMovieScene* MovieScene, IMovieScen
 			UMovieScenePropertyTrack* PropertyTrack = (UMovieScenePropertyTrack*)Track;
 			ExportLevelSequencePropertyTrack(FbxActor, *PropertyTrack, MovieScene->GetPlaybackRange(), RootToLocalTransform);
 		}
-		else if (Track->IsA(UMovieSceneSkeletalAnimationTrack::StaticClass()))
+		else if (Track->IsA(UMovieSceneSkeletalAnimationTrack::StaticClass()) && !bExportedAnimTrack)
 		{
 			if (SkeletalMeshComp)
 			{
+				bExportedAnimTrack = true;
 				FLevelSequenceAnimTrackAdapter AnimTrackAdapter(MovieScenePlayer, MovieScene, RootToLocalTransform);
 				ExportAnimTrack(AnimTrackAdapter, Actor, SkeletalMeshComp, 1.0 / DisplayRate.AsDecimal());
 			}
@@ -1639,7 +1667,7 @@ bool FFbxExporter::ExportLevelSequenceTracks(UMovieScene* MovieScene, IMovieScen
 }
 
 
-bool FFbxExporter::ExportLevelSequence( UMovieScene* MovieScene, const TArray<FGuid>& Bindings, IMovieScenePlayer* MovieScenePlayer, FMovieSceneSequenceIDRef SequenceID, const FMovieSceneSequenceTransform& RootToLocalTransform)
+bool FFbxExporter::ExportLevelSequence( UMovieScene* MovieScene, const TArray<FGuid>& Bindings, IMovieScenePlayer* MovieScenePlayer, INodeNameAdapter& NodeNameAdapter, FMovieSceneSequenceIDRef SequenceID, const FMovieSceneSequenceTransform& RootToLocalTransform)
 {
 	if ( MovieScene == nullptr || MovieScenePlayer == nullptr )
 	{
@@ -1672,7 +1700,7 @@ bool FFbxExporter::ExportLevelSequence( UMovieScene* MovieScene, const TArray<FG
 					continue;
 				}
 
-				FbxNode* FbxActor = FindActor( Actor );
+				FbxNode* FbxActor = FindActor(Actor, &NodeNameAdapter);
 
 				// now it should export everybody
 				if ( FbxActor )
@@ -1702,7 +1730,7 @@ FbxNode* FFbxExporter::ExportActor(AActor* Actor, bool bExportComponents, INodeN
 {
 	// Verify that this actor isn't already exported, create a structure for it
 	// and buffer it.
-	FbxNode* ActorNode = FindActor(Actor);
+	FbxNode* ActorNode = FindActor(Actor, &NodeNameAdapter);
 	if (ActorNode == NULL)
 	{
 		FString FbxNodeName = NodeNameAdapter.GetActorNodeName(Actor);
@@ -1711,7 +1739,7 @@ FbxNode* FFbxExporter::ExportActor(AActor* Actor, bool bExportComponents, INodeN
 
 		AActor* ParentActor = Actor->GetAttachParentActor();
 		// this doesn't work with skeletalmeshcomponent
-		FbxNode* ParentNode = FindActor(ParentActor);
+		FbxNode* ParentNode = FindActor(ParentActor, &NodeNameAdapter);
 		FVector ActorLocation, ActorRotation, ActorScale;
 
 		// For cameras and lights: always add a rotation to get the correct coordinate system.
@@ -1766,6 +1794,7 @@ FbxNode* FFbxExporter::ExportActor(AActor* Actor, bool bExportComponents, INodeN
 
 		ParentNode->AddChild(ActorNode);
 		FbxActors.Add(Actor, ActorNode);
+		NodeNameAdapter.AddFbxNode(Actor, ActorNode);
 
 		// Set the default position of the actor on the transforms
 		// The transformation is different from FBX's Z-up: invert the Y-axis for translations and the Y/Z angle values in rotations.
@@ -1887,7 +1916,7 @@ FbxNode* FFbxExporter::ExportActor(AActor* Actor, bool bExportComponents, INodeN
 				}
 				else if (SkelMeshComp && SkelMeshComp->SkeletalMesh)
 				{
-					ExportSkeletalMeshComponent(SkelMeshComp, *SkelMeshComp->GetName(), ExportNode, bSaveAnimSeq);
+					ExportSkeletalMeshComponent(SkelMeshComp, *SkelMeshComp->GetName(), ExportNode, NodeNameAdapter, bSaveAnimSeq);
 				}
 				else if (Component->IsA(UCameraComponent::StaticClass()))
 				{
@@ -1905,6 +1934,7 @@ FbxNode* FFbxExporter::ExportActor(AActor* Actor, bool bExportComponents, INodeN
 				{
 					FbxNode* ChildActorNode = ExportActor(ChildActorComp->GetChildActor(), true, NodeNameAdapter, bSaveAnimSeq);
 					FbxActors.Add(ChildActorComp->GetChildActor(), ChildActorNode);
+					NodeNameAdapter.AddFbxNode(ChildActorComp->GetChildActor(), ChildActorNode);
 				}
 			}
 		}
@@ -2790,8 +2820,18 @@ void FFbxExporter::ExportLevelSequencePropertyTrack( FbxNode* FbxNode, UMovieSce
 /**
  * Finds the given actor in the already-exported list of structures
  */
-FbxNode* FFbxExporter::FindActor(AActor* Actor)
+FbxNode* FFbxExporter::FindActor(AActor* Actor, INodeNameAdapter* NodeNameAdapter)
 {
+	if (NodeNameAdapter)
+	{
+		FbxNode* ActorNode = NodeNameAdapter->GetFbxNode(Actor);
+
+		if (ActorNode)
+		{
+			return ActorNode;
+		}
+	}
+
 	if (FbxActors.Find(Actor))
 	{
 		return *FbxActors.Find(Actor);
@@ -2809,8 +2849,20 @@ FbxNode* FFbxExporter::CreateNode(const FString& NodeName)
 	return FbxNode;
 }
 
-bool FFbxExporter::FindSkeleton(const USkeletalMeshComponent* SkelComp, TArray<FbxNode*>& BoneNodes)
+bool FFbxExporter::FindSkeleton(USkeletalMeshComponent* SkelComp, TArray<FbxNode*>& BoneNodes, INodeNameAdapter* NodeNameAdapter)
 {
+	if (NodeNameAdapter)
+	{
+		FbxNode* SkelRoot = NodeNameAdapter->GetFbxNode(SkelComp);
+		if (SkelRoot)
+		{
+			BoneNodes.Empty();
+			GetSkeleton(SkelRoot, BoneNodes);
+
+			return true;
+		}
+	}
+
 	FbxNode** SkelRoot = FbxSkeletonRoots.Find(SkelComp);
 
 	if (SkelRoot)
diff --git a/Engine/Source/Editor/UnrealEd/Private/Fbx/FbxSkeletalMeshExport.cpp b/Engine/Source/Editor/UnrealEd/Private/Fbx/FbxSkeletalMeshExport.cpp
--- a/Engine/Source/Editor/UnrealEd/Private/Fbx/FbxSkeletalMeshExport.cpp
+++ b/Engine/Source/Editor/UnrealEd/Private/Fbx/FbxSkeletalMeshExport.cpp
@@ -543,7 +543,7 @@ void FFbxExporter::CreateBindPose(FbxNode* MeshRootNode)
 	}
 }
 
-void FFbxExporter::ExportSkeletalMeshComponent(USkeletalMeshComponent* SkelMeshComp, const TCHAR* MeshName, FbxNode* ActorRootNode, bool bSaveAnimSeq)
+void FFbxExporter::ExportSkeletalMeshComponent(USkeletalMeshComponent* SkelMeshComp, const TCHAR* MeshName, FbxNode* ActorRootNode, INodeNameAdapter& NodeNameAdapter, bool bSaveAnimSeq)
 {
 	if (SkelMeshComp && SkelMeshComp->SkeletalMesh)
 	{
@@ -553,6 +553,7 @@ void FFbxExporter::ExportSkeletalMeshComponent(USkeletalMeshComponent* SkelMeshC
 		if(SkeletonRootNode)
 		{
 			FbxSkeletonRoots.Add(SkelMeshComp, SkeletonRootNode);
+			NodeNameAdapter.AddFbxNode(SkelMeshComp, SkeletonRootNode);
 		}
 	}
 }
diff --git a/Engine/Source/Editor/UnrealEd/Private/FbxExporter.h b/Engine/Source/Editor/UnrealEd/Private/FbxExporter.h
--- a/Engine/Source/Editor/UnrealEd/Private/FbxExporter.h
+++ b/Engine/Source/Editor/UnrealEd/Private/FbxExporter.h
@@ -134,7 +134,7 @@ public:
 	 *
 	 * @return	true, if successful
 	 */
-	bool ExportLevelSequence(UMovieScene* MovieScene, const TArray<FGuid>& InBindings, IMovieScenePlayer* MovieScenePlayer, FMovieSceneSequenceIDRef SequenceID, const FMovieSceneSequenceTransform& RootToLocalTransform);
+	bool ExportLevelSequence(UMovieScene* MovieScene, const TArray<FGuid>& InBindings, IMovieScenePlayer* MovieScenePlayer, INodeNameAdapter& NodeNameAdapter, FMovieSceneSequenceIDRef SequenceID, const FMovieSceneSequenceTransform& RootToLocalTransform);
 
 	/**
 	 * Exports the given level sequence track information into a FBX document.
@@ -219,10 +219,13 @@ public:
 	public:
 		FLevelSequenceNodeNameAdapter( UMovieScene* InMovieScene, IMovieScenePlayer* InMovieScenePlayer, FMovieSceneSequenceIDRef InSequenceID);
 		virtual FString GetActorNodeName(const AActor* InActor) override;
+		virtual void AddFbxNode(UObject* InObject, FbxNode* InFbxNode) override;
+		virtual FbxNode* GetFbxNode(UObject* InObject) override;
 	private:
 		UMovieScene* MovieScene;
 		IMovieScenePlayer* MovieScenePlayer;
 		FMovieSceneSequenceID SequenceID;
+		TMap<FGuid, FbxNode*> GuidToFbxNodeMap;
 	};
 
 	/* Get a valid unique name from a name */
@@ -399,7 +402,7 @@ private:
 	FbxNode* ExportSkeletalMeshToFbx(const USkeletalMesh* SkelMesh, const UAnimSequence* AnimSeq, const TCHAR* MeshName, FbxNode* ActorRootNode);
 
 	/** Export SkeletalMeshComponent */
-	void ExportSkeletalMeshComponent(USkeletalMeshComponent* SkelMeshComp, const TCHAR* MeshName, FbxNode* ActorRootNode, bool bSaveAnimSeq = true);
+	void ExportSkeletalMeshComponent(USkeletalMeshComponent* SkelMeshComp, const TCHAR* MeshName, FbxNode* ActorRootNode, INodeNameAdapter& NodeNameAdapter, bool bSaveAnimSeq = true);
 
 	/** Initializing the AnimStack playrate from the AnimSequence */
 	bool SetupAnimStack(const UAnimSequence* AnimSeq);
@@ -476,7 +479,7 @@ private:
 	 * Finds the given actor in the already-exported list of structures
 	 * @return FbxNode* the FBX node created from the UE4 actor
 	 */
-	FbxNode* FindActor(AActor* Actor);
+	FbxNode* FindActor(AActor* Actor, INodeNameAdapter* NodeNameAdapter = nullptr);
 
 	/** Create fbx node with the given name */
 	FbxNode* CreateNode(const FString& NodeName);
@@ -484,7 +487,7 @@ private:
 	/**
 	 * Find bone array of FbxNOdes of the given skeletalmeshcomponent  
 	 */
-	bool FindSkeleton(const USkeletalMeshComponent* SkelComp, TArray<FbxNode*>& BoneNodes);
+	bool FindSkeleton(USkeletalMeshComponent* SkelComp, TArray<FbxNode*>& BoneNodes, INodeNameAdapter* NodeNameAdapter = nullptr);
 
 	/** recursively get skeleton */
 	void GetSkeleton(FbxNode* RootNode, TArray<FbxNode*>& BoneNodes);
diff --git a/Engine/Source/Editor/UnrealEd/Public/MatineeExporter.h b/Engine/Source/Editor/UnrealEd/Public/MatineeExporter.h
--- a/Engine/Source/Editor/UnrealEd/Public/MatineeExporter.h
+++ b/Engine/Source/Editor/UnrealEd/Public/MatineeExporter.h
@@ -14,14 +14,22 @@ class ACameraActor;
 class ALight;
 class UModel;
 class UStaticMeshComponent;
+class USkeletalMeshComponent;
+
+namespace fbxsdk
+{
+	class FbxNode;
+}
 
 /** Adapter interface which allows finding the corresponding actor node name to act on both sequencer and matinee data. */
 class INodeNameAdapter
 {
 public:
+	virtual ~INodeNameAdapter() {}
 	virtual FString GetActorNodeName(const AActor* InActor) { return InActor->GetName(); }
+	virtual void AddFbxNode(UObject* InObject, fbxsdk::FbxNode* InFbxNode) {}
+	virtual fbxsdk::FbxNode* GetFbxNode(UObject* InObject) { return nullptr; }
 };
-
 /**
  * Main Matinee Exporter class.
  * Except for CImporter, consider the other classes as private.
