ToolsFramework: cleaned up IToolsContextAssetAPI, now only exposes a single function to create new package at a given path/basename combo, and returns unique asset name. Rewrote FEditorToolAssetAPI to use AssetToolsModule. AssetGenerationUtil::GenerateStaticMeshActor no longer tries to rename actor as this can result in collisions (which cause check), now only changes display name.
#rnx

diff --git a/Engine/Plugins/Experimental/MeshModelingToolset/Source/ModelingComponents/Private/AssetGenerationUtil.cpp b/Engine/Plugins/Experimental/MeshModelingToolset/Source/ModelingComponents/Private/AssetGenerationUtil.cpp
--- a/Engine/Plugins/Experimental/MeshModelingToolset/Source/ModelingComponents/Private/AssetGenerationUtil.cpp
+++ b/Engine/Plugins/Experimental/MeshModelingToolset/Source/ModelingComponents/Private/AssetGenerationUtil.cpp
@@ -41,21 +41,19 @@ AActor* AssetGenerationUtil::GenerateStaticMeshActor(
 	FString MeshName = ObjectName;
 
 	// create new package
-	FString UniqueMeshName = AssetAPI->MakeUniqueAssetName(MeshName, PackageFolderPath);
-	UPackage* AssetPackage = AssetAPI->CreateNewPackage(UniqueMeshName, PackageFolderPath);
-
+	FString UniqueAssetName;
+	UPackage* AssetPackage = AssetAPI->MakeNewAssetPackage(PackageFolderPath, ObjectName, UniqueAssetName);
 
 	// create new actor
 	FRotator Rotation(0.0f, 0.0f, 0.0f);
 	FActorSpawnParameters SpawnInfo;
 	// @todo nothing here is specific to AStaticMeshActor...could we pass in a CDO and clone it instead of using SpawnActor?
 	AStaticMeshActor* NewActor = TargetWorld->SpawnActor<AStaticMeshActor>(FVector::ZeroVector, Rotation, SpawnInfo);
-	NewActor->Rename(*UniqueMeshName);
-	NewActor->SetActorLabel(*UniqueMeshName);
+	NewActor->SetActorLabel(*UniqueAssetName);
 
 	// construct new static mesh
 	FStaticMeshComponentBuilder Builder;
-	Builder.Initialize(AssetPackage, FName(*UniqueMeshName));
+	Builder.Initialize(AssetPackage, FName(*UniqueAssetName));
 
 	// Populate the MeshDescription associated with the asset.
 	FMeshDescriptionBuilder MeshDescripBuilder;
diff --git a/Engine/Source/Editor/UnrealEd/Private/Tools/EditorToolAssetAPI.cpp b/Engine/Source/Editor/UnrealEd/Private/Tools/EditorToolAssetAPI.cpp
--- a/Engine/Source/Editor/UnrealEd/Private/Tools/EditorToolAssetAPI.cpp
+++ b/Engine/Source/Editor/UnrealEd/Private/Tools/EditorToolAssetAPI.cpp
@@ -6,9 +6,8 @@
 
 
 #include "AssetRegistryModule.h"
+#include "AssetToolsModule.h"
 #include "FileHelpers.h"
-#include "PackageTools.h"
-#include "ObjectTools.h"
 
 // for content-browser things
 #include "ContentBrowserModule.h"
@@ -30,37 +29,17 @@ FString FEditorToolAssetAPI::GetActiveAssetFolderPath()
 
 
 
-FString FEditorToolAssetAPI::MakePackageName(const FString& AssetName, const FString& FolderPath)
-{
-	return FolderPath + TEXT("/") + AssetName;
-}
 
-
-FString FEditorToolAssetAPI::MakeUniqueAssetName(const FString& AssetName, const FString& FolderPath)
+UPackage* FEditorToolAssetAPI::MakeNewAssetPackage(const FString& FolderPath, const FString& AssetBaseName, FString& UniqueAssetName)
 {
-	FString NewPackageName = MakePackageName(AssetName, FolderPath);
-	if (FPackageName::DoesPackageExist(NewPackageName) == false)
-	{
-		return AssetName;
-	}
-
-	int Counter = 1;
-	FString UniqueName = AssetName;
-	while ( FPackageName::DoesPackageExist(NewPackageName) )
-	{
-		UniqueName = AssetName + FString::Printf(TEXT("_%d"), Counter++);
-		NewPackageName = MakePackageName(UniqueName, FolderPath);
-	}
-	return UniqueName;
-}
+	FString UniquePackageName;
+	FAssetToolsModule& AssetToolsModule = FModuleManager::LoadModuleChecked<FAssetToolsModule>("AssetTools");
+	AssetToolsModule.Get().CreateUniqueAssetName(
+		FolderPath + TEXT("/") + AssetBaseName, TEXT(""), UniquePackageName, UniqueAssetName);
 
-
-UPackage* FEditorToolAssetAPI::CreateNewPackage(const FString& AssetName, const FString& FolderPath)
-{
-	FString NewPackageName = MakePackageName(AssetName, FolderPath);
-	NewPackageName = UPackageTools::SanitizePackageName(NewPackageName);
-	UPackage* AssetPackage = CreatePackage(nullptr, *NewPackageName);
+	UPackage* AssetPackage = CreatePackage(nullptr, *UniquePackageName);
 	return AssetPackage;
+
 }
 
 
diff --git a/Engine/Source/Editor/UnrealEd/Public/Tools/EditorToolAssetAPI.h b/Engine/Source/Editor/UnrealEd/Public/Tools/EditorToolAssetAPI.h
--- a/Engine/Source/Editor/UnrealEd/Public/Tools/EditorToolAssetAPI.h
+++ b/Engine/Source/Editor/UnrealEd/Public/Tools/EditorToolAssetAPI.h
@@ -14,13 +14,7 @@ class UNREALED_API FEditorToolAssetAPI : public IToolsContextAssetAPI
 public:
 	virtual FString GetActiveAssetFolderPath() override;
 
-	virtual FString MakePackageName(const FString& AssetName, const FString& FolderPath) override;
-
-	virtual FString MakeUniqueAssetName(const FString& AssetName, const FString& FolderPath) override;
-
-	// FolderPath should neither start with, nor end with, path separators
-	// Generated path will be /Game/FolderPath/AssetName
-	virtual UPackage* CreateNewPackage(const FString& AssetName, const FString& FolderPath) override;
+	virtual UPackage* MakeNewAssetPackage(const FString& FolderPath, const FString& AssetBaseName, FString& UniqueAssetName);
 
 	virtual void InteractiveSaveGeneratedAsset(UObject* Asset, UPackage* AssetPackage) override;
 
diff --git a/Engine/Source/Runtime/Experimental/InteractiveToolsFramework/Public/ToolContextInterfaces.h b/Engine/Source/Runtime/Experimental/InteractiveToolsFramework/Public/ToolContextInterfaces.h
--- a/Engine/Source/Runtime/Experimental/InteractiveToolsFramework/Public/ToolContextInterfaces.h
+++ b/Engine/Source/Runtime/Experimental/InteractiveToolsFramework/Public/ToolContextInterfaces.h
@@ -312,14 +312,14 @@ public:
 	/** Get default path to save assets in. For example the currently-visible path in the Editor. */
 	virtual FString GetActiveAssetFolderPath() = 0;
 
-	/** Combines folder and asset names */
-	virtual FString MakePackageName(const FString& AssetName, const FString& FolderPath) = 0;
-
-	/** return "unique" version of AssetName, created by appending _1, _2, ... if AssetName already exists */
-	virtual FString MakeUniqueAssetName(const FString& AssetName, const FString& FolderPath) = 0;
-
-	/** create a new package at the given path (calls MakePackageName) */
-	virtual UPackage* CreateNewPackage(const FString& AssetName, const FString& FolderPath) = 0;
+	/**
+	 * Creates a new package for an asset
+	 * @param FolderPath path for new package
+	 * @param AssetBaseName base name for asset
+	 * @param UniqueAssetNameOut unique name in form of AssetBaseName##, where ## is a unqiue index
+	 * @return new package
+	 */
+	virtual UPackage* MakeNewAssetPackage(const FString& FolderPath, const FString& AssetBaseName, FString& UniqueAssetNameOut) = 0;
 
 	/** Request saving of asset to persistent storage via something like an interactive popup dialog */
 	virtual void InteractiveSaveGeneratedAsset(UObject* Asset, UPackage* AssetPackage) = 0;
