Modes Toolbar Palettes

diff --git a/Engine/Source/Editor/EditorStyle/Private/SlateEditorStyle.cpp b/Engine/Source/Editor/EditorStyle/Private/SlateEditorStyle.cpp
--- a/Engine/Source/Editor/EditorStyle/Private/SlateEditorStyle.cpp
+++ b/Engine/Source/Editor/EditorStyle/Private/SlateEditorStyle.cpp
@@ -2020,7 +2020,7 @@ void FSlateEditorStyle::FStyle::SetupGeneralStyles()
 		Set( "ToolBar.Block.IndentedPadding", FMargin( 18.0f, 2.0f, 4.0f, 4.0f ) );
 		Set( "ToolBar.Block.Padding", FMargin( 2.0f, 2.0f, 4.0f, 4.0f ) );
 
-		Set( "ToolBar.Separator", new BOX_BRUSH( "Old/Button", 4.0f/32.0f ) );
+		Set( "ToolBar.Separator", new FSlateColorBrush( FLinearColor(FColor(47, 47, 47)) ) );
 		Set( "ToolBar.Separator.Padding", FMargin( 0.5f ) );
 
 		Set( "ToolBar.Label", FTextBlockStyle(NormalText) .SetFont( DEFAULT_FONT( "Regular", 9 ) ) );
@@ -4909,6 +4909,22 @@ void FSlateEditorStyle::FStyle::SetupLevelEditorStyle()
 		Set( "LevelEditor.EditorModes.Small", new IMAGE_BRUSH( "Icons/icon_Editor_Modes_40x", Icon20x20 ) );
 		Set( "LevelEditor.EditorModes.Menu", new IMAGE_BRUSH( "Icons/icon_Editor_Modes_16x", Icon16x16 ) );
 
+
+		Set( "ToolPalette.DockingTab", FCheckBoxStyle()
+			.SetCheckBoxType(ESlateCheckBoxType::ToggleButton)
+			.SetPadding( FMargin(16.0f, 2.0f, 16.0f, 2.0f ) )
+			.SetCheckedImage(             FSlateBoxBrush( RootToCoreContentDir("Docking/Tab_Shape.png"),  2.f/8.0f, FColor(62, 62, 62) ) )
+			.SetCheckedHoveredImage(      FSlateBoxBrush( RootToCoreContentDir("Docking/Tab_Shape.png"),  2.f/8.0f, FColor(62, 62, 62) ) )
+			.SetCheckedPressedImage(      FSlateBoxBrush( RootToCoreContentDir("Docking/Tab_Shape.png"),  2.f/8.0f, FColor(62, 62, 62) ) )
+			.SetUncheckedImage(           FSlateNoResource() )
+			.SetUncheckedHoveredImage(    FSlateNoResource() )
+			.SetUncheckedPressedImage(    FSlateNoResource() )
+			.SetUndeterminedImage(        FSlateNoResource() )
+			.SetUndeterminedHoveredImage( FSlateNoResource() )
+			.SetUndeterminedPressedImage( FSlateNoResource() )
+		);	
+		Set( "ToolPalette.DockingWell", new FSlateColorBrush(FLinearColor(FColor(34, 34, 34, 255))));
+
 		Set( "LevelEditor.PlacementMode", new IMAGE_BRUSH( "Icons/icon_Mode_Placement_40px", Icon40x40 ) );
 		Set( "LevelEditor.PlacementMode.Small", new IMAGE_BRUSH( "Icons/icon_Mode_Placement_40px", Icon20x20 ) );
 		Set( "LevelEditor.PlacementMode.Selected", new IMAGE_BRUSH( "Icons/icon_Mode_Placement_selected_40x", Icon40x40 ) );
diff --git a/Engine/Source/Editor/LevelEditor/Private/SLevelEditorToolBox.cpp b/Engine/Source/Editor/LevelEditor/Private/SLevelEditorToolBox.cpp
--- a/Engine/Source/Editor/LevelEditor/Private/SLevelEditorToolBox.cpp
+++ b/Engine/Source/Editor/LevelEditor/Private/SLevelEditorToolBox.cpp
@@ -40,7 +40,7 @@ void SLevelEditorToolBox::Construct( const FArguments& InArgs, const TSharedRef<
 		+ SVerticalBox::Slot()
 		.AutoHeight()
 		.HAlign( HAlign_Left )
-		.Padding(0, 0, 0, 0)
+		.Padding(1.f)
 		[
 			SAssignNew( ModeToolBarContainer, SBorder )
 			.BorderImage( FEditorStyle::GetBrush( "NoBorder" ) )
diff --git a/Engine/Source/Editor/UnrealEd/Private/EditorModeManager.cpp b/Engine/Source/Editor/UnrealEd/Private/EditorModeManager.cpp
--- a/Engine/Source/Editor/UnrealEd/Private/EditorModeManager.cpp
+++ b/Engine/Source/Editor/UnrealEd/Private/EditorModeManager.cpp
@@ -16,6 +16,7 @@
 #include "Toolkits/IToolkitHost.h"
 #include "Framework/Notifications/NotificationManager.h"
 #include "Widgets/Notifications/SNotificationList.h"
+#include "Widgets/Input/SButton.h"
 #include "Engine/LevelStreaming.h"
 #include "EditorWorldExtension.h"
 #include "ViewportWorldInteraction.h"
@@ -24,6 +25,7 @@
 #include "Editor/UnrealEdEngine.h"
 #include "Bookmarks/IBookmarkTypeTools.h"
 #include "Widgets/Docking/SDockTab.h"
+#include "Widgets/Layout/SWidgetSwitcher.h"
 #include "EditorStyleSet.h"
 #include "Framework/Commands/UICommandList.h"
 #include "Framework/MultiBox/MultiBoxBuilder.h"
@@ -709,9 +711,9 @@ void FEditorModeTools::DeactivateModeAtIndex(int32 InIndex)
 	// Remove the toolbar widget
 	ActiveToolBarRows.RemoveAll(
 		[&Mode](FEdModeToolbarRow& Row)
-	{
-		return Row.ModeID == Mode->GetID();
-	}
+		{
+			return Row.ModeID == Mode->GetID();
+		}
 	);
 
 	RebuildModeToolBar();
@@ -731,9 +733,9 @@ void FEditorModeTools::DeactivateScriptableModeAtIndex(int32 InIndex)
 	// Remove the toolbar widget
 	ActiveToolBarRows.RemoveAll(
 		[&Mode](FEdModeToolbarRow& Row)
-	{
-		return Row.ModeID == Mode->GetID();
-	}
+		{
+			return Row.ModeID == Mode->GetID();
+		}
 	);
 
 	RebuildModeToolBar();
@@ -750,22 +752,72 @@ void FEditorModeTools::RebuildModeToolBar()
 	{
 		ModeToolbarBoxPinned->ClearChildren();
 
-		if(ActiveToolBarRows.Num())
+		TSharedRef< SHorizontalBox > PaletteTabBox = SNew(SHorizontalBox);
+		TSharedRef< SWidgetSwitcher > PaletteSwitcher = SNew(SWidgetSwitcher);
+
+		int32 PaletteCount = ActiveToolBarRows.Num();
+		if(PaletteCount > 0)
 		{
-			for(int32 RowIdx = 0; RowIdx < ActiveToolBarRows.Num(); ++RowIdx)
+			for(int32 RowIdx = 0; RowIdx < PaletteCount; ++RowIdx)
 			{
 				const FEdModeToolbarRow& Row = ActiveToolBarRows[RowIdx];
 				if (ensure(Row.ToolbarWidget.IsValid()))
 				{
-					ModeToolbarBoxPinned->AddSlot()
-					.HAlign(HAlign_Left)
-					.AutoHeight()
-					.Padding(0.0f, RowIdx > 0 ? 5.0f : 0.0f, 0.0f, 0.0f)
+					TSharedRef<SWidget> PaletteWidget = Row.ToolbarWidget.ToSharedRef();
+
+					FEdMode* Mode = GetActiveMode(Row.ModeID);
+					TSharedPtr<FModeToolkit> RowToolkit = Mode->GetToolkit();
+
+					// Don't show Palette Tabs if there is only one
+					if (PaletteCount > 1)
+					{
+						PaletteTabBox->AddSlot()
+						.AutoWidth()
+						.Padding(FMargin(3.f, 1.0f, 3.0f, 0.0f))
+						[
+							SNew(SCheckBox)
+							.Style( FEditorStyle::Get(),  "ToolPalette.DockingTab" )
+							.OnCheckStateChanged_Lambda( [PaletteSwitcher, Row, RowToolkit] (const ECheckBoxState) { 
+									PaletteSwitcher->SetActiveWidget(Row.ToolbarWidget.ToSharedRef());
+									RowToolkit->OnToolPaletteChanged(Row.PaletteName);
+								} 
+							)
+							.IsChecked_Lambda( [PaletteSwitcher, PaletteWidget] () -> ECheckBoxState { return PaletteSwitcher->GetActiveWidget() == PaletteWidget ? ECheckBoxState::Checked : ECheckBoxState::Unchecked; })
+							[
+								SNew(STextBlock)
+								.Text(Row.DisplayName)
+							]
+						];
+					}
+
+					PaletteSwitcher->AddSlot()
 					[
-						Row.ToolbarWidget.ToSharedRef()
-					];
+						PaletteWidget
+					]; 
 				}
 			}
+
+			ModeToolbarBoxPinned->AddSlot()
+			.AutoHeight()
+			[
+				SNew(SOverlay)
+				+SOverlay::Slot()
+				[
+					SNew(SImage)
+					.Image(FEditorStyle::GetBrush("ToolPalette.DockingWell"))
+				]
+
+				+SOverlay::Slot()
+				[
+					PaletteTabBox
+				]
+			];
+
+			ModeToolbarBoxPinned->AddSlot()
+			.Padding(1.f)
+			[
+				PaletteSwitcher	
+			];
 		}
 		else
 		{
@@ -774,7 +826,6 @@ void FEditorModeTools::RebuildModeToolBar()
 	}
 }
 
-
 void FEditorModeTools::SpawnOrUpdateModeToolbar()
 {
 	if(ShouldShowModeToolbar())
@@ -876,6 +927,7 @@ TSharedRef<SDockTab> FEditorModeTools::MakeModeToolbarTab()
 		SNew(SDockTab)
 		.Label(NSLOCTEXT("EditorModes", "EditorModesToolbarTitle", "Mode Toolbar"))
 		.ShouldAutosize(true)
+		.ContentPadding(0.0f)
 		.Icon(FEditorStyle::GetBrush("ToolBar.Icon"))
 		[
 			SAssignNew(ModeToolbarBox, SVerticalBox)
@@ -985,17 +1037,25 @@ void FEditorModeTools::ActivateMode(FEditorModeID InID, bool bToggle)
 			if (Toolkit.IsValid())
 			{
 				CommandList = Toolkit->GetToolkitCommands();
-			}
 
-			FToolBarBuilder ModeToolbarBuilder(CommandList, FMultiBoxCustomization(Mode->GetModeInfo().ToolbarCustomizationName), TSharedPtr<FExtender>(), Orient_Horizontal, false);
-			Mode->BuildModeToolbar(ModeToolbarBuilder);
+				// Also build the toolkit here 
+				int32 PaletteCount = 0;
+				TArray<FName> PaletteNames;
+				Toolkit->GetToolPaletteNames(PaletteNames);
+				for(auto Palette : PaletteNames)
+				{
+					FToolBarBuilder ModeToolbarBuilder(CommandList, FMultiBoxCustomization(Mode->GetModeInfo().ToolbarCustomizationName), TSharedPtr<FExtender>(), Orient_Horizontal, false);
+					Toolkit->BuildToolPalette(Palette, ModeToolbarBuilder);
 
-			if (ModeToolbarBuilder.GetMultiBox()->GetBlocks().Num())
-			{
-				TSharedRef<SWidget> ToolbarWidget = ModeToolbarBuilder.MakeWidget();
-				ActiveToolBarRows.Emplace(Mode->GetID(), ToolbarWidget);
+					// ActiveToolBarRows.Emplace(Mode, Toolkit->GetToolPaletteName(i), ModeToolbarBuilder.MakeWidget());
+					ActiveToolBarRows.Emplace(Mode->GetID(), Palette, Toolkit->GetToolPaletteDisplayName(Palette), ModeToolbarBuilder.MakeWidget());
+					PaletteCount++;
+				}
 
-				SpawnOrUpdateModeToolbar();
+				if (PaletteCount > 0) 
+				{
+					SpawnOrUpdateModeToolbar();
+				}
 			}
 
 		}
@@ -1049,25 +1109,29 @@ void FEditorModeTools::ActivateMode(FEditorModeID InID, bool bToggle)
 			if (Toolkit.IsValid())
 			{
 				CommandList = Toolkit->GetToolkitCommands();
-			}
 
-			FToolBarBuilder ModeToolbarBuilder(CommandList, FMultiBoxCustomization(Mode->GetModeInfo().ToolbarCustomizationName), TSharedPtr<FExtender>(), Orient_Horizontal, false);
-			Mode->BuildModeToolbar(ModeToolbarBuilder);
+				// Also build the toolkit here 
+				int32 PaletteCount = 0;
+				TArray<FName> PaletteNames;
+				Toolkit->GetToolPaletteNames(PaletteNames);
+				for(auto Palette : PaletteNames)
+				{
+					FToolBarBuilder ModeToolbarBuilder(CommandList, FMultiBoxCustomization(Mode->GetModeInfo().ToolbarCustomizationName), TSharedPtr<FExtender>(), Orient_Horizontal, false);
+					Toolkit->BuildToolPalette(Palette, ModeToolbarBuilder);
 
-			if (ModeToolbarBuilder.GetMultiBox()->GetBlocks().Num())
-			{
-				TSharedRef<SWidget> ToolbarWidget = ModeToolbarBuilder.MakeWidget();
-				ActiveToolBarRows.Emplace(Mode->GetID(), ToolbarWidget);
+					// ActiveToolBarRows.Emplace(Mode, Toolkit->GetToolPaletteName(i), ModeToolbarBuilder.MakeWidget());
+					ActiveToolBarRows.Emplace(Mode->GetID(), Palette, Toolkit->GetToolPaletteDisplayName(Palette), ModeToolbarBuilder.MakeWidget());
+					PaletteCount++;
+				}
 
-				SpawnOrUpdateModeToolbar();
+				if (PaletteCount > 0) 
+				{
+					SpawnOrUpdateModeToolbar();
+				}
 			}
 		}
 	}
 
-
-
-
-
 	// Update the editor UI
 	FEditorSupportDelegates::UpdateUI.Broadcast();
 }
diff --git a/Engine/Source/Editor/UnrealEd/Public/EditorModeManager.h b/Engine/Source/Editor/UnrealEd/Public/EditorModeManager.h
--- a/Engine/Source/Editor/UnrealEd/Public/EditorModeManager.h
+++ b/Engine/Source/Editor/UnrealEd/Public/EditorModeManager.h
@@ -593,11 +593,15 @@ protected:
 private:
 	struct FEdModeToolbarRow
 	{
-		FEdModeToolbarRow(FName InModeID, TSharedRef<SWidget>& InToolbarWidget)
+		FEdModeToolbarRow(FName InModeID, FName InPaletteName, FText InDisplayName, TSharedRef<SWidget> InToolbarWidget)
 			: ModeID(InModeID)
+			, PaletteName(InPaletteName)
+			, DisplayName(InDisplayName)
 			, ToolbarWidget(InToolbarWidget)
 		{}
 		FName ModeID;
+		FName PaletteName;
+		FText DisplayName;
 		TSharedPtr<SWidget> ToolbarWidget;
 	};
 
diff --git a/Engine/Source/Editor/UnrealEd/Public/Toolkits/IToolkit.h b/Engine/Source/Editor/UnrealEd/Public/Toolkits/IToolkit.h
--- a/Engine/Source/Editor/UnrealEd/Public/Toolkits/IToolkit.h
+++ b/Engine/Source/Editor/UnrealEd/Public/Toolkits/IToolkit.h
@@ -136,5 +136,22 @@ public:
 
 	/* Returns true if this is the default generic asset editor used by multiple classes */
 	virtual bool IsSimpleAssetEditor() const { return false; }
+
+	/** Returns the number of Mode specific tabs in the mode toolbar **/ 
+	virtual void GetToolPaletteNames(TArray<FName>& PaletteNames) const {}
+
+	/** 
+	 * @param PaletteIndex      The index of the ToolPalette to build
+	 * @returns the name of Tool Palette 
+	 **/ 
+	virtual FText GetToolPaletteDisplayName(FName Palette) { return FText(); }
+
+	/** 
+	 * @param PaletteIndex      The index of the ToolPalette to build
+	 * @param ToolbarBuilder    The builder to use for given PaletteIndex
+	**/ 
+	virtual void BuildToolPalette(FName Palette, class FToolBarBuilder& ToolbarBuilder) {}
+
+	virtual void OnToolPaletteChanged(FName NewActivePalette) {}
 };
 
diff --git a/Engine/Source/Runtime/Slate/Private/Framework/MultiBox/MultiBoxBuilder.cpp b/Engine/Source/Runtime/Slate/Private/Framework/MultiBox/MultiBoxBuilder.cpp
--- a/Engine/Source/Runtime/Slate/Private/Framework/MultiBox/MultiBoxBuilder.cpp
+++ b/Engine/Source/Runtime/Slate/Private/Framework/MultiBox/MultiBoxBuilder.cpp
@@ -5,6 +5,7 @@
 #include "Widgets/Text/STextBlock.h"
 #include "Framework/MultiBox/MultiBox.h"
 #include "Widgets/Layout/SBox.h"
+#include "Widgets/SBoxPanel.h"
 #include "Framework/MultiBox/SHeadingBlock.h"
 #include "Framework/MultiBox/SMenuEntryBlock.h"
 #include "Framework/MultiBox/SMenuSeparatorBlock.h"
@@ -416,6 +417,42 @@ void FToolBarBuilder::AddComboButton( const FUIAction& InAction, const FOnGetCon
 	MultiBox->AddMultiBlock( NewToolBarComboButtonBlock );
 }
 
+void FToolBarBuilder::AddToolBarWidget( TSharedRef<SWidget> InWidget, const TAttribute<FText>& InLabel, FName InTutorialHighlightName, bool bSearchable )
+{
+	ApplySectionBeginning();
+
+	// If tutorial name specified, wrap in tutorial wrapper
+	const FName WrapperName = GenerateTutorialIdentfierName(InTutorialHighlightName, NAME_None, nullptr, MultiBox->GetBlocks().Num());
+
+	TSharedRef<SWidget> ChildWidget = InWidget;
+	InWidget = 
+		SNew( SVerticalBox )
+		.AddMetaData<FTagMetaData>(FTagMetaData(InTutorialHighlightName))
+
+		+SVerticalBox::Slot()
+		.AutoHeight()
+		.HAlign( HAlign_Center )
+		[
+			ChildWidget
+		]
+
+		+SVerticalBox::Slot()
+		.AutoHeight()
+		.HAlign( HAlign_Center )
+		[
+			SNew( STextBlock )
+			// .Visibility_Lambda( [this] () -> EVisibility { return /*LabelVisibility.IsSet() ? LabelVisibility.GetValue() :*/ (bForceSmallIcons || FMultiBoxSettings::UseSmallToolBarIcons.Get()) ? EVisibility::Collapsed : EVisibility::Visible; } ) 
+			.Visibility_Lambda( [] () -> EVisibility { return FMultiBoxSettings::UseSmallToolBarIcons.Get() ? EVisibility::Collapsed : EVisibility::Visible; } ) 
+			.Text( InLabel )
+			.TextStyle( GetStyleSet(), ISlateStyle::Join( GetStyleName(), ".Label" ) )	// Smaller font for tool tip labels
+			.ShadowOffset( FVector2D::UnitVector )
+		] ;
+	
+	TSharedRef< FWidgetBlock > NewWidgetBlock( new FWidgetBlock( InWidget, FText::GetEmpty(), true ) );
+	MultiBox->AddMultiBlock( NewWidgetBlock );
+	NewWidgetBlock->SetSearchable(bSearchable);
+}
+
 void FToolBarBuilder::AddWidget( TSharedRef<SWidget> InWidget, FName InTutorialHighlightName, bool bSearchable )
 {
 	ApplySectionBeginning();
diff --git a/Engine/Source/Runtime/Slate/Private/Framework/MultiBox/SToolBarSeparatorBlock.cpp b/Engine/Source/Runtime/Slate/Private/Framework/MultiBox/SToolBarSeparatorBlock.cpp
--- a/Engine/Source/Runtime/Slate/Private/Framework/MultiBox/SToolBarSeparatorBlock.cpp
+++ b/Engine/Source/Runtime/Slate/Private/Framework/MultiBox/SToolBarSeparatorBlock.cpp
@@ -45,7 +45,6 @@ void SToolBarSeparatorBlock::Construct( const FArguments& InArgs )
 }
 
 
-
 /**
  * Builds this MultiBlock widget up from the MultiBlock associated with it
  */
@@ -56,10 +55,11 @@ void SToolBarSeparatorBlock::BuildMultiBlockWidget(const ISlateStyle* StyleSet,
 		SNew( SHorizontalBox )
 		+SHorizontalBox::Slot()
 		.AutoWidth()
-		.Padding( 0.0f, 0.0f, 0.0f, 0.0f )
+		.Padding( 8.0f, 0.0f, 8.0f, 0.0f )
 		[
 			SNew(SSeparator)
 				.Orientation(Orient_Vertical)
+				.Thickness(2.0f)
 				.SeparatorImage( StyleSet->GetBrush( ISlateStyle::Join( StyleName, ".Separator" ) ) )
 		]
 	];
diff --git a/Engine/Source/Runtime/Slate/Public/Framework/MultiBox/MultiBoxBuilder.h b/Engine/Source/Runtime/Slate/Public/Framework/MultiBox/MultiBoxBuilder.h
--- a/Engine/Source/Runtime/Slate/Public/Framework/MultiBox/MultiBoxBuilder.h
+++ b/Engine/Source/Runtime/Slate/Public/Framework/MultiBox/MultiBoxBuilder.h
@@ -433,6 +433,18 @@ public:
 	 */
 	void AddComboButton( const FUIAction& InAction, const FOnGetContent& InMenuContentGenerator, const TAttribute<FText>& InLabelOverride = TAttribute<FText>(), const TAttribute<FText>& InToolTipOverride = TAttribute<FText>(), const TAttribute<FSlateIcon>& InIconOverride = TAttribute<FSlateIcon>(), bool bInSimpleComboBox = false, FName InTutorialHighlightName = NAME_None );
 
+
+	/**
+	 * Adds any widget to the toolbar
+	 * 
+	 * @param	InWidget				The widget that should be shown in the toolbar
+	 * @param	InLabel                 Optional Label.  
+	 * @param	InTutorialHighlightName	Name to identify this widget and highlight during tutorials
+	 * @param	bSearchable			If true, widget will be searchable (default == true)
+	 */
+	void AddToolBarWidget(TSharedRef<SWidget> InWidget, const TAttribute<FText>& InLabel = TAttribute<FText>(), FName InTutorialHighlightName = NAME_None, bool bSearchable = true);
+
+
 	/**
 	 * Adds any widget to the toolbar
 	 * 
