add predeclarations so InteractiveToolManager.h does not need to be included in UEdMode.h
#rnx

diff --git a/Engine/Source/Editor/UnrealEd/Private/Tools/UEdMode.cpp b/Engine/Source/Editor/UnrealEd/Private/Tools/UEdMode.cpp
--- a/Engine/Source/Editor/UnrealEd/Private/Tools/UEdMode.cpp
+++ b/Engine/Source/Editor/UnrealEd/Private/Tools/UEdMode.cpp
@@ -21,6 +21,7 @@
 #include "Materials/Material.h"
 #include "Materials/MaterialInterface.h"
 #include "LevelEditorViewport.h"
+#include "InteractiveToolManager.h"
 #include "InteractiveToolObjects.h"
 
 
diff --git a/Engine/Source/Editor/UnrealEd/Public/Tools/UEdMode.h b/Engine/Source/Editor/UnrealEd/Public/Tools/UEdMode.h
--- a/Engine/Source/Editor/UnrealEd/Public/Tools/UEdMode.h
+++ b/Engine/Source/Editor/UnrealEd/Public/Tools/UEdMode.h
@@ -9,8 +9,8 @@
 #include "EditorModeRegistry.h"
 #include "Templates/SharedPointer.h"
 #include "UObject/Object.h"
-#include "InteractiveToolManager.h"
 #include "Math/Ray.h"
+#include "InputState.h"
 #include "UEdMode.generated.h"
 
 
@@ -26,6 +26,10 @@ class UTexture2D;
 struct FConvexVolume;
 struct FViewportClick;
 class UInteractiveToolsContext;
+class UInteractiveToolManager;
+class IToolsContextQueriesAPI;
+class IToolsContextTransactionsAPI;
+class IToolsContextAssetAPI;
 
 enum EModeTools : int8;
 class FEditorViewportClient;
