Adding missing SImage include

diff --git a/Engine/Source/Editor/UnrealEd/Private/EditorModeManager.cpp b/Engine/Source/Editor/UnrealEd/Private/EditorModeManager.cpp
--- a/Engine/Source/Editor/UnrealEd/Private/EditorModeManager.cpp
+++ b/Engine/Source/Editor/UnrealEd/Private/EditorModeManager.cpp
@@ -31,6 +31,7 @@
 #include "Framework/MultiBox/MultiBoxBuilder.h"
 #include "Toolkits/BaseToolkit.h"
 #include "Tools/UEdMode.h"
+#include "Widgets/Images/SImage.h"
 
 /*------------------------------------------------------------------------------
 	FEditorModeTools.
