Fix for Replace Actors not properly finding actor references in sublevels. Instead of performing the reference fixup on the editor world, we now step through each actor being replaced and gather up it's level so that we know where to search.

diff --git a/Engine/Source/Editor/UnrealEd/Private/EditorEngine.cpp b/Engine/Source/Editor/UnrealEd/Private/EditorEngine.cpp
--- a/Engine/Source/Editor/UnrealEd/Private/EditorEngine.cpp
+++ b/Engine/Source/Editor/UnrealEd/Private/EditorEngine.cpp
@@ -5241,23 +5241,31 @@ void UEditorEngine::ReplaceActors(UActorFactory* Factory, const FAssetData& Asse
 	ReattachActorsHelper::ReattachActors(ConvertedMap, AttachmentInfo);
 
 	// Perform reference replacement on all Actors referenced by World
-	UWorld* CurrentEditorWorld = GetEditorWorldContext().World();
-	FArchiveReplaceObjectRef<AActor> Ar(CurrentEditorWorld, ConvertedMap, false, true, false);
+	TArray<UObject*> ReferencedLevels;
 
-	// Go through modified objects, marking their packages as dirty and informing them of property changes
-	for (const auto& MapItem : Ar.GetReplacedReferences())
+	for (const TPair<AActor*, AActor*>& ReplacedObj : ConvertedMap)
 	{
-		UObject* ModifiedObject = MapItem.Key;
+		ReferencedLevels.AddUnique(ReplacedObj.Value->GetLevel());
+	}
 
-		if (!ModifiedObject->HasAnyFlags(RF_Transient) && ModifiedObject->GetOutermost() != GetTransientPackage() && !ModifiedObject->RootPackageHasAnyFlags(PKG_CompiledIn))
-		{
-			ModifiedObject->MarkPackageDirty();
-		}
+	for (UObject* Referencer : ReferencedLevels)
+	{
+		FArchiveReplaceObjectRef<AActor> Ar(Referencer, ConvertedMap, false, true, false);
 
-		for (UProperty* Property : MapItem.Value)
+		for (const auto& MapItem : Ar.GetReplacedReferences())
 		{
-			FPropertyChangedEvent PropertyEvent(Property);
-			ModifiedObject->PostEditChangeProperty(PropertyEvent);
+			UObject* ModifiedObject = MapItem.Key;
+
+			if (!ModifiedObject->HasAnyFlags(RF_Transient) && ModifiedObject->GetOutermost() != GetTransientPackage() && !ModifiedObject->RootPackageHasAnyFlags(PKG_CompiledIn))
+			{
+				ModifiedObject->MarkPackageDirty();
+			}
+
+			for (UProperty* Property : MapItem.Value)
+			{
+				FPropertyChangedEvent PropertyEvent(Property);
+				ModifiedObject->PostEditChangeProperty(PropertyEvent);
+			}
 		}
 	}
 
