Rename a setting to bEnableEditToolMenusUI
#rnx

diff --git a/Engine/Source/Editor/UnrealEd/Private/EditorEngine.cpp b/Engine/Source/Editor/UnrealEd/Private/EditorEngine.cpp
--- a/Engine/Source/Editor/UnrealEd/Private/EditorEngine.cpp
+++ b/Engine/Source/Editor/UnrealEd/Private/EditorEngine.cpp
@@ -798,9 +798,9 @@ void UEditorEngine::InitEditor(IEngineLoop* InEngineLoop)
 			}
 		}));
 
-		bool bShowEditMenusModeCheckbox = false;
-		GConfig->GetBool(TEXT("/Script/UnrealEd.EditorExperimentalSettings"), TEXT("bShowEditMenusModeCheckbox"), bShowEditMenusModeCheckbox, GEditorPerProjectIni);
-		if (bShowEditMenusModeCheckbox)
+		bool bEnableEditToolMenusUI = false;
+		GConfig->GetBool(TEXT("/Script/UnrealEd.EditorExperimentalSettings"), TEXT("bEnableEditToolMenusUI"), bEnableEditToolMenusUI, GEditorPerProjectIni);
+		if (bEnableEditToolMenusUI)
 		{
 			IToolMenusEditorModule::Get().RegisterShowEditMenusModeCheckbox();
 
