ToolsFramework: Add ToggleWireframe as a standard mode command (used by InteractiveTools to turn single-object wireframe overlay on/off)
#rnx

diff --git a/Engine/Source/Editor/UnrealEd/Private/Tools/StandardToolModeCommands.cpp b/Engine/Source/Editor/UnrealEd/Private/Tools/StandardToolModeCommands.cpp
--- a/Engine/Source/Editor/UnrealEd/Private/Tools/StandardToolModeCommands.cpp
+++ b/Engine/Source/Editor/UnrealEd/Private/Tools/StandardToolModeCommands.cpp
@@ -25,6 +25,11 @@ void FStandardToolModeCommands::RegisterCommands()
 	TSharedPtr<FUICommandInfo> DecreaseBrushSize;
 	UI_COMMAND(DecreaseBrushSize, "Decrease Brush Size", "Decreases the size of the brush", EUserInterfaceActionType::Button, FInputChord(EKeys::LeftBracket));
 	Commands.Add(EStandardToolModeCommands::DecreaseBrushSize, DecreaseBrushSize);
+
+	TSharedPtr<FUICommandInfo> ToggleWireframe;
+	UI_COMMAND(ToggleWireframe, "Toggle Wireframe", "Toggle Wireframe overlay on/off", EUserInterfaceActionType::Button, FInputChord(EModifierKey::Control, EKeys::W));
+	Commands.Add(EStandardToolModeCommands::ToggleWireframe, ToggleWireframe);
+
 }
 
 
diff --git a/Engine/Source/Editor/UnrealEd/Public/Tools/InteractiveToolsCommands.h b/Engine/Source/Editor/UnrealEd/Public/Tools/InteractiveToolsCommands.h
--- a/Engine/Source/Editor/UnrealEd/Public/Tools/InteractiveToolsCommands.h
+++ b/Engine/Source/Editor/UnrealEd/Public/Tools/InteractiveToolsCommands.h
@@ -222,6 +222,8 @@ TSharedPtr<FUICommandInfo> TInteractiveToolCommands<CommandContextType>::FindSta
 		return StdCommands.FindStandardCommand(EStandardToolModeCommands::IncreaseBrushSize);
 	case EStandardToolActions::DecreaseBrushSize:
 		return StdCommands.FindStandardCommand(EStandardToolModeCommands::DecreaseBrushSize);
+	case EStandardToolActions::ToggleWireframe:
+		return StdCommands.FindStandardCommand(EStandardToolModeCommands::ToggleWireframe);
 	default:
 		return nullptr;
 	}
diff --git a/Engine/Source/Editor/UnrealEd/Public/Tools/StandardToolModeCommands.h b/Engine/Source/Editor/UnrealEd/Public/Tools/StandardToolModeCommands.h
--- a/Engine/Source/Editor/UnrealEd/Public/Tools/StandardToolModeCommands.h
+++ b/Engine/Source/Editor/UnrealEd/Public/Tools/StandardToolModeCommands.h
@@ -14,7 +14,8 @@
 enum class EStandardToolModeCommands
 {
 	IncreaseBrushSize = 100,
-	DecreaseBrushSize = 101
+	DecreaseBrushSize = 101,
+	ToggleWireframe = 102
 };
 
 
diff --git a/Engine/Source/Runtime/Experimental/InteractiveToolsFramework/Public/InteractiveToolActionSet.h b/Engine/Source/Runtime/Experimental/InteractiveToolsFramework/Public/InteractiveToolActionSet.h
--- a/Engine/Source/Runtime/Experimental/InteractiveToolsFramework/Public/InteractiveToolActionSet.h
+++ b/Engine/Source/Runtime/Experimental/InteractiveToolsFramework/Public/InteractiveToolActionSet.h
@@ -18,6 +18,7 @@ enum class EStandardToolActions
 {
 	IncreaseBrushSize = 100,
 	DecreaseBrushSize = 101,
+	ToggleWireframe = 102,
 
 
 	BaseClientDefinedActionID = 10000
