Skeletal mesh build refactor, the skeletal mesh now have the capability to be rebuild without reimport like the staticmesh.

diff --git a/Engine/Source/Editor/UnrealEd/Private/Fbx/FbxFactory.cpp b/Engine/Source/Editor/UnrealEd/Private/Fbx/FbxFactory.cpp
--- a/Engine/Source/Editor/UnrealEd/Private/Fbx/FbxFactory.cpp
+++ b/Engine/Source/Editor/UnrealEd/Private/Fbx/FbxFactory.cpp
@@ -549,6 +549,10 @@ UObject* UFbxFactory::FactoryCreateFile
 						int32 LODIndex;
 						int32 SuccessfulLodIndex = 0;
 						bool bImportSkeletalMeshLODs = ImportUI->SkeletalMeshImportData->bImportMeshLODs;
+						
+						//The skeletalmesh will be set after we import the LOD 0 since it is not created yet.
+						FScopedSkeletalMeshPostEditChange ScopedPostEditChange(nullptr);
+
 						for (LODIndex = 0; LODIndex < MaxLODLevel; LODIndex++)
 						{
 							//We need to know what is the imported lod index when importing the morph targets
@@ -605,6 +609,9 @@ UObject* UFbxFactory::FactoryCreateFile
 
 								USkeletalMesh* NewMesh = FbxImporter->ImportSkeletalMesh( ImportSkeletalMeshArgs );
 								CreatedObject = NewMesh;
+								
+								//Set the base skeletalmesh to the scoped post edit change variable
+								ScopedPostEditChange.SetSkeletalMesh(NewMesh);
 
 								if(bOperationCanceled)
 								{
@@ -643,7 +650,7 @@ UObject* UFbxFactory::FactoryCreateFile
 								LODInfo.ReductionSettings.BaseLOD = 0;
 								LODInfo.bImportWithBaseMesh = true;
 								LODInfo.SourceImportFilename = FString(TEXT(""));
-								FLODUtilities::SimplifySkeletalMeshLOD(UpdateContext, SuccessfulLodIndex, false);
+								FLODUtilities::SimplifySkeletalMeshLOD(UpdateContext, SuccessfulLodIndex);
 								ImportedSuccessfulLodIndex = SuccessfulLodIndex;
 								SuccessfulLodIndex++;
 							}
@@ -662,7 +669,7 @@ UObject* UFbxFactory::FactoryCreateFile
 								ImportSkeletalMeshArgs.OutData = &OutData;
 
 								USkeletalMesh *LODObject = FbxImporter->ImportSkeletalMesh( ImportSkeletalMeshArgs );
-								bool bImportSucceeded = !bOperationCanceled && FbxImporter->ImportSkeletalMeshLOD(LODObject, BaseSkeletalMesh, SuccessfulLodIndex, false);
+								bool bImportSucceeded = !bOperationCanceled && FbxImporter->ImportSkeletalMeshLOD(LODObject, BaseSkeletalMesh, SuccessfulLodIndex);
 
 								if (bImportSucceeded)
 								{
@@ -704,15 +711,6 @@ UObject* UFbxFactory::FactoryCreateFile
 							
 							USkeletalMesh* SkeletalMesh = Cast<USkeletalMesh>(CreatedObject);
 							UnFbx::FFbxImporter::UpdateSkeletalMeshImportData(SkeletalMesh, ImportUI->SkeletalMeshImportData, INDEX_NONE, nullptr, nullptr);
-							
-							//If we have import some morph target we have to rebuild the render resources since morph target are now using GPU
-							if (SkeletalMesh)// && SkeletalMesh->MorphTargets.Num() > 0)
-							{
-								SkeletalMesh->ReleaseResources();
-								//Rebuild the resources with a post edit change since we have added some morph targets
-								SkeletalMesh->PostEditChange();
-							}
-
 						}
 					}
 				
@@ -1391,7 +1389,7 @@ namespace ImportCompareHelper
 		}
 		else if (!bCombineMeshes && !bCombineMeshesLOD)
 		{
-			Node = FFbxImporter->GetMeshNodesFromName(StaticMesh, FbxMeshArray);
+			Node = FFbxImporter->GetMeshNodesFromName(StaticMesh->GetName(), FbxMeshArray);
 		}
 
 		// If there is no match it may be because an LOD group was imported where
@@ -1481,7 +1479,38 @@ namespace ImportCompareHelper
 			return;
 		}
 
-		const TArray<FbxNode*>& SkeletalMeshNodes = *(SkeletalMeshArray[0]);
+		TArray<FbxNode*>* FbxNodes = nullptr;
+		// if there is only one mesh, use it without name checking 
+		// (because the "Used As Full Name" option enables users name the Unreal mesh by themselves
+		if (SkeletalMeshArray.Num() == 1)
+		{
+			FbxNodes = SkeletalMeshArray[0];
+		}
+		else if(SkeletalMeshArray.Num() > 1)
+		{
+			FbxNodes = nullptr;
+			//Search a set corresponding
+			for (TArray<FbxNode*>*SkeletalMeshNodes : SkeletalMeshArray)
+			{
+				FbxNode* Node = FFbxImporter->GetMeshNodesFromName(SkeletalMesh->GetName(), *SkeletalMeshNodes);
+				if (Node != nullptr)
+				{
+					FbxNodes = SkeletalMeshNodes;
+					break;
+				}
+			}
+			if (FbxNodes == nullptr)
+			{
+				//Re import the first skeletalmesh found in the fbx file
+				FbxNodes = SkeletalMeshArray[0];
+			}
+		}
+		if (FbxNodes == nullptr)
+		{
+			return;
+		}
+
+		const TArray<FbxNode*>& SkeletalMeshNodes = *(FbxNodes);
 		if (SkeletalMeshNodes.Num() == 0)
 		{
 			return;
diff --git a/Engine/Source/Editor/UnrealEd/Private/Fbx/FbxSkeletalMeshImport.cpp b/Engine/Source/Editor/UnrealEd/Private/Fbx/FbxSkeletalMeshImport.cpp
--- a/Engine/Source/Editor/UnrealEd/Private/Fbx/FbxSkeletalMeshImport.cpp
+++ b/Engine/Source/Editor/UnrealEd/Private/Fbx/FbxSkeletalMeshImport.cpp
@@ -27,6 +27,8 @@
 #include "Animation/MorphTarget.h"
 #include "PhysicsAssetUtils.h"
 #include "Rendering/SkeletalMeshModel.h"
+#include "Rendering/SkeletalMeshRenderData.h"
+#include "Rendering/SkeletalMeshLODRenderData.h"
 
 #include "SkelImport.h"
 #include "Logging/TokenizedMessage.h"
@@ -54,6 +56,7 @@
 #include "Engine/SkeletalMeshSocket.h"
 #include "Assets/ClothingAsset.h"
 #include "LODUtilities.h"
+#include "IMeshBuilderModule.h"
 
 #define LOCTEXT_NAMESPACE "FBXImpoter"
 
@@ -1433,7 +1436,7 @@ bool UnFbx::FFbxImporter::FillSkeletalMeshImportData(TArray<FbxNode*>& NodeArray
 
 	if (ImportOptions->bPreserveSmoothingGroups)
 	{
-		bool bDuplicateUnSmoothWedges = (ImportOptions->NormalGenerationMethod != EFBXNormalGenerationMethod::MikkTSpace);
+		bool bDuplicateUnSmoothWedges = false; //We deprecate legacy build
 		DoUnSmoothVerts(*SkelMeshImportDataPtr, bDuplicateUnSmoothWedges);
 	}
 	else
@@ -1754,9 +1757,21 @@ USkeletalMesh* UnFbx::FFbxImporter::ImportSkeletalMesh(FImportSkeletalMeshArgs &
 		return nullptr;
 	}
 
+	//Stack the PostEditChange call, it will call post edit change when it will go out of scope
+	FScopedSkeletalMeshPostEditChange ScopedPostEditChange(ExistingSkelMesh);
+
+	ESkeletalMeshGeoImportVersions GeoImportVersion = ESkeletalMeshGeoImportVersions::LatestVersion;
+	ESkeletalMeshSkinningImportVersions SkinningImportVersion = ESkeletalMeshSkinningImportVersions::LatestVersion;
 	//Adjust the import data from the import options
 	if (ExistingSkelMesh != nullptr)
 	{
+		if (ExistingSkelMesh->GetImportedModel() && ExistingSkelMesh->GetImportedModel()->LODModels.IsValidIndex(ImportSkeletalMeshArgs.LodIndex))
+		{
+			const FRawSkeletalMeshBulkData& RawSkeletalMeshBulkData = ExistingSkelMesh->GetImportedModel()->LODModels[ImportSkeletalMeshArgs.LodIndex].RawSkeletalMeshBulkData;
+			GeoImportVersion = RawSkeletalMeshBulkData.GeoImportVersion;
+			SkinningImportVersion = RawSkeletalMeshBulkData.SkinningImportVersion;
+		}
+
 		if ((ImportOptions->bImportAsSkeletalSkinning || ImportOptions->bImportAsSkeletalGeometry))
 		{
 			if (ImportOptions->bImportAsSkeletalSkinning)
@@ -1805,18 +1820,8 @@ USkeletalMesh* UnFbx::FFbxImporter::ImportSkeletalMesh(FImportSkeletalMeshArgs &
 		return nullptr;
 	}
 
-	// Array of re-import contexts for components using this mesh
-	// Will unregister before import, then re-register afterwards
-	TIndirectArray<FComponentReregisterContext> ComponentContexts;
-	TUniquePtr< FSkinnedMeshComponentRecreateRenderStateContext > SkinnedMeshComponentRecreateRenderStateContext;
-
-	FSkinnedMeshComponentRecreateRenderStateContext* RecreateExistingRenderStateContext = ExistingSkelMesh ? new FSkinnedMeshComponentRecreateRenderStateContext(ExistingSkelMesh, false) : nullptr;
-	//Backup the data before importing the new one
 	if (ExistingSkelMesh)
 	{
-		//Release the Renderdata resources before reimporting the skeletal mesh
-		SkinnedMeshComponentRecreateRenderStateContext = MakeUnique<FSkinnedMeshComponentRecreateRenderStateContext>(ExistingSkelMesh);
-
 		ExistingSkelMesh->PreEditChange(NULL);
 		//The backup of the skeletal mesh data empty the LOD array in the ImportedResource of the skeletal mesh
 		//If the import fail after this step the editor can crash when updating the bone later since the LODModel will not exist anymore
@@ -1825,12 +1830,18 @@ USkeletalMesh* UnFbx::FFbxImporter::ImportSkeletalMesh(FImportSkeletalMeshArgs &
 
 	if (SkeletalMesh == nullptr)
 	{
+		//Backup the PostEditchangeStackCounter before overriding the skeletalmesh
+		int32 PostEditChangeStackCounter = ExistingSkelMesh != nullptr ? ExistingSkelMesh->GetPostEditChangeStackCounter() : 0;
 		// Create the new mesh after saving the old data, since it will replace the old skeletalmesh
 		// This should happen only when doing a re-import. Otherwise the skeletal mesh is create before.
 		SkeletalMesh = NewObject<USkeletalMesh>(ImportSkeletalMeshArgs.InParent, ImportSkeletalMeshArgs.Name, ImportSkeletalMeshArgs.Flags);
+		SkeletalMesh->SetPostEditChangeStackCounter(PostEditChangeStackCounter);
 	}
 
+
 	SkeletalMesh->PreEditChange(NULL);
+	//Dirty the DDC Key for any imported Skeletal Mesh
+	SkeletalMesh->InvalidateDeriveDataCacheGUID();
 
 	FSkeletalMeshModel *ImportedResource = SkeletalMesh->GetImportedModel();
 	check(ImportedResource->LODModels.Num() == 0);
@@ -1858,9 +1869,24 @@ USkeletalMesh* UnFbx::FFbxImporter::ImportSkeletalMesh(FImportSkeletalMeshArgs &
 	// process bone influences from import data
 	ProcessImportMeshInfluences(*SkelMeshImportDataPtr);
 
-	//Store the original fbx import data
+	//Store the original fbx import data the SkelMeshImportDataPtr should not be modified after this
 	LODModel.RawSkeletalMeshBulkData.SaveRawMesh(*SkelMeshImportDataPtr);
-
+	if (ImportOptions->bImportAsSkeletalSkinning)
+	{
+		LODModel.RawSkeletalMeshBulkData.SkinningImportVersion = ESkeletalMeshSkinningImportVersions::LatestVersion;
+		LODModel.RawSkeletalMeshBulkData.GeoImportVersion = GeoImportVersion;
+	}
+	else if(ImportOptions->bImportAsSkeletalGeometry)
+	{
+		LODModel.RawSkeletalMeshBulkData.GeoImportVersion = ESkeletalMeshGeoImportVersions::LatestVersion;
+		LODModel.RawSkeletalMeshBulkData.SkinningImportVersion = SkinningImportVersion;
+	}
+	else
+	{
+		//We reimport both
+		LODModel.RawSkeletalMeshBulkData.SkinningImportVersion = ESkeletalMeshSkinningImportVersions::LatestVersion;
+		LODModel.RawSkeletalMeshBulkData.GeoImportVersion = ESkeletalMeshGeoImportVersions::LatestVersion;
+	}
 	SkeletalMesh->ResetLODInfo();
 	FSkeletalMeshLODInfo& NewLODInfo = SkeletalMesh->AddLODInfo();
 	NewLODInfo.ReductionSettings.NumOfTrianglesPercentage = 1.0f;
@@ -1886,38 +1912,61 @@ USkeletalMesh* UnFbx::FFbxImporter::ImportSkeletalMesh(FImportSkeletalMeshArgs &
 		TArray<int32> LODPointToRawMap;
 		SkelMeshImportDataPtr->CopyLODImportData(LODPoints,LODWedges,LODFaces,LODInfluences,LODPointToRawMap);
 
-		IMeshUtilities::MeshBuildOptions BuildOptions;
-		BuildOptions.OverlappingThresholds = ImportOptions->OverlappingThresholds;
-		BuildOptions.bComputeNormals = !ImportOptions->ShouldImportNormals() || !SkelMeshImportDataPtr->bHasNormals;
-		BuildOptions.bComputeTangents = !ImportOptions->ShouldImportTangents() || !SkelMeshImportDataPtr->bHasTangents;
+		FSkeletalMeshBuildSettings BuildOptions;
+		BuildOptions.bBuildAdjacencyBuffer = true;
+		BuildOptions.bRecomputeNormals = !ImportOptions->ShouldImportNormals() || !SkelMeshImportDataPtr->bHasNormals;
+		BuildOptions.bRecomputeTangents = !ImportOptions->ShouldImportTangents() || !SkelMeshImportDataPtr->bHasTangents;
 		BuildOptions.bUseMikkTSpace = (ImportOptions->NormalGenerationMethod == EFBXNormalGenerationMethod::MikkTSpace) && (!ImportOptions->ShouldImportNormals() || !ImportOptions->ShouldImportTangents());
 		BuildOptions.bComputeWeightedNormals = ImportOptions->bComputeWeightedNormals;
-		BuildOptions.bRemoveDegenerateTriangles = false;
-
-		IMeshUtilities& MeshUtilities = FModuleManager::Get().LoadModuleChecked<IMeshUtilities>("MeshUtilities");
-		
-		TArray<FText> WarningMessages;
-		TArray<FName> WarningNames;
-		// Create actual rendering data.
-		bool bBuildSuccess = MeshUtilities.BuildSkeletalMesh(ImportedResource->LODModels[0],SkeletalMesh->RefSkeleton,LODInfluences, LODWedges,LODFaces,LODPoints,LODPointToRawMap,BuildOptions, &WarningMessages, &WarningNames);
+		BuildOptions.bRemoveDegenerates = ImportOptions->bRemoveDegenerates;
+		BuildOptions.ThresholdPosition = ImportOptions->OverlappingThresholds.ThresholdPosition;
+		BuildOptions.ThresholdTangentNormal = ImportOptions->OverlappingThresholds.ThresholdTangentNormal;
+		BuildOptions.ThresholdUV = ImportOptions->OverlappingThresholds.ThresholdUV;
 
-		// temporary hack of message/names, should be one token or a struct
-		if(WarningMessages.Num() > 0 && WarningNames.Num() == WarningMessages.Num())
+		bool bUseLegacyMeshUtilities = false;
+		bool bBuildSuccess = false;
+		if (bUseLegacyMeshUtilities)
 		{
-			EMessageSeverity::Type MessageSeverity = bBuildSuccess ? EMessageSeverity::Warning : EMessageSeverity::Error;
+			IMeshUtilities::MeshBuildOptions LegacyBuildOptions;
+			LegacyBuildOptions.FillOptions(BuildOptions);
+
+			IMeshUtilities& MeshUtilities = FModuleManager::Get().LoadModuleChecked<IMeshUtilities>("MeshUtilities");
+
+			TArray<FText> WarningMessages;
+			TArray<FName> WarningNames;
+			// Create actual rendering data.
+			bBuildSuccess = MeshUtilities.BuildSkeletalMesh(ImportedResource->LODModels[0], SkeletalMesh->RefSkeleton, LODInfluences, LODWedges, LODFaces, LODPoints, LODPointToRawMap, LegacyBuildOptions, &WarningMessages, &WarningNames);
 
-			for(int32 MessageIdx = 0; MessageIdx<WarningMessages.Num(); ++MessageIdx)	
+			// temporary hack of message/names, should be one token or a struct
+			if (WarningMessages.Num() > 0 && WarningNames.Num() == WarningMessages.Num())
 			{
-				AddTokenizedErrorMessage(FTokenizedMessage::Create(MessageSeverity, WarningMessages[MessageIdx]), WarningNames[MessageIdx]);
+				EMessageSeverity::Type MessageSeverity = bBuildSuccess ? EMessageSeverity::Warning : EMessageSeverity::Error;
+
+				for (int32 MessageIdx = 0; MessageIdx < WarningMessages.Num(); ++MessageIdx)
+				{
+					AddTokenizedErrorMessage(FTokenizedMessage::Create(MessageSeverity, WarningMessages[MessageIdx]), WarningNames[MessageIdx]);
+				}
 			}
 		}
+		else
+		{
+			//The imported LOD is always 0 here, the LOD custom import will import the LOD alone(in a temporary skeletalmesh) and add it to the base skeletal mesh later
+			const int32 LODIndex = 0;
+			check(SkeletalMesh->GetLODInfo(LODIndex) != nullptr);
+			//Set the build options
+			SkeletalMesh->GetLODInfo(LODIndex)->BuildSettings = BuildOptions;
+			//New MeshDescription build process
+			IMeshBuilderModule& MeshBuilderModule = FModuleManager::Get().LoadModuleChecked<IMeshBuilderModule>("MeshBuilder");
+			//We must build the LODModel so we can restore properly the mesh, but we do not have to regenerate LODs
+			bBuildSuccess = MeshBuilderModule.BuildSkeletalMesh(SkeletalMesh, LODIndex, false);
+		}
 
 		if( !bBuildSuccess )
 		{
 			SkeletalMesh->MarkPendingKill();
 			return NULL;
 		}
-
+		
 		// Presize the per-section shadow casting array with the number of sections in the imported LOD.
 		const int32 NumSections = LODModel.Sections.Num();
 
@@ -1975,18 +2024,12 @@ USkeletalMesh* UnFbx::FFbxImporter::ImportSkeletalMesh(FImportSkeletalMeshArgs &
 		}
 
 		SkeletalMesh->CalculateInvRefMatrices();
-		SkeletalMesh->PostEditChange();
-		SkeletalMesh->MarkPackageDirty();
-
-		// Now iterate over all skeletal mesh components re-initialising them.
-		for (TObjectIterator<USkinnedMeshComponent> It; It; ++It)
+		if ((!SkeletalMesh->GetResourceForRendering() || !SkeletalMesh->GetResourceForRendering()->LODRenderData.IsValidIndex(0)) && ImportOptions->bCreatePhysicsAsset)
 		{
-			USkinnedMeshComponent* SkinComp = *It;
-			if (SkinComp->SkeletalMesh == SkeletalMesh)
-			{
-				ComponentContexts.Add(new FComponentReregisterContext(SkinComp));
-			}
+			//We need to have a valid render data to create physic asset
+			SkeletalMesh->Build();
 		}
+		SkeletalMesh->MarkPackageDirty();
 	}
 
 	if(ImportSkeletalMeshArgs.LodIndex == 0)
@@ -2142,13 +2185,6 @@ USkeletalMesh* UnFbx::FFbxImporter::ImportSkeletalMesh(FImportSkeletalMeshArgs &
 		ImportNodeCustomProperties(SkeletalMesh->Skeleton, SkeletonNode, true);
 	}
 
-	// ComponentContexts will now go out of scope, causing components to be re-registered
-
-	if (RecreateExistingRenderStateContext)
-	{
-		delete RecreateExistingRenderStateContext;
-	}
-
 	return SkeletalMesh;
 }
 
@@ -2352,8 +2388,6 @@ USkeletalMesh* UnFbx::FFbxImporter::ReimportSkeletalMesh(USkeletalMesh* Mesh, UF
 		return NULL;
 	}
 
-	char MeshName[1024];
-	FCStringAnsi::Strcpy(MeshName, TCHAR_TO_UTF8(*Mesh->GetName()));
 	TArray<FbxNode*>* FbxNodes = NULL;
 	USkeletalMesh* NewMesh = NULL;
 
@@ -2362,16 +2396,35 @@ USkeletalMesh* UnFbx::FFbxImporter::ReimportSkeletalMesh(USkeletalMesh* Mesh, UF
 	bool Old_ImportTextures = ImportOptions->bImportTextures;
 	bool Old_ImportAnimations = ImportOptions->bImportAnimations;
 
+	auto CleanUpImportOptionsData = [&]()
+	{
+		ImportOptions->bImportRigidMesh = Old_ImportRigidMesh;
+		ImportOptions->bImportMaterials = Old_ImportMaterials;
+		ImportOptions->bImportTextures = Old_ImportTextures;
+		ImportOptions->bImportAnimations = Old_ImportAnimations;
+	};
+
 	// support to update rigid animation mesh
 	ImportOptions->bImportRigidMesh = true;
 
-	// Unbind any existing clothing assets before we reimport the geometry
-	TArray<ClothingAssetUtils::FClothingAssetMeshBinding> ClothingBindings;
-	ClothingAssetUtils::GetMeshClothingAssetBindings(Mesh, ClothingBindings);
-
-	for(ClothingAssetUtils::FClothingAssetMeshBinding& Binding : ClothingBindings)
+	UFbxSkeletalMeshImportData* SKImportData = Cast<UFbxSkeletalMeshImportData>(Mesh->AssetImportData);
+	if (SKImportData)
 	{
-		Binding.Asset->UnbindFromSkeletalMesh(Mesh, Binding.LODIndex);
+		FSkeletalMeshLODInfo* LODInfo = Mesh->GetLODInfo(0);
+		if(LODInfo && Mesh->GetImportedModel() && Mesh->GetImportedModel()->LODModels.IsValidIndex(0))
+		{
+			const FSkeletalMeshLODModel& LODModel = Mesh->GetImportedModel()->LODModels[0];
+			if (!LODModel.RawSkeletalMeshBulkData.IsBuildDataAvailable())
+			{
+				//Set the build settings
+				LODInfo->BuildSettings.bBuildAdjacencyBuffer = true;
+				LODInfo->BuildSettings.bComputeWeightedNormals = SKImportData->bComputeWeightedNormals;
+				LODInfo->BuildSettings.bRecomputeNormals = SKImportData->NormalImportMethod == EFBXNormalImportMethod::FBXNIM_ComputeNormals;
+				LODInfo->BuildSettings.bRecomputeTangents = SKImportData->NormalImportMethod != EFBXNormalImportMethod::FBXNIM_ImportNormalsAndTangents;
+				LODInfo->BuildSettings.bRemoveDegenerates = true;
+				LODInfo->BuildSettings.bUseMikkTSpace = SKImportData->NormalGenerationMethod == EFBXNormalGenerationMethod::MikkTSpace;
+			}
+		}
 	}
 
 	// get meshes in Fbx file
@@ -2427,26 +2480,53 @@ USkeletalMesh* UnFbx::FFbxImporter::ReimportSkeletalMesh(USkeletalMesh* Mesh, UF
 			if (FbxNodes != nullptr)
 				break;
 		}
-		if (FbxNodes == nullptr)
-			return nullptr;
 	}
 	else
 	{
 		// if there is only one mesh, use it without name checking 
 		// (because the "Used As Full Name" option enables users name the Unreal mesh by themselves
-		if (FbxSkelMeshArray.Num() > 0)
+		if (FbxSkelMeshArray.Num() == 1)
 		{
 			FbxNodes = FbxSkelMeshArray[0];
 		}
-		else
+		else if (FbxSkelMeshArray.Num() > 1)
 		{
-			// @todo - FBX Importing - We need proper support for reimport if the file contains more than one skeletal mesh. 
-
+			//Search a set corresponding
+			for (TArray<FbxNode*>*SkeletalMeshNodes : FbxSkelMeshArray)
+			{
+				FbxNode* Node = GetMeshNodesFromName(Mesh->GetName(), *SkeletalMeshNodes);
+				if (Node != nullptr)
+				{
+					FbxNodes = SkeletalMeshNodes;
+					break;
+				}
+			}
+			if (FbxNodes == nullptr)
+			{
+				//Re import the first skeletalmesh found in the fbx file
+				FbxNodes = FbxSkelMeshArray[0];
+			}
 		}
 	}
-
+	
 	if (FbxNodes)
 	{
+		//Posteditchange will be call at the end of the reimport which will reallocate the render ressource
+		FScopedSkeletalMeshPostEditChange ScopedPostEditChange(Mesh);
+
+		FScopedSlowTask SlowTask(100.0f, NSLOCTEXT("SkeltalMeshReimportFactory", "SkeletalMeshReimportTask", "Reimporting skeletal mesh..."));
+		SlowTask.MakeDialog();
+
+		//The fbx file is valid we can reimport the skeletal mesh
+		Mesh->PreEditChange(nullptr);
+
+		// Unbind any existing clothing assets before we reimport the geometry
+		TArray<ClothingAssetUtils::FClothingAssetMeshBinding> ClothingBindings;
+		FLODUtilities::UnbindClothingAndBackup(Mesh, ClothingBindings);
+
+		//Reimport must force a new guid
+		Mesh->InvalidateDeriveDataCacheGUID();
+
 		for (int32 LODCounter = 1; LODCounter < Mesh->GetLODNum(); ++LODCounter)
 		{
 			//Reset all import with base mesh flags
@@ -2491,8 +2571,10 @@ USkeletalMesh* UnFbx::FFbxImporter::ReimportSkeletalMesh(USkeletalMesh* Mesh, UF
 		int32 LODIndex;
 		int32 SuccessfulLodIndex = 0;
 		bool bOperationCanceled = false;
+		float ProgressStep = 90.0 / MaxLODLevel;
 		for (LODIndex = 0; LODIndex < MaxLODLevel; LODIndex++)
 		{
+			SlowTask.EnterProgressFrame(ProgressStep);
 			int32 ImportedSuccessfulLodIndex = INDEX_NONE;
 			TArray<FbxNode*> SkelMeshNodeArray;
 			for (int32 j = 0; j < (*FbxNodes).Num(); j++)
@@ -2542,8 +2624,9 @@ USkeletalMesh* UnFbx::FFbxImporter::ReimportSkeletalMesh(USkeletalMesh* Mesh, UF
 				NewMesh = ImportSkeletalMesh( ImportSkeletalMeshArgs );
 				if (bOperationCanceled)
 				{
-					// User cancelled, clean up and return
-					return nullptr;
+					// User cancel, clean up and return
+					NewMesh = nullptr;
+					break;
 				}
 				if (NewMesh)
 				{
@@ -2590,7 +2673,7 @@ USkeletalMesh* UnFbx::FFbxImporter::ReimportSkeletalMesh(USkeletalMesh* Mesh, UF
 				ImportSkeletalMeshArgs.OutData = &OutData;
 
 				UObject *LODObject = ImportSkeletalMesh( ImportSkeletalMeshArgs );
-				bool bImportSucceeded = !bOperationCanceled && ImportSkeletalMeshLOD( Cast<USkeletalMesh>(LODObject), BaseSkeletalMesh, SuccessfulLodIndex, false, nullptr, TemplateImportData);
+				bool bImportSucceeded = !bOperationCanceled && ImportSkeletalMeshLOD( Cast<USkeletalMesh>(LODObject), BaseSkeletalMesh, SuccessfulLodIndex, TemplateImportData);
 
 				if (bImportSucceeded)
 				{
@@ -2626,47 +2709,51 @@ USkeletalMesh* UnFbx::FFbxImporter::ReimportSkeletalMesh(USkeletalMesh* Mesh, UF
 
 			if (NewMesh)
 			{
-				// import morph target
 				if ((ImportOptions->bImportSkeletalMeshLODs || LODIndex == 0) &&
 					ImportOptions->bImportMorph &&
 					ImportedSuccessfulLodIndex != INDEX_NONE &&
 					NewMesh->GetImportedModel() &&
-					NewMesh->GetImportedModel()->LODModels.IsValidIndex(LODIndex))
+					NewMesh->GetImportedModel()->LODModels.IsValidIndex(ImportedSuccessfulLodIndex))
 				{
-					// @fixme: @question : where do they import this morph? where to? What morph target sets?
 					ImportFbxMorphTarget(SkelMeshNodeArray, NewMesh, ImportedSuccessfulLodIndex, OutData);
 				}
 			}
 		}
-		if (NewMesh)
+
+		auto RebindClothing = [&ClothingBindings](USkeletalMesh* SkelMesh)
 		{
-			FSkeletalMeshModel* NewMeshResource = NewMesh->GetImportedModel();
-			if(NewMeshResource && ClothingBindings.Num() > 0)
+			FSkeletalMeshModel* MeshResource = SkelMesh->GetImportedModel();
+			if (MeshResource && ClothingBindings.Num() > 0)
 			{
-				NewMesh->PreEditChange(nullptr);
-
-				for(ClothingAssetUtils::FClothingAssetMeshBinding& Binding : ClothingBindings)
+				FLODUtilities::RestoreClothingFromBackup(SkelMesh, ClothingBindings, FOnPostRebindCloth::CreateLambda([&MeshResource](int32 LODIndex, const FSkelMeshSection& BindSection)
 				{
-					if(NewMeshResource->LODModels.IsValidIndex(Binding.LODIndex) &&
-						NewMeshResource->LODModels[Binding.LODIndex].Sections.IsValidIndex(Binding.SectionIndex))
+					if (MeshResource->LODModels.IsValidIndex(LODIndex))
 					{
-						Binding.Asset->BindToSkeletalMesh(NewMesh, Binding.LODIndex, Binding.SectionIndex, Binding.AssetInternalLodIndex);
+						FSkeletalMeshLODModel& LODModel = MeshResource->LODModels[LODIndex];
+						//If successfull set back the section user data
+						FSkelMeshSourceSectionUserData& SectionUserData = LODModel.UserSectionsData.FindOrAdd(BindSection.OriginalDataSectionIndex);
+						SectionUserData.CorrespondClothAssetIndex = BindSection.CorrespondClothAssetIndex;
+						SectionUserData.ClothingData = BindSection.ClothingData;
 					}
-				}
-
-				NewMesh->PostEditChange();
+				}));
 			}
+		};
+
+
+		if (!bOperationCanceled && NewMesh)
+		{
+			RebindClothing(NewMesh);
 
 			//Update the import data so we can re-import correctly
 			UpdateSkeletalMeshImportData(NewMesh, TemplateImportData, INDEX_NONE, &ImportMaterialOriginalNameData, &ImportMeshLodData);
-			//If we have import some morph target we have to rebuild the render resources since morph target are now using GPU
-			if (NewMesh->MorphTargets.Num() > 0)
-			{
-				NewMesh->ReleaseResources();
-				//Rebuild the resources with a post edit change since we have added some morph targets
-				NewMesh->PostEditChange();
-			}
 		}
+		else if(bOperationCanceled && Mesh != nullptr)
+		{
+			//Operation was cancel
+			//rebind the clothing on the original mesh
+			RebindClothing(Mesh);
+		}
+		SlowTask.EnterProgressFrame(10.0f);
 	}
 	else
 	{
@@ -2674,10 +2761,7 @@ USkeletalMesh* UnFbx::FFbxImporter::ReimportSkeletalMesh(USkeletalMesh* Mesh, UF
 		AddTokenizedErrorMessage(FTokenizedMessage::Create(EMessageSeverity::Warning, FText::Format(LOCTEXT("FbxSkeletaLMeshimport_NoFBXMeshMatch", "No FBX mesh matches the Unreal mesh '{0}'."), FText::FromString(Mesh->GetName()))), FFbxErrors::Generic_Mesh_MeshNotFound);
 	}
 
-	ImportOptions->bImportRigidMesh = Old_ImportRigidMesh;
-	ImportOptions->bImportMaterials = Old_ImportMaterials;
-	ImportOptions->bImportTextures = Old_ImportTextures;
-	ImportOptions->bImportAnimations = Old_ImportAnimations;
+	CleanUpImportOptionsData();
 
 	return NewMesh;
 }
@@ -3411,7 +3495,7 @@ bool UnFbx::FFbxImporter::FillSkelMeshImporterFromFbx( FSkeletalMeshImportData&
 			//
 			// normals, tangents and binormals
 			//
-			if( ImportOptions->ShouldImportNormals() && bHasNormalInformation )
+			if( bHasNormalInformation )
 			{
 				int32 TmpIndex = LocalIndex*3 + VertexIndex;
 				//normals may have different reference and mapping mode than tangents and binormals
@@ -3425,7 +3509,7 @@ bool UnFbx::FFbxImporter::FillSkelMeshImporterFromFbx( FSkeletalMeshImportData&
 
 				FbxVector4 TempValue;
 
-				if( ImportOptions->ShouldImportTangents() && bHasTangentInformation )
+				if( bHasTangentInformation )
 				{
 					TempValue = LayerElementTangent->GetDirectArray().GetAt(TangentMapIndex);
 					TempValue = TotalMatrixForNormal.MultT(TempValue);
@@ -3722,6 +3806,9 @@ bool UnFbx::FFbxImporter::FillSkelMeshImporterFromFbx( FSkeletalMeshImportData&
 
 void UnFbx::FFbxImporter::InsertNewLODToBaseSkeletalMesh(USkeletalMesh* InSkeletalMesh, USkeletalMesh* BaseSkeletalMesh, int32 DesiredLOD, UFbxSkeletalMeshImportData* TemplateImportData)
 {
+	//Scope a post edit change
+	FScopedSkeletalMeshPostEditChange ScopedPostEditChange(BaseSkeletalMesh);
+
 	FSkeletalMeshModel* ImportedResource = InSkeletalMesh->GetImportedModel();
 	FSkeletalMeshModel* DestImportedResource = BaseSkeletalMesh->GetImportedModel();
 
@@ -3781,6 +3868,13 @@ void UnFbx::FFbxImporter::InsertNewLODToBaseSkeletalMesh(USkeletalMesh* InSkelet
 
 	// Set up LODMaterialMap to number of materials in new mesh.
 	FSkeletalMeshLODInfo& LODInfo = *(BaseSkeletalMesh->GetLODInfo(DesiredLOD));
+
+	//Copy the build settings
+	if(InSkeletalMesh->GetLODInfo(0))
+	{
+		const FSkeletalMeshLODInfo& ImportedLODInfo = *(InSkeletalMesh->GetLODInfo(0));
+		LODInfo.BuildSettings = ImportedLODInfo.BuildSettings;
+	}
 	
 	LODInfo.LODMaterialMap.Empty();
 	// Now set up the material mapping array.
@@ -3874,12 +3968,9 @@ void UnFbx::FFbxImporter::InsertNewLODToBaseSkeletalMesh(USkeletalMesh* InSkelet
 			}
 		}
 	}
-
-	// rebuild vertex buffers and reinit RHI resources
-	BaseSkeletalMesh->PostEditChange();
 }
 
-bool UnFbx::FFbxImporter::ImportSkeletalMeshLOD(USkeletalMesh* InSkeletalMesh, USkeletalMesh* BaseSkeletalMesh, int32 DesiredLOD, bool bNeedToReregister, TArray<UActorComponent*>* ReregisterAssociatedComponents, UFbxSkeletalMeshImportData* TemplateImportData)
+bool UnFbx::FFbxImporter::ImportSkeletalMeshLOD(USkeletalMesh* InSkeletalMesh, USkeletalMesh* BaseSkeletalMesh, int32 DesiredLOD, UFbxSkeletalMeshImportData* TemplateImportData)
 {
 	check(InSkeletalMesh);
 	check(BaseSkeletalMesh);
@@ -3943,6 +4034,8 @@ bool UnFbx::FFbxImporter::ImportSkeletalMeshLOD(USkeletalMesh* InSkeletalMesh, U
 		}
 	}
 
+	FScopedSkeletalMeshPostEditChange ScopedPostEditChange(BaseSkeletalMesh);
+
 	//Skeleton do not match and we are importing skeletal geometry only so apply the base LOD Skinning
 	if (bApplyBaseSkinning)
 	{
@@ -3950,8 +4043,10 @@ bool UnFbx::FFbxImporter::ImportSkeletalMeshLOD(USkeletalMesh* InSkeletalMesh, U
 		//Apply the base LOD skinning to the new LOD geometry, DestImportedResource here is the base lod and is the source for the apply skinning.
 		//The Dest here mean we merge the just import LOD to the base LodModel.
 		SkeletalMeshHelper::ApplySkinning(InSkeletalMesh, DestImportedResource->LODModels[0], ImportedResource->LODModels[0]);
-		IMeshUtilities& MeshUtilities = FModuleManager::Get().LoadModuleChecked<IMeshUtilities>("MeshUtilities");
-		MeshUtilities.RemoveBonesFromMesh(BaseSkeletalMesh, DesiredLOD, NULL);
+
+		//Excluded bones will be remove when we rebuild the asset
+		//IMeshUtilities& MeshUtilities = FModuleManager::Get().LoadModuleChecked<IMeshUtilities>("MeshUtilities");
+		//MeshUtilities.RemoveBonesFromMesh(BaseSkeletalMesh, DesiredLOD, NULL);
 	}
 
 	FSkeletalMeshLODModel& NewLODModel = ImportedResource->LODModels[0];
@@ -4114,32 +4209,7 @@ bool UnFbx::FFbxImporter::ImportSkeletalMeshLOD(USkeletalMesh* InSkeletalMesh, U
 		}
 	}
 
-	if (bNeedToReregister)
-	{
-		// Shut down the skeletal mesh component that is previewing this mesh.
-		if (ReregisterAssociatedComponents)
-		{
-			FMultiComponentReregisterContext ReregisterContext(*ReregisterAssociatedComponents);
-			// wait until resources are released
-			FlushRenderingCommands();
-
-			InsertNewLODToBaseSkeletalMesh(InSkeletalMesh, BaseSkeletalMesh, DesiredLOD, TemplateImportData);
-
-			// ReregisterContexts go out of scope here, reregistering associated components with the scene.
-		}
-		else
-		{
-			TComponentReregisterContext<USkinnedMeshComponent> ComponentReregisterContext;
-
-			InsertNewLODToBaseSkeletalMesh(InSkeletalMesh, BaseSkeletalMesh, DesiredLOD, TemplateImportData);
-
-			// ReregisterContexts go out of scope here, reregistering skel components with the scene.
-		}
-	}
-	else
-	{
-		InsertNewLODToBaseSkeletalMesh(InSkeletalMesh, BaseSkeletalMesh, DesiredLOD, TemplateImportData);
-	}
+	InsertNewLODToBaseSkeletalMesh(InSkeletalMesh, BaseSkeletalMesh, DesiredLOD, TemplateImportData);
 
 	return true;
 }
@@ -4280,9 +4350,21 @@ void UnFbx::FFbxImporter::ImportMorphTargetsInternal( TArray<FbxNode*>& SkelMesh
 		check(BaseImportData.MorphTargetNames.Num() == BaseImportData.MorphTargets.Num() && BaseImportData.MorphTargetNames.Num() == BaseImportData.MorphTargetModifiedPoints.Num());
 	}
 
-	//Build MorphTargets
-	// Calculate overlapping corners and tangents
-	FLODUtilities::BuildMorphTargets(BaseSkelMesh, BaseImportData, LODIndex, ImportOptions->ShouldImportNormals(), ImportOptions->ShouldImportTangents(), (ImportOptions->NormalGenerationMethod == EFBXNormalGenerationMethod::MikkTSpace));
+	if (BaseSkelMesh->GetImportedModel() && BaseSkelMesh->GetImportedModel()->LODModels.IsValidIndex(LODIndex))
+	{
+		//If we can build the skeletal mesh there is no need to build the morph target now, all the necessary build morph target data was copied before.
+		if (!BaseSkelMesh->GetImportedModel()->LODModels[LODIndex].RawSkeletalMeshBulkData.IsBuildDataAvailable())
+		{
+			//Build MorphTargets
+			FLODUtilities::BuildMorphTargets(
+				BaseSkelMesh,
+				BaseImportData,
+				LODIndex,
+				ImportOptions->ShouldImportNormals(),
+				ImportOptions->ShouldImportTangents(),
+				(ImportOptions->NormalGenerationMethod == EFBXNormalGenerationMethod::MikkTSpace));
+		}
+	}
 
 	GWarn->EndSlowTask();
 }
@@ -4290,6 +4372,9 @@ void UnFbx::FFbxImporter::ImportMorphTargetsInternal( TArray<FbxNode*>& SkelMesh
 // Import Morph target
 void UnFbx::FFbxImporter::ImportFbxMorphTarget(TArray<FbxNode*> &SkelMeshNodeArray, USkeletalMesh* BaseSkelMesh, int32 LODIndex, FSkeletalMeshImportData &BaseSkeletalMeshImportData)
 {
+	//Stack the PostEditChange call, it will call post edit change when it will go out of scope
+	FScopedSkeletalMeshPostEditChange ScopedPostEditChange(BaseSkelMesh);
+
 	bool bHasMorph = false;
 	int32 NodeIndex;
 	// check if there are morph in this geometry
@@ -4308,7 +4393,12 @@ void UnFbx::FFbxImporter::ImportFbxMorphTarget(TArray<FbxNode*> &SkelMeshNodeArr
 	
 	if (bHasMorph)
 	{
-		ImportMorphTargetsInternal( SkelMeshNodeArray, BaseSkelMesh, LODIndex, BaseSkeletalMeshImportData );
+		ImportMorphTargetsInternal( SkelMeshNodeArray, BaseSkelMesh, LODIndex, BaseSkeletalMeshImportData);
+		//Save the rawMesh
+		if (BaseSkelMesh->GetImportedModel() && BaseSkelMesh->GetImportedModel()->LODModels.IsValidIndex(LODIndex))
+		{
+			BaseSkelMesh->GetImportedModel()->LODModels[LODIndex].RawSkeletalMeshBulkData.SaveRawMesh(BaseSkeletalMeshImportData);
+		}
 	}
 }
 
diff --git a/Engine/Source/Editor/UnrealEd/Private/Fbx/FbxStaticMeshImport.cpp b/Engine/Source/Editor/UnrealEd/Private/Fbx/FbxStaticMeshImport.cpp
--- a/Engine/Source/Editor/UnrealEd/Private/Fbx/FbxStaticMeshImport.cpp
+++ b/Engine/Source/Editor/UnrealEd/Private/Fbx/FbxStaticMeshImport.cpp
@@ -1071,51 +1071,6 @@ void UnFbx::FFbxImporter::AddStaticMeshSourceModelGeneratedLOD(UStaticMesh* Stat
 	}
 }
 
-FbxNode* UnFbx::FFbxImporter::GetMeshNodesFromName(UStaticMesh* StaticMesh, TArray<FbxNode*>& FbxMeshArray)
-{
-	char MeshName[1024];
-	FCStringAnsi::Strcpy(MeshName, 1024, TCHAR_TO_UTF8(*StaticMesh->GetName()));
-	// find the Fbx mesh node that the Unreal Mesh matches according to name
-	int32 MeshIndex;
-	for (MeshIndex = 0; MeshIndex < FbxMeshArray.Num(); MeshIndex++)
-	{
-		const char* FbxMeshName = FbxMeshArray[MeshIndex]->GetName();
-		// The name of Unreal mesh may have a prefix, so we match from end
-		int32 i = 0;
-		char* MeshPtr = MeshName + FCStringAnsi::Strlen(MeshName) - 1;
-		if (FCStringAnsi::Strlen(FbxMeshName) <= FCStringAnsi::Strlen(MeshName))
-		{
-			const char* FbxMeshPtr = FbxMeshName + FCStringAnsi::Strlen(FbxMeshName) - 1;
-			while (i < FCStringAnsi::Strlen(FbxMeshName))
-			{
-				bool bIsPointAndUnderscore = *FbxMeshPtr == '.' && *MeshPtr == '_';
-
-				if (*MeshPtr != *FbxMeshPtr && !bIsPointAndUnderscore)
-				{
-					break;
-				}
-				else
-				{
-					i++;
-					MeshPtr--;
-					FbxMeshPtr--;
-				}
-			}
-		}
-
-		if (i == FCStringAnsi::Strlen(FbxMeshName)) // matched
-		{
-			// check further
-			if (FCStringAnsi::Strlen(FbxMeshName) == FCStringAnsi::Strlen(MeshName) || // the name of Unreal mesh is full match
-				*MeshPtr == '_')														// or the name of Unreal mesh has a prefix
-			{
-				return FbxMeshArray[MeshIndex];
-			}
-		}
-	}
-	return nullptr;
-}
-
 UStaticMesh* UnFbx::FFbxImporter::ReimportStaticMesh(UStaticMesh* Mesh, UFbxStaticMeshImportData* TemplateImportData)
 {
 	TArray<FbxNode*> FbxMeshArray;
@@ -1198,7 +1153,7 @@ UStaticMesh* UnFbx::FFbxImporter::ReimportStaticMesh(UStaticMesh* Mesh, UFbxStat
 	}
 	else if(!bCombineMeshes && !bCombineMeshesLOD)
 	{
-		Node = GetMeshNodesFromName(Mesh, FbxMeshArray);
+		Node = GetMeshNodesFromName(Mesh->GetName(), FbxMeshArray);
 	}
 
 	// If there is no match it may be because an LOD group was imported where
diff --git a/Engine/Source/Runtime/Engine/Classes/Engine/SkeletalMeshLODSettings.h b/Engine/Source/Runtime/Engine/Classes/Engine/SkeletalMeshLODSettings.h
--- a/Engine/Source/Runtime/Engine/Classes/Engine/SkeletalMeshLODSettings.h
+++ b/Engine/Source/Runtime/Engine/Classes/Engine/SkeletalMeshLODSettings.h
@@ -63,6 +63,9 @@ struct FSkeletalMeshLODGroupSettings
 	/** Get Skeletal mesh optimizations setting structure for the given LOD level */
 	ENGINE_API FSkeletalMeshOptimizationSettings GetReductionSettings() const;
 
+	/** Get Skeletal mesh build setting structure for the given LOD level */
+	ENGINE_API FSkeletalMeshBuildSettings GetBuildSettings() const;
+
 	/** Get Skeletal mesh optimizations setting structure for the given LOD level */
 	UE_DEPRECATED(4.20, "Please use GetReductionSettings instead")
 	FSkeletalMeshOptimizationSettings GetSettings() const
@@ -107,6 +110,10 @@ struct FSkeletalMeshLODGroupSettings
 	/** The optimization settings to use for the respective LOD level */
 	UPROPERTY(EditAnywhere, Category = Reduction)
 	FSkeletalMeshOptimizationSettings ReductionSettings;
+
+	/** The build settings to use for the respective LOD level */
+	UPROPERTY(EditAnywhere, Category = Build)
+	FSkeletalMeshBuildSettings BuildSettings;
 };
 
 UCLASS(config = Engine, defaultconfig, BlueprintType, MinimalAPI)
diff --git a/Engine/Source/Runtime/Engine/Private/SkeletalMeshRenderData.cpp b/Engine/Source/Runtime/Engine/Private/SkeletalMeshRenderData.cpp
--- a/Engine/Source/Runtime/Engine/Private/SkeletalMeshRenderData.cpp
+++ b/Engine/Source/Runtime/Engine/Private/SkeletalMeshRenderData.cpp
@@ -15,6 +15,8 @@
 #include "Interfaces/ITargetPlatformManagerModule.h"
 #include "Interfaces/ITargetPlatform.h"
 #include "PlatformInfo.h"
+#include "IMeshBuilderModule.h"
+#include "EngineUtils.h"
 
 
 #if ENABLE_COOK_STATS
@@ -28,11 +30,36 @@ namespace SkeletalMeshCookStats
 }
 #endif
 
+//Serialize the LODInfo and append the result to the KeySuffix to build the LODInfo part of the DDC KEY
+//Note: this serializer is only used to build the mesh DDC key, no versioning is required
+static void SerializeLODInfoForDDC(USkeletalMesh* SkeletalMesh, FString& KeySuffix)
+{
+	TArray<FSkeletalMeshLODInfo>& LODInfos = SkeletalMesh->GetLODInfoArray();
+	const bool bIs16BitfloatBufferSupported = GVertexElementTypeSupport.IsSupported(VET_Half2);
+	for (int32 LODIndex = 0; LODIndex < SkeletalMesh->GetLODNum(); ++LODIndex)
+	{
+		check(LODInfos.IsValidIndex(LODIndex));
+		FSkeletalMeshLODInfo& LODInfo = LODInfos[LODIndex];
+		bool bValidLODSettings = false;
+		if (SkeletalMesh->LODSettings != nullptr)
+		{
+			const int32 NumSettings = FMath::Min(SkeletalMesh->LODSettings->GetNumberOfSettings(), SkeletalMesh->GetLODNum());
+			if (LODIndex < NumSettings)
+			{
+				bValidLODSettings = true;
+			}
+		}
+		const FSkeletalMeshLODGroupSettings* SkeletalMeshLODGroupSettings = bValidLODSettings ? &SkeletalMesh->LODSettings->GetSettingsForLODLevel(LODIndex) : nullptr;
+		LODInfo.BuildGUID = LODInfo.ComputeDeriveDataCacheKey(SkeletalMeshLODGroupSettings);
+		KeySuffix += LODInfo.BuildGUID.ToString(EGuidFormats::Digits);
+	}
+}
+
 // If skeletal mesh derived data needs to be rebuilt (new format, serialization
 // differences, etc.) replace the version GUID below with a new one.
 // In case of merge conflicts with DDC versions, you MUST generate a new GUID
 // and set this new GUID as the version.                                       
-#define SKELETALMESH_DERIVEDDATA_VER TEXT("51BDAF351B6941548818C2F8B33397D7")
+#define SKELETALMESH_DERIVEDDATA_VER TEXT("E33979BDFD4B42C9A4B97BE162A8C0FF")
 
 static const FString& GetSkeletalMeshDerivedDataVersion()
 {
@@ -47,11 +74,29 @@ static const FString& GetSkeletalMeshDerivedDataVersion()
 static FString BuildSkeletalMeshDerivedDataKey(USkeletalMesh* SkelMesh)
 {
 	FString KeySuffix(TEXT(""));
+	FString tmpDebugString;
+
+	//Synchronize the user data that are part of the key
+	SkelMesh->GetImportedModel()->SyncronizeLODUserSectionsData();
+
+	tmpDebugString = SkelMesh->GetImportedModel()->GetIdString();
+	KeySuffix += tmpDebugString;
+	tmpDebugString = SkelMesh->GetImportedModel()->GetLODModelIdString();
+	KeySuffix += tmpDebugString;
+
+	//Add the max gpu bone per section
+	const uint32 MaxGPUSkinBones = FGPUBaseSkinVertexFactory::GetMaxGPUSkinBones();
+	KeySuffix += FString::FromInt((int32)MaxGPUSkinBones);
+
+	tmpDebugString = TEXT("");
+	SerializeLODInfoForDDC(SkelMesh, tmpDebugString);
+	KeySuffix += tmpDebugString;
 
-	KeySuffix += SkelMesh->GetImportedModel()->GetIdString();
-	KeySuffix += (SkelMesh->bUseFullPrecisionUVs || !GVertexElementTypeSupport.IsSupported(VET_Half2)) ? "1" : "0";
 	KeySuffix += SkelMesh->bHasVertexColors ? "1" : "0";
-	KeySuffix += SkelMesh->VertexColorGuid.ToString(EGuidFormats::Digits);
+	if (SkelMesh->bHasVertexColors)
+	{
+		KeySuffix += SkelMesh->VertexColorGuid.ToString(EGuidFormats::Digits);
+	}
 
 	const ITargetPlatform* TargetPlatform = GetTargetPlatformManagerRef().GetRunningTargetPlatform();
 	check(TargetPlatform);
@@ -99,7 +144,28 @@ void FSkeletalMeshRenderData::Cache(USkeletalMesh* Owner)
 		if (GetDerivedDataCacheRef().GetSynchronous(*DerivedDataKey, DerivedData))
 		{
 			COOK_STAT(Timer.AddHit(DerivedData.Num()));
+			
+			FSkeletalMeshModel* SkelMeshModel = Owner->GetImportedModel();
+			check(SkelMeshModel);
+
 			FMemoryReader Ar(DerivedData, /*bIsPersistent=*/ true);
+			
+			//Serialize the LODModel sections since they are dependent on the reduction
+			for (int32 LODIndex = 0; LODIndex < SkelMeshModel->LODModels.Num(); LODIndex++)
+			{
+				FSkeletalMeshLODModel* LODModel = &(SkelMeshModel->LODModels[LODIndex]);
+				int32 SectionNum = 0;
+				Ar << SectionNum;
+				LODModel->Sections.Reset(SectionNum);
+				LODModel->Sections.AddDefaulted(SectionNum);
+				for (int32 SectionIndex = 0; SectionIndex < LODModel->Sections.Num(); ++SectionIndex)
+				{
+					Ar << LODModel->Sections[SectionIndex];
+					LODModel->Sections[SectionIndex].OverlappingVertices;
+				}
+				LODModel->SyncronizeUserSectionsDataArray();
+			}
+
 			Serialize(Ar, Owner);
 			for (int32 LODIndex = 0; LODIndex < LODRenderData.Num(); ++LODIndex)
 			{
@@ -129,14 +195,72 @@ void FSkeletalMeshRenderData::Cache(USkeletalMesh* Owner)
 
 			for (int32 LODIndex = 0; LODIndex < SkelMeshModel->LODModels.Num(); LODIndex++)
 			{
-				FSkeletalMeshLODModel& LODModel = SkelMeshModel->LODModels[LODIndex];
+				FSkeletalMeshLODModel* LODModel = &(SkelMeshModel->LODModels[LODIndex]);
+				FSkeletalMeshLODInfo* LODInfo = Owner->GetLODInfo(LODIndex);
+				check(LODInfo);
+				//Build the source model before the render data, if we are a purely generated LOD we do not need to be build
+				IMeshBuilderModule& MeshBuilderModule = FModuleManager::Get().LoadModuleChecked<IMeshBuilderModule>("MeshBuilder");
+				if (!LODModel->RawSkeletalMeshBulkData.IsEmpty() && LODModel->RawSkeletalMeshBulkData.IsBuildDataAvailable())
+				{
+					MeshBuilderModule.BuildSkeletalMesh(Owner, LODIndex, true);
+					LODModel = &(SkelMeshModel->LODModels[LODIndex]);
+				}
+				else
+				{
+					//We need to synchronize when we are generated mesh or if we have load an old asset that was not re-imported
+					LODModel->SyncronizeUserSectionsDataArray();
+				}
+
 				FSkeletalMeshLODRenderData* LODData = new FSkeletalMeshLODRenderData();
 				LODRenderData.Add(LODData);
-
-				LODData->BuildFromLODModel(&LODModel, VertexBufferBuildFlags);
+				
+				//Get the UVs and tangents precision build settings flag specific for this LOD index
+				{
+					bool bUseFullPrecisionUVs = LODInfo->BuildSettings.bUseFullPrecisionUVs;
+					bool bUseHighPrecisionTangentBasis = LODInfo->BuildSettings.bUseHighPrecisionTangentBasis;
+					bool bBuildAdjacencyBuffer = LODInfo->BuildSettings.bBuildAdjacencyBuffer;
+					if (Owner->LODSettings != nullptr)
+					{
+						const int32 NumSettings = Owner->LODSettings->GetNumberOfSettings();
+						if (LODIndex < NumSettings)
+						{
+							const FSkeletalMeshLODGroupSettings& SkeletalMeshLODGroupSettings = Owner->LODSettings->GetSettingsForLODLevel(LODIndex);
+							bUseFullPrecisionUVs = SkeletalMeshLODGroupSettings.GetBuildSettings().bUseFullPrecisionUVs;
+							bUseHighPrecisionTangentBasis = SkeletalMeshLODGroupSettings.GetBuildSettings().bUseHighPrecisionTangentBasis;
+							bBuildAdjacencyBuffer = SkeletalMeshLODGroupSettings.GetBuildSettings().bBuildAdjacencyBuffer;
+						}
+					}
+					if (bUseFullPrecisionUVs || !GVertexElementTypeSupport.IsSupported(VET_Half2))
+					{
+						VertexBufferBuildFlags |= ESkeletalMeshVertexFlags::UseFullPrecisionUVs;
+					}
+					if (bUseHighPrecisionTangentBasis)
+					{
+						VertexBufferBuildFlags |= ESkeletalMeshVertexFlags::UseHighPrecisionTangentBasis;
+					}
+					if (bBuildAdjacencyBuffer)
+					{
+						VertexBufferBuildFlags |= ESkeletalMeshVertexFlags::BuildAdjacencyIndexBuffer;
+					}
+				}
+				LODData->BuildFromLODModel(LODModel, VertexBufferBuildFlags);
 			}
 
 			FMemoryWriter Ar(DerivedData, /*bIsPersistent=*/ true);
+			
+			//Serialize the LODModel sections since they are dependent on the reduction
+			for (int32 LODIndex = 0; LODIndex < SkelMeshModel->LODModels.Num(); LODIndex++)
+			{
+				FSkeletalMeshLODModel* LODModel = &(SkelMeshModel->LODModels[LODIndex]);
+				int32 SectionNum = LODModel->Sections.Num();
+				Ar << SectionNum;
+				for (int32 SectionIndex = 0; SectionIndex < LODModel->Sections.Num(); ++SectionIndex)
+				{
+					Ar << LODModel->Sections[SectionIndex];
+				}
+			}
+
+			//Serialize the render data
 			Serialize(Ar, Owner);
 			for (int32 LODIndex = 0; LODIndex < LODRenderData.Num(); ++LODIndex)
 			{
