Custom FChanges that have Expired need to be ignored. This can be hit when processing a multi-record transaction where only some records are expired.
#rnx

diff --git a/Engine/Source/Editor/UnrealEd/Private/EditorTransaction.cpp b/Engine/Source/Editor/UnrealEd/Private/EditorTransaction.cpp
--- a/Engine/Source/Editor/UnrealEd/Private/EditorTransaction.cpp
+++ b/Engine/Source/Editor/UnrealEd/Private/EditorTransaction.cpp
@@ -234,22 +234,25 @@ void FTransaction::FObjectRecord::Load(FTransaction* Owner)
 
 		if (CustomChange.IsValid())
 		{
-			if (CustomChange->GetChangeType() == FChange::EChangeStyle::InPlaceSwap)
+			if (CustomChange->HasExpired(Object.Get()) == false)		// skip expired changes
 			{
-				TUniquePtr<FChange> InvertedChange = CustomChange->Execute(Object.Get());
-				ensure(InvertedChange->GetChangeType() == FChange::EChangeStyle::InPlaceSwap);
-				CustomChange = MoveTemp(InvertedChange);
-			}
-			else
-			{
-				bool bIsRedo = (Owner->Inc == 1);
-				if (bIsRedo)
+				if (CustomChange->GetChangeType() == FChange::EChangeStyle::InPlaceSwap)
 				{
-					CustomChange->Apply(Object.Get());
+					TUniquePtr<FChange> InvertedChange = CustomChange->Execute(Object.Get());
+					ensure(InvertedChange->GetChangeType() == FChange::EChangeStyle::InPlaceSwap);
+					CustomChange = MoveTemp(InvertedChange);
 				}
 				else
 				{
-					CustomChange->Revert(Object.Get());
+					bool bIsRedo = (Owner->Inc == 1);
+					if (bIsRedo)
+					{
+						CustomChange->Apply(Object.Get());
+					}
+					else
+					{
+						CustomChange->Revert(Object.Get());
+					}
 				}
 			}
 		}
