Fix CIS static analysis
#rnx

diff --git a/Engine/Source/Editor/SkeletalMeshEditor/Private/SkeletalMeshEditor.cpp b/Engine/Source/Editor/SkeletalMeshEditor/Private/SkeletalMeshEditor.cpp
--- a/Engine/Source/Editor/SkeletalMeshEditor/Private/SkeletalMeshEditor.cpp
+++ b/Engine/Source/Editor/SkeletalMeshEditor/Private/SkeletalMeshEditor.cpp
@@ -893,7 +893,7 @@ void FSkeletalMeshEditor::ApplyClothing(UClothingAssetBase* InAsset, int32 InLod
 {
 	USkeletalMesh* Mesh = GetPersonaToolkit()->GetPreviewMesh();
 
-	if (Mesh->GetImportedModel() == nullptr || !Mesh->GetImportedModel()->LODModels.IsValidIndex(InLodIndex))
+	if (Mesh == nullptr || Mesh->GetImportedModel() == nullptr || !Mesh->GetImportedModel()->LODModels.IsValidIndex(InLodIndex))
 	{
 		return;
 	}
diff --git a/Engine/Source/Editor/UnrealEd/Private/Fbx/FbxSkeletalMeshImport.cpp b/Engine/Source/Editor/UnrealEd/Private/Fbx/FbxSkeletalMeshImport.cpp
--- a/Engine/Source/Editor/UnrealEd/Private/Fbx/FbxSkeletalMeshImport.cpp
+++ b/Engine/Source/Editor/UnrealEd/Private/Fbx/FbxSkeletalMeshImport.cpp
@@ -1766,9 +1766,10 @@ USkeletalMesh* UnFbx::FFbxImporter::ImportSkeletalMesh(FImportSkeletalMeshArgs &
 	//Adjust the import data from the import options
 	if (ExistingSkelMesh != nullptr)
 	{
-		if (ExistingSkelMesh->GetImportedModel() && ExistingSkelMesh->GetImportedModel()->LODModels.IsValidIndex(ImportSkeletalMeshArgs.LodIndex))
+		int32 SafeLODIndex = ImportSkeletalMeshArgs.LodIndex < 0 ? 0 : ImportSkeletalMeshArgs.LodIndex;
+		if (ExistingSkelMesh->GetImportedModel() && ExistingSkelMesh->GetImportedModel()->LODModels.IsValidIndex(SafeLODIndex))
 		{
-			const FRawSkeletalMeshBulkData& RawSkeletalMeshBulkData = ExistingSkelMesh->GetImportedModel()->LODModels[ImportSkeletalMeshArgs.LodIndex].RawSkeletalMeshBulkData;
+			const FRawSkeletalMeshBulkData& RawSkeletalMeshBulkData = ExistingSkelMesh->GetImportedModel()->LODModels[SafeLODIndex].RawSkeletalMeshBulkData;
 			GeoImportVersion = RawSkeletalMeshBulkData.GeoImportVersion;
 			SkinningImportVersion = RawSkeletalMeshBulkData.SkinningImportVersion;
 		}
@@ -1778,20 +1779,19 @@ USkeletalMesh* UnFbx::FFbxImporter::ImportSkeletalMesh(FImportSkeletalMeshArgs &
 			if (ImportOptions->bImportAsSkeletalSkinning)
 			{
 				//Replace geometry import data by original existing skel mesh geometry data
-				ReplaceSkeletalMeshGeometryImportData(ExistingSkelMesh, SkelMeshImportDataPtr, ImportSkeletalMeshArgs.LodIndex);
+				ReplaceSkeletalMeshGeometryImportData(ExistingSkelMesh, SkelMeshImportDataPtr, SafeLODIndex);
 			}
 			else if (ImportOptions->bImportAsSkeletalGeometry)
 			{
 				//Replace skinning import data by original existing skel mesh skinning data
-				ReplaceSkeletalMeshSkinningImportData(ExistingSkelMesh, SkelMeshImportDataPtr, ImportSkeletalMeshArgs.LodIndex);
+				ReplaceSkeletalMeshSkinningImportData(ExistingSkelMesh, SkelMeshImportDataPtr, SafeLODIndex);
 			}
 		}
 
 		//Reuse the vertex color in case we are not re-importing weights only and there is some vertexcolor in the existing skeleMesh
 		if (ExistingSkelMesh->bHasVertexColors && ImportOptions->VertexColorImportOption == EVertexColorImportOption::Ignore)
 		{
-			int32 LODIndex = ImportSkeletalMeshArgs.LodIndex < 0 ? 0 : ImportSkeletalMeshArgs.LodIndex;
-			RemapSkeletalMeshVertexColorToImportData(ExistingSkelMesh, LODIndex, SkelMeshImportDataPtr);
+			RemapSkeletalMeshVertexColorToImportData(ExistingSkelMesh, SafeLODIndex, SkelMeshImportDataPtr);
 		}
 	}
 
