Rename GetViewportClient to GetAssetViewportClient on IAssetViewport to avoid name collision during multiple inheritance.
#rnx

diff --git a/Engine/Plugins/Compositing/Composure/Source/ComposureLayersEditor/Private/CompElementEditorModule.cpp b/Engine/Plugins/Compositing/Composure/Source/ComposureLayersEditor/Private/CompElementEditorModule.cpp
--- a/Engine/Plugins/Compositing/Composure/Source/ComposureLayersEditor/Private/CompElementEditorModule.cpp
+++ b/Engine/Plugins/Compositing/Composure/Source/ComposureLayersEditor/Private/CompElementEditorModule.cpp
@@ -36,7 +36,7 @@ namespace CompElementEditor_Impl
 		TSharedPtr<IAssetViewport> Viewport = LevelEditorModule.GetFirstActiveViewport();
 		if (Viewport.IsValid())
 		{
-			Viewport->GetViewportClient().RedrawRequested(Viewport->GetActiveViewport());
+			Viewport->GetAssetViewportClient().RedrawRequested(Viewport->GetActiveViewport());
 		}
 		else if (GCurrentLevelEditingViewportClient)
 		{
diff --git a/Engine/Source/Developer/GameplayDebugger/Private/Editor/GameplayDebuggerToolkit.cpp b/Engine/Source/Developer/GameplayDebugger/Private/Editor/GameplayDebuggerToolkit.cpp
--- a/Engine/Source/Developer/GameplayDebugger/Private/Editor/GameplayDebuggerToolkit.cpp
+++ b/Engine/Source/Developer/GameplayDebugger/Private/Editor/GameplayDebuggerToolkit.cpp
@@ -94,7 +94,7 @@ FReply FGameplayDebuggerToolkit::OnClickedDisableTool()
 		{
 			if (ViewportWindow.IsValid())
 			{
-				FEditorViewportClient& Viewport = ViewportWindow->GetViewportClient();
+				FEditorViewportClient& Viewport = ViewportWindow->GetAssetViewportClient();
 				Viewport.EngineShowFlags.SetDebugAI(false);
 			}
 		}
diff --git a/Engine/Source/Editor/FoliageEdit/Private/FoliageEdMode.cpp b/Engine/Source/Editor/FoliageEdit/Private/FoliageEdMode.cpp
--- a/Engine/Source/Editor/FoliageEdit/Private/FoliageEdMode.cpp
+++ b/Engine/Source/Editor/FoliageEdit/Private/FoliageEdMode.cpp
@@ -3843,7 +3843,7 @@ void FEdModeFoliage::ForceRealTimeViewports(const bool bEnable, const bool bStor
 	TSharedPtr< IAssetViewport > ViewportWindow = LevelEditorModule.GetFirstActiveViewport();
 	if (ViewportWindow.IsValid())
 	{
-		FEditorViewportClient &Viewport = ViewportWindow->GetViewportClient();
+		FEditorViewportClient &Viewport = ViewportWindow->GetAssetViewportClient();
 		if (Viewport.IsPerspective())
 		{
 			if (bEnable)
diff --git a/Engine/Source/Editor/LandscapeEditor/Private/LandscapeEdMode.cpp b/Engine/Source/Editor/LandscapeEditor/Private/LandscapeEdMode.cpp
--- a/Engine/Source/Editor/LandscapeEditor/Private/LandscapeEdMode.cpp
+++ b/Engine/Source/Editor/LandscapeEditor/Private/LandscapeEdMode.cpp
@@ -3535,7 +3535,7 @@ void FEdModeLandscape::ForceRealTimeViewports(const bool bEnable, const bool bSt
 		{
 			if (ViewportWindow.IsValid())
 			{
-				FEditorViewportClient& Viewport = ViewportWindow->GetViewportClient();
+				FEditorViewportClient& Viewport = ViewportWindow->GetAssetViewportClient();
 				if (bEnable)
 				{
 					Viewport.SetRealtime(bEnable, bStoreCurrentState);
diff --git a/Engine/Source/Editor/LevelEditor/Public/SLevelViewport.h b/Engine/Source/Editor/LevelEditor/Public/SLevelViewport.h
--- a/Engine/Source/Editor/LevelEditor/Public/SLevelViewport.h
+++ b/Engine/Source/Editor/LevelEditor/Public/SLevelViewport.h
@@ -94,7 +94,7 @@ public:
 		return *LevelViewportClient;
 	}
 
-	virtual FEditorViewportClient& GetViewportClient() override
+	virtual FEditorViewportClient& GetAssetViewportClient() override
 	{
 		return *LevelViewportClient;
 	}
diff --git a/Engine/Source/Editor/MeshPaint/Private/MeshPaintHelpers.cpp b/Engine/Source/Editor/MeshPaint/Private/MeshPaintHelpers.cpp
--- a/Engine/Source/Editor/MeshPaint/Private/MeshPaintHelpers.cpp
+++ b/Engine/Source/Editor/MeshPaint/Private/MeshPaintHelpers.cpp
@@ -1153,7 +1153,7 @@ void MeshPaintHelpers::SetRealtimeViewport(bool bRealtime)
 	const bool bRememberCurrentState = false;
 	if (ViewportWindow.IsValid())
 	{
-		FEditorViewportClient &Viewport = ViewportWindow->GetViewportClient();
+		FEditorViewportClient &Viewport = ViewportWindow->GetAssetViewportClient();
 		if (Viewport.IsPerspective())
 		{
 			if (bRealtime)
diff --git a/Engine/Source/Editor/UnrealEd/Private/EditorViewportClient.cpp b/Engine/Source/Editor/UnrealEd/Private/EditorViewportClient.cpp
--- a/Engine/Source/Editor/UnrealEd/Private/EditorViewportClient.cpp
+++ b/Engine/Source/Editor/UnrealEd/Private/EditorViewportClient.cpp
@@ -459,7 +459,7 @@ FEditorViewportClient::~FEditorViewportClient()
 
 	if(Viewport)
 	{
-		UE_LOG(LogEditorViewport, Fatal, TEXT("Viewport != NULL in FLevelEditorViewportClient destructor."));
+		UE_LOG(LogEditorViewport, Fatal, TEXT("Viewport != NULL in FEditorViewportClient destructor."));
 	}
 
 	if(GEditor)
diff --git a/Engine/Source/Editor/UnrealEd/Private/Kismet2/DebuggerCommands.cpp b/Engine/Source/Editor/UnrealEd/Private/Kismet2/DebuggerCommands.cpp
--- a/Engine/Source/Editor/UnrealEd/Private/Kismet2/DebuggerCommands.cpp
+++ b/Engine/Source/Editor/UnrealEd/Private/Kismet2/DebuggerCommands.cpp
@@ -1356,9 +1356,9 @@ void FPlayWorldCommandCallbacks::StartPlayFromHere()
 	const bool bSimulateInEditor = false;
 	if( ActiveLevelViewport.IsValid() )
 	{
-		if( ActiveLevelViewport->GetViewportClient().IsPerspective() )
+		if( ActiveLevelViewport->GetAssetViewportClient().IsPerspective() )
 		{
-			StartRotation = ActiveLevelViewport->GetViewportClient().GetViewRotation();
+			StartRotation = ActiveLevelViewport->GetAssetViewportClient().GetViewRotation();
 		}
 
 		// If there is an active level view port, play the game in it.
@@ -1772,8 +1772,8 @@ void FInternalPlayWorldCommandCallbacks::PlayInViewport_Clicked( )
 		if( !bAtPlayerStart )
 		{
 			// Start the player where the camera is if not forcing from player start
-			StartLoc = &ActiveLevelViewport->GetViewportClient().GetViewLocation();
-			StartRot = &ActiveLevelViewport->GetViewportClient().GetViewRotation();
+			StartLoc = &ActiveLevelViewport->GetAssetViewportClient().GetViewLocation();
+			StartRot = &ActiveLevelViewport->GetAssetViewportClient().GetViewRotation();
 		}
 
 		// @todo UE4: Not supported yet
@@ -1841,8 +1841,8 @@ void FInternalPlayWorldCommandCallbacks::PlayInEditorFloating_Clicked( )
 				FSlateApplication::Get().FindWidgetWindow( ActiveLevelViewport->AsWidget() ).IsValid() )
 			{
 				// Start the player where the camera is if not forcing from player start
-				StartLoc = &ActiveLevelViewport->GetViewportClient().GetViewLocation();
-				StartRot = &ActiveLevelViewport->GetViewportClient().GetViewRotation();
+				StartLoc = &ActiveLevelViewport->GetAssetViewportClient().GetViewLocation();
+				StartRot = &ActiveLevelViewport->GetAssetViewportClient().GetViewRotation();
 			}
 		}
 
@@ -1890,8 +1890,8 @@ void FInternalPlayWorldCommandCallbacks::PlayInVR_Clicked()
 				FSlateApplication::Get().FindWidgetWindow(ActiveLevelViewport->AsWidget()).IsValid())
 			{
 				// Start the player where the camera is if not forcing from player start
-				StartLoc = &ActiveLevelViewport->GetViewportClient().GetViewLocation();
-				StartRot = &ActiveLevelViewport->GetViewportClient().GetViewRotation();
+				StartLoc = &ActiveLevelViewport->GetAssetViewportClient().GetViewLocation();
+				StartRot = &ActiveLevelViewport->GetAssetViewportClient().GetViewRotation();
 			}
 		}
 
@@ -1945,8 +1945,8 @@ void FInternalPlayWorldCommandCallbacks::PlayInNewProcess_Clicked(EPlayModeType
 
 			if ( ActiveLevelViewport.IsValid() && FSlateApplication::Get().FindWidgetWindow( ActiveLevelViewport->AsWidget() ).IsValid() )
 			{
-				StartLoc = &ActiveLevelViewport->GetViewportClient().GetViewLocation();
-				StartRot = &ActiveLevelViewport->GetViewportClient().GetViewRotation();
+				StartLoc = &ActiveLevelViewport->GetAssetViewportClient().GetViewLocation();
+				StartRot = &ActiveLevelViewport->GetAssetViewportClient().GetViewRotation();
 			}
 		}
 
diff --git a/Engine/Source/Editor/UnrealEd/Private/PlayLevel.cpp b/Engine/Source/Editor/UnrealEd/Private/PlayLevel.cpp
--- a/Engine/Source/Editor/UnrealEd/Private/PlayLevel.cpp
+++ b/Engine/Source/Editor/UnrealEd/Private/PlayLevel.cpp
@@ -620,25 +620,25 @@ void UEditorEngine::TeardownPlaySession(FWorldContext& PieWorldContext)
 					if (bLastViewAndLocationValid == true && !GEngine->IsStereoscopic3D( Viewport->GetActiveViewport() ) )
 					{
 						bLastViewAndLocationValid = false;
-						Viewport->GetViewportClient().SetViewLocation( LastViewLocation );
+						Viewport->GetAssetViewportClient().SetViewLocation( LastViewLocation );
 
-						if( Viewport->GetViewportClient().IsPerspective() )
+						if( Viewport->GetAssetViewportClient().IsPerspective() )
 						{
 							// Rotation only matters for perspective viewports not orthographic
-							Viewport->GetViewportClient().SetViewRotation( LastViewRotation );
+							Viewport->GetAssetViewportClient().SetViewRotation( LastViewRotation );
 						}
 					}
 				}
 
 				// No longer simulating in the viewport
-				Viewport->GetViewportClient().SetIsSimulateInEditorViewport( false );
+				Viewport->GetAssetViewportClient().SetIsSimulateInEditorViewport( false );
 
 				
 				FEditorModeRegistry::Get().UnregisterMode(FBuiltinEditorModes::EM_Physics);
 				
 
 				// Clear out the hit proxies before GC'ing
-				Viewport->GetViewportClient().Viewport->InvalidateHitProxy();
+				Viewport->GetAssetViewportClient().Viewport->InvalidateHitProxy();
 			}
 			else if (SlatePlayInEditorSession->SlatePlayInEditorWindow.IsValid())
 			{
@@ -808,7 +808,7 @@ void UEditorEngine::TeardownPlaySession(FWorldContext& PieWorldContext)
 				Viewport->OnSimulateSessionFinished();
 			}
 
-			StaticCast<FLevelEditorViewportClient&>(Viewport->GetViewportClient()).SetReferenceToWorldContext(EditorWorldContext);
+			StaticCast<FLevelEditorViewportClient&>(Viewport->GetAssetViewportClient()).SetReferenceToWorldContext(EditorWorldContext);
 		}
 
 		// Remove the slate info from the map (note that the UWorld* is long gone at this point, but the WorldContext still exists. It will be removed outside of this function)
@@ -3659,7 +3659,7 @@ bool UEditorEngine::GetSimulateInEditorViewTransform(FTransform& OutViewTransfor
 					TSharedPtr<IAssetViewport> LevelViewport = SlateInfoPtr->DestinationSlateViewport.Pin();
 					if (LevelViewport.IsValid())
 					{
-						FEditorViewportClient& EditorViewportClient = LevelViewport->GetViewportClient();
+						FEditorViewportClient& EditorViewportClient = LevelViewport->GetAssetViewportClient();
 						OutViewTransform = FTransform(EditorViewportClient.GetViewRotation(), EditorViewportClient.GetViewLocation());
 						return true;
 					}
@@ -3707,7 +3707,7 @@ void UEditorEngine::ToggleBetweenPIEandSIE( bool bNewSession )
 	TSharedPtr<IAssetViewport> LevelViewport = SlatePlayInEditorSession.DestinationSlateViewport.Pin();
 	if( ensure(LevelViewport.IsValid()) )
 	{
-		FEditorViewportClient& EditorViewportClient = LevelViewport->GetViewportClient();
+		FEditorViewportClient& EditorViewportClient = LevelViewport->GetAssetViewportClient();
 
 		// Toggle to pie if currently simulating
 		if( bIsSimulatingInEditor )
diff --git a/Engine/Source/Editor/UnrealEd/Public/IAssetViewport.h b/Engine/Source/Editor/UnrealEd/Public/IAssetViewport.h
--- a/Engine/Source/Editor/UnrealEd/Public/IAssetViewport.h
+++ b/Engine/Source/Editor/UnrealEd/Public/IAssetViewport.h
@@ -62,7 +62,7 @@ public:
 	/**
 	 * @return The editor client for this viewport
 	 */
-	virtual FEditorViewportClient& GetViewportClient() = 0;
+	virtual FEditorViewportClient& GetAssetViewportClient() = 0;
 
 	/**
 	 * Gets the active viewport.
diff --git a/Engine/Source/Editor/VREditor/VREditorMode.cpp b/Engine/Source/Editor/VREditor/VREditorMode.cpp
--- a/Engine/Source/Editor/VREditor/VREditorMode.cpp
+++ b/Engine/Source/Editor/VREditor/VREditorMode.cpp
@@ -546,7 +546,7 @@ void UVREditorMode::PostTick( float DeltaTime )
 	}
 
 	TickHandle.Broadcast( DeltaTime );
-	UISystem->Tick( &GetLevelViewportPossessedForVR().GetViewportClient(), DeltaTime );
+	UISystem->Tick( GetLevelViewportPossessedForVR().GetViewportClient().Get(), DeltaTime );
 
 	// Update avatar meshes
 	{
@@ -1106,7 +1106,7 @@ void UVREditorMode::StartViewport(TSharedPtr<SLevelViewport> Viewport)
 
 	if (WorldInteraction != nullptr)
 	{
-		TSharedPtr<FEditorViewportClient> VRViewportClient = MakeShareable(&Viewport->GetViewportClient());
+		TSharedPtr<FEditorViewportClient> VRViewportClient = Viewport->GetViewportClient();
 		WorldInteraction->SetDefaultOptionalViewportClient(VRViewportClient);
 	}
 }
diff --git a/Engine/Source/Editor/ViewportInteraction/ViewportInteractionModule.cpp b/Engine/Source/Editor/ViewportInteraction/ViewportInteractionModule.cpp
--- a/Engine/Source/Editor/ViewportInteraction/ViewportInteractionModule.cpp
+++ b/Engine/Source/Editor/ViewportInteraction/ViewportInteractionModule.cpp
@@ -80,8 +80,7 @@ void FViewportInteractionModule::ToggleMode()
 					TSharedPtr<IAssetViewport> ActiveLevelViewport = LevelEditor->GetActiveViewportInterface();
 					if (ActiveLevelViewport.IsValid())
 					{
- 						const TSharedRef<SLevelViewport> LevelViewportWidget = StaticCastSharedRef<SLevelViewport>(ActiveLevelViewport->AsWidget());
-						ViewportClient = MakeShareable(&LevelViewportWidget->GetViewportClient());
+						ViewportClient = StaticCastSharedRef<SLevelViewport>(ActiveLevelViewport->AsWidget())->GetViewportClient();
 					}
 				}
 
