Broadcast OnAssetsDeleted when using ForceDelete

diff --git a/Engine/Source/Editor/UnrealEd/Private/ObjectTools.cpp b/Engine/Source/Editor/UnrealEd/Private/ObjectTools.cpp
--- a/Engine/Source/Editor/UnrealEd/Private/ObjectTools.cpp
+++ b/Engine/Source/Editor/UnrealEd/Private/ObjectTools.cpp
@@ -2766,17 +2766,20 @@ namespace ObjectTools
 			}
 		}
 
+		TArray<UClass*> DeletedObjectClasses;
 		TArray<UPackage*> PotentialPackagesToDelete;
 		for(TWeakObjectPtr<UObject>& Object : ObjectsToDelete)
 		{
 			if(Object.IsValid())
 			{
+				DeletedObjectClasses.AddUnique(Object->GetClass());
 				PotentialPackagesToDelete.AddUnique(Object->GetOutermost());
 			}
 		}
 
 		if (PotentialPackagesToDelete.Num() > 0)
 		{
+			FEditorDelegates::OnAssetsDeleted.Broadcast(DeletedObjectClasses);
 			CleanupAfterSuccessfulDelete(PotentialPackagesToDelete);
 		}
 		ObjectsToDelete.Empty();
