Fixed issue where blueprint compilation errors trigger a modal dialog in unattended mode

diff --git a/Engine/Source/Editor/UnrealEd/Private/PlayLevel.cpp b/Engine/Source/Editor/UnrealEd/Private/PlayLevel.cpp
--- a/Engine/Source/Editor/UnrealEd/Private/PlayLevel.cpp
+++ b/Engine/Source/Editor/UnrealEd/Private/PlayLevel.cpp
@@ -2301,6 +2301,12 @@ bool UEditorEngine::SpawnPlayFromHereStart( UWorld* World, AActor*& PlayerStart,
 
 static bool ShowBlueprintErrorDialog( TArray<UBlueprint*> ErroredBlueprints )
 {
+	if (FApp::IsUnattended() || GIsRunningUnattendedScript)
+	{
+		// App is running in unattended mode, so we should avoid modal dialogs and proceed
+		return true;
+	}
+
 	struct Local
 	{
 		static void OnHyperlinkClicked( TWeakObjectPtr<UBlueprint> InBlueprint, TSharedPtr<SCustomDialog> InDialog )
