Fix missing virtual destructor.

diff --git a/Engine/Source/Editor/UnrealEd/Private/FbxExporter.h b/Engine/Source/Editor/UnrealEd/Private/FbxExporter.h
--- a/Engine/Source/Editor/UnrealEd/Private/FbxExporter.h
+++ b/Engine/Source/Editor/UnrealEd/Private/FbxExporter.h
@@ -270,6 +270,7 @@ private:
 	class IAnimTrackAdapter
 	{
 	public:
+		virtual ~IAnimTrackAdapter() {}
 		virtual int32 GetLocalStartFrame() const = 0;
 		virtual int32 GetStartFrame() const { return GetLocalStartFrame(); }
 		virtual int32 GetLength() const = 0;
