Adds Editor Viewport layout system.  #rb lauren.barnes #jira UETOOL-1723

diff --git a/Engine/Source/Editor/EditorStyle/Private/SlateEditorStyle.cpp b/Engine/Source/Editor/EditorStyle/Private/SlateEditorStyle.cpp
--- a/Engine/Source/Editor/EditorStyle/Private/SlateEditorStyle.cpp
+++ b/Engine/Source/Editor/EditorStyle/Private/SlateEditorStyle.cpp
@@ -5133,6 +5133,37 @@ void FSlateEditorStyle::FStyle::SetupLevelEditorStyle()
 		Set( "LevelEditor.RecompileGameCode", new IMAGE_BRUSH( "Old/MainToolBar/RecompileGameCode", Icon40x40 ) );
 	}
 
+	// Editor viewport layout command icons
+	{
+		const FVector2D IconLayoutSize(47.0f, 37.0f);
+		const FVector2D IconLayoutSizeSmall(47.0f, 37.0f);		// small version set to same size as these are in their own menu and don't clutter the UI
+
+		Set("EditorViewport.ViewportConfig_OnePane", new IMAGE_BRUSH("Icons/ViewportLayout_OnePane", IconLayoutSize));
+		Set("EditorViewport.ViewportConfig_OnePane.Small", new IMAGE_BRUSH("Icons/ViewportLayout_OnePane", IconLayoutSizeSmall));
+		Set("EditorViewport.ViewportConfig_TwoPanesH", new IMAGE_BRUSH("Icons/ViewportLayout_TwoPanesHoriz", IconLayoutSize));
+		Set("EditorViewport.ViewportConfig_TwoPanesH.Small", new IMAGE_BRUSH("Icons/ViewportLayout_TwoPanesHoriz", IconLayoutSizeSmall));
+		Set("EditorViewport.ViewportConfig_TwoPanesV", new IMAGE_BRUSH("Icons/ViewportLayout_TwoPanesVert", IconLayoutSize));
+		Set("EditorViewport.ViewportConfig_TwoPanesV.Small", new IMAGE_BRUSH("Icons/ViewportLayout_TwoPanesVert", IconLayoutSizeSmall));
+		Set("EditorViewport.ViewportConfig_ThreePanesLeft", new IMAGE_BRUSH("Icons/ViewportLayout_ThreePanesLeft", IconLayoutSize));
+		Set("EditorViewport.ViewportConfig_ThreePanesLeft.Small", new IMAGE_BRUSH("Icons/ViewportLayout_ThreePanesLeft", IconLayoutSizeSmall));
+		Set("EditorViewport.ViewportConfig_ThreePanesRight", new IMAGE_BRUSH("Icons/ViewportLayout_ThreePanesRight", IconLayoutSize));
+		Set("EditorViewport.ViewportConfig_ThreePanesRight.Small", new IMAGE_BRUSH("Icons/ViewportLayout_ThreePanesRight", IconLayoutSizeSmall));
+		Set("EditorViewport.ViewportConfig_ThreePanesTop", new IMAGE_BRUSH("Icons/ViewportLayout_ThreePanesTop", IconLayoutSize));
+		Set("EditorViewport.ViewportConfig_ThreePanesTop.Small", new IMAGE_BRUSH("Icons/ViewportLayout_ThreePanesTop", IconLayoutSizeSmall));
+		Set("EditorViewport.ViewportConfig_ThreePanesBottom", new IMAGE_BRUSH("Icons/ViewportLayout_ThreePanesBottom", IconLayoutSize));
+		Set("EditorViewport.ViewportConfig_ThreePanesBottom.Small", new IMAGE_BRUSH("Icons/ViewportLayout_ThreePanesBottom", IconLayoutSizeSmall));
+		Set("EditorViewport.ViewportConfig_FourPanesLeft", new IMAGE_BRUSH("Icons/ViewportLayout_FourPanesLeft", IconLayoutSize));
+		Set("EditorViewport.ViewportConfig_FourPanesLeft.Small", new IMAGE_BRUSH("Icons/ViewportLayout_FourPanesLeft", IconLayoutSizeSmall));
+		Set("EditorViewport.ViewportConfig_FourPanesRight", new IMAGE_BRUSH("Icons/ViewportLayout_FourPanesRight", IconLayoutSize));
+		Set("EditorViewport.ViewportConfig_FourPanesRight.Small", new IMAGE_BRUSH("Icons/ViewportLayout_FourPanesRight", IconLayoutSizeSmall));
+		Set("EditorViewport.ViewportConfig_FourPanesTop", new IMAGE_BRUSH("Icons/ViewportLayout_FourPanesTop", IconLayoutSize));
+		Set("EditorViewport.ViewportConfig_FourPanesTop.Small", new IMAGE_BRUSH("Icons/ViewportLayout_FourPanesTop", IconLayoutSizeSmall));
+		Set("EditorViewport.ViewportConfig_FourPanesBottom", new IMAGE_BRUSH("Icons/ViewportLayout_FourPanesBottom", IconLayoutSize));
+		Set("EditorViewport.ViewportConfig_FourPanesBottom.Small", new IMAGE_BRUSH("Icons/ViewportLayout_FourPanesBottom", IconLayoutSizeSmall));
+		Set("EditorViewport.ViewportConfig_FourPanes2x2", new IMAGE_BRUSH("Icons/ViewportLayout_FourPanes2x2", IconLayoutSize));
+		Set("EditorViewport.ViewportConfig_FourPanes2x2.Small", new IMAGE_BRUSH("Icons/ViewportLayout_FourPanes2x2", IconLayoutSizeSmall));
+	}
+
 	// Level viewport layout command icons
 	{
 		const FVector2D IconLayoutSize(47.0f, 37.0f);
diff --git a/Engine/Source/Editor/LevelEditor/Public/LevelViewportTabContent.h b/Engine/Source/Editor/LevelEditor/Public/LevelViewportTabContent.h
--- a/Engine/Source/Editor/LevelEditor/Public/LevelViewportTabContent.h
+++ b/Engine/Source/Editor/LevelEditor/Public/LevelViewportTabContent.h
@@ -12,7 +12,7 @@ class ILevelEditor;
  * Represents the content in a viewport tab in the level editor.
  * Each SDockTab holding viewports in the level editor contains and owns one of these.
  */
-class LEVELEDITOR_API FLevelViewportTabContent : public FViewportTabContent
+class LEVELEDITOR_API FLevelViewportTabContent : public FViewportTabContent, public TSharedFromThis<FLevelViewportTabContent>
 {
 public:
 	/** Starts the tab content object and creates the initial layout based on the layout string */
diff --git a/Engine/Source/Editor/UnrealEd/Private/AssetEditorViewportLayout.cpp b/Engine/Source/Editor/UnrealEd/Private/AssetEditorViewportLayout.cpp
new file mode 100644
--- /dev/null
+++ b/Engine/Source/Editor/UnrealEd/Private/AssetEditorViewportLayout.cpp
@@ -0,0 +1,315 @@
+// Copyright 1998-2019 Epic Games, Inc. All Rights Reserved.
+
+#include "AssetEditorViewportLayout.h"
+#include "Layout/Margin.h"
+#include "Widgets/DeclarativeSyntaxSupport.h"
+#include "Widgets/SCompoundWidget.h"
+#include "Widgets/SOverlay.h"
+#include "Framework/Docking/LayoutService.h"
+#include "Misc/ConfigCacheIni.h"
+#include "Modules/ModuleManager.h"
+#include "Layout/WidgetPath.h"
+#include "Framework/Application/SlateApplication.h"
+#include "Widgets/Layout/SBorder.h"
+#include "Widgets/Layout/SSpacer.h"
+#include "Widgets/SCanvas.h"
+#include "EditorStyleSet.h"
+#include "Editor/UnrealEdEngine.h"
+#include "UnrealEdGlobals.h"
+#include "Widgets/Docking/SDockTab.h"
+#include "EditorViewportLayout.h"
+#include "EditorViewportTypeDefinition.h"
+#include "EditorViewportLayoutEntity.h"
+#include "EditorViewportCommands.h"
+#include "SAssetEditorViewport.h"
+#include "Templates/SharedPointer.h"
+
+
+#define LOCTEXT_NAMESPACE "AssetEditorViewportToolBar"
+
+
+namespace ViewportLayoutDefs
+{
+	/** How many seconds to interpolate from restored to maximized state */
+	static const float MaximizeTransitionTime = 0.15f;
+
+	/** How many seconds to interpolate from maximized to restored state */
+	static const float RestoreTransitionTime = 0.2f;
+
+	/** Default maximized state for new layouts - will only be applied when no config data is restoring state */
+	static const bool bDefaultShouldBeMaximized = true;
+
+	/** Default immersive state for new layouts - will only be applied when no config data is restoring state */
+	static const bool bDefaultShouldBeImmersive = false;
+}
+
+
+// SViewportsOverlay ////////////////////////////////////////////////
+
+/**
+* Overlay wrapper class so that we can cache the size of the widget
+* It will also store the ViewportLayout data because that data can't be stored
+* per app; it must be stored per viewport overlay in case the app that made it closes.
+*/
+class SViewportsOverlay : public SCompoundWidget
+{
+
+public:
+
+	SLATE_BEGIN_ARGS( SViewportsOverlay ){}
+	SLATE_DEFAULT_SLOT( FArguments, Content )
+		SLATE_ARGUMENT( TSharedPtr<FViewportTabContent>, ViewportTab )
+		SLATE_END_ARGS()
+
+		void Construct( const FArguments& InArgs );
+
+	/** Default constructor */
+	SViewportsOverlay()
+		: CachedSize( FVector2D::ZeroVector )
+	{}
+
+	/** Overridden from SWidget */
+	virtual void Tick( const FGeometry& AllottedGeometry, const double InCurrentTime, const float InDeltaTime ) override;
+
+	/** Wraps SOverlay::AddSlot() */
+	SOverlay::FOverlaySlot& AddSlot();
+
+	/** Wraps SOverlay::RemoveSlot() */
+	void RemoveSlot();
+
+	/**
+	* Returns the cached size of this viewport overlay
+	*
+	* @return	The size that was cached
+	*/
+	const FVector2D& GetCachedSize() const;
+
+	/** Gets the  Viewport Tab that created this overlay */
+	TSharedPtr<FViewportTabContent> GetViewportTab() const;
+
+private:
+
+	/** Reference to the owning  viewport tab */
+	TSharedPtr<FViewportTabContent> ViewportTab;
+
+	/** The overlay widget we're containing */
+	TSharedPtr< SOverlay > OverlayWidget;
+
+	/** Cache our size, so that we can use this when animating a viewport maximize/restore */
+	FVector2D CachedSize;
+};
+
+
+void SViewportsOverlay::Construct( const FArguments& InArgs )
+{
+	const TSharedRef<SWidget>& ContentWidget = InArgs._Content.Widget;
+	ViewportTab = InArgs._ViewportTab;
+
+	ChildSlot
+		[
+			SAssignNew( OverlayWidget, SOverlay )
+			+ SOverlay::Slot()
+		[
+			ContentWidget
+		]
+		];
+}
+
+void SViewportsOverlay::Tick( const FGeometry& AllottedGeometry, const double InCurrentTime, const float InDeltaTime )
+{
+	CachedSize = AllottedGeometry.Size;
+}
+
+SOverlay::FOverlaySlot& SViewportsOverlay::AddSlot()
+{
+	return OverlayWidget->AddSlot();
+}
+
+void SViewportsOverlay::RemoveSlot()
+{
+	return OverlayWidget->RemoveSlot();
+}
+
+const FVector2D& SViewportsOverlay::GetCachedSize() const
+{
+	return CachedSize;
+}
+
+TSharedPtr<FViewportTabContent> SViewportsOverlay::GetViewportTab() const
+{
+	return ViewportTab;
+}
+
+
+// FAssetEditorViewportLayout /////////////////////////////
+
+FAssetEditorViewportLayout::FAssetEditorViewportLayout() {}
+
+
+FAssetEditorViewportLayout::~FAssetEditorViewportLayout()
+{
+	for (auto& Pair : Viewports)
+	{
+		Pair.Value->OnLayoutDestroyed();
+	}
+}
+
+TSharedRef<IEditorViewportLayoutEntity> FAssetEditorViewportLayout::FactoryViewport(TFunction<TSharedRef<SEditorViewport>(void)> &Func, FName InTypeName, const FAssetEditorViewportConstructionArgs& ConstructionArgs) const
+{
+	TSharedRef<FEditorViewportLayoutEntity> LayoutEntity  = MakeShareable(new FEditorViewportLayoutEntity(Func, ConstructionArgs));
+	return LayoutEntity;
+}
+
+
+
+TSharedRef<SWidget> FAssetEditorViewportLayout::BuildViewportLayout(TFunction<TSharedRef<SEditorViewport>(void)> &Func, TSharedPtr<SDockTab> InParentDockTab, TSharedPtr<class FViewportTabContent> InParentTab, const FString& LayoutString)
+{
+	// We don't support reconfiguring an existing layout object, as this makes handling of transitions
+	// particularly difficult.  Instead just destroy the old layout and create a new layout object.
+	check(!ParentTab.IsValid());
+	ParentTab = InParentDockTab;
+	ParentTabContent = InParentTab;
+
+	// We use an overlay so that we can draw a maximized viewport on top of the other viewports
+	TSharedPtr<SBorder> ViewportsBorder;
+	TSharedRef<SViewportsOverlay> ViewportsOverlay =
+		SNew(SViewportsOverlay)
+		.ViewportTab(InParentTab)
+		[
+			SAssignNew(ViewportsBorder, SBorder)
+			.Padding(0.0f)
+		.BorderImage(FEditorStyle::GetBrush("NoBorder"))
+		];
+
+	ViewportsOverlayPtr = ViewportsOverlay;
+
+	// Don't set the content until the OverlayPtr has been set, because it access this when we want to start with the viewports maximized.
+	ViewportsBorder->SetContent(MakeViewportLayout(Func, LayoutString));
+
+	return ViewportsOverlay;
+
+}
+
+TSharedRef<SWidget> FAssetEditorViewportLayout::GenerateLayoutMenu(TSharedPtr<SAssetEditorViewport> AssetEditorViewport)
+{
+	const bool bInShouldCloseWindowAfterMenuSelection = true;
+
+	TSharedPtr<FUICommandList> CommandList = AssetEditorViewport->GetCommandList();
+
+	FMenuBuilder MenuBuilder(bInShouldCloseWindowAfterMenuSelection, CommandList);
+
+	MenuBuilder.BeginSection(NAME_None, LOCTEXT("OnePaneConfigHeader", "One Pane"));
+	{
+ 		FToolBarBuilder OnePaneButton(CommandList, FMultiBoxCustomization::None);
+		OnePaneButton.SetLabelVisibility(EVisibility::Collapsed);
+		OnePaneButton.SetStyle(&FEditorStyle::Get(), "ViewportLayoutToolbar");
+
+ 		OnePaneButton.AddToolBarButton(FEditorViewportCommands::Get().ViewportConfig_OnePane);
+
+		MenuBuilder.AddWidget(
+			SNew(SHorizontalBox)
+			+ SHorizontalBox::Slot()
+			.AutoWidth()
+			[
+				OnePaneButton.MakeWidget()
+			]
+		+ SHorizontalBox::Slot()
+			.FillWidth(1)
+			[
+				SNullWidget::NullWidget
+			],
+			FText::GetEmpty(), true
+			);
+	}
+	MenuBuilder.EndSection();
+ 
+	MenuBuilder.BeginSection("EditorViewportTwoPaneConfigs", LOCTEXT("TwoPaneConfigHeader", "Two Panes"));
+	{
+		FToolBarBuilder TwoPaneButtons(CommandList, FMultiBoxCustomization::None);
+		TwoPaneButtons.SetLabelVisibility(EVisibility::Collapsed);
+		TwoPaneButtons.SetStyle(&FEditorStyle::Get(), "ViewportLayoutToolbar");
+
+		TwoPaneButtons.AddToolBarButton(FEditorViewportCommands::Get().ViewportConfig_TwoPanesH, NAME_None, FText());
+		TwoPaneButtons.AddToolBarButton(FEditorViewportCommands::Get().ViewportConfig_TwoPanesV, NAME_None, FText());
+
+		MenuBuilder.AddWidget(
+			SNew(SHorizontalBox)
+			+ SHorizontalBox::Slot()
+			.AutoWidth()
+			[
+				TwoPaneButtons.MakeWidget()
+			]
+		+ SHorizontalBox::Slot()
+			.FillWidth(1)
+			[
+				SNullWidget::NullWidget
+			],
+			FText::GetEmpty(), true
+			);
+	}
+	MenuBuilder.EndSection();
+
+	MenuBuilder.BeginSection("EditorViewportThreePaneConfigs", LOCTEXT("ThreePaneConfigHeader", "Three Panes"));
+	{
+		FToolBarBuilder ThreePaneButtons(CommandList, FMultiBoxCustomization::None);
+		ThreePaneButtons.SetLabelVisibility(EVisibility::Collapsed);
+		ThreePaneButtons.SetStyle(&FEditorStyle::Get(), "ViewportLayoutToolbar");
+
+		ThreePaneButtons.AddToolBarButton(FEditorViewportCommands::Get().ViewportConfig_ThreePanesLeft, NAME_None, FText());
+		ThreePaneButtons.AddToolBarButton(FEditorViewportCommands::Get().ViewportConfig_ThreePanesRight, NAME_None, FText());
+		ThreePaneButtons.AddToolBarButton(FEditorViewportCommands::Get().ViewportConfig_ThreePanesTop, NAME_None, FText());
+		ThreePaneButtons.AddToolBarButton(FEditorViewportCommands::Get().ViewportConfig_ThreePanesBottom, NAME_None, FText());
+
+		MenuBuilder.AddWidget(
+			SNew(SHorizontalBox)
+			+ SHorizontalBox::Slot()
+			.AutoWidth()
+			[
+				ThreePaneButtons.MakeWidget()
+			]
+		+ SHorizontalBox::Slot()
+			.FillWidth(1)
+			[
+				SNullWidget::NullWidget
+			],
+			FText::GetEmpty(), true
+			);
+	}
+	MenuBuilder.EndSection();
+
+	MenuBuilder.BeginSection("EditorViewportFourPaneConfigs", LOCTEXT("FourPaneConfigHeader", "Four Panes"));
+	{
+		FToolBarBuilder FourPaneButtons(CommandList, FMultiBoxCustomization::None);
+		FourPaneButtons.SetLabelVisibility(EVisibility::Collapsed);
+		FourPaneButtons.SetStyle(&FEditorStyle::Get(), "ViewportLayoutToolbar");
+
+		FourPaneButtons.AddToolBarButton(FEditorViewportCommands::Get().ViewportConfig_FourPanes2x2, NAME_None, FText());
+		FourPaneButtons.AddToolBarButton(FEditorViewportCommands::Get().ViewportConfig_FourPanesLeft, NAME_None, FText());
+		FourPaneButtons.AddToolBarButton(FEditorViewportCommands::Get().ViewportConfig_FourPanesRight, NAME_None, FText());
+		FourPaneButtons.AddToolBarButton(FEditorViewportCommands::Get().ViewportConfig_FourPanesTop, NAME_None, FText());
+		FourPaneButtons.AddToolBarButton(FEditorViewportCommands::Get().ViewportConfig_FourPanesBottom, NAME_None, FText());
+
+		MenuBuilder.AddWidget(
+			SNew(SHorizontalBox)
+			+ SHorizontalBox::Slot()
+			.AutoWidth()
+			[
+				FourPaneButtons.MakeWidget()
+			]
+		+ SHorizontalBox::Slot()
+			.FillWidth(1)
+			[
+				SNullWidget::NullWidget
+			],
+			FText::GetEmpty(), true
+			);
+	}
+	MenuBuilder.EndSection();
+
+
+	return MenuBuilder.MakeWidget();
+
+}
+
+#undef LOCTEXT_NAMESPACE
+
diff --git a/Engine/Source/Editor/UnrealEd/Private/EditorViewportLayout2x2.cpp b/Engine/Source/Editor/UnrealEd/Private/EditorViewportLayout2x2.cpp
new file mode 100644
--- /dev/null
+++ b/Engine/Source/Editor/UnrealEd/Private/EditorViewportLayout2x2.cpp
@@ -0,0 +1,90 @@
+// Copyright 1998-2019 Epic Games, Inc. All Rights Reserved.
+
+#include "EditorViewportLayout2x2.h"
+#include "Framework/Docking/LayoutService.h"
+#include "Editor.h"
+#include "Misc/ConfigCacheIni.h"
+#include "Modules/ModuleManager.h"
+#include "Framework/Application/SlateApplication.h"
+// #include "LevelEditor.h"
+
+namespace ViewportLayout2x2Defs
+{
+	/** Default 2x2 splitters to equal 50/50 splits */
+	static const FVector2D DefaultSplitterPercentages(0.5f, 0.5f);
+}
+
+// FLevelViewportLayout2x2 //////////////////////////////////////////
+
+TSharedRef<SWidget> FEditorViewportLayout2x2::MakeViewportLayout(TFunction<TSharedRef<SEditorViewport>(void)> &Func, const FString& LayoutString)
+{
+// 	FString SpecificLayoutString = GetTypeSpecificLayoutString(LayoutString);
+
+	FString TopLeftKey, BottomLeftKey, TopRightKey, BottomRightKey;
+
+	FString TopLeftType = TEXT("Default"), BottomLeftType = TEXT("Default"), TopRightType = TEXT("Default"), BottomRightType = TEXT("Default");
+
+	TArray<FVector2D> SplitterPercentages;
+	
+	// Set up the viewports
+	FAssetEditorViewportConstructionArgs Args;
+ 	Args.ParentLayout = AsShared();
+
+	// Left viewport
+	Args.bRealtime = false;
+	Args.ConfigKey = *TopLeftKey;
+	Args.ViewportType = LVT_OrthoYZ;
+	TSharedPtr< IEditorViewportLayoutEntity > ViewportTL = FactoryViewport(Func, *TopLeftType, Args);
+
+	// Persp viewport
+	Args.bRealtime = !FPlatformMisc::IsRemoteSession();
+	Args.ConfigKey = *BottomLeftKey;
+	Args.ViewportType = LVT_Perspective;
+	TSharedPtr< IEditorViewportLayoutEntity > ViewportBL = FactoryViewport(Func, *BottomLeftType, Args);
+
+	// Front viewport
+	Args.bRealtime = false;
+	Args.ConfigKey = *TopRightKey;
+	Args.ViewportType = LVT_OrthoXZ;
+	TSharedPtr< IEditorViewportLayoutEntity > ViewportTR = FactoryViewport(Func, *TopRightType, Args);
+
+	// Top Viewport
+	Args.bRealtime = false;
+	Args.ConfigKey = *BottomRightKey;
+	Args.ViewportType = LVT_OrthoXY;
+	TSharedPtr< IEditorViewportLayoutEntity > ViewportBR = FactoryViewport(Func, *BottomRightType, Args);
+
+	Viewports.Add( *TopLeftKey, ViewportTL );
+	Viewports.Add( *BottomLeftKey, ViewportBL );
+	Viewports.Add( *TopRightKey, ViewportTR );
+	Viewports.Add( *BottomRightKey, ViewportBR );
+
+	// Set up the splitter
+	SplitterWidget = 
+	SNew( SSplitter2x2 )
+	.TopLeft()
+	[
+		ViewportTL->AsWidget()
+	]
+	.BottomLeft()
+	[
+		ViewportBL->AsWidget()
+	]
+	.TopRight()
+	[
+		ViewportTR->AsWidget()
+	]
+	.BottomRight()
+	[
+		ViewportBR->AsWidget()
+	];
+	
+	if (SplitterPercentages.Num() > 0)
+	{
+		SplitterWidget->SetSplitterPercentages(SplitterPercentages);
+	}
+
+
+	return SplitterWidget.ToSharedRef();
+}
+
diff --git a/Engine/Source/Editor/UnrealEd/Private/EditorViewportLayoutEntity.cpp b/Engine/Source/Editor/UnrealEd/Private/EditorViewportLayoutEntity.cpp
new file mode 100644
--- /dev/null
+++ b/Engine/Source/Editor/UnrealEd/Private/EditorViewportLayoutEntity.cpp
@@ -0,0 +1,56 @@
+// Copyright 1998-2019 Epic Games, Inc. All Rights Reserved.
+
+#include "EditorViewportLayoutEntity.h"
+#include "Widgets/DeclarativeSyntaxSupport.h"
+#include "Editor/UnrealEdEngine.h"
+#include "LevelEditorViewport.h"
+#include "UnrealEdGlobals.h"
+#include "SAssetEditorViewport.h"
+#include "EditorViewportCommands.h"
+#include "EditorViewportTabContent.h"
+
+FEditorViewportLayoutEntity::FEditorViewportLayoutEntity(TFunction<TSharedRef<SEditorViewport>(void)> &Func, const FAssetEditorViewportConstructionArgs& ConstructionArgs)
+{
+ 	TSharedPtr<SEditorViewport> NewViewport = Func();
+	AssetEditorViewport = StaticCastSharedPtr<SAssetEditorViewport>( NewViewport);
+	AssetEditorViewport->ParentLayout = ConstructionArgs.ParentLayout;
+}
+
+
+TSharedRef<SWidget> FEditorViewportLayoutEntity::AsWidget() const
+{
+	return AssetEditorViewport.ToSharedRef();
+}
+
+TSharedPtr<SAssetEditorViewport> FEditorViewportLayoutEntity::AsAssetEditorViewport() const
+{
+	return AssetEditorViewport;
+}
+
+FName FEditorViewportLayoutEntity::GetType() const
+{
+	static FName DefaultName("Default");
+	return DefaultName;
+}
+
+TSharedPtr<FEditorViewportClient> FEditorViewportLayoutEntity::GetViewportClient() const
+{
+	return AssetEditorViewport->GetViewportClient();
+}
+
+void FEditorViewportLayoutEntity::SetKeyboardFocus()
+{
+}
+
+void FEditorViewportLayoutEntity::OnLayoutDestroyed()
+{
+}
+
+void FEditorViewportLayoutEntity::SaveConfig(const FString& ConfigSection)
+{
+}
+
+void FEditorViewportLayoutEntity::TakeHighResScreenShot() const
+{
+	GetViewportClient()->TakeHighResScreenShot();
+}
diff --git a/Engine/Source/Editor/UnrealEd/Private/EditorViewportLayoutEntity.h b/Engine/Source/Editor/UnrealEd/Private/EditorViewportLayoutEntity.h
new file mode 100644
--- /dev/null
+++ b/Engine/Source/Editor/UnrealEd/Private/EditorViewportLayoutEntity.h
@@ -0,0 +1,31 @@
+// Copyright 1998-2019 Epic Games, Inc. All Rights Reserved.
+
+#pragma once
+
+#include "CoreMinimal.h"
+#include "Widgets/SWidget.h"
+#include "AssetEditorViewportLayout.h"
+#include "SEditorViewport.h"
+
+class FEditorViewportClient;
+
+class UNREALED_API FEditorViewportLayoutEntity : public IEditorViewportLayoutEntity
+{
+public:
+
+	FEditorViewportLayoutEntity(TFunction<TSharedRef<SEditorViewport>(void)> &Func, const FAssetEditorViewportConstructionArgs& ConstructionArgs);
+	virtual TSharedRef<SWidget> AsWidget() const override;
+	virtual TSharedPtr<SAssetEditorViewport> AsAssetEditorViewport() const;
+
+	void SetKeyboardFocus() override;
+	void OnLayoutDestroyed() override;
+	void SaveConfig(const FString& ConfigSection) override;
+	TSharedPtr<FEditorViewportClient> GetViewportClient() const;
+	FName GetType() const override;
+	void TakeHighResScreenShot() const override;
+
+protected:
+
+	/** This entity's editor viewport */
+	TSharedPtr<SAssetEditorViewport> AssetEditorViewport;
+};
diff --git a/Engine/Source/Editor/UnrealEd/Private/EditorViewportLayoutFourPanes.cpp b/Engine/Source/Editor/UnrealEd/Private/EditorViewportLayoutFourPanes.cpp
new file mode 100644
--- /dev/null
+++ b/Engine/Source/Editor/UnrealEd/Private/EditorViewportLayoutFourPanes.cpp
@@ -0,0 +1,232 @@
+// Copyright 1998-2019 Epic Games, Inc. All Rights Reserved.
+
+#include "EditorViewportLayoutFourPanes.h"
+#include "Framework/Docking/LayoutService.h"
+#include "Editor.h"
+#include "Misc/ConfigCacheIni.h"
+#include "Modules/ModuleManager.h"
+#include "Framework/Application/SlateApplication.h"
+
+namespace ViewportLayoutFourPanesDefs
+{
+	/** Default main splitter to equal 50/50 split */
+	static const float DefaultPrimarySplitterPercentage = 0.5f;
+
+	/** Default secondary splitter to equal three-way split */
+	static const float DefaultSecondarySplitterPercentage = 0.333f;
+}
+
+// FEditorViewportLayoutFourPanes /////////////////////////////
+
+TSharedRef<SWidget> FEditorViewportLayoutFourPanes::MakeViewportLayout(TFunction<TSharedRef<SEditorViewport>(void)> &Func, const FString& LayoutString)
+{
+	FString ViewportKey0, ViewportKey1, ViewportKey2, ViewportKey3;
+	FString ViewportType0 = TEXT("Default"), ViewportType1 = TEXT("Default"), ViewportType2 = TEXT("Default"), ViewportType3 = TEXT("Default");
+
+	float PrimarySplitterPercentage = ViewportLayoutFourPanesDefs::DefaultPrimarySplitterPercentage;
+	float SecondarySplitterPercentage0 = ViewportLayoutFourPanesDefs::DefaultSecondarySplitterPercentage;
+	float SecondarySplitterPercentage1 = ViewportLayoutFourPanesDefs::DefaultSecondarySplitterPercentage;
+
+	// Set up the viewports
+	FAssetEditorViewportConstructionArgs Args;
+	Args.ParentLayout = AsShared();
+	Args.IsEnabled = FSlateApplication::Get().GetNormalExecutionAttribute();
+
+	Args.bRealtime = !FPlatformMisc::IsRemoteSession();
+	Args.ConfigKey = *ViewportKey0;
+	Args.ViewportType = LVT_Perspective;
+	TSharedRef<IEditorViewportLayoutEntity> Viewport0 = FactoryViewport(Func, *ViewportType0, Args);
+
+	Args.bRealtime = false;
+	Args.ConfigKey = *ViewportKey1;
+	Args.ViewportType = LVT_OrthoXY;
+	TSharedRef<IEditorViewportLayoutEntity> Viewport1 = FactoryViewport(Func, *ViewportType1, Args);
+
+	// Front viewport
+	Args.bRealtime = false;
+	Args.ConfigKey = *ViewportKey2;
+	Args.ViewportType = LVT_OrthoXZ;
+	TSharedRef<IEditorViewportLayoutEntity> Viewport2 = FactoryViewport(Func, *ViewportType2, Args);
+
+	// Top Viewport
+	Args.bRealtime = false;
+	Args.ConfigKey = *ViewportKey3;
+	Args.ViewportType = LVT_OrthoYZ;
+	TSharedRef<IEditorViewportLayoutEntity> Viewport3 = FactoryViewport(Func, *ViewportType2, Args);
+
+	Viewports.Add( *ViewportKey0, Viewport0 );
+	Viewports.Add( *ViewportKey1, Viewport1 );
+	Viewports.Add( *ViewportKey2, Viewport2 );
+	Viewports.Add( *ViewportKey3, Viewport3 );
+
+	TSharedRef<SWidget> LayoutWidget = MakeFourPanelWidget(
+		Viewports,
+		Viewport0->AsWidget(), Viewport1->AsWidget(), Viewport2->AsWidget(), Viewport3->AsWidget(),
+		PrimarySplitterPercentage, SecondarySplitterPercentage0, SecondarySplitterPercentage1);
+
+	return LayoutWidget;
+}
+
+
+// FEditorViewportLayoutFourPanesLeft /////////////////////////////
+
+TSharedRef<SWidget> FEditorViewportLayoutFourPanesLeft::MakeFourPanelWidget(
+	TMap<FName, TSharedPtr< IEditorViewportLayoutEntity >>& ViewportWidgets,
+	TSharedRef<SWidget> Viewport0, TSharedRef<SWidget> Viewport1, TSharedRef<SWidget> Viewport2, TSharedRef<SWidget> Viewport3,
+	float PrimarySplitterPercentage, float SecondarySplitterPercentage0, float SecondarySplitterPercentage1)
+{
+	TSharedRef<SWidget> Widget = 
+		SAssignNew( PrimarySplitterWidget, SSplitter )
+		.Orientation(EOrientation::Orient_Horizontal)
+		+SSplitter::Slot()
+		.Value(PrimarySplitterPercentage)
+		[
+			Viewport0
+		]
+		+SSplitter::Slot()
+		.Value(1.0f - PrimarySplitterPercentage)
+		[
+			SAssignNew( SecondarySplitterWidget, SSplitter )
+			.Orientation(EOrientation::Orient_Vertical)
+			+SSplitter::Slot()
+			.Value(SecondarySplitterPercentage0)
+			[
+				Viewport1
+			]
+			+SSplitter::Slot()
+			.Value(SecondarySplitterPercentage1)
+			[
+				Viewport2
+			]
+			+SSplitter::Slot()
+			.Value(1.0f - SecondarySplitterPercentage0 - SecondarySplitterPercentage1)
+			[
+				Viewport3
+			]
+		];
+
+	return Widget;
+}
+
+
+// FEditorViewportLayoutFourPanesRight /////////////////////////////
+
+TSharedRef<SWidget> FEditorViewportLayoutFourPanesRight::MakeFourPanelWidget(
+	TMap<FName, TSharedPtr< IEditorViewportLayoutEntity >>& ViewportWidgets,
+	TSharedRef<SWidget> Viewport0, TSharedRef<SWidget> Viewport1, TSharedRef<SWidget> Viewport2, TSharedRef<SWidget> Viewport3,
+	float PrimarySplitterPercentage, float SecondarySplitterPercentage0, float SecondarySplitterPercentage1)
+{
+	TSharedRef<SWidget> Widget = 
+		SAssignNew( PrimarySplitterWidget, SSplitter )
+		.Orientation(EOrientation::Orient_Horizontal)
+		+SSplitter::Slot()
+		.Value(PrimarySplitterPercentage)
+		[
+			SAssignNew( SecondarySplitterWidget, SSplitter )
+			.Orientation(EOrientation::Orient_Vertical)
+			+SSplitter::Slot()
+			.Value(SecondarySplitterPercentage0)
+			[
+				Viewport1
+			]
+			+SSplitter::Slot()
+			.Value(SecondarySplitterPercentage1)
+			[
+				Viewport2
+			]
+			+SSplitter::Slot()
+			.Value(1.0f - SecondarySplitterPercentage0 - SecondarySplitterPercentage1)
+			[
+				Viewport3
+			]
+		]
+		+SSplitter::Slot()
+		.Value(1.0f - PrimarySplitterPercentage)
+		[
+			Viewport0
+		];
+
+	return Widget;
+}
+
+
+// FEditorViewportLayoutFourPanesTop /////////////////////////////
+
+TSharedRef<SWidget> FEditorViewportLayoutFourPanesTop::MakeFourPanelWidget(
+	TMap<FName, TSharedPtr< IEditorViewportLayoutEntity >>& ViewportWidgets,
+	TSharedRef<SWidget> Viewport0, TSharedRef<SWidget> Viewport1, TSharedRef<SWidget> Viewport2, TSharedRef<SWidget> Viewport3,
+	float PrimarySplitterPercentage, float SecondarySplitterPercentage0, float SecondarySplitterPercentage1)
+{
+	TSharedRef<SWidget> Widget = 
+		SAssignNew( PrimarySplitterWidget, SSplitter )
+		.Orientation(EOrientation::Orient_Vertical)
+		+SSplitter::Slot()
+		.Value(PrimarySplitterPercentage)
+		[
+			Viewport0
+		]
+		+SSplitter::Slot()
+		.Value(1.0f - PrimarySplitterPercentage)
+		[
+			SAssignNew( SecondarySplitterWidget, SSplitter )
+			.Orientation(EOrientation::Orient_Horizontal)
+			+SSplitter::Slot()
+			.Value(SecondarySplitterPercentage0)
+			[
+				Viewport1
+			]
+			+SSplitter::Slot()
+			.Value(SecondarySplitterPercentage1)
+			[
+				Viewport2
+			]
+			+SSplitter::Slot()
+			.Value(1.0f - SecondarySplitterPercentage0 - SecondarySplitterPercentage1)
+			[
+				Viewport3
+			]
+		];
+
+	return Widget;
+}
+
+
+// FEditorViewportLayoutFourPanesBottom /////////////////////////////
+
+TSharedRef<SWidget> FEditorViewportLayoutFourPanesBottom::MakeFourPanelWidget(
+	TMap<FName, TSharedPtr< IEditorViewportLayoutEntity >>& ViewportWidgets,
+	TSharedRef<SWidget> Viewport0, TSharedRef<SWidget> Viewport1, TSharedRef<SWidget> Viewport2, TSharedRef<SWidget> Viewport3,
+	float PrimarySplitterPercentage, float SecondarySplitterPercentage0, float SecondarySplitterPercentage1)
+{
+	TSharedRef<SWidget> Widget = 
+		SAssignNew( PrimarySplitterWidget, SSplitter )
+		.Orientation(EOrientation::Orient_Vertical)
+		+SSplitter::Slot()
+		.Value(PrimarySplitterPercentage)
+		[
+			SAssignNew( SecondarySplitterWidget, SSplitter )
+			.Orientation(EOrientation::Orient_Horizontal)
+			+SSplitter::Slot()
+			.Value(SecondarySplitterPercentage0)
+			[
+				Viewport1
+			]
+			+SSplitter::Slot()
+			.Value(SecondarySplitterPercentage1)
+			[
+				Viewport2
+			]
+			+SSplitter::Slot()
+			.Value(1.0f - SecondarySplitterPercentage0 - SecondarySplitterPercentage1)
+			[
+				Viewport3
+			]
+		]
+		+SSplitter::Slot()
+		.Value(1.0f - PrimarySplitterPercentage)
+		[
+			Viewport0
+		];
+
+	return Widget;
+}
diff --git a/Engine/Source/Editor/UnrealEd/Private/EditorViewportLayoutOnePane.cpp b/Engine/Source/Editor/UnrealEd/Private/EditorViewportLayoutOnePane.cpp
new file mode 100644
--- /dev/null
+++ b/Engine/Source/Editor/UnrealEd/Private/EditorViewportLayoutOnePane.cpp
@@ -0,0 +1,35 @@
+// Copyright 1998-2019 Epic Games, Inc. All Rights Reserved.
+
+#include "EditorViewportLayoutOnePane.h"
+#include "Widgets/DeclarativeSyntaxSupport.h"
+#include "Widgets/SBoxPanel.h"
+#include "Framework/Docking/LayoutService.h"
+#include "ShowFlags.h"
+#include "Editor.h"
+#include "Misc/ConfigCacheIni.h"
+#include "Modules/ModuleManager.h"
+#include "Framework/Application/SlateApplication.h"
+
+
+TSharedRef<SWidget> FEditorViewportLayoutOnePane::MakeViewportLayout(TFunction<TSharedRef<SEditorViewport>(void)> &Func, const FString& LayoutString)
+{
+	FString ViewportKey, ViewportType;
+
+	// Set up the viewport
+	FAssetEditorViewportConstructionArgs Args;
+ 	Args.ParentLayout = AsShared();
+	Args.bRealtime = !FPlatformMisc::IsRemoteSession();
+ 	Args.ViewportType = LVT_Perspective;
+	TSharedRef<IEditorViewportLayoutEntity> Viewport = FactoryViewport(Func, *ViewportType, Args);
+
+	ViewportBox =
+		SNew(SHorizontalBox)
+		+ SHorizontalBox::Slot()
+		[
+			Viewport->AsWidget()
+		];
+
+	Viewports.Add( *ViewportKey, Viewport );
+
+	return ViewportBox.ToSharedRef();
+}
diff --git a/Engine/Source/Editor/UnrealEd/Private/EditorViewportLayoutThreePanes.cpp b/Engine/Source/Editor/UnrealEd/Private/EditorViewportLayoutThreePanes.cpp
new file mode 100644
--- /dev/null
+++ b/Engine/Source/Editor/UnrealEd/Private/EditorViewportLayoutThreePanes.cpp
@@ -0,0 +1,207 @@
+// Copyright 1998-2019 Epic Games, Inc. All Rights Reserved.
+
+#include "EditorViewportLayoutThreePanes.h"
+#include "Framework/Docking/LayoutService.h"
+#include "ShowFlags.h"
+#include "Misc/ConfigCacheIni.h"
+#include "Modules/ModuleManager.h"
+#include "Framework/Application/SlateApplication.h"
+#include "LevelEditor.h"
+
+namespace ViewportLayoutThreePanesDefs
+{
+	/** Default splitters to equal 50/50 split */
+	static const float DefaultSplitterPercentage = 0.5f;
+}
+
+
+// FEditorViewportLayoutThreePanes /////////////////////////////
+
+TSharedRef<SWidget> FEditorViewportLayoutThreePanes::MakeViewportLayout(TFunction<TSharedRef<SEditorViewport>(void)> &Func, const FString& LayoutString)
+{
+	FEngineShowFlags OrthoShowFlags(ESFIM_Editor);	
+	ApplyViewMode(VMI_BrushWireframe, false, OrthoShowFlags);
+
+	FEngineShowFlags PerspectiveShowFlags(ESFIM_Editor);	
+	ApplyViewMode(VMI_Lit, true, PerspectiveShowFlags);
+
+	FString ViewportKey0, ViewportKey1, ViewportKey2;
+	FString ViewportType0, ViewportType1, ViewportType2;
+	float PrimarySplitterPercentage = ViewportLayoutThreePanesDefs::DefaultSplitterPercentage;
+	float SecondarySplitterPercentage = ViewportLayoutThreePanesDefs::DefaultSplitterPercentage;
+
+	// Set up the viewports
+	FAssetEditorViewportConstructionArgs Args;
+	Args.ParentLayout = AsShared();
+	Args.IsEnabled = FSlateApplication::Get().GetNormalExecutionAttribute();
+
+	Args.bRealtime = !FPlatformMisc::IsRemoteSession();
+	Args.ConfigKey = *ViewportKey0;
+	Args.ViewportType = LVT_Perspective;
+	TSharedRef<IEditorViewportLayoutEntity> Viewport0 =FactoryViewport(Func, *ViewportType0, Args);
+
+	Args.bRealtime = false;
+	Args.ConfigKey = *ViewportKey1;
+	Args.ViewportType = LVT_OrthoXY;
+	TSharedRef<IEditorViewportLayoutEntity> Viewport1 =FactoryViewport(Func, *ViewportType1, Args);
+
+	// Front viewport
+	Args.bRealtime = false;
+	Args.ConfigKey = *ViewportKey2;
+	Args.ViewportType = LVT_OrthoXZ;
+	TSharedRef<IEditorViewportLayoutEntity> Viewport2 =FactoryViewport(Func, *ViewportType2, Args);
+
+	Viewports.Add( *ViewportKey0, Viewport0 );
+	Viewports.Add( *ViewportKey1, Viewport1 );
+	Viewports.Add( *ViewportKey2, Viewport2 );
+
+	TSharedRef<SWidget> LayoutWidget = MakeThreePanelWidget(
+		Viewports,
+		Viewport0->AsWidget(), Viewport1->AsWidget(), Viewport2->AsWidget(),
+		PrimarySplitterPercentage, SecondarySplitterPercentage);
+
+	return LayoutWidget;
+}
+
+// FEditorViewportLayoutThreePanesLeft /////////////////////////////
+
+TSharedRef<SWidget> FEditorViewportLayoutThreePanesLeft::MakeThreePanelWidget(
+	TMap< FName, TSharedPtr< IEditorViewportLayoutEntity >>& ViewportWidgets,
+	const TSharedRef<SWidget>& Viewport0, const TSharedRef<SWidget>& Viewport1, const TSharedRef<SWidget>& Viewport2,
+	float PrimarySplitterPercentage, float SecondarySplitterPercentage)
+{
+	TSharedRef<SWidget> Widget = 
+		SAssignNew( PrimarySplitterWidget, SSplitter )
+		.Orientation(EOrientation::Orient_Horizontal)
+		+SSplitter::Slot()
+		.Value(PrimarySplitterPercentage)
+		[
+			Viewport0
+		]
+		+SSplitter::Slot()
+		.Value(1.0f - PrimarySplitterPercentage)
+		[
+			SAssignNew( SecondarySplitterWidget, SSplitter )
+			.Orientation(EOrientation::Orient_Vertical)
+			+SSplitter::Slot()
+			.Value(SecondarySplitterPercentage)
+			[
+				Viewport1
+			]
+			+SSplitter::Slot()
+			.Value(1.0f - SecondarySplitterPercentage)
+			[
+				Viewport2
+			]
+		];
+
+	return Widget;
+}
+
+
+// FEditorViewportLayoutThreePanesRight /////////////////////////////
+
+TSharedRef<SWidget> FEditorViewportLayoutThreePanesRight::MakeThreePanelWidget(
+	TMap< FName, TSharedPtr< IEditorViewportLayoutEntity >>& ViewportWidgets,
+	const TSharedRef<SWidget>& Viewport0, const TSharedRef<SWidget>& Viewport1, const TSharedRef<SWidget>& Viewport2,
+	float PrimarySplitterPercentage, float SecondarySplitterPercentage)
+{
+	TSharedRef<SWidget> Widget = 
+		SAssignNew( PrimarySplitterWidget, SSplitter )
+		.Orientation(EOrientation::Orient_Horizontal)
+		+SSplitter::Slot()
+		.Value(PrimarySplitterPercentage)
+		[
+			SAssignNew( SecondarySplitterWidget, SSplitter )
+			.Orientation(EOrientation::Orient_Vertical)
+			+SSplitter::Slot()
+			.Value(SecondarySplitterPercentage)
+			[
+				Viewport1
+			]
+			+SSplitter::Slot()
+			.Value(1.0f - SecondarySplitterPercentage)
+			[
+				Viewport2
+			]
+		]
+		+SSplitter::Slot()
+		.Value(1.0f - PrimarySplitterPercentage)
+		[
+			Viewport0
+		];
+
+	return Widget;
+}
+
+
+// FEditorViewportLayoutThreePanesTop /////////////////////////////
+
+TSharedRef<SWidget> FEditorViewportLayoutThreePanesTop::MakeThreePanelWidget(
+	TMap< FName, TSharedPtr< IEditorViewportLayoutEntity >>& ViewportWidgets,
+	const TSharedRef<SWidget>& Viewport0, const TSharedRef<SWidget>& Viewport1, const TSharedRef<SWidget>& Viewport2,
+	float PrimarySplitterPercentage, float SecondarySplitterPercentage)
+{
+	TSharedRef<SWidget> Widget = 
+		SAssignNew( PrimarySplitterWidget, SSplitter )
+		.Orientation(EOrientation::Orient_Vertical)
+		+SSplitter::Slot()
+		.Value(PrimarySplitterPercentage)
+		[
+			Viewport0
+		]
+		+SSplitter::Slot()
+		.Value(1.0f - PrimarySplitterPercentage)
+		[
+			SAssignNew( SecondarySplitterWidget, SSplitter )
+			.Orientation(EOrientation::Orient_Horizontal)
+			+SSplitter::Slot()
+			.Value(SecondarySplitterPercentage)
+			[
+				Viewport1
+			]
+			+SSplitter::Slot()
+			.Value(1.0f - SecondarySplitterPercentage)
+			[
+				Viewport2
+			]
+		];
+
+	return Widget;
+}
+
+
+// FEditorViewportLayoutThreePanesBottom /////////////////////////////
+
+TSharedRef<SWidget> FEditorViewportLayoutThreePanesBottom::MakeThreePanelWidget(
+	TMap< FName, TSharedPtr< IEditorViewportLayoutEntity >>& ViewportWidgets,
+	const TSharedRef<SWidget>& Viewport0, const TSharedRef<SWidget>& Viewport1, const TSharedRef<SWidget>& Viewport2,
+	float PrimarySplitterPercentage, float SecondarySplitterPercentage)
+{
+	TSharedRef<SWidget> Widget = 
+		SAssignNew( PrimarySplitterWidget, SSplitter )
+		.Orientation(EOrientation::Orient_Vertical)
+		+SSplitter::Slot()
+		.Value(PrimarySplitterPercentage)
+		[
+			SAssignNew( SecondarySplitterWidget, SSplitter )
+			.Orientation(EOrientation::Orient_Horizontal)
+			+SSplitter::Slot()
+			.Value(SecondarySplitterPercentage)
+			[
+				Viewport1
+			]
+			+SSplitter::Slot()
+			.Value(1.0f - SecondarySplitterPercentage)
+			[
+				Viewport2
+			]
+		]
+		+SSplitter::Slot()
+		.Value(1.0f - PrimarySplitterPercentage)
+		[
+			Viewport0
+		];
+
+	return Widget;
+}
diff --git a/Engine/Source/Editor/UnrealEd/Private/EditorViewportLayoutTwoPanes.cpp b/Engine/Source/Editor/UnrealEd/Private/EditorViewportLayoutTwoPanes.cpp
new file mode 100644
--- /dev/null
+++ b/Engine/Source/Editor/UnrealEd/Private/EditorViewportLayoutTwoPanes.cpp
@@ -0,0 +1,71 @@
+// Copyright 1998-2019 Epic Games, Inc. All Rights Reserved.
+
+#include "EditorViewportLayoutTwoPanes.h"
+#include "Framework/Docking/LayoutService.h"
+#include "Editor.h"
+#include "Misc/ConfigCacheIni.h"
+#include "Modules/ModuleManager.h"
+#include "Framework/Application/SlateApplication.h"
+
+
+namespace ViewportLayoutTwoPanesDefs
+{
+	/** Default splitters to equal 50/50 split */
+	static const float DefaultSplitterPercentage = 0.5f;
+}
+
+
+template <EOrientation TOrientation>
+TSharedRef<SWidget> TEditorViewportLayoutTwoPanes<TOrientation>::MakeViewportLayout(TFunction<TSharedRef<SEditorViewport>(void)> &Func, const FString& LayoutString)
+{
+	FString ViewportKey0, ViewportKey1;
+	FString ViewportType0, ViewportType1;
+	float PrimarySplitterPercentage = ViewportLayoutTwoPanesDefs::DefaultSplitterPercentage;
+
+	// Set up the viewports
+	FAssetEditorViewportConstructionArgs Args;
+	Args.ParentLayout = AsShared();
+	Args.IsEnabled = FSlateApplication::Get().GetNormalExecutionAttribute();
+
+	Args.bRealtime = false;
+	Args.ConfigKey = *ViewportKey0;
+	Args.ViewportType = LVT_OrthoXY;
+	TSharedRef<IEditorViewportLayoutEntity> Viewport0 = FactoryViewport(Func, *ViewportType0, Args);
+
+	Args.bRealtime = !FPlatformMisc::IsRemoteSession();
+	Args.ConfigKey = *ViewportKey1;
+	Args.ViewportType = LVT_Perspective;
+	TSharedRef<IEditorViewportLayoutEntity> Viewport1 = FactoryViewport(Func, *ViewportType1, Args);
+
+	Viewports.Add(*ViewportKey0, Viewport0);
+	Viewports.Add(*ViewportKey1, Viewport1);
+
+	SplitterWidget =
+		SNew(SSplitter)
+		.Orientation(TOrientation)
+		+ SSplitter::Slot()
+		.Value(PrimarySplitterPercentage)
+		[
+			Viewport0->AsWidget()
+		]
+	+ SSplitter::Slot()
+		.Value(1.0f - PrimarySplitterPercentage)
+		[
+			Viewport1->AsWidget()
+		];
+
+	return SplitterWidget.ToSharedRef();
+}
+
+
+/**
+* Function avoids linker errors on the template class functions in this cpp file.
+* It avoids the need to put the contents of this file into the header.
+* It doesn't get called.
+*/
+void EditorViewportLayoutTwoPanes_LinkErrorFixFunc()
+{
+	check(0);
+	FEditorViewportLayoutTwoPanesVert DummyVert;
+	FEditorViewportLayoutTwoPanesHoriz DummyHoriz;
+}
diff --git a/Engine/Source/Editor/UnrealEd/Private/EditorViewportTabContent.cpp b/Engine/Source/Editor/UnrealEd/Private/EditorViewportTabContent.cpp
new file mode 100644
--- /dev/null
+++ b/Engine/Source/Editor/UnrealEd/Private/EditorViewportTabContent.cpp
@@ -0,0 +1,98 @@
+
+#include "EditorViewportTabContent.h"
+#include "SEditorViewport.h"
+#include "EditorViewportLayoutOnePane.h"
+#include "EditorViewportLayoutTwoPanes.h"
+#include "EditorViewportLayoutThreePanes.h"
+#include "EditorViewportLayoutFourPanes.h"
+#include "EditorViewportLayout2x2.h"
+
+
+
+
+TSharedPtr< class FEditorViewportLayout > FEditorViewportTabContent::ConstructViewportLayoutByTypeName(const FName& TypeName, bool bSwitchingLayouts)
+{
+	TSharedPtr< class FEditorViewportLayout > ViewportLayout;
+
+	//The items in these ifs should match the names in namespace EditorViewportConfigurationNames
+	if (TypeName == EditorViewportConfigurationNames::TwoPanesHoriz) ViewportLayout = MakeShareable(new FEditorViewportLayoutTwoPanesHoriz);
+	else if (TypeName == EditorViewportConfigurationNames::TwoPanesVert) ViewportLayout = MakeShareable(new FEditorViewportLayoutTwoPanesVert);
+ 	else if (TypeName == EditorViewportConfigurationNames::FourPanes2x2) ViewportLayout = MakeShareable(new FEditorViewportLayout2x2);
+	else if (TypeName == EditorViewportConfigurationNames::ThreePanesLeft) ViewportLayout = MakeShareable(new FEditorViewportLayoutThreePanesLeft);
+	else if (TypeName == EditorViewportConfigurationNames::ThreePanesRight) ViewportLayout = MakeShareable(new FEditorViewportLayoutThreePanesRight);
+	else if (TypeName == EditorViewportConfigurationNames::ThreePanesTop) ViewportLayout = MakeShareable(new FEditorViewportLayoutThreePanesTop);
+	else if (TypeName == EditorViewportConfigurationNames::ThreePanesBottom) ViewportLayout = MakeShareable(new FEditorViewportLayoutThreePanesBottom);
+	else if (TypeName == EditorViewportConfigurationNames::FourPanesLeft) ViewportLayout = MakeShareable(new FEditorViewportLayoutFourPanesLeft);
+	else if (TypeName == EditorViewportConfigurationNames::FourPanesRight) ViewportLayout = MakeShareable(new FEditorViewportLayoutFourPanesRight);
+	else if (TypeName == EditorViewportConfigurationNames::FourPanesBottom) ViewportLayout = MakeShareable(new FEditorViewportLayoutFourPanesBottom);
+  	else if (TypeName == EditorViewportConfigurationNames::FourPanesTop) ViewportLayout = MakeShareable(new FEditorViewportLayoutFourPanesTop);
+	else /*(TypeName == EditorViewportConfigurationNames::OnePane)*/ ViewportLayout = MakeShareable(new FEditorViewportLayoutOnePane);
+
+	if (!ensure(ViewportLayout.IsValid()))
+	{
+		ViewportLayout = MakeShareable(new FEditorViewportLayoutOnePane);
+	}
+	return ViewportLayout;
+}
+
+void FEditorViewportTabContent::Initialize(TFunction<TSharedRef<SEditorViewport>(void)> Func, TSharedPtr<SDockTab> InParentTab, const FString& InLayoutString)
+{
+	ViewportCreationFunc = Func;
+	ParentTab = InParentTab;
+	LayoutString = InLayoutString;
+
+	FName LayoutType(*LayoutString);
+	SetViewportConfiguration(Func, LayoutType);
+}
+
+void FEditorViewportTabContent::SetViewportConfiguration(TFunction<TSharedRef<SEditorViewport>(void)> &Func, const FName& ConfigurationName)
+{
+	bool bSwitchingLayouts = ActiveViewportLayout.IsValid();
+
+	if (bSwitchingLayouts)
+	{
+		ActiveViewportLayout.Reset();
+	}
+
+	ActiveViewportLayout = ConstructViewportLayoutByTypeName(ConfigurationName, bSwitchingLayouts);
+	check(ActiveViewportLayout.IsValid());
+
+	UpdateViewportTabWidget(Func);
+}
+
+void FEditorViewportTabContent::SetViewportConfiguration(const FName& ConfigurationName)
+{
+	bool bSwitchingLayouts = ActiveViewportLayout.IsValid();
+
+	if (bSwitchingLayouts)
+	{
+		ActiveViewportLayout.Reset();
+	}
+
+	ActiveViewportLayout = ConstructViewportLayoutByTypeName(ConfigurationName, bSwitchingLayouts);
+	check(ActiveViewportLayout.IsValid());
+
+  	UpdateViewportTabWidget(ViewportCreationFunc);
+
+}
+
+void FEditorViewportTabContent::UpdateViewportTabWidget(TFunction<TSharedRef<SEditorViewport>(void)> &Func)
+{
+	TSharedPtr<SDockTab> ParentTabPinned = ParentTab.Pin();
+	if (ParentTabPinned.IsValid() && ActiveViewportLayout.IsValid())
+	{
+ 		TSharedRef<SWidget> LayoutWidget = StaticCastSharedPtr<FAssetEditorViewportLayout>(ActiveViewportLayout)->BuildViewportLayout(Func, ParentTabPinned, SharedThis(this), LayoutString);
+ 		ParentTabPinned->SetContent(LayoutWidget);
+
+		if (PreviouslyFocusedViewport.IsSet())
+		{
+			TSharedPtr<IEditorViewportLayoutEntity> ViewportToFocus = ActiveViewportLayout->GetViewports().FindRef(PreviouslyFocusedViewport.GetValue());
+			if (ViewportToFocus.IsValid())
+			{
+				ViewportToFocus->SetKeyboardFocus();
+			}
+			PreviouslyFocusedViewport = TOptional<FName>();
+		}
+	}
+}
+
diff --git a/Engine/Source/Editor/UnrealEd/Private/SAssetEditorViewport.cpp b/Engine/Source/Editor/UnrealEd/Private/SAssetEditorViewport.cpp
new file mode 100644
--- /dev/null
+++ b/Engine/Source/Editor/UnrealEd/Private/SAssetEditorViewport.cpp
@@ -0,0 +1,121 @@
+// Copyright 1998-2019 Epic Games, Inc. All Rights Reserved.
+
+#include "SAssetEditorViewport.h"
+#include "Widgets/DeclarativeSyntaxSupport.h"
+#include "Editor/UnrealEdEngine.h"
+#include "UnrealEdGlobals.h"
+#include "EditorViewportCommands.h"
+#include "EditorViewportTabContent.h"
+
+
+void SAssetEditorViewport::BindCommands()
+{
+	FUICommandList& CommandListRef = *CommandList;
+
+	const FEditorViewportCommands& Commands = FEditorViewportCommands::Get();
+
+	TSharedRef<FEditorViewportClient> ClientRef = Client.ToSharedRef();
+
+	CommandListRef.MapAction(
+		Commands.ViewportConfig_OnePane,
+		FExecuteAction::CreateSP(this, &SAssetEditorViewport::OnSetViewportConfiguration, EditorViewportConfigurationNames::OnePane),
+		FCanExecuteAction(),
+		FIsActionChecked::CreateSP(this, &SAssetEditorViewport::IsViewportConfigurationSet, EditorViewportConfigurationNames::OnePane));
+	
+	CommandListRef.MapAction(
+		Commands.ViewportConfig_TwoPanesH,
+		FExecuteAction::CreateSP(this, &SAssetEditorViewport::OnSetViewportConfiguration, EditorViewportConfigurationNames::TwoPanesHoriz),
+		FCanExecuteAction(),
+		FIsActionChecked::CreateSP(this, &SAssetEditorViewport::IsViewportConfigurationSet, EditorViewportConfigurationNames::TwoPanesHoriz));
+
+	CommandListRef.MapAction(
+		Commands.ViewportConfig_TwoPanesV,
+		FExecuteAction::CreateSP(this, &SAssetEditorViewport::OnSetViewportConfiguration, EditorViewportConfigurationNames::TwoPanesVert),
+		FCanExecuteAction(),
+		FIsActionChecked::CreateSP(this, &SAssetEditorViewport::IsViewportConfigurationSet, EditorViewportConfigurationNames::TwoPanesVert));
+
+	CommandListRef.MapAction(
+		Commands.ViewportConfig_ThreePanesLeft,
+		FExecuteAction::CreateSP(this, &SAssetEditorViewport::OnSetViewportConfiguration, EditorViewportConfigurationNames::ThreePanesLeft),
+		FCanExecuteAction(),
+		FIsActionChecked::CreateSP(this, &SAssetEditorViewport::IsViewportConfigurationSet, EditorViewportConfigurationNames::ThreePanesLeft));
+
+	CommandListRef.MapAction(
+		Commands.ViewportConfig_ThreePanesRight,
+		FExecuteAction::CreateSP(this, &SAssetEditorViewport::OnSetViewportConfiguration, EditorViewportConfigurationNames::ThreePanesRight),
+		FCanExecuteAction(),
+		FIsActionChecked::CreateSP(this, &SAssetEditorViewport::IsViewportConfigurationSet, EditorViewportConfigurationNames::ThreePanesRight));
+
+	CommandListRef.MapAction(
+		Commands.ViewportConfig_ThreePanesTop,
+		FExecuteAction::CreateSP(this, &SAssetEditorViewport::OnSetViewportConfiguration, EditorViewportConfigurationNames::ThreePanesTop),
+		FCanExecuteAction(),
+		FIsActionChecked::CreateSP(this, &SAssetEditorViewport::IsViewportConfigurationSet, EditorViewportConfigurationNames::ThreePanesTop));
+
+	CommandListRef.MapAction(
+		Commands.ViewportConfig_ThreePanesBottom,
+		FExecuteAction::CreateSP(this, &SAssetEditorViewport::OnSetViewportConfiguration, EditorViewportConfigurationNames::ThreePanesBottom),
+		FCanExecuteAction(),
+		FIsActionChecked::CreateSP(this, &SAssetEditorViewport::IsViewportConfigurationSet, EditorViewportConfigurationNames::ThreePanesBottom));
+
+	CommandListRef.MapAction(
+		Commands.ViewportConfig_FourPanesLeft,
+		FExecuteAction::CreateSP(this, &SAssetEditorViewport::OnSetViewportConfiguration, EditorViewportConfigurationNames::FourPanesLeft),
+		FCanExecuteAction(),
+		FIsActionChecked::CreateSP(this, &SAssetEditorViewport::IsViewportConfigurationSet, EditorViewportConfigurationNames::FourPanesLeft));
+
+	CommandListRef.MapAction(
+		Commands.ViewportConfig_FourPanesRight,
+		FExecuteAction::CreateSP(this, &SAssetEditorViewport::OnSetViewportConfiguration, EditorViewportConfigurationNames::FourPanesRight),
+		FCanExecuteAction(),
+		FIsActionChecked::CreateSP(this, &SAssetEditorViewport::IsViewportConfigurationSet, EditorViewportConfigurationNames::FourPanesRight));
+
+	CommandListRef.MapAction(
+		Commands.ViewportConfig_FourPanesTop,
+		FExecuteAction::CreateSP(this, &SAssetEditorViewport::OnSetViewportConfiguration, EditorViewportConfigurationNames::FourPanesTop),
+		FCanExecuteAction(),
+		FIsActionChecked::CreateSP(this, &SAssetEditorViewport::IsViewportConfigurationSet, EditorViewportConfigurationNames::FourPanesTop));
+
+	CommandListRef.MapAction(
+		Commands.ViewportConfig_FourPanesBottom,
+		FExecuteAction::CreateSP(this, &SAssetEditorViewport::OnSetViewportConfiguration, EditorViewportConfigurationNames::FourPanesBottom),
+		FCanExecuteAction(),
+		FIsActionChecked::CreateSP(this, &SAssetEditorViewport::IsViewportConfigurationSet, EditorViewportConfigurationNames::FourPanesBottom));
+
+	CommandListRef.MapAction(
+		Commands.ViewportConfig_FourPanes2x2,
+		FExecuteAction::CreateSP(this, &SAssetEditorViewport::OnSetViewportConfiguration, EditorViewportConfigurationNames::FourPanes2x2),
+		FCanExecuteAction(),
+		FIsActionChecked::CreateSP(this, &SAssetEditorViewport::IsViewportConfigurationSet, EditorViewportConfigurationNames::FourPanes2x2));
+
+	SEditorViewport::BindCommands();
+}
+
+
+void SAssetEditorViewport::OnSetViewportConfiguration(FName ConfigurationName)
+{
+	TSharedPtr<FAssetEditorViewportLayout> LayoutPinned = ParentLayout.Pin();
+	if (LayoutPinned.IsValid())
+	{
+		TSharedPtr<FViewportTabContent> ViewportTabPinned = LayoutPinned->GetParentTabContent().Pin();
+		if (ViewportTabPinned.IsValid())
+		{
+			ViewportTabPinned->SetViewportConfiguration(ConfigurationName);
+ 			FSlateApplication::Get().DismissAllMenus();
+		}
+	}
+}
+
+bool SAssetEditorViewport::IsViewportConfigurationSet(FName ConfigurationName) const
+{
+	TSharedPtr<FAssetEditorViewportLayout> LayoutPinned = ParentLayout.Pin();
+	if (LayoutPinned.IsValid())
+	{
+		TSharedPtr<FViewportTabContent> ViewportTabPinned = LayoutPinned->GetParentTabContent().Pin();
+		if (ViewportTabPinned.IsValid())
+		{
+			return ViewportTabPinned->IsViewportConfigurationSet(ConfigurationName);
+		}
+	}
+	return false;
+}
diff --git a/Engine/Source/Editor/UnrealEd/Private/ViewportTabContent.cpp b/Engine/Source/Editor/UnrealEd/Private/ViewportTabContent.cpp
--- a/Engine/Source/Editor/UnrealEd/Private/ViewportTabContent.cpp
+++ b/Engine/Source/Editor/UnrealEd/Private/ViewportTabContent.cpp
@@ -14,3 +14,12 @@ bool FViewportTabContent::BelongsToTab(TSharedRef<class SDockTab> InParentTab) c
 	TSharedPtr<SDockTab> ParentTabPinned = ParentTab.Pin();
 	return ParentTabPinned == InParentTab;
 }
+
+bool FViewportTabContent::IsViewportConfigurationSet(const FName& ConfigurationName) const
+{
+	if (ActiveViewportLayout.IsValid())
+	{
+		return ActiveViewportLayout->GetLayoutTypeName() == ConfigurationName;
+	}
+	return false;
+}
diff --git a/Engine/Source/Editor/UnrealEd/Public/AssetEditorViewportLayout.h b/Engine/Source/Editor/UnrealEd/Public/AssetEditorViewportLayout.h
new file mode 100644
--- /dev/null
+++ b/Engine/Source/Editor/UnrealEd/Public/AssetEditorViewportLayout.h
@@ -0,0 +1,115 @@
+// Copyright 1998-2019 Epic Games, Inc. All Rights Reserved.
+
+#pragma once
+
+#include "CoreMinimal.h"
+#include "SlateFwd.h"
+#include "Stats/Stats.h"
+#include "Misc/Attribute.h"
+#include "Animation/CurveSequence.h"
+#include "Layout/Visibility.h"
+#include "Widgets/SWidget.h"
+#include "Editor/UnrealEdTypes.h"
+#include "Application/ThrottleManager.h"
+#include "EditorViewportLayout.h"
+#include "TickableEditorObject.h"
+
+class FAssetEditorViewportLayout;
+class SViewportsOverlay;
+class SWindow;
+class SAssetEditorViewport;
+
+/** Arguments for constructing a viewport */
+struct FAssetEditorViewportConstructionArgs
+{
+	FAssetEditorViewportConstructionArgs()
+		: ViewportType(LVT_Perspective)
+		, bRealtime(false)
+	{}
+
+	/** The viewport's parent layout */
+	TSharedPtr<FAssetEditorViewportLayout> ParentLayout;
+	/** The viewport's desired type */
+	ELevelViewportType ViewportType;
+	/** Whether the viewport should default to realtime */
+	bool bRealtime;
+	/** A config key for loading/saving settings for the viewport */
+	FName ConfigKey;
+	/** Widget enabled attribute */
+	TAttribute<bool> IsEnabled;
+};
+
+
+
+namespace EditorViewportConfigurationNames
+{
+	static FName TwoPanesHoriz("TwoPanesHoriz");
+	static FName TwoPanesVert("TwoPanesVert");
+	static FName ThreePanesLeft("ThreePanesLeft");
+	static FName ThreePanesRight("ThreePanesRight");
+	static FName ThreePanesTop("ThreePanesTop");
+	static FName ThreePanesBottom("ThreePanesBottom");
+	static FName FourPanesLeft("FourPanesLeft");
+	static FName FourPanesRight("FourPanesRight");
+	static FName FourPanesTop("FourPanesTop");
+	static FName FourPanesBottom("FourPanesBottom");
+	static FName FourPanes2x2("FourPanes2x2");
+	static FName OnePane("OnePane");
+}
+/**
+ * Base class for viewport layout configurations
+ * Handles maximizing and restoring well as visibility of specific viewports.
+ */
+class UNREALED_API FAssetEditorViewportLayout : public TSharedFromThis<FAssetEditorViewportLayout>, public FEditorViewportLayout, public FTickableEditorObject
+{
+public:
+	/**
+	 * Constructor
+	 */
+	FAssetEditorViewportLayout();
+
+	/**
+	 * Destructor
+	 */
+	virtual ~FAssetEditorViewportLayout();
+
+	/** Create an instance of a custom viewport from the specified viewport type name */
+	TSharedRef<IEditorViewportLayoutEntity> FactoryViewport(TFunction<TSharedRef<SEditorViewport>(void)> &Func, FName InTypeName, const FAssetEditorViewportConstructionArgs& ConstructionArgs) const;
+
+	/** FTickableEditorObject interface */
+	virtual void Tick(float DeltaTime) override {}
+	virtual bool IsTickable() const override { return false; }
+ 	virtual TStatId GetStatId() const override { return TStatId(); }
+
+	/**
+	 * Builds a viewport layout and returns the widget containing the layout
+	 * 
+	 * @param InParentDockTab		The parent dock tab widget of this viewport configuration
+	 * @param InParentTab			The parent tab object
+	 * @param LayoutString			The layout string loaded from file to custom build the layout with
+	 */
+ 	TSharedRef<SWidget> BuildViewportLayout(TFunction<TSharedRef<SEditorViewport>(void)> &Func,  TSharedPtr<SDockTab> InParentDockTab, TSharedPtr<class FViewportTabContent> InParentTab, const FString& LayoutString );
+
+	static TSharedRef<SWidget> GenerateLayoutMenu(TSharedPtr<SAssetEditorViewport> AssetEditorViewport) ;
+
+	/** Returns the parent tab content object */
+	TWeakPtr< class FViewportTabContent > GetParentTabContent() const { return ParentTabContent; }
+
+	/**
+	 * Overridden in derived classes to set up custom layouts  
+	 *
+	 * @param LayoutString		The layout string loaded from a file
+	 * @return The base widget representing the layout.  Usually a splitter
+	 */
+  	virtual TSharedRef<SWidget> MakeViewportLayout(TFunction<TSharedRef<SEditorViewport>(void)> &Func, const FString& LayoutString) = 0;
+
+	/** The overlay widget that handles what viewports should be on top (non-maximized or maximized) */
+	TWeakPtr< class SViewportsOverlay > ViewportsOverlayPtr;
+
+	/** The parent tab content object where this layout resides */
+	TWeakPtr< class FViewportTabContent > ParentTabContent;
+
+	/** The parent tab where this layout resides */
+	TWeakPtr< SDockTab > ParentTab;
+
+};
diff --git a/Engine/Source/Editor/UnrealEd/Public/EditorViewportCommands.h b/Engine/Source/Editor/UnrealEd/Public/EditorViewportCommands.h
--- a/Engine/Source/Editor/UnrealEd/Public/EditorViewportCommands.h
+++ b/Engine/Source/Editor/UnrealEd/Public/EditorViewportCommands.h
@@ -29,6 +29,21 @@ public:
 	{
 	}
 	
+	/** Viewport pane configurations */
+	TSharedPtr< FUICommandInfo > ViewportConfig_OnePane;
+	TSharedPtr< FUICommandInfo > ViewportConfig_TwoPanesH;
+	TSharedPtr< FUICommandInfo > ViewportConfig_TwoPanesV;
+	TSharedPtr< FUICommandInfo > ViewportConfig_ThreePanesLeft;
+	TSharedPtr< FUICommandInfo > ViewportConfig_ThreePanesRight;
+	TSharedPtr< FUICommandInfo > ViewportConfig_ThreePanesTop;
+	TSharedPtr< FUICommandInfo > ViewportConfig_ThreePanesBottom;
+	TSharedPtr< FUICommandInfo > ViewportConfig_FourPanesLeft;
+	TSharedPtr< FUICommandInfo > ViewportConfig_FourPanesRight;
+	TSharedPtr< FUICommandInfo > ViewportConfig_FourPanesTop;
+	TSharedPtr< FUICommandInfo > ViewportConfig_FourPanesBottom;
+	TSharedPtr< FUICommandInfo > ViewportConfig_FourPanes2x2;
+
+
 	/** Changes the viewport to perspective view */
 	TSharedPtr< FUICommandInfo > Perspective;
 
diff --git a/Engine/Source/Editor/UnrealEd/Public/EditorViewportLayout.h b/Engine/Source/Editor/UnrealEd/Public/EditorViewportLayout.h
--- a/Engine/Source/Editor/UnrealEd/Public/EditorViewportLayout.h
+++ b/Engine/Source/Editor/UnrealEd/Public/EditorViewportLayout.h
@@ -44,7 +44,7 @@ public:
 	/**
 	* Saves viewport layout information between editor sessions
 	*/
-	virtual void SaveLayoutString(const FString& LayoutString) const = 0;
+	virtual void SaveLayoutString(const FString& LayoutString) const {}
 
 	virtual const FName& GetLayoutTypeName() const = 0;
 
diff --git a/Engine/Source/Editor/UnrealEd/Public/EditorViewportLayout2x2.h b/Engine/Source/Editor/UnrealEd/Public/EditorViewportLayout2x2.h
new file mode 100644
--- /dev/null
+++ b/Engine/Source/Editor/UnrealEd/Public/EditorViewportLayout2x2.h
@@ -0,0 +1,24 @@
+// Copyright 1998-2019 Epic Games, Inc. All Rights Reserved.
+
+#pragma once
+
+#include "CoreMinimal.h"
+#include "Widgets/SWidget.h"
+#include "Widgets/Layout/SSplitter.h"
+#include "AssetEditorViewportLayout.h"
+
+class FEditorViewportLayout2x2 : public FAssetEditorViewportLayout
+{
+public:
+
+	virtual const FName& GetLayoutTypeName() const override { return EditorViewportConfigurationNames::FourPanes2x2; }
+protected:
+	/**
+	 * Creates the viewports and splitter for the 2x2 layout                   
+	 */
+	virtual TSharedRef<SWidget> MakeViewportLayout(TFunction<TSharedRef<SEditorViewport>(void)> &Func, const FString& LayoutString) override;
+
+private:
+	/** The splitter widget */
+	TSharedPtr< class SSplitter2x2 > SplitterWidget;
+};
diff --git a/Engine/Source/Editor/UnrealEd/Public/EditorViewportLayoutFourPanes.h b/Engine/Source/Editor/UnrealEd/Public/EditorViewportLayoutFourPanes.h
new file mode 100644
--- /dev/null
+++ b/Engine/Source/Editor/UnrealEd/Public/EditorViewportLayoutFourPanes.h
@@ -0,0 +1,89 @@
+// Copyright 1998-2019 Epic Games, Inc. All Rights Reserved.
+
+#pragma once
+
+#include "CoreMinimal.h"
+#include "Widgets/SWidget.h"
+#include "EditorViewportLayout.h"
+#include "Widgets/Layout/SSplitter.h"
+#include "AssetEditorViewportLayout.h"
+
+class FEditorViewportLayoutFourPanes : public FAssetEditorViewportLayout
+{
+public:
+	/**
+	 * Saves viewport layout information between editor sessions
+	 */
+// 	virtual void SaveLayoutString(const FString& LayoutString) const override;
+protected:
+	/**
+	 * Creates the viewports and splitter for the four-pane layout              
+	 */
+	virtual TSharedRef<SWidget> MakeViewportLayout(TFunction<TSharedRef<SEditorViewport>(void)> &Func, const FString& LayoutString) override;
+
+	virtual TSharedRef<SWidget> MakeFourPanelWidget(
+		TMap<FName, TSharedPtr< IEditorViewportLayoutEntity >>& ViewportWidgets,
+		TSharedRef<SWidget> Viewport0, TSharedRef<SWidget> Viewport1, TSharedRef<SWidget> Viewport2, TSharedRef<SWidget> Viewport3,
+		float PrimarySplitterPercentage, float SecondarySplitterPercentage0, float SecondarySplitterPercentage1) = 0;
+
+protected:
+	/** The splitter widgets */
+	TSharedPtr< class SSplitter > PrimarySplitterWidget;
+	TSharedPtr< class SSplitter > SecondarySplitterWidget;
+};
+
+
+// FEditorlViewportLayoutFourPanesLeft /////////////////////////////
+
+class FEditorViewportLayoutFourPanesLeft : public FEditorViewportLayoutFourPanes
+{
+public:
+	virtual const FName& GetLayoutTypeName() const override { return EditorViewportConfigurationNames::FourPanesLeft; }
+
+	virtual TSharedRef<SWidget> MakeFourPanelWidget(
+		TMap<FName, TSharedPtr< IEditorViewportLayoutEntity >>& ViewportWidgets,
+		TSharedRef<SWidget> Viewport0, TSharedRef<SWidget> Viewport1, TSharedRef<SWidget> Viewport2, TSharedRef<SWidget> Viewport3,
+		float PrimarySplitterPercentage, float SecondarySplitterPercentage0, float SecondarySplitterPercentage1) override;
+};
+
+
+// FEditorViewportLayoutFourPanesRight /////////////////////////////
+
+class FEditorViewportLayoutFourPanesRight : public FEditorViewportLayoutFourPanes
+{
+public:
+	virtual const FName& GetLayoutTypeName() const override { return EditorViewportConfigurationNames::FourPanesRight; }
+
+	virtual TSharedRef<SWidget> MakeFourPanelWidget(
+		TMap<FName, TSharedPtr< IEditorViewportLayoutEntity >>& ViewportWidgets,
+		TSharedRef<SWidget> Viewport0, TSharedRef<SWidget> Viewport1, TSharedRef<SWidget> Viewport2, TSharedRef<SWidget> Viewport3,
+		float PrimarySplitterPercentage, float SecondarySplitterPercentage0, float SecondarySplitterPercentage1) override;
+};
+
+
+// FEditorViewportLayoutFourPanesTop /////////////////////////////
+
+class FEditorViewportLayoutFourPanesTop : public FEditorViewportLayoutFourPanes
+{
+public:
+	virtual const FName& GetLayoutTypeName() const override { return EditorViewportConfigurationNames::FourPanesTop; }
+
+	virtual TSharedRef<SWidget> MakeFourPanelWidget(
+		TMap<FName, TSharedPtr< IEditorViewportLayoutEntity >>& ViewportWidgets,
+		TSharedRef<SWidget> Viewport0, TSharedRef<SWidget> Viewport1, TSharedRef<SWidget> Viewport2, TSharedRef<SWidget> Viewport3,
+		float PrimarySplitterPercentage, float SecondarySplitterPercentage0, float SecondarySplitterPercentage1) override;
+};
+
+
+// FEditorViewportLayoutFourPanesBottom /////////////////////////////
+
+class FEditorViewportLayoutFourPanesBottom : public FEditorViewportLayoutFourPanes
+{
+public:
+	virtual const FName& GetLayoutTypeName() const override { return EditorViewportConfigurationNames::FourPanesBottom; }
+
+	virtual TSharedRef<SWidget> MakeFourPanelWidget(
+		TMap<FName, TSharedPtr< IEditorViewportLayoutEntity >>& ViewportWidgets,
+		TSharedRef<SWidget> Viewport0, TSharedRef<SWidget> Viewport1, TSharedRef<SWidget> Viewport2, TSharedRef<SWidget> Viewport3,
+		float PrimarySplitterPercentage, float SecondarySplitterPercentage0, float SecondarySplitterPercentage1) override;
+};
diff --git a/Engine/Source/Editor/UnrealEd/Public/EditorViewportLayoutOnePane.h b/Engine/Source/Editor/UnrealEd/Public/EditorViewportLayoutOnePane.h
new file mode 100644
--- /dev/null
+++ b/Engine/Source/Editor/UnrealEd/Public/EditorViewportLayoutOnePane.h
@@ -0,0 +1,25 @@
+// Copyright 1998-2019 Epic Games, Inc. All Rights Reserved.
+
+#pragma once
+
+#include "CoreMinimal.h"
+#include "Widgets/SWidget.h"
+#include "EditorViewportLayout.h"
+#include "AssetEditorViewportLayout.h"
+
+class SHorizontalBox;
+
+class FEditorViewportLayoutOnePane : public FAssetEditorViewportLayout
+{
+public:
+	virtual const FName& GetLayoutTypeName() const override{ return EditorViewportConfigurationNames::OnePane; }
+protected:
+	/**
+	* Creates the viewport for the single pane
+	*/
+	virtual TSharedRef<SWidget> MakeViewportLayout(TFunction<TSharedRef<SEditorViewport>(void)> &Func, const FString& LayoutString) override;
+
+protected:
+	/** The viewport widget parent box */
+	TSharedPtr< SHorizontalBox > ViewportBox;
+};
diff --git a/Engine/Source/Editor/UnrealEd/Public/EditorViewportLayoutThreePanes.h b/Engine/Source/Editor/UnrealEd/Public/EditorViewportLayoutThreePanes.h
new file mode 100644
--- /dev/null
+++ b/Engine/Source/Editor/UnrealEd/Public/EditorViewportLayoutThreePanes.h
@@ -0,0 +1,83 @@
+// Copyright 1998-2019 Epic Games, Inc. All Rights Reserved.
+
+#pragma once
+
+#include "CoreMinimal.h"
+#include "Widgets/SWidget.h"
+#include "EditorViewportLayout.h"
+#include "Widgets/Layout/SSplitter.h"
+#include "AssetEditorViewportLayout.h"
+
+class FEditorViewportLayoutThreePanes : public FAssetEditorViewportLayout
+{
+protected:
+	/**
+	 * Creates the viewports and splitter for the two panes vertical layout                   
+	 */
+	virtual TSharedRef<SWidget> MakeViewportLayout(TFunction<TSharedRef<SEditorViewport>(void)> &Func, const FString& LayoutString) override;
+
+	virtual TSharedRef<SWidget> MakeThreePanelWidget(
+		TMap< FName, TSharedPtr< IEditorViewportLayoutEntity >>& ViewportWidgets,
+		const TSharedRef<SWidget>& ViewportKey0, const TSharedRef<SWidget>& ViewportKey1, const TSharedRef<SWidget>& ViewportKey2,
+		float PrimarySplitterPercentage, float SecondarySplitterPercentage) = 0;
+
+protected:
+	/** The splitter widgets */
+	TSharedPtr< class SSplitter > PrimarySplitterWidget;
+	TSharedPtr< class SSplitter > SecondarySplitterWidget;
+};
+
+// FEditorViewportLayoutThreePanesLeft /////////////////////////////
+
+class FEditorViewportLayoutThreePanesLeft : public FEditorViewportLayoutThreePanes
+{
+public:
+	virtual const FName& GetLayoutTypeName() const override { return EditorViewportConfigurationNames::ThreePanesLeft; }
+
+	virtual TSharedRef<SWidget> MakeThreePanelWidget(
+		TMap< FName, TSharedPtr< IEditorViewportLayoutEntity >>& ViewportWidgets,
+		const TSharedRef<SWidget>& ViewportKey0, const TSharedRef<SWidget>& ViewportKey1, const TSharedRef<SWidget>& ViewportKey2,
+		float PrimarySplitterPercentage, float SecondarySplitterPercentage) override;
+};
+
+
+// FEditorViewportLayoutThreePanesRight /////////////////////////////
+
+class FEditorViewportLayoutThreePanesRight : public FEditorViewportLayoutThreePanes
+{
+public:
+	virtual const FName& GetLayoutTypeName() const override { return EditorViewportConfigurationNames::ThreePanesRight; }
+
+	virtual TSharedRef<SWidget> MakeThreePanelWidget(
+		TMap< FName, TSharedPtr< IEditorViewportLayoutEntity >>& ViewportWidgets,
+		const TSharedRef<SWidget>& ViewportKey0, const TSharedRef<SWidget>& ViewportKey1, const TSharedRef<SWidget>& ViewportKey2,
+		float PrimarySplitterPercentage, float SecondarySplitterPercentage) override;
+};
+
+
+// FEditorViewportLayoutThreePanesTop /////////////////////////////
+
+class FEditorViewportLayoutThreePanesTop : public FEditorViewportLayoutThreePanes
+{
+public:
+	virtual const FName& GetLayoutTypeName() const override { return EditorViewportConfigurationNames::ThreePanesTop; }
+
+	virtual TSharedRef<SWidget> MakeThreePanelWidget(
+		TMap< FName, TSharedPtr< IEditorViewportLayoutEntity >>& ViewportWidgets,
+		const TSharedRef<SWidget>& ViewportKey0, const TSharedRef<SWidget>& ViewportKey1, const TSharedRef<SWidget>& ViewportKey2,
+		float PrimarySplitterPercentage, float SecondarySplitterPercentage) override;
+};
+
+
+// FEditorViewportLayoutThreePanesBottom /////////////////////////////
+
+class FEditorViewportLayoutThreePanesBottom : public FEditorViewportLayoutThreePanes
+{
+public:
+	virtual const FName& GetLayoutTypeName() const override { return EditorViewportConfigurationNames::ThreePanesBottom; }
+
+	virtual TSharedRef<SWidget> MakeThreePanelWidget(
+		TMap< FName, TSharedPtr< IEditorViewportLayoutEntity >>& ViewportWidgets,
+		const TSharedRef<SWidget>& ViewportKey0, const TSharedRef<SWidget>& ViewportKey1, const TSharedRef<SWidget>& ViewportKey2,
+		float PrimarySplitterPercentage, float SecondarySplitterPercentage) override;
+};
diff --git a/Engine/Source/Editor/UnrealEd/Public/EditorViewportLayoutTwoPanes.h b/Engine/Source/Editor/UnrealEd/Public/EditorViewportLayoutTwoPanes.h
new file mode 100644
--- /dev/null
+++ b/Engine/Source/Editor/UnrealEd/Public/EditorViewportLayoutTwoPanes.h
@@ -0,0 +1,47 @@
+// Copyright 1998-2019 Epic Games, Inc. All Rights Reserved.
+
+#pragma once
+
+#include "CoreMinimal.h"
+#include "Widgets/SWidget.h"
+#include "AssetEditorViewportLayout.h"
+#include "Widgets/Layout/SSplitter.h"
+#include "Types/SlateEnums.h"
+
+
+template <EOrientation TOrientation>
+class TEditorViewportLayoutTwoPanes : public FAssetEditorViewportLayout
+{
+public:
+	/**
+	 * Saves viewport layout information between editor sessions
+	 */
+  	virtual void SaveLayoutString(const FString& LayoutString) const override {}
+
+protected:
+	virtual TSharedRef<SWidget> MakeViewportLayout(TFunction<TSharedRef<SEditorViewport>(void)> &Func, const FString& LayoutString) override ;
+
+
+private:
+	/** The splitter widget */
+	TSharedPtr< class SSplitter > SplitterWidget;
+};
+
+
+// FEditorViewportLayoutTwoPanesVert /////////////////////////////
+
+
+class FEditorViewportLayoutTwoPanesVert : public TEditorViewportLayoutTwoPanes<EOrientation::Orient_Vertical>
+{
+public:
+	virtual const FName& GetLayoutTypeName() const override { return EditorViewportConfigurationNames::TwoPanesVert; }
+};
+
+
+// FEditorViewportLayoutTwoPanesHoriz /////////////////////////////
+
+class FEditorViewportLayoutTwoPanesHoriz : public TEditorViewportLayoutTwoPanes<EOrientation::Orient_Horizontal>
+{
+public:
+	virtual const FName& GetLayoutTypeName() const override { return EditorViewportConfigurationNames::TwoPanesHoriz; }
+};
diff --git a/Engine/Source/Editor/UnrealEd/Public/EditorViewportTabContent.h b/Engine/Source/Editor/UnrealEd/Public/EditorViewportTabContent.h
new file mode 100644
--- /dev/null
+++ b/Engine/Source/Editor/UnrealEd/Public/EditorViewportTabContent.h
@@ -0,0 +1,57 @@
+// Copyright 1998-2019 Epic Games, Inc. All Rights Reserved.
+
+#pragma once
+
+#include "CoreMinimal.h"
+#include "Widgets/Docking/SDockTab.h"
+#include "ViewportTabContent.h"
+
+class FEditorViewportLayout;
+/**
+ * Represents the content in a viewport tab in an editor.
+ * Each SDockTab holding viewports in an editor contains and owns one of these.
+ */
+class UNREALED_API FEditorViewportTabContent : public FViewportTabContent, public TSharedFromThis<FEditorViewportTabContent>
+{
+public:
+	/** Returns whether the tab is currently shown */
+	bool IsVisible() const;
+
+	/** @return True if this viewport belongs to the tab given */
+	bool BelongsToTab(TSharedRef<class SDockTab> InParentTab) const;
+
+	/** @return The string used to identify the layout of this viewport tab */
+	const FString& GetLayoutString() const
+	{
+		return LayoutString;
+	}
+
+
+	TSharedPtr< class FEditorViewportLayout > ConstructViewportLayoutByTypeName(const FName& TypeName, bool bSwitchingLayouts);
+
+	void Initialize(TFunction<TSharedRef<SEditorViewport>(void)> Func, TSharedPtr<SDockTab> InParentTab, const FString& InLayoutString);
+
+	void SetViewportConfiguration(TFunction<TSharedRef<SEditorViewport>(void)> &Func, const FName& ConfigurationName);
+
+	void UpdateViewportTabWidget(TFunction<TSharedRef<SEditorViewport>(void)> &Func);
+
+	/**
+	* Sets the current layout by changing the contained layout object
+	* 
+	* @param ConfigurationName		The name of the layout (for the names in namespace EditorViewportConfigurationNames)
+	*/
+	void SetViewportConfiguration(const FName& ConfigurationName) override;
+
+
+
+protected:
+	TWeakPtr<class SDockTab> ParentTab;
+
+	/** Current layout */
+	TSharedPtr< class FEditorViewportLayout > ActiveViewportLayout;
+
+	TOptional<FName> PreviouslyFocusedViewport;
+
+private:
+	TFunction<TSharedRef<SEditorViewport>(void)> ViewportCreationFunc;
+};
diff --git a/Engine/Source/Editor/UnrealEd/Public/EditorViewportTypeDefinition.h b/Engine/Source/Editor/UnrealEd/Public/EditorViewportTypeDefinition.h
new file mode 100644
--- /dev/null
+++ b/Engine/Source/Editor/UnrealEd/Public/EditorViewportTypeDefinition.h
@@ -0,0 +1,34 @@
+// Copyright 1998-2019 Epic Games, Inc. All Rights Reserved.
+
+#pragma once
+
+#include "CoreMinimal.h"
+
+class FUICommandInfo;
+class IEditorViewportLayoutEntity;
+struct FAssetEditorViewportConstructionArgs;
+
+/** Definition of a custom viewport */
+struct FEditorViewportTypeDefinition
+{
+	typedef TFunction<TSharedRef<IEditorViewportLayoutEntity>(const FAssetEditorViewportConstructionArgs&)> FFactoryFunctionType;
+
+	template<typename T>
+	static FEditorViewportTypeDefinition FromType(const TSharedPtr<FUICommandInfo>& ActivationCommand)
+	{
+		return FViewportTypeDefinition([](const FAssetEditorViewportConstructionArgs& Args) -> TSharedRef<IEditorViewportLayoutEntity> {
+			return MakeShareable(new T(Args));
+		}, ActivationCommand);
+	}
+
+	FEditorViewportTypeDefinition(const FFactoryFunctionType& InFactoryFunction, const TSharedPtr<FUICommandInfo>& InActivationCommand)
+		: ActivationCommand(InActivationCommand)
+		, FactoryFunction(InFactoryFunction)
+	{}
+
+	/** A UI command for toggling activation this viewport */
+	TSharedPtr<FUICommandInfo> ActivationCommand;
+
+	/** Function used to create a new instance of the viewport */
+	FFactoryFunctionType FactoryFunction;
+};
diff --git a/Engine/Source/Editor/UnrealEd/Public/SAssetEditorViewport.h b/Engine/Source/Editor/UnrealEd/Public/SAssetEditorViewport.h
new file mode 100644
--- /dev/null
+++ b/Engine/Source/Editor/UnrealEd/Public/SAssetEditorViewport.h
@@ -0,0 +1,51 @@
+// Copyright 1998-2019 Epic Games, Inc. All Rights Reserved.
+
+#pragma once
+
+#include "CoreMinimal.h"
+#include "Widgets/SWidget.h"
+#include "SEditorViewport.h"
+
+class FEditorViewportClient;
+class FAssetEditorViewportLayout;
+
+class UNREALED_API SAssetEditorViewport : public SEditorViewport
+{
+public:
+
+	SLATE_BEGIN_ARGS(SAssetEditorViewport)
+		: _ViewportType(LVT_Perspective)
+		, _Realtime(false)
+	{
+	}
+
+	SLATE_ARGUMENT(TWeakPtr<class FEditorModeTools>, EditorModeTools)
+		SLATE_ARGUMENT(TSharedPtr<class FAssetEditorViewportLayout>, ParentLayout)
+		SLATE_ARGUMENT(TSharedPtr<FEditorViewportClient>, EditorViewportClient)
+		SLATE_ARGUMENT(ELevelViewportType, ViewportType)
+		SLATE_ARGUMENT(bool, Realtime)
+		SLATE_ARGUMENT(FName, ConfigKey)
+	SLATE_END_ARGS()
+
+
+	void OnSetViewportConfiguration(FName ConfigurationName);
+	bool IsViewportConfigurationSet(FName ConfigurationName) const;
+
+	TWeakPtr<class FAssetEditorViewportLayout> ParentLayout;
+
+protected:
+	virtual void BindCommands() override;
+
+	TSharedRef<FEditorViewportClient> MakeEditorViewportClient() override {
+		EditorViewportClient = MakeShareable(new FEditorViewportClient(nullptr));
+
+		return EditorViewportClient.ToSharedRef();
+	};
+
+
+private:
+
+	// Viewport client
+	TSharedPtr<FEditorViewportClient> EditorViewportClient;
+
+};
diff --git a/Engine/Source/Editor/UnrealEd/Public/ViewportTabContent.h b/Engine/Source/Editor/UnrealEd/Public/ViewportTabContent.h
--- a/Engine/Source/Editor/UnrealEd/Public/ViewportTabContent.h
+++ b/Engine/Source/Editor/UnrealEd/Public/ViewportTabContent.h
@@ -3,28 +3,39 @@
 #pragma once
 
 #include "CoreMinimal.h"
+#include "EditorViewportLayout.h"
+#include "EditorViewportLayoutTwoPanes.h"
+#include "Widgets/Docking/SDockTab.h"
 
 /**
  * Represents the content in a viewport tab in an editor.
  * Each SDockTab holding viewports in an editor contains and owns one of these.
  */
-class UNREALED_API FViewportTabContent : public TSharedFromThis<FViewportTabContent>
+class UNREALED_API FViewportTabContent
 {
 public:
+
+	virtual ~FViewportTabContent() {}
 	/** Returns whether the tab is currently shown */
 	bool IsVisible() const;
 
 	/** @return True if this viewport belongs to the tab given */
 	bool BelongsToTab(TSharedRef<class SDockTab> InParentTab) const;
 
-	/** @return The string used to identify the layout of this viewport tab */
-	const FString& GetLayoutString() const
-	{
-		return LayoutString;
-	}
+
+	/**
+	* Returns whether the named layout is currently selected
+	*
+	* @param ConfigurationName		The name of the layout (for the names in namespace LevelViewportConfigurationNames)
+	* @return						True, if the named layout is currently active
+	*/
+	bool IsViewportConfigurationSet(const FName& ConfigurationName) const;
+
+	virtual void SetViewportConfiguration(const FName& ConfigurationName) {}
 
 protected:
 	TWeakPtr<class SDockTab> ParentTab;
+
 	FString LayoutString;
 
 	/** Current layout */
