add widget header #rb none

diff --git a/Engine/Source/Editor/UnrealEd/Public/EditorViewportLayout.h b/Engine/Source/Editor/UnrealEd/Public/EditorViewportLayout.h
--- a/Engine/Source/Editor/UnrealEd/Public/EditorViewportLayout.h
+++ b/Engine/Source/Editor/UnrealEd/Public/EditorViewportLayout.h
@@ -3,6 +3,7 @@
 #pragma once
 
 #include "CoreMinimal.h"
+#include "Widgets/SWidget.h"
 
 
 class IEditorViewportLayoutEntity : public TSharedFromThis<IEditorViewportLayoutEntity>
