fix missing include #rb none

diff --git a/Engine/Source/Editor/UnrealEd/Public/AssetEditorViewportLayout.h b/Engine/Source/Editor/UnrealEd/Public/AssetEditorViewportLayout.h
--- a/Engine/Source/Editor/UnrealEd/Public/AssetEditorViewportLayout.h
+++ b/Engine/Source/Editor/UnrealEd/Public/AssetEditorViewportLayout.h
@@ -13,6 +13,7 @@
 #include "Application/ThrottleManager.h"
 #include "EditorViewportLayout.h"
 #include "TickableEditorObject.h"
+#include "SEditorViewport.h"
 
 class FAssetEditorViewportLayout;
 class SViewportsOverlay;
