more missing headers #rb none

diff --git a/Engine/Source/Editor/UnrealEd/Private/AssetEditorViewportLayout.cpp b/Engine/Source/Editor/UnrealEd/Private/AssetEditorViewportLayout.cpp
--- a/Engine/Source/Editor/UnrealEd/Private/AssetEditorViewportLayout.cpp
+++ b/Engine/Source/Editor/UnrealEd/Private/AssetEditorViewportLayout.cpp
@@ -23,6 +23,7 @@
 #include "EditorViewportCommands.h"
 #include "SAssetEditorViewport.h"
 #include "Templates/SharedPointer.h"
+#include "Framework/MultiBox/MultiBoxBuilder.h"
 
 
 #define LOCTEXT_NAMESPACE "AssetEditorViewportToolBar"
