Memory stomp detection in FTickableEditorObject

diff --git a/Engine/Source/Editor/UnrealEd/Private/UnrealEdMisc.cpp b/Engine/Source/Editor/UnrealEd/Private/UnrealEdMisc.cpp
--- a/Engine/Source/Editor/UnrealEd/Private/UnrealEdMisc.cpp
+++ b/Engine/Source/Editor/UnrealEd/Private/UnrealEdMisc.cpp
@@ -92,6 +92,7 @@ DEFINE_LOG_CATEGORY_STATIC(LogUnrealEdMisc, Log, All);
 
 bool FTickableEditorObject::bCollectionIntact = true;
 bool FTickableEditorObject::bIsTickingObjects = false;
+bool FTickableEditorObject::bIsInObjectsTick = false;
 
 
 namespace
diff --git a/Engine/Source/Editor/UnrealEd/Public/TickableEditorObject.h b/Engine/Source/Editor/UnrealEd/Public/TickableEditorObject.h
--- a/Engine/Source/Editor/UnrealEd/Public/TickableEditorObject.h
+++ b/Engine/Source/Editor/UnrealEd/Public/TickableEditorObject.h
@@ -40,7 +40,9 @@ public:
 				{
 					if ((TickableEntry.TickType == ETickableTickType::Always) || TickableObject->IsTickable())
 					{
+						bIsInObjectsTick = true;
 						TickableObject->Tick(DeltaSeconds);
+						bIsInObjectsTick = false;
 					}
 
 					// In case it was removed during tick
@@ -76,6 +78,8 @@ public:
 	/** Removes this instance from the static array of tickable objects. */
 	virtual ~FTickableEditorObject()
 	{
+		ensureMsgf(!bIsInObjectsTick, TEXT("Detected possible memory stomp. We are in the Tickable objects Tick function but hit its deconstructor, the 'this' pointer for the Object will now be invalid"));
+
 		ensure(IsInGameThread() || IsInAsyncLoadingThread());
 		if (bCollectionIntact && GetPendingTickableObjects().Remove(this) == 0)
 		{
@@ -107,6 +111,9 @@ private:
 	/** True if currently ticking of tickable editor objects. */
 	static bool bIsTickingObjects;
 
+	/** True if we are in the Tick function for an editor tickable object */
+	static bool bIsInObjectsTick;
+
 	static TArray<FTickableObjectEntry>& GetTickableObjects()
 	{
 		static TTickableObjectsCollection TickableObjects;
