Set ViewportType on FEditorViewportLayoutEntity creation.  #rb none

diff --git a/Engine/Source/Editor/UnrealEd/Private/EditorViewportLayoutEntity.cpp b/Engine/Source/Editor/UnrealEd/Private/EditorViewportLayoutEntity.cpp
--- a/Engine/Source/Editor/UnrealEd/Private/EditorViewportLayoutEntity.cpp
+++ b/Engine/Source/Editor/UnrealEd/Private/EditorViewportLayoutEntity.cpp
@@ -14,6 +14,7 @@ FEditorViewportLayoutEntity::FEditorViewportLayoutEntity(TFunction<TSharedRef<SE
  	TSharedPtr<SEditorViewport> NewViewport = Func();
 	AssetEditorViewport = StaticCastSharedPtr<SAssetEditorViewport>( NewViewport);
 	AssetEditorViewport->ParentLayout = ConstructionArgs.ParentLayout;
+	AssetEditorViewport->GetViewportClient()->SetViewportType(ConstructionArgs.ViewportType);
 }
 
 
