Combo boxes for enum bitflags will no longer display enum values that have been marked with UMETA(Hidden).

diff --git a/Engine/Source/Editor/GraphEditor/Private/KismetPins/SGraphPinInteger.cpp b/Engine/Source/Editor/GraphEditor/Private/KismetPins/SGraphPinInteger.cpp
--- a/Engine/Source/Editor/GraphEditor/Private/KismetPins/SGraphPinInteger.cpp
+++ b/Engine/Source/Editor/GraphEditor/Private/KismetPins/SGraphPinInteger.cpp
@@ -54,7 +54,8 @@ TSharedRef<SWidget>	SGraphPinInteger::GetDefaultValueWidget()
 			for(int32 BitmaskEnumIndex = 0; BitmaskEnumIndex < BitmaskEnum->NumEnums() - 1; ++BitmaskEnumIndex)
 			{
 				const int64 EnumValue = BitmaskEnum->GetValueByIndex(BitmaskEnumIndex);
-				if (EnumValue >= 0)
+				const bool bIsHidden = BitmaskEnum->HasMetaData(TEXT("Hidden"), BitmaskEnumIndex);
+				if (EnumValue >= 0 && !bIsHidden)
 				{
 					if (bUseEnumValuesAsMaskValues)
 					{
diff --git a/Engine/Source/Editor/PropertyEditor/Private/UserInterface/PropertyEditor/SPropertyEditorNumeric.h b/Engine/Source/Editor/PropertyEditor/Private/UserInterface/PropertyEditor/SPropertyEditorNumeric.h
--- a/Engine/Source/Editor/PropertyEditor/Private/UserInterface/PropertyEditor/SPropertyEditorNumeric.h
+++ b/Engine/Source/Editor/PropertyEditor/Private/UserInterface/PropertyEditor/SPropertyEditorNumeric.h
@@ -92,7 +92,8 @@ public:
 					for (int32 BitmaskEnumIndex = 0; BitmaskEnumIndex < BitmaskEnum->NumEnums() - 1; ++BitmaskEnumIndex)
 					{
 						const int64 EnumValue = BitmaskEnum->GetValueByIndex(BitmaskEnumIndex);
-						if (EnumValue >= 0)
+						const bool bIsHidden = BitmaskEnum->HasMetaData(TEXT("Hidden"), BitmaskEnumIndex);
+						if (EnumValue >= 0 && !bIsHidden)
 						{
 							if (bUseEnumValuessAsMaskValues)
 							{
diff --git a/Engine/Source/Editor/UnrealEd/Classes/Editor/PropertyEditorTestObject.h b/Engine/Source/Editor/UnrealEd/Classes/Editor/PropertyEditorTestObject.h
--- a/Engine/Source/Editor/UnrealEd/Classes/Editor/PropertyEditorTestObject.h
+++ b/Engine/Source/Editor/UnrealEd/Classes/Editor/PropertyEditorTestObject.h
@@ -41,6 +41,16 @@ enum PropertyEditorTestEnum
 	PropertyEditorTest_MAX,
 };
 
+UENUM(meta = (Bitflags))
+enum class PropertyEditorTestBitflags : uint8
+{
+	First,
+	Second,
+	Third,
+	Hidden UMETA(Hidden, ToolTip = "This value shouldn't be used or even visible in the editor")
+};
+ENUM_CLASS_FLAGS(PropertyEditorTestBitflags)
+
 UENUM()
 enum ArrayLabelEnum
 {
@@ -320,6 +330,9 @@ class UPropertyEditorTestObject : public UObject
 	UPROPERTY(VisibleAnywhere, Category=AdvancedProperties)
 	UPrimitiveComponent* ObjectThatCannotBeChanged;
 
+	UPROPERTY(EditAnywhere, Category = AdvancedProperties, meta = (Bitmask, BitmaskEnum = "PropertyEditorTestBitflags"))
+	int32 EnumBitflags = 0;
+
 	UPROPERTY(EditAnywhere, Category=AdvancedProperties, meta=(PasswordField=true))
 	FString StringPasswordProperty;
 
