Fixed crash trying to get SlateApplication in ShouldThrottleCPUUsage if called from a commandlet

diff --git a/Engine/Source/Editor/UnrealEd/Private/EditorEngine.cpp b/Engine/Source/Editor/UnrealEd/Private/EditorEngine.cpp
--- a/Engine/Source/Editor/UnrealEd/Private/EditorEngine.cpp
+++ b/Engine/Source/Editor/UnrealEd/Private/EditorEngine.cpp
@@ -6183,9 +6183,11 @@ bool UEditorEngine::ShouldThrottleCPUUsage() const
 {
 	bool bShouldThrottle = false;
 
-	bool bIsForeground = FPlatformApplicationMisc::IsThisApplicationForeground();
+	const bool bRunningCommandlet = IsRunningCommandlet();
 
-	if( !bIsForeground )
+	const bool bIsForeground = FPlatformApplicationMisc::IsThisApplicationForeground();
+
+	if( !bIsForeground && !bRunningCommandlet )
 	{
 		const UEditorPerformanceSettings* Settings = GetDefault<UEditorPerformanceSettings>();
 		bShouldThrottle = Settings->bThrottleCPUWhenNotForeground;
@@ -6214,7 +6216,7 @@ bool UEditorEngine::ShouldThrottleCPUUsage() const
 		}
 	}
 
-	return bShouldThrottle && !IsRunningCommandlet();
+	return bShouldThrottle;
 }
 
 bool UEditorEngine::AreAllWindowsHidden() const
