Fixing the factory for UEdMode
#rnx

diff --git a/Engine/Source/Editor/UnrealEd/Public/EditorModeRegistry.h b/Engine/Source/Editor/UnrealEd/Public/EditorModeRegistry.h
--- a/Engine/Source/Editor/UnrealEd/Public/EditorModeRegistry.h
+++ b/Engine/Source/Editor/UnrealEd/Public/EditorModeRegistry.h
@@ -122,10 +122,11 @@ struct FEditorModeFactory : IEditorModeFactory
 	virtual bool ForScriptableMode() const override { return false; }
 };
 
-class UEditorModeFactory : public UObject, public IEditorModeFactory
+struct UNREALED_API FScriptableEditorModeFactory : IEditorModeFactory
 {
-	UEditorModeFactory(FEditorModeInfo InModeInfo) : ModeInfo(InModeInfo) {}
-	virtual ~UEditorModeFactory() {}
+public:
+	FScriptableEditorModeFactory(FEditorModeInfo InModeInfo) : ModeInfo(InModeInfo) {}
+	virtual ~FScriptableEditorModeFactory() {}
 
 	/** Information pertaining to this factory's mode */
 	FEditorModeInfo ModeInfo;
@@ -208,6 +209,23 @@ public:
 		RegisterMode(ModeID, Factory);
 	}
 
+	/**
+ * Registers an editor mode type. Typically called from a module's StartupModule() routine.
+ *
+ * @param ModeID	ID of the mode to register
+ */
+	template<class T>
+	void RegisterScriptableMode(FEditorModeID ModeID, FText Name = FText(), FSlateIcon IconBrush = FSlateIcon(), bool bVisible = false, int32 PriorityOrder = MAX_int32)
+	{
+		TSharedRef<FScriptableEditorModeFactory> Factory = MakeShareable(new FScriptableEditorModeFactory(FEditorModeInfo(ModeID, Name, IconBrush, bVisible, PriorityOrder)));
+
+		Factory->FactoryDelegate = FEditorModeFactoryDelegate::CreateStatic([]() -> UEdMode* {
+			return NewObject<T>();
+		});
+
+		RegisterMode(ModeID, Factory);
+	}
+
 	/**
 	 * Unregisters an editor mode. Typically called from a module's ShutdownModule() routine.
 	 * Note: will exit the edit mode if it is currently active.
