Making a function on the toolkit for getting data const
#rnx

diff --git a/Engine/Plugins/Experimental/ChaosEditor/Source/FractureEditor/Private/FractureEditorModeToolkit.cpp b/Engine/Plugins/Experimental/ChaosEditor/Source/FractureEditor/Private/FractureEditorModeToolkit.cpp
--- a/Engine/Plugins/Experimental/ChaosEditor/Source/FractureEditor/Private/FractureEditorModeToolkit.cpp
+++ b/Engine/Plugins/Experimental/ChaosEditor/Source/FractureEditor/Private/FractureEditorModeToolkit.cpp
@@ -306,7 +306,7 @@ void FFractureEditorModeToolkit::Init(const TSharedPtr<IToolkitHost>& InitToolki
 
 const TArray<FName> FFractureEditorModeToolkit::PaletteNames = { FName(TEXT("Fracture")), FName(TEXT("Cluster")) };
 
-FText FFractureEditorModeToolkit::GetToolPaletteDisplayName(FName Palette) 
+FText FFractureEditorModeToolkit::GetToolPaletteDisplayName(FName Palette) const
 { 
 	return FText::FromName(Palette);
 }
diff --git a/Engine/Plugins/Experimental/ChaosEditor/Source/FractureEditor/Public/FractureEditorModeToolkit.h b/Engine/Plugins/Experimental/ChaosEditor/Source/FractureEditor/Public/FractureEditorModeToolkit.h
--- a/Engine/Plugins/Experimental/ChaosEditor/Source/FractureEditor/Public/FractureEditorModeToolkit.h
+++ b/Engine/Plugins/Experimental/ChaosEditor/Source/FractureEditor/Public/FractureEditorModeToolkit.h
@@ -99,7 +99,7 @@ public:
 	/** Returns the number of Mode specific tabs in the mode toolbar **/ 
 	const static TArray<FName> PaletteNames;
 	virtual void GetToolPaletteNames( TArray<FName>& InPaletteName ) const { InPaletteName = PaletteNames; }
-	virtual FText GetToolPaletteDisplayName(FName PaletteName); 
+	virtual FText GetToolPaletteDisplayName(FName PaletteName) const; 
 	virtual void BuildToolPalette(FName PaletteName, class FToolBarBuilder& ToolbarBuilder);
 	virtual void OnToolPaletteChanged(FName PaletteName) override;
 
diff --git a/Engine/Plugins/Experimental/ModelingToolsEditorMode/Source/ModelingToolsEditorMode/Private/ModelingToolsEditorModeToolkit.cpp b/Engine/Plugins/Experimental/ModelingToolsEditorMode/Source/ModelingToolsEditorMode/Private/ModelingToolsEditorModeToolkit.cpp
--- a/Engine/Plugins/Experimental/ModelingToolsEditorMode/Source/ModelingToolsEditorMode/Private/ModelingToolsEditorModeToolkit.cpp
+++ b/Engine/Plugins/Experimental/ModelingToolsEditorMode/Source/ModelingToolsEditorMode/Private/ModelingToolsEditorModeToolkit.cpp
@@ -237,7 +237,7 @@ void FModelingToolsEditorModeToolkit::GetToolPaletteNames(TArray<FName>& InPalet
 }
 
 
-FText FModelingToolsEditorModeToolkit::GetToolPaletteDisplayName(FName Palette) 
+FText FModelingToolsEditorModeToolkit::GetToolPaletteDisplayName(FName Palette) const
 { 
 	return FText::FromName(Palette);
 }
diff --git a/Engine/Plugins/Experimental/ModelingToolsEditorMode/Source/ModelingToolsEditorMode/Public/ModelingToolsEditorModeToolkit.h b/Engine/Plugins/Experimental/ModelingToolsEditorMode/Source/ModelingToolsEditorMode/Public/ModelingToolsEditorModeToolkit.h
--- a/Engine/Plugins/Experimental/ModelingToolsEditorMode/Source/ModelingToolsEditorMode/Public/ModelingToolsEditorModeToolkit.h
+++ b/Engine/Plugins/Experimental/ModelingToolsEditorMode/Source/ModelingToolsEditorMode/Public/ModelingToolsEditorModeToolkit.h
@@ -41,7 +41,7 @@ public:
 
 	/** Returns the Mode specific tabs in the mode toolbar **/ 
 	virtual void GetToolPaletteNames(TArray<FName>& InPaletteName) const;
-	virtual FText GetToolPaletteDisplayName(FName PaletteName); 
+	virtual FText GetToolPaletteDisplayName(FName PaletteName) const; 
 	virtual void BuildToolPalette(FName PaletteName, class FToolBarBuilder& ToolbarBuilder);
 	virtual void OnToolPaletteChanged(FName PaletteName) override;
 
diff --git a/Engine/Source/Editor/UnrealEd/Private/EditorModeManager.cpp b/Engine/Source/Editor/UnrealEd/Private/EditorModeManager.cpp
--- a/Engine/Source/Editor/UnrealEd/Private/EditorModeManager.cpp
+++ b/Engine/Source/Editor/UnrealEd/Private/EditorModeManager.cpp
@@ -726,8 +726,16 @@ void FEditorModeTools::RebuildModeToolBar()
 				{
 					TSharedRef<SWidget> PaletteWidget = Row.ToolbarWidget.ToSharedRef();
 
-					FEdMode* Mode = GetActiveMode(Row.ModeID);
-					TSharedPtr<FModeToolkit> RowToolkit = Mode->GetToolkit();
+					
+					TSharedPtr<FModeToolkit> RowToolkit;
+					if (FEdMode* Mode = GetActiveMode(Row.ModeID))
+					{
+						RowToolkit = Mode->GetToolkit();
+					}
+					else if (UEdMode* ScriptableMode = GetActiveScriptableMode(Row.ModeID))
+					{
+						RowToolkit = ScriptableMode->GetToolkit();
+					}
 
 					// Don't show Palette Tabs if there is only one
 					if (PaletteCount > 1)
@@ -1080,7 +1088,6 @@ void FEditorModeTools::ActivateMode(FEditorModeID InID, bool bToggle)
 					FToolBarBuilder ModeToolbarBuilder(CommandList, FMultiBoxCustomization(Mode->GetModeInfo().ToolbarCustomizationName), TSharedPtr<FExtender>(), Orient_Horizontal, false);
 					Toolkit->BuildToolPalette(Palette, ModeToolbarBuilder);
 
-					// ActiveToolBarRows.Emplace(Mode, Toolkit->GetToolPaletteName(i), ModeToolbarBuilder.MakeWidget());
 					ActiveToolBarRows.Emplace(Mode->GetID(), Palette, Toolkit->GetToolPaletteDisplayName(Palette), ModeToolbarBuilder.MakeWidget());
 					PaletteCount++;
 				}
diff --git a/Engine/Source/Editor/UnrealEd/Public/Toolkits/IToolkit.h b/Engine/Source/Editor/UnrealEd/Public/Toolkits/IToolkit.h
--- a/Engine/Source/Editor/UnrealEd/Public/Toolkits/IToolkit.h
+++ b/Engine/Source/Editor/UnrealEd/Public/Toolkits/IToolkit.h
@@ -144,7 +144,7 @@ public:
 	 * @param PaletteIndex      The index of the ToolPalette to build
 	 * @returns the name of Tool Palette 
 	 **/ 
-	virtual FText GetToolPaletteDisplayName(FName Palette) { return FText(); }
+	virtual FText GetToolPaletteDisplayName(FName Palette) const { return FText(); }
 
 	/** 
 	 * @param PaletteIndex      The index of the ToolPalette to build
