Moved generation of Layouts menu to SAssetEditoriViewport and moved layouts menu to submenu to match Level Editor.
Adds missing command binds.

diff --git a/Engine/Source/Editor/UnrealEd/Private/AssetEditorViewportLayout.cpp b/Engine/Source/Editor/UnrealEd/Private/AssetEditorViewportLayout.cpp
--- a/Engine/Source/Editor/UnrealEd/Private/AssetEditorViewportLayout.cpp
+++ b/Engine/Source/Editor/UnrealEd/Private/AssetEditorViewportLayout.cpp
@@ -191,126 +191,6 @@ TSharedRef<SWidget> FAssetEditorViewportLayout::BuildViewportLayout(TFunction<TS
 
 }
 
-TSharedRef<SWidget> FAssetEditorViewportLayout::GenerateLayoutMenu(TSharedPtr<SAssetEditorViewport> AssetEditorViewport)
-{
-	const bool bInShouldCloseWindowAfterMenuSelection = true;
-
-	TSharedPtr<FUICommandList> CommandList = AssetEditorViewport->GetCommandList();
-
-	FMenuBuilder MenuBuilder(bInShouldCloseWindowAfterMenuSelection, CommandList);
-
-	MenuBuilder.BeginSection(NAME_None, LOCTEXT("OnePaneConfigHeader", "One Pane"));
-	{
- 		FToolBarBuilder OnePaneButton(CommandList, FMultiBoxCustomization::None);
-		OnePaneButton.SetLabelVisibility(EVisibility::Collapsed);
-		OnePaneButton.SetStyle(&FEditorStyle::Get(), "ViewportLayoutToolbar");
-
- 		OnePaneButton.AddToolBarButton(FEditorViewportCommands::Get().ViewportConfig_OnePane);
-
-		MenuBuilder.AddWidget(
-			SNew(SHorizontalBox)
-			+ SHorizontalBox::Slot()
-			.AutoWidth()
-			[
-				OnePaneButton.MakeWidget()
-			]
-		+ SHorizontalBox::Slot()
-			.FillWidth(1)
-			[
-				SNullWidget::NullWidget
-			],
-			FText::GetEmpty(), true
-			);
-	}
-	MenuBuilder.EndSection();
- 
-	MenuBuilder.BeginSection("EditorViewportTwoPaneConfigs", LOCTEXT("TwoPaneConfigHeader", "Two Panes"));
-	{
-		FToolBarBuilder TwoPaneButtons(CommandList, FMultiBoxCustomization::None);
-		TwoPaneButtons.SetLabelVisibility(EVisibility::Collapsed);
-		TwoPaneButtons.SetStyle(&FEditorStyle::Get(), "ViewportLayoutToolbar");
-
-		TwoPaneButtons.AddToolBarButton(FEditorViewportCommands::Get().ViewportConfig_TwoPanesH, NAME_None, FText());
-		TwoPaneButtons.AddToolBarButton(FEditorViewportCommands::Get().ViewportConfig_TwoPanesV, NAME_None, FText());
-
-		MenuBuilder.AddWidget(
-			SNew(SHorizontalBox)
-			+ SHorizontalBox::Slot()
-			.AutoWidth()
-			[
-				TwoPaneButtons.MakeWidget()
-			]
-		+ SHorizontalBox::Slot()
-			.FillWidth(1)
-			[
-				SNullWidget::NullWidget
-			],
-			FText::GetEmpty(), true
-			);
-	}
-	MenuBuilder.EndSection();
-
-	MenuBuilder.BeginSection("EditorViewportThreePaneConfigs", LOCTEXT("ThreePaneConfigHeader", "Three Panes"));
-	{
-		FToolBarBuilder ThreePaneButtons(CommandList, FMultiBoxCustomization::None);
-		ThreePaneButtons.SetLabelVisibility(EVisibility::Collapsed);
-		ThreePaneButtons.SetStyle(&FEditorStyle::Get(), "ViewportLayoutToolbar");
-
-		ThreePaneButtons.AddToolBarButton(FEditorViewportCommands::Get().ViewportConfig_ThreePanesLeft, NAME_None, FText());
-		ThreePaneButtons.AddToolBarButton(FEditorViewportCommands::Get().ViewportConfig_ThreePanesRight, NAME_None, FText());
-		ThreePaneButtons.AddToolBarButton(FEditorViewportCommands::Get().ViewportConfig_ThreePanesTop, NAME_None, FText());
-		ThreePaneButtons.AddToolBarButton(FEditorViewportCommands::Get().ViewportConfig_ThreePanesBottom, NAME_None, FText());
-
-		MenuBuilder.AddWidget(
-			SNew(SHorizontalBox)
-			+ SHorizontalBox::Slot()
-			.AutoWidth()
-			[
-				ThreePaneButtons.MakeWidget()
-			]
-		+ SHorizontalBox::Slot()
-			.FillWidth(1)
-			[
-				SNullWidget::NullWidget
-			],
-			FText::GetEmpty(), true
-			);
-	}
-	MenuBuilder.EndSection();
-
-	MenuBuilder.BeginSection("EditorViewportFourPaneConfigs", LOCTEXT("FourPaneConfigHeader", "Four Panes"));
-	{
-		FToolBarBuilder FourPaneButtons(CommandList, FMultiBoxCustomization::None);
-		FourPaneButtons.SetLabelVisibility(EVisibility::Collapsed);
-		FourPaneButtons.SetStyle(&FEditorStyle::Get(), "ViewportLayoutToolbar");
-
-		FourPaneButtons.AddToolBarButton(FEditorViewportCommands::Get().ViewportConfig_FourPanes2x2, NAME_None, FText());
-		FourPaneButtons.AddToolBarButton(FEditorViewportCommands::Get().ViewportConfig_FourPanesLeft, NAME_None, FText());
-		FourPaneButtons.AddToolBarButton(FEditorViewportCommands::Get().ViewportConfig_FourPanesRight, NAME_None, FText());
-		FourPaneButtons.AddToolBarButton(FEditorViewportCommands::Get().ViewportConfig_FourPanesTop, NAME_None, FText());
-		FourPaneButtons.AddToolBarButton(FEditorViewportCommands::Get().ViewportConfig_FourPanesBottom, NAME_None, FText());
-
-		MenuBuilder.AddWidget(
-			SNew(SHorizontalBox)
-			+ SHorizontalBox::Slot()
-			.AutoWidth()
-			[
-				FourPaneButtons.MakeWidget()
-			]
-		+ SHorizontalBox::Slot()
-			.FillWidth(1)
-			[
-				SNullWidget::NullWidget
-			],
-			FText::GetEmpty(), true
-			);
-	}
-	MenuBuilder.EndSection();
-
-
-	return MenuBuilder.MakeWidget();
-
-}
 
 #undef LOCTEXT_NAMESPACE
 
diff --git a/Engine/Source/Editor/UnrealEd/Private/EditorViewportCommands.cpp b/Engine/Source/Editor/UnrealEd/Private/EditorViewportCommands.cpp
--- a/Engine/Source/Editor/UnrealEd/Private/EditorViewportCommands.cpp
+++ b/Engine/Source/Editor/UnrealEd/Private/EditorViewportCommands.cpp
@@ -28,6 +28,19 @@ void FEditorViewportCommands::RegisterCommands()
 	UI_COMMAND( Right, "Right", "Switches the viewport to right view", EUserInterfaceActionType::RadioButton, FInputChord( EModifierKey::Alt | EModifierKey::Shift, EKeys::K ) );
 	UI_COMMAND( Next, "Next", "Rotate through each view options", EUserInterfaceActionType::RadioButton, FInputChord( EModifierKey::Control | EModifierKey::Shift, EKeys::SpaceBar ) );
 
+	UI_COMMAND(ViewportConfig_OnePane, "Layout One Pane", "Changes the viewport arrangement to one pane", EUserInterfaceActionType::ToggleButton, FInputChord());
+	UI_COMMAND(ViewportConfig_TwoPanesH, "Layout Two Panes (horizontal)", "Changes the viewport arrangement to two panes, side-by-side", EUserInterfaceActionType::ToggleButton, FInputChord());
+	UI_COMMAND(ViewportConfig_TwoPanesV, "Layout Two Panes (vertical)", "Changes the viewport arrangement to two panes, one above the other", EUserInterfaceActionType::ToggleButton, FInputChord());
+	UI_COMMAND(ViewportConfig_ThreePanesLeft, "Layout Three Panes (one left, two right)", "Changes the viewport arrangement to three panes, one on the left, two on the right", EUserInterfaceActionType::ToggleButton, FInputChord());
+	UI_COMMAND(ViewportConfig_ThreePanesRight, "Layout Three Panes (one right, two left)", "Changes the viewport arrangement to three panes, one on the right, two on the left", EUserInterfaceActionType::ToggleButton, FInputChord());
+	UI_COMMAND(ViewportConfig_ThreePanesTop, "Layout Three Panes (one top, two bottom)", "Changes the viewport arrangement to three panes, one on the top, two on the bottom", EUserInterfaceActionType::ToggleButton, FInputChord());
+	UI_COMMAND(ViewportConfig_ThreePanesBottom, "Layout Three Panes (one bottom, two top)", "Changes the viewport arrangement to three panes, one on the bottom, two on the top", EUserInterfaceActionType::ToggleButton, FInputChord());
+	UI_COMMAND(ViewportConfig_FourPanesLeft, "Layout Four Panes (one left, three right)", "Changes the viewport arrangement to four panes, one on the left, three on the right", EUserInterfaceActionType::ToggleButton, FInputChord());
+	UI_COMMAND(ViewportConfig_FourPanesRight, "Layout Four Panes (one right, three left)", "Changes the viewport arrangement to four panes, one on the right, three on the left", EUserInterfaceActionType::ToggleButton, FInputChord());
+	UI_COMMAND(ViewportConfig_FourPanesTop, "Layout Four Panes (one top, three bottom)", "Changes the viewport arrangement to four panes, one on the top, three on the bottom", EUserInterfaceActionType::ToggleButton, FInputChord());
+	UI_COMMAND(ViewportConfig_FourPanesBottom, "Layout Four Panes (one bottom, three top)", "Changes the viewport arrangement to four panes, one on the bottom, three on the top", EUserInterfaceActionType::ToggleButton, FInputChord());
+	UI_COMMAND(ViewportConfig_FourPanes2x2, "Layout Four Panes (2x2)", "Changes the viewport arrangement to four panes, in a 2x2 grid", EUserInterfaceActionType::ToggleButton, FInputChord());
+
 	UI_COMMAND( WireframeMode, "Brush Wireframe View Mode", "Renders the scene in brush wireframe", EUserInterfaceActionType::RadioButton, FInputChord( EModifierKey::Alt, EKeys::Two ) );
 	UI_COMMAND( UnlitMode, "Unlit View Mode", "Renders the scene with no lights", EUserInterfaceActionType::RadioButton, FInputChord( EModifierKey::Alt, EKeys::Three ) );
 	UI_COMMAND( LitMode, "Lit View Mode", "Renders the scene with normal lighting", EUserInterfaceActionType::RadioButton, FInputChord( EModifierKey::Alt, EKeys::Four ) );
diff --git a/Engine/Source/Editor/UnrealEd/Private/SAssetEditorViewport.cpp b/Engine/Source/Editor/UnrealEd/Private/SAssetEditorViewport.cpp
--- a/Engine/Source/Editor/UnrealEd/Private/SAssetEditorViewport.cpp
+++ b/Engine/Source/Editor/UnrealEd/Private/SAssetEditorViewport.cpp
@@ -7,6 +7,8 @@
 #include "EditorViewportCommands.h"
 #include "EditorViewportTabContent.h"
 
+#define LOCTEXT_NAMESPACE "SAssetEditorViewport"
+
 
 void SAssetEditorViewport::BindCommands()
 {
@@ -119,3 +121,118 @@ bool SAssetEditorViewport::IsViewportConfigurationSet(FName ConfigurationName) c
 	}
 	return false;
 }
+
+
+void SAssetEditorViewport::GenerateLayoutMenu(FMenuBuilder& MenuBuilder) const
+{
+	MenuBuilder.BeginSection("EditorViewportOnePaneConfigs", LOCTEXT("OnePaneConfigHeader", "One Pane"));
+	{
+		FToolBarBuilder OnePaneButton(CommandList, FMultiBoxCustomization::None);
+		OnePaneButton.SetLabelVisibility(EVisibility::Collapsed);
+		OnePaneButton.SetStyle(&FEditorStyle::Get(), "ViewportLayoutToolbar");
+
+		OnePaneButton.AddToolBarButton(FEditorViewportCommands::Get().ViewportConfig_OnePane);
+
+		MenuBuilder.AddWidget(
+			SNew(SHorizontalBox)
+			+ SHorizontalBox::Slot()
+			.AutoWidth()
+			[
+				OnePaneButton.MakeWidget()
+			]
+		+ SHorizontalBox::Slot()
+			.FillWidth(1)
+			[
+				SNullWidget::NullWidget
+			],
+			FText::GetEmpty(), true
+			);
+	}
+	MenuBuilder.EndSection();
+
+	MenuBuilder.BeginSection("EditorViewportTwoPaneConfigs", LOCTEXT("TwoPaneConfigHeader", "Two Panes"));
+	{
+		FToolBarBuilder TwoPaneButtons(CommandList, FMultiBoxCustomization::None);
+		TwoPaneButtons.SetLabelVisibility(EVisibility::Collapsed);
+		TwoPaneButtons.SetStyle(&FEditorStyle::Get(), "ViewportLayoutToolbar");
+
+		TwoPaneButtons.AddToolBarButton(FEditorViewportCommands::Get().ViewportConfig_TwoPanesH, NAME_None, FText());
+		TwoPaneButtons.AddToolBarButton(FEditorViewportCommands::Get().ViewportConfig_TwoPanesV, NAME_None, FText());
+
+		MenuBuilder.AddWidget(
+			SNew(SHorizontalBox)
+			+ SHorizontalBox::Slot()
+			.AutoWidth()
+			[
+				TwoPaneButtons.MakeWidget()
+			]
+		+ SHorizontalBox::Slot()
+			.FillWidth(1)
+			[
+				SNullWidget::NullWidget
+			],
+			FText::GetEmpty(), true
+			);
+	}
+	MenuBuilder.EndSection();
+
+	MenuBuilder.BeginSection("EditorViewportThreePaneConfigs", LOCTEXT("ThreePaneConfigHeader", "Three Panes"));
+	{
+		FToolBarBuilder ThreePaneButtons(CommandList, FMultiBoxCustomization::None);
+		ThreePaneButtons.SetLabelVisibility(EVisibility::Collapsed);
+		ThreePaneButtons.SetStyle(&FEditorStyle::Get(), "ViewportLayoutToolbar");
+
+		ThreePaneButtons.AddToolBarButton(FEditorViewportCommands::Get().ViewportConfig_ThreePanesLeft, NAME_None, FText());
+		ThreePaneButtons.AddToolBarButton(FEditorViewportCommands::Get().ViewportConfig_ThreePanesRight, NAME_None, FText());
+		ThreePaneButtons.AddToolBarButton(FEditorViewportCommands::Get().ViewportConfig_ThreePanesTop, NAME_None, FText());
+		ThreePaneButtons.AddToolBarButton(FEditorViewportCommands::Get().ViewportConfig_ThreePanesBottom, NAME_None, FText());
+
+		MenuBuilder.AddWidget(
+			SNew(SHorizontalBox)
+			+ SHorizontalBox::Slot()
+			.AutoWidth()
+			[
+				ThreePaneButtons.MakeWidget()
+			]
+		+ SHorizontalBox::Slot()
+			.FillWidth(1)
+			[
+				SNullWidget::NullWidget
+			],
+			FText::GetEmpty(), true
+			);
+	}
+	MenuBuilder.EndSection();
+
+	MenuBuilder.BeginSection("EditorViewportFourPaneConfigs", LOCTEXT("FourPaneConfigHeader", "Four Panes"));
+	{
+		FToolBarBuilder FourPaneButtons(CommandList, FMultiBoxCustomization::None);
+		FourPaneButtons.SetLabelVisibility(EVisibility::Collapsed);
+		FourPaneButtons.SetStyle(&FEditorStyle::Get(), "ViewportLayoutToolbar");
+
+		FourPaneButtons.AddToolBarButton(FEditorViewportCommands::Get().ViewportConfig_FourPanes2x2, NAME_None, FText());
+		FourPaneButtons.AddToolBarButton(FEditorViewportCommands::Get().ViewportConfig_FourPanesLeft, NAME_None, FText());
+		FourPaneButtons.AddToolBarButton(FEditorViewportCommands::Get().ViewportConfig_FourPanesRight, NAME_None, FText());
+		FourPaneButtons.AddToolBarButton(FEditorViewportCommands::Get().ViewportConfig_FourPanesTop, NAME_None, FText());
+		FourPaneButtons.AddToolBarButton(FEditorViewportCommands::Get().ViewportConfig_FourPanesBottom, NAME_None, FText());
+
+		MenuBuilder.AddWidget(
+			SNew(SHorizontalBox)
+			+ SHorizontalBox::Slot()
+			.AutoWidth()
+			[
+				FourPaneButtons.MakeWidget()
+			]
+		+ SHorizontalBox::Slot()
+			.FillWidth(1)
+			[
+				SNullWidget::NullWidget
+			],
+			FText::GetEmpty(), true
+			);
+	}
+	MenuBuilder.EndSection();
+
+}
+
+#undef LOCTEXT_NAMESPACE
diff --git a/Engine/Source/Editor/UnrealEd/Private/SCommonEditorViewportToolbarBase.cpp b/Engine/Source/Editor/UnrealEd/Private/SCommonEditorViewportToolbarBase.cpp
--- a/Engine/Source/Editor/UnrealEd/Private/SCommonEditorViewportToolbarBase.cpp
+++ b/Engine/Source/Editor/UnrealEd/Private/SCommonEditorViewportToolbarBase.cpp
@@ -14,8 +14,10 @@
 #include "EditorViewportCommands.h"
 #include "SEditorViewportToolBarMenu.h"
 #include "SEditorViewportViewMenu.h"
+#include "AssetEditorViewportLayout.h"
+#include "SAssetEditorViewport.h"
 
-#define LOCTEXT_NAMESPACE "LevelViewportToolBar"
+#define LOCTEXT_NAMESPACE "SCommonEditorViewportToolbarBase"
 
 //////////////////////////////////////////////////////////////////////////
 // SCommonEditorViewportToolbarBase
@@ -181,6 +183,20 @@ TSharedRef<SWidget> SCommonEditorViewportToolbarBase::GenerateOptionsMenu() cons
 			OptionsMenuBuilder.AddWidget(GenerateScreenPercentageMenu(), LOCTEXT("ScreenPercentage", "Screen Percentage"));
 		}
 		OptionsMenuBuilder.EndSection();
+
+ 		TSharedPtr<SAssetEditorViewport> AssetEditorViewportPtr = StaticCastSharedRef<SAssetEditorViewport>(ViewportRef);
+ 		if (AssetEditorViewportPtr.IsValid())
+		{
+			OptionsMenuBuilder.BeginSection("LevelViewportLayouts");
+			{
+				OptionsMenuBuilder.AddSubMenu(
+					LOCTEXT("ConfigsSubMenu", "Layouts"),
+					FText::GetEmpty(),
+					FNewMenuDelegate::CreateSP(AssetEditorViewportPtr.Get(), &SAssetEditorViewport::GenerateLayoutMenu));
+			}
+			OptionsMenuBuilder.EndSection();
+		}
+
 		ExtendOptionsMenu(OptionsMenuBuilder);
 	}
 
diff --git a/Engine/Source/Editor/UnrealEd/Public/AssetEditorViewportLayout.h b/Engine/Source/Editor/UnrealEd/Public/AssetEditorViewportLayout.h
--- a/Engine/Source/Editor/UnrealEd/Public/AssetEditorViewportLayout.h
+++ b/Engine/Source/Editor/UnrealEd/Public/AssetEditorViewportLayout.h
@@ -91,8 +91,6 @@ public:
 	 */
  	TSharedRef<SWidget> BuildViewportLayout(TFunction<TSharedRef<SEditorViewport>(void)> &Func,  TSharedPtr<SDockTab> InParentDockTab, TSharedPtr<class FViewportTabContent> InParentTab, const FString& LayoutString );
 
-	static TSharedRef<SWidget> GenerateLayoutMenu(TSharedPtr<SAssetEditorViewport> AssetEditorViewport) ;
-
 	/** Returns the parent tab content object */
 	TWeakPtr< class FViewportTabContent > GetParentTabContent() const { return ParentTabContent; }
 
diff --git a/Engine/Source/Editor/UnrealEd/Public/SAssetEditorViewport.h b/Engine/Source/Editor/UnrealEd/Public/SAssetEditorViewport.h
--- a/Engine/Source/Editor/UnrealEd/Public/SAssetEditorViewport.h
+++ b/Engine/Source/Editor/UnrealEd/Public/SAssetEditorViewport.h
@@ -31,6 +31,8 @@ public:
 	void OnSetViewportConfiguration(FName ConfigurationName);
 	bool IsViewportConfigurationSet(FName ConfigurationName) const;
 
+	void GenerateLayoutMenu(FMenuBuilder& MenuBuilder) const;
+
 	TWeakPtr<class FAssetEditorViewportLayout> ParentLayout;
 
 protected:
