small name fix  #rb none

diff --git a/Engine/Source/Editor/UnrealEd/Private/SCommonEditorViewportToolbarBase.cpp b/Engine/Source/Editor/UnrealEd/Private/SCommonEditorViewportToolbarBase.cpp
--- a/Engine/Source/Editor/UnrealEd/Private/SCommonEditorViewportToolbarBase.cpp
+++ b/Engine/Source/Editor/UnrealEd/Private/SCommonEditorViewportToolbarBase.cpp
@@ -187,7 +187,7 @@ TSharedRef<SWidget> SCommonEditorViewportToolbarBase::GenerateOptionsMenu() cons
  		TSharedPtr<SAssetEditorViewport> AssetEditorViewportPtr = StaticCastSharedRef<SAssetEditorViewport>(ViewportRef);
  		if (AssetEditorViewportPtr.IsValid())
 		{
-			OptionsMenuBuilder.BeginSection("LevelViewportLayouts");
+			OptionsMenuBuilder.BeginSection("EditorViewportLayouts");
 			{
 				OptionsMenuBuilder.AddSubMenu(
 					LOCTEXT("ConfigsSubMenu", "Layouts"),
