add missing header #rb none

diff --git a/Engine/Source/Editor/UnrealEd/Private/SAssetEditorViewport.cpp b/Engine/Source/Editor/UnrealEd/Private/SAssetEditorViewport.cpp
--- a/Engine/Source/Editor/UnrealEd/Private/SAssetEditorViewport.cpp
+++ b/Engine/Source/Editor/UnrealEd/Private/SAssetEditorViewport.cpp
@@ -6,6 +6,7 @@
 #include "UnrealEdGlobals.h"
 #include "EditorViewportCommands.h"
 #include "EditorViewportTabContent.h"
+#include "Framework/MultiBox/MultiBoxBuilder.h"
 
 #define LOCTEXT_NAMESPACE "SAssetEditorViewport"
 
