Significantly refactored PIE startup  and consolidated a fair bit of the existing code into structures that better reflect the lifespan that they are valid for. This refactoring will make it easier for systems that are requesting PIE startup to have better control over the instance that is started up. This refactoring also changes network testing to focus on the role you want to test (standalone, client, listen server) and the number of clients you want to test instead of of trying to determine this information based on other variables.

diff --git a/Engine/Plugins/Lumin/MagicLeap/Source/MagicLeap/Private/MagicLeapHMD.cpp b/Engine/Plugins/Lumin/MagicLeap/Source/MagicLeap/Private/MagicLeapHMD.cpp
--- a/Engine/Plugins/Lumin/MagicLeap/Source/MagicLeap/Private/MagicLeapHMD.cpp
+++ b/Engine/Plugins/Lumin/MagicLeap/Source/MagicLeap/Private/MagicLeapHMD.cpp
@@ -570,7 +570,7 @@ bool FMagicLeapHMD::IsHeadTrackingAllowed() const
 	if (GIsEditor)
 	{
 		UEditorEngine* EdEngine = Cast<UEditorEngine>(GEngine);
-		return ((!EdEngine || EdEngine->bUseVRPreviewForPlayWorld || GetDefault<ULevelEditorPlaySettings>()->ViewportGetsHMDControl) && GEngine->IsStereoscopic3D());
+		return ((!EdEngine || EdEngine->IsVRPreviewActive() || GetDefault<ULevelEditorPlaySettings>()->ViewportGetsHMDControl) && GEngine->IsStereoscopic3D());
 	}
 #endif // WITH_EDITOR
 	return (GEngine->IsStereoscopic3D());
diff --git a/Engine/Plugins/Tests/EditorTests/Source/EditorTests/Private/UnrealEd/EditorBuildPromotionTests.cpp b/Engine/Plugins/Tests/EditorTests/Source/EditorTests/Private/UnrealEd/EditorBuildPromotionTests.cpp
--- a/Engine/Plugins/Tests/EditorTests/Source/EditorTests/Private/UnrealEd/EditorBuildPromotionTests.cpp
+++ b/Engine/Plugins/Tests/EditorTests/Source/EditorTests/Private/UnrealEd/EditorBuildPromotionTests.cpp
@@ -395,8 +395,15 @@ namespace EditorBuildPromotionTestUtils
 	{
 		FLevelEditorModule& LevelEditorModule = FModuleManager::Get().GetModuleChecked<FLevelEditorModule>(TEXT("LevelEditor"));
 		TSharedPtr<class IAssetViewport> ActiveLevelViewport = LevelEditorModule.GetFirstActiveViewport();
+		
+		FRequestPlaySessionParams SessionParams;
+		SessionParams.DestinationSlateViewport = ActiveLevelViewport;
+		if (bSimulateInEditor)
+		{
+			SessionParams.WorldType = EPlaySessionWorldType::SimulateInEditor;
+		}
 
-		GUnrealEd->RequestPlaySession(false, ActiveLevelViewport, bSimulateInEditor, NULL, NULL, -1, false);
+		GUnrealEd->RequestPlaySession(SessionParams);
 	}
 
 	/**
diff --git a/Engine/Source/Developer/GameplayDebugger/Private/GameplayDebuggerAddonBase.cpp b/Engine/Source/Developer/GameplayDebugger/Private/GameplayDebuggerAddonBase.cpp
--- a/Engine/Source/Developer/GameplayDebugger/Private/GameplayDebuggerAddonBase.cpp
+++ b/Engine/Source/Developer/GameplayDebugger/Private/GameplayDebuggerAddonBase.cpp
@@ -36,7 +36,10 @@ bool FGameplayDebuggerAddonBase::IsSimulateInEditor()
 {
 #if WITH_EDITOR
 	extern UNREALED_API UEditorEngine* GEditor;
-	return GIsEditor && (GEditor->bIsSimulateInEditorQueued || GEditor->bIsSimulatingInEditor);
+	if (GEditor)
+	{
+		return GEditor->IsSimulateInEditorInProgress();
+	}
 #endif
 	return false;
 }
diff --git a/Engine/Source/Developer/HotReload/Private/HotReload.cpp b/Engine/Source/Developer/HotReload/Private/HotReload.cpp
--- a/Engine/Source/Developer/HotReload/Private/HotReload.cpp
+++ b/Engine/Source/Developer/HotReload/Private/HotReload.cpp
@@ -1442,7 +1442,7 @@ bool FHotReloadModule::Tick(float DeltaTime)
 	if (GEditor)
 	{
 		// Don't try to do an IDE reload yet if we're PIE - wait until we leave
-		if (GEditor->bIsPlayWorldQueued || GEditor->PlayWorld)
+		if (GEditor->IsPlaySessionInProgress())
 		{
 			return true;
 		}
diff --git a/Engine/Source/Editor/LevelEditor/Private/LevelEditor.cpp b/Engine/Source/Editor/LevelEditor/Private/LevelEditor.cpp
--- a/Engine/Source/Editor/LevelEditor/Private/LevelEditor.cpp
+++ b/Engine/Source/Editor/LevelEditor/Private/LevelEditor.cpp
@@ -484,38 +484,32 @@ void FLevelEditorModule::SetLevelEditorTabManager( const TSharedPtr<SDockTab>& O
 
 void FLevelEditorModule::StartPlayInEditorSession()
 {
-	TSharedPtr<SLevelViewport> ActiveLevelViewport = GetFirstActiveLevelViewport();
+	TSharedPtr<IAssetViewport> ActiveLevelViewport = GetFirstActiveViewport();
+	FRequestPlaySessionParams SessionParams;
 
 	if( ActiveLevelViewport.IsValid() )
 	{
-		const FVector* StartLocation = NULL;
-		const FRotator* StartRotation = NULL;
-
 		// We never want to play from the camera's location at startup, because the camera could have
 		// been abandoned in a strange location in the map
 		if( 0 )	// @todo immersive
 		{
 			// If this is a perspective viewport, then we'll Play From Here
-			const FLevelEditorViewportClient& LevelViewportClient = ActiveLevelViewport->GetLevelViewportClient();
+			const FEditorViewportClient& LevelViewportClient = ActiveLevelViewport->GetAssetViewportClient();
 			if( LevelViewportClient.IsPerspective() )
 			{
 				// Start PIE from the camera's location and orientation!
-				StartLocation = &LevelViewportClient.GetViewLocation();
-				StartRotation = &LevelViewportClient.GetViewRotation();
+				SessionParams.StartLocation = LevelViewportClient.GetViewLocation();
+				SessionParams.StartRotation = LevelViewportClient.GetViewRotation();
 			}
 		}
 
-		// Queue up the PIE session
-		const bool bSimulateInEditor = false;
-		const bool bUseMobilePreview = false;
-		GUnrealEd->RequestPlaySession( true, ActiveLevelViewport, bSimulateInEditor, StartLocation, StartRotation, -1, bUseMobilePreview );
+		SessionParams.DestinationSlateViewport = ActiveLevelViewport;
+
+		GUnrealEd->RequestPlaySession(SessionParams);
+
 		// Kick off the queued PIE session immediately.  This is so that at startup, we don't need to
 		// wait for the next engine tick.  We want to see PIE gameplay when the editor first appears!
-		GUnrealEd->StartQueuedPlayMapRequest();
-
-		// Special case for immersive pie startup, When in immersive pie at startup we use the player start but we want to move the camera where the player
-		// was at when pie ended.
-		GEditor->bHasPlayWorldPlacement = true;
+		GUnrealEd->StartQueuedPlaySessionRequest();
 	}
 }
 
diff --git a/Engine/Source/Editor/LevelEditor/Private/LevelEditorActions.cpp b/Engine/Source/Editor/LevelEditor/Private/LevelEditorActions.cpp
--- a/Engine/Source/Editor/LevelEditor/Private/LevelEditorActions.cpp
+++ b/Engine/Source/Editor/LevelEditor/Private/LevelEditorActions.cpp
@@ -1065,7 +1065,7 @@ bool FLevelEditorActionCallbacks::Recompile_CanExecute()
 #endif
 
 	// We can't recompile while in PIE
-	if (GEditor->bIsPlayWorldQueued || GEditor->PlayWorld)
+	if (GEditor->IsPlaySessionInProgress())
 	{
 		return false;
 	}
diff --git a/Engine/Source/Editor/MovieSceneCaptureDialog/Private/MovieSceneCaptureDialogModule.cpp b/Engine/Source/Editor/MovieSceneCaptureDialog/Private/MovieSceneCaptureDialogModule.cpp
--- a/Engine/Source/Editor/MovieSceneCaptureDialog/Private/MovieSceneCaptureDialogModule.cpp
+++ b/Engine/Source/Editor/MovieSceneCaptureDialog/Private/MovieSceneCaptureDialogModule.cpp
@@ -441,7 +441,26 @@ void FInEditorCapture::Start()
 		AudioDevice->SetTransientMasterVolume(0.0f);
 	}
 
-	GEditor->RequestPlaySession(true, nullptr, false);
+	TSharedRef<SWindow> CustomWindow = SNew(SWindow)
+		.Title(LOCTEXT("MovieRenderPreviewTitle", "Movie Render - Preview"))
+		.AutoCenter(EAutoCenter::PrimaryWorkArea)
+		.UseOSWindowBorder(true)
+		.FocusWhenFirstShown(false)
+		.ActivationPolicy(EWindowActivationPolicy::Never)
+		.HasCloseButton(true)
+		.SupportsMaximize(false)
+		.SupportsMinimize(true)
+		.MaxWidth(CaptureObject->GetSettings().Resolution.ResX)
+		.MaxHeight(CaptureObject->GetSettings().Resolution.ResY)
+		.SizingRule(ESizingRule::FixedSize);
+
+	FSlateApplication::Get().AddWindow(CustomWindow);
+
+	FRequestPlaySessionParams Params;
+	Params.EditorPlaySettings = PlayInEditorSettings;
+	Params.CustomPIEWindow = CustomWindow;
+
+	GEditor->RequestPlaySession(Params);
 }
 
 void FInEditorCapture::Cancel()
@@ -461,23 +480,6 @@ void FInEditorCapture::OverridePlaySettings(ULevelEditorPlaySettings* PlayInEdit
 	PlayInEditorSettings->CenterNewWindow = true;
 	PlayInEditorSettings->LastExecutedPlayModeType = EPlayModeType::PlayMode_InEditorFloating;
 
-	TSharedRef<SWindow> CustomWindow = SNew(SWindow)
-		.Title(LOCTEXT("MovieRenderPreviewTitle", "Movie Render - Preview"))
-		.AutoCenter(EAutoCenter::PrimaryWorkArea)
-		.UseOSWindowBorder(true)
-		.FocusWhenFirstShown(false)
-		.ActivationPolicy(EWindowActivationPolicy::Never)
-		.HasCloseButton(true)
-		.SupportsMaximize(false)
-		.SupportsMinimize(true)
-		.MaxWidth( Settings.Resolution.ResX )
-		.MaxHeight( Settings.Resolution.ResY )
-		.SizingRule(ESizingRule::FixedSize);
-
-	FSlateApplication::Get().AddWindow(CustomWindow);
-
-	PlayInEditorSettings->CustomPIEWindow = CustomWindow;
-
 	// Reset everything else
 	PlayInEditorSettings->GameGetsMouseControl = false;
 	PlayInEditorSettings->ShowMouseControlLabel = false;
@@ -493,7 +495,7 @@ void FInEditorCapture::OverridePlaySettings(ULevelEditorPlaySettings* PlayInEdit
 	PlayInEditorSettings->LaunchConfiguration = EPlayOnLaunchConfiguration::LaunchConfig_Default;
 	PlayInEditorSettings->SetPlayNetMode(EPlayNetMode::PIE_Standalone);
 	PlayInEditorSettings->SetRunUnderOneProcess(true);
-	PlayInEditorSettings->SetPlayNetDedicated(false);
+	PlayInEditorSettings->bLaunchSeparateServer = false;
 	PlayInEditorSettings->SetPlayNumberOfClients(1);
 }
 
diff --git a/Engine/Source/Editor/UnrealEd/Classes/Editor/EditorEngine.h b/Engine/Source/Editor/UnrealEd/Classes/Editor/EditorEngine.h
--- a/Engine/Source/Editor/UnrealEd/Classes/Editor/EditorEngine.h
+++ b/Engine/Source/Editor/UnrealEd/Classes/Editor/EditorEngine.h
@@ -25,7 +25,7 @@
 #include "Misc/CompilationResult.h"
 #include "Interfaces/ITargetPlatform.h"
 #include "Interfaces/ITargetPlatformManagerModule.h"
-
+#include "PlayInEditorDataTypes.h"
 #include "EditorSubsystem.h"
 #include "Subsystems/SubsystemCollection.h"
 #include "RHI.h"
@@ -91,83 +91,5 @@
 };
 
 USTRUCT()
-struct FSlatePlayInEditorInfo
-{
-	GENERATED_USTRUCT_BODY()
-
-	/** The spawned player for updating viewport location from player when pie closes */
-	TWeakObjectPtr<class ULocalPlayer>	EditorPlayer;
-
-	/** The current play in editor SWindow if playing in a floating window */
-	TWeakPtr<class SWindow>				SlatePlayInEditorWindow;
-	
-	/** The current play in editor rendering and I/O viewport if playing in a floating window*/
-	TSharedPtr<class FSceneViewport>	SlatePlayInEditorWindowViewport;
-	
-	/** The slate viewport that should be used for play in viewport */
-	TWeakPtr<class IAssetViewport>		DestinationSlateViewport;
-
-	FSlatePlayInEditorInfo()
-	: SlatePlayInEditorWindow(NULL), DestinationSlateViewport(NULL)
-	{}
-};
-
-/**
- * Data structure for storing PIE login credentials
- */
-USTRUCT()
-struct FPIELoginInfo
-{
-public:
-
-	GENERATED_USTRUCT_BODY()
-
-	/** Type of account. Needed to identity the auth method to use (epic, internal, facebook, etc) */
-	UPROPERTY()
-	FString Type;
-	/** Id of the user logging in (email, display name, facebook id, etc) */
-	UPROPERTY()
-	FString Id;
-	/** Credentials of the user logging in (password or auth token) */
-	UPROPERTY()
-	FString Token;
-};
-
-/**
- * Holds various data to pass to the post login delegate for PIE logins
- */
-struct FPieLoginStruct
-{
-	/** World context handle for this login */
-	FName WorldContextHandle;
-	/** Setting index for window positioning */
-	int32 SettingsIndex;
-	/** X location for window positioning */
-	int32 NextX;
-	/** Y location for window positioning */
-	int32 NextY;
-	/** What net mode to run this instance as */
-	EPlayNetMode NetMode;
-	/** Passthrough condition of blueprint compilation*/
-	bool bAnyBlueprintErrors;
-	/** Passthrough condition of spectator mode */
-	bool bStartInSpectatorMode;
-	/** Passthrough start time of PIE */
-	double PIEStartTime;
-
-	FPieLoginStruct() :
-		WorldContextHandle(NAME_None),
-		SettingsIndex(0),
-		NextX(0),
-		NextY(0),
-		NetMode(EPlayNetMode::PIE_Standalone),
-		bAnyBlueprintErrors(false),
-		bStartInSpectatorMode(false),
-		PIEStartTime(0)
-	{
-	}
-};
-
-USTRUCT()
 struct FCopySelectedInfo
 {
@@ -348,19 +270,6 @@
 	TArray<FString> SelectedComponents;
 };
 
-/** Overrides you can pass when starting PIE to temporarily suppress the user's options for dedicated server, number of clients..etc. */
-struct FPlayInEditorOverrides
-{
-	FPlayInEditorOverrides()
-		: bDedicatedServer()
-		, NumberOfClients()
-	{
-	}
-
-	TOptional<bool> bDedicatedServer;
-	TOptional<int32> NumberOfClients;
-};
-
 struct FPreviewPlatformInfo
 {
 	FPreviewPlatformInfo()
@@ -416,7 +325,7 @@
  * Separate from UGameEngine because it may have much different functionality than desired for an instance of a game itself.
  */
 UCLASS(config=Engine, transient)
-class UNREALED_API UEditorEngine : public UEngine
+class UNREALED_API UEditorEngine : public UEngine, public FGCObject
 {
 public:
 	GENERATED_BODY()
@@ -499,6 +408,7 @@
 	UPROPERTY(EditAnywhere, config, Category=Advanced)
 	uint32 UseAxisIndicator:1;
 
+	UE_DEPRECATED(4.25, "This variable is no longer read.")
 	UPROPERTY(EditAnywhere, config, Category=Advanced)
 	uint32 GodMode:1;
 
@@ -543,27 +453,13 @@
 	UPROPERTY()
 	class UWorld* PlayWorld;
 
-	/** An optional location for the starting location for "Play From Here"																*/
-	UPROPERTY()
-	FVector PlayWorldLocation;
 
-	/** An optional rotation for the starting location for "Play From Here"																*/
-	UPROPERTY()
-	FRotator PlayWorldRotation;
 
-	/** Has a request for "Play From Here" been made?													 								*/
-	UPROPERTY()
-	uint32 bIsPlayWorldQueued:1;
-
 	/** Has a request to toggle between PIE and SIE been made? */
 	UPROPERTY()
 	uint32 bIsToggleBetweenPIEandSIEQueued:1;
 
-	/** True if we are requesting to start a simulation-in-editor session */
-	UPROPERTY()
-	uint32 bIsSimulateInEditorQueued:1;
-
-	/** Allows multipel PIE worlds under a single instance. If false, you can only do multiple UE4 processes for pie networking */
+	/** Allows multiple PIE worlds under a single instance. If false, you can only do multiple UE4 processes for pie networking */
 	UPROPERTY(globalconfig)
 	uint32 bAllowMultiplePIEWorlds:1;
 
@@ -571,22 +467,6 @@
 	UPROPERTY()
 	uint32 bRequestEndPlayMapQueued:1;
 
-	/** Did the request include the optional location and rotation?										 								*/
-	UPROPERTY()
-	uint32 bHasPlayWorldPlacement:1;
-
-	/** True to enable mobile preview mode when launching the game from the editor on PC platform */
-	UPROPERTY()
-	uint32 bUseMobilePreviewForPlayWorld:1;
-
-	/** True to enable VR preview mode when launching the game from the editor on PC platform */
-	UPROPERTY()
-	uint32 bUseVRPreviewForPlayWorld:1;
-
-	/** True if we're Simulating In Editor, as opposed to Playing In Editor.  In this mode, simulation takes place right the level editing environment */
-	UPROPERTY()
-	uint32 bIsSimulatingInEditor:1;
-
 	/** True if we should not display notifications about undo/redo */
 	UPROPERTY()
 	uint32 bSquelchTransactionNotification:1;
@@ -623,16 +503,10 @@
 	UPROPERTY(config)
 	int32 BuildPlayDevice;
 
+
 	/** Maps world contexts to their slate data */
 	TMap<FName, FSlatePlayInEditorInfo>	SlatePlayInEditorMap;
 
-	/** Viewport the next PlaySession was requested to happen on */
-	TWeakPtr<class IAssetViewport>		RequestedDestinationSlateViewport;
-
-	/** When set to anything other than -1, indicates a specific In-Editor viewport index that PIE should use */
-	UPROPERTY()
-	int32 PlayInEditorViewportIndex;
-
 	/** Play world url string edited by a user. */
 	UPROPERTY()
 	FString UserEditedPlayWorldURL;
@@ -713,23 +587,6 @@
 	UPROPERTY()
 	UEditorWorldExtensionManager* EditorWorldExtensionsManager;
 
-protected:
-
-	/** Count of how many PIE instances are waiting to log in */
-	int32 PIEInstancesToLogInCount;
-	bool bAtLeastOnePIELoginFailed;
-
-	/* These are parameters that we need to cache for late joining */
-	FString ServerPrefix;
-	int32 PIEInstance;
-	int32 SettingsIndex;
-	bool bStartLateJoinersInSpectatorMode;
-
-private:
-
-	/** Additional launch options requested for the next PlaySession */
-	FString RequestedAdditionalStandaloneLaunchOptions;
-
 public:
 
 	/** The "manager" of all the layers for the UWorld currently being edited */
@@ -938,7 +795,6 @@
 	virtual void HandleNetworkFailure(UWorld *World, UNetDriver *NetDriver, ENetworkFailure::Type FailureType, const FString& ErrorString) override;
 	virtual ERHIFeatureLevel::Type GetDefaultWorldFeatureLevel() const override { return DefaultWorldFeatureLevel; }
 
-	FString GetPlayOnTargetPlatformName() const;
 protected:
 	virtual void InitializeObjectReferences() override;
 	virtual void ProcessToggleFreezeCommand(UWorld* InWorld) override;
@@ -1823,19 +1679,24 @@
 	 * @param	bUseMobilePreview		True to enable mobile preview mode (PC platform only)
 	 * @param	bUseVRPreview			True to enable VR preview mode (PC platform only)
 	 */
+	UE_DEPRECATED(4.25, "Use the overload of RequestPlaySession which takes a FRequestPlaySessionParams instead.")
 	void RequestPlaySession( bool bAtPlayerStart, TSharedPtr<class IAssetViewport> DestinationViewport, bool bInSimulateInEditor, const FVector* StartLocation = NULL, const FRotator* StartRotation = NULL, int32 DestinationConsole = -1, bool bUseMobilePreview = false, bool bUseVRPreview = false, bool bUseVulkanPreview = false);
 
 	// @todo gmp: temp hack for Rocket demo
+	UE_DEPRECATED(4.25, "Use the overload of RequestPlaySession which takes a FRequestPlaySessionParams instead.")
 	void RequestPlaySession(const FVector* StartLocation, const FRotator* StartRotation, bool MobilePreview, bool VulkanPreview, const FString& MobilePreviewTargetDevice, FString AdditionalStandaloneLaunchParameters = TEXT(""));
 
 	/** Request to play a game on a remote device */
+	UE_DEPRECATED(4.25, "Use the overload of RequestPlaySession which takes a FRequestPlaySessionParams instead.")
 	void RequestPlaySession( const FString& DeviceId, const FString& DeviceName );
+	
+	/** Request a play session (Play in Editor, Play in New Process, Launcher) with the configured parameters. See FRequestPlaySessionParams for more details. */
+	void RequestPlaySession(const FRequestPlaySessionParams& InParams);
 
 	/** Cancel request to start a play session */
 	void CancelRequestPlaySession();
 
-	/** Asks the player to save dirty maps, if this fails it will return false and call CancelRequestPlaySession */
-	bool SaveMapsForPlaySession();
 
+
 	/** Makes a request to start a play from a Slate editor session */
 	void RequestToggleBetweenPIEandSIE() { bIsToggleBetweenPIEandSIEQueued = true; }
@@ -1849,14 +1710,45 @@
 	/** Called when the debugger has single-stepped the active PIE or SIE session */
 	void PlaySessionSingleStepped();
 
+	/** Returns true if we are currently either PIE/SIE in the editor, false if we are not (even if we would start next tick). See IsPlaySessionInProgress() */
+	bool IsPlayingSessionInEditor() const { return PlayInEditorSessionInfo.IsSet(); }
+	/** Returns true if we are going to start PIE/SIE on the next tick, false if we are not (or if we are already in progress). See IsPlaySessionInProgress() */
+	bool IsPlaySessionRequestQueued() const { return PlaySessionRequest.IsSet(); }
+	/** Returns true if Playing in Editor, Simulating in Editor, or are either of these queued to start on the next tick. */
+	bool IsPlaySessionInProgress() const { return IsPlayingSessionInEditor() || IsPlaySessionRequestQueued(); }
+	
+	/** Returns true if we are currently Simulating in Editor, false if we are not (even if we would next tick). See IsSimulateInEditorInProgress() */
+	bool IsSimulatingInEditor() const { return PlayInEditorSessionInfo.IsSet() && PlayInEditorSessionInfo->OriginalRequestParams.WorldType == EPlaySessionWorldType::SimulateInEditor; }	
+	/** Returns true if we are going to start Simulating in Editor on the next tick, false if we are not (or we are already simulating). See IsSimulateInEditorInProgress() */
+	bool IsSimulateInEditorQueued() const { return PlaySessionRequest.IsSet() && PlaySessionRequest->WorldType == EPlaySessionWorldType::SimulateInEditor; }
+	/** Returns true if we are currently Simulating in Editor, or are going to start Simulating in Editor on the next tick. */
+	bool IsSimulateInEditorInProgress() const { return IsSimulatingInEditor() || IsSimulateInEditorQueued(); }
+
+	/** Returns the currently set PlaySessionRequest (if any). Unset if there is no session queued or has already been started (see GetPlayInEditorSessionInfo()) */
+	const TOptional<FRequestPlaySessionParams> GetPlaySessionRequest() const { return PlaySessionRequest; }
+
+	/** Returns information about the currently active Play In Editor session. Unset if there is no active PIE session. */
+	const TOptional<FPlayInEditorSessionInfo> GetPlayInEditorSessionInfo() const { return PlayInEditorSessionInfo; }
+	
+	/** Returns true if the user has queued a Play in Editor (in VR Mode) or one is currently running. */
+	bool IsVRPreviewActive() const
+	{
+		return PlaySessionRequest.IsSet() && PlaySessionRequest->SessionPreviewTypeOverride == EPlaySessionPreviewType::VRPreview || 
+			PlayInEditorSessionInfo.IsSet() && PlayInEditorSessionInfo->OriginalRequestParams.SessionPreviewTypeOverride == EPlaySessionPreviewType::VRPreview;
+	}
+
 	/** Called when game client received input key */
 	bool ProcessDebuggerCommands(const FKey InKey, const FModifierKeysState ModifierKeyState, EInputEvent EventType);
 
 	/**
 	 * Kicks off a "Play From Here" request that was most likely made during a transaction
 	 */
+	UE_DEPRECATED(4.25, "Call StartQueuedPlaySessionRequest, or override StartQueuedPlaySessionRequestImpl instead.")
 	virtual void StartQueuedPlayMapRequest();
 
+	/** Kicks off a "Request Play Session" request if there is one set. This should not be called mid-transaction. */
+	void StartQueuedPlaySessionRequest();
+
 	/**
 	 * Request that the current PIE/SIE session should end.  
 	 * This should be used to end the game if its not safe to directly call EndPlayMap in your stack frame
@@ -1889,8 +1781,12 @@
 	 * @param	InWorld		World context
 	 * @param	bInSimulateInEditor	True to start an in-editor simulation session, or false to start a play-in-editor session
 	 */
+	UE_DEPRECATED(4.25, "Override StartPlayInEditorSession(FRequestPlaySessionParams& InRequestParams) instead. InWorld is GetEditorWorldContext().World(), and the other arguments come from InRequestParams now.") 
+	PRAGMA_DISABLE_DEPRECATION_WARNINGS
 	virtual void PlayInEditor( UWorld* InWorld, bool bInSimulateInEditor, FPlayInEditorOverrides Overrides = FPlayInEditorOverrides());
+	PRAGMA_ENABLE_DEPRECATION_WARNINGS
 
+	UE_DEPRECATED(4.25, "Override CreateInnerProcessPIEGameInstance instead.")
 	virtual UGameInstance* CreatePIEGameInstance(int32 InPIEInstance, bool bInSimulateInEditor, bool bAnyBlueprintErrors, bool bStartInSpectatorMode, bool bPlayNetDedicated, bool bPlayStereoscopic, float PIEStartTime);
 
 	/**
@@ -2600,27 +2496,6 @@
 	 */
 	void UnregisterForUndo( class FEditorUndoClient* UndoEditor );
 
-	/** 
-	 * Are we playing on a local PC session?
-	 */
-	bool IsPlayingOnLocalPCSession() const { return bPlayOnLocalPcSession && !bIsPlayWorldQueued; }
-
-	/** 
-	 * Are we playing via the Launcher?
-	 */
-	bool IsPlayingViaLauncher() const { return bPlayUsingLauncher && !bIsPlayWorldQueued; }
-
-	/** 
-	 * Cancel playing via the Launcher
-	 */
-	void CancelPlayingViaLauncher();
-
-	/** @return true if the editor is able to launch PIE with online platform support */
-	bool SupportsOnlinePIE() const;
-
-	/** @return true if there are active PIE instances logged into an online platform */
-	bool IsPlayingWithOnlinePIE() const { return NumOnlinePIEInstances > 0; }
-
 	/**
 	 * Ensures the assets specified are loaded and adds them to the global selection set
 	 * @param Assets		An array of assets to load and select
@@ -2633,11 +2508,6 @@
 	 */
 	bool UsePercentageBasedScaling() const;
 
-	DECLARE_DELEGATE(FPIEInstanceWindowSwitch);
-
-	/** Sets the delegate for when the focused PIE window is changed */
-	void SetPIEInstanceWindowSwitchDelegate(FPIEInstanceWindowSwitch PIEInstanceWindowSwitchDelegate);
-
 	/**
 	 * Returns the actor grouping utility class that performs all grouping related tasks
 	 * This will create the class instance if it doesn't exist.
@@ -2694,24 +2564,12 @@
 	 * @param bStartInSpectatorMode passthrough value if it is expected that PIE will start in spectator mode
 	 * @param PIEStartTime passthrough value of the time that PIE was initiated by the user
 	 */
+	UE_DEPRECATED(4.25, "This is now done as part of StartPlayInEditorSession as online-login/non-logged in flow is now combined.")
 	virtual void LoginPIEInstances(bool bAnyBlueprintErrors, bool bStartInSpectatorMode, double PIEStartTime);
 
-	/**
-	 * Delegate called as each PIE instance login is complete, continues creating the PIE world for a single instance
-	 * 
-	 * @param LocalUserNum local user id, for PIE is going to be 0 (there is no splitscreen)
-	 * @param bWasSuccessful was the login successful
-	 * @param ErrorString descriptive error when applicable
-	 * @param DataStruct data required to continue PIE creation, set at login time
-	 */
-	virtual void OnLoginPIEComplete(int32 LocalUserNum, bool bWasSuccessful, const FString& ErrorString, FPieLoginStruct DataStruct);
-
-	/** Above function but called a frame later, to stop PIE login from happening from a network callback */
-	virtual void OnLoginPIEComplete_Deferred(int32 LocalUserNum, bool bWasSuccessful, FString ErrorString, FPieLoginStruct DataStruct);
-
 	/** Called when all PIE instances have been successfully logged in */
+	UE_DEPRECATED(4.25, "Override OnAllPIEInstancesStarted instead, which is called in both online and offline PIE sessions.")
 	virtual void OnLoginPIEAllComplete();
-
 	/** Called when a recompilation of a module is beginning */
 	void OnModuleCompileStarted(bool bIsAsyncCompile);
 
@@ -2730,6 +2588,7 @@
 	 * @param DataStruct data required to continue PIE creation, set at login time
 	 * @return	true if world created successfully
 	 */
+	UE_DEPRECATED(4.25, "This is now handled as part of OnLoginPIEComplete_Deferred.")
 	bool CreatePIEWorldFromLogin(FWorldContext& PieWorldContext, EPlayNetMode PlayNetMode, FPieLoginStruct& DataStruct);
 
 	/** Called before creating a PIE server instance. */
@@ -2750,53 +2609,35 @@
 	 */
 	bool GetSimulateInEditorViewTransform(FTransform& OutViewTransform) const;
 
-private:
-	/** Gets the scene viewport for a viewport client */
-	FSceneViewport* GetGameSceneViewport(UGameViewportClient* ViewportClient) const;
-
+public:
 	/**
 	 * Non Online PIE creation flow, creates all instances of PIE at once when online isn't requested/required
 	 *
 	 * @param bAnyBlueprintErrors passthrough value of blueprint errors encountered during PIE creation
 	 * @param bStartInSpectatorMode passthrough value if it is expected that PIE will start in spectator mode
 	 */
+	UE_DEPRECATED(4.25, "This non-online flow is now deprecated and is shared with normal online flow. See OnLoginPIEComplete_Deferred instead.")
 	void SpawnIntraProcessPIEWorlds(bool bAnyBlueprintErrors, bool bStartInSpectatorMode);
 
-	/** Common init shared by CreatePIEWorldByDuplication and CreatePIEWorldBySavingToTemp */
-	void PostCreatePIEWorld(UWorld *InWorld);
-
-	/**
-	 * Toggles PIE to SIE or vice-versa
-	 */
-	void ToggleBetweenPIEandSIE( bool bNewSession = false);
-
-	/**
-	 * Hack to switch worlds for the PIE window before and after a slate event
-	 *
-	 * @param WorldID	The id of the world to restore or -1 if no world
-	 * @return The ID of the world to restore later or -1 if no world to restore
-	 */
-	int32 OnSwitchWorldForSlatePieWindow( int32 WorldID );
-
-	/**
-	 * Called via a delegate to toggle between the editor and pie world
-	 */
-	void OnSwitchWorldsForPIE( bool bSwitchToPieWorld, UWorld* OverrideWorld = nullptr );
-
-	/**
-	 * Gives focus to the server or first PIE client viewport
-	 */
-	void GiveFocusToFirstClientPIEViewport();
-
-public:
-	/**
+	/*
 	 * Spawns a PlayFromHere playerstart in the given world
 	 * @param	World		The World to spawn in (for PIE this may not be GWorld)
 	 * @param	PlayerStartPIE	A reference to the resulting PlayerStartPIE actor
+	 * @param	StartLocation	The location to spawn the player in.
+	 * @param	StartRotation	The rotation to spawn the player at.
 	 *
 	 * @return	true if spawn succeeded, false if failed
 	 */
-	bool SpawnPlayFromHereStart(UWorld* World, AActor*& PlayerStartPIE, const FVector& StartLocation, const FRotator& StartRotation );
+	bool SpawnPlayFromHereStart(UWorld* World, AActor*& PlayerStartPIE, const FVector& StartLocation, const FRotator& StartRotation);
+	
+	/**
+	 * Spawns a PlayFromHere playerstart in the given world if there is a location set to spawn in.
+	 * @param	World		The World to spawn in (for PIE this may not be GWorld)
+	 * @param	PlayerStartPIE	A reference to the resulting PlayerStartPIE actor
+	 *
+	 * @return	true if spawn succeeded, false if failed. Returns true even if no object was spawned (due to location not being set).
+	 */
+	bool SpawnPlayFromHereStart(UWorld* World, AActor*& PlayerStartPIE);
 
 
 private:
@@ -3000,71 +2841,217 @@
 	/** The Timer manager for all timer delegates */
 	TSharedPtr<class FTimerManager> TimerManager;
 
-	/** The output log -> message log redirector for use during PIE */
-	TSharedPtr<class FOutputLogErrorsToMessageLogProxy> OutputLogErrorsToMessageLogProxyPtr;
+	// This chunk is used for Play In New Process
+public:
+	/**
+	 * Are we playing on a local PC session?
+	 */
+	bool IsPlayingOnLocalPCSession() const { return PlayOnLocalPCSessions.Num() > 0; }
+	
+	/** 
+	 * Launch a new process with settings pulled from the specified parameters.
+	 * @param InParams 				Shared parameters to launch with.
+	 * @param InInstanceNum 		Which instance index is this? Used to do load settins per-instance if needed.
+	 * @param bIsHost				Is this instance a host (either dedicated or listen).
+	 * @param bIsDedicatedServer 	Is this instance a dedicated server? If true, implies bIsHost.
+	 */
+	virtual void LaunchNewProcess(const FRequestPlaySessionParams& InParams, const int32 InInstanceNum, bool bIsHost, bool bIsDedicatedServer);
 
+	// This chunk is used for Play In New Process
+protected:
 	struct FPlayOnPCInfo
 	{
 		FProcHandle ProcessHandle;
 	};
 
+	/** List of New Process sessions that are currently active, started by this Editor instance. */
 	TArray<FPlayOnPCInfo> PlayOnLocalPCSessions;
 
-	bool bPlayOnLocalPcSession;
+	/** Start a Play in New Process session with the given parameters. Called by StartQueuedPlaySessionRequestImpl based on request settings. */
+	virtual void StartPlayInNewProcessSession(FRequestPlaySessionParams& InRequestParams);
+	
+	// This chunk is used for Play Via Launcher settings.
+public:
+	/**
+	 * Are we playing via the Launcher?
+	 */
+	bool IsPlayingViaLauncher() const { return LauncherSessionInfo.IsSet(); }
 
-	/** True if we are using the "Play On Device" launcher mode (ie UAT) */
-	// @todo: This should be an enum along with bPlayOnLocalPcSession, or completely refactored
-	bool bPlayUsingLauncher;
-	bool bPlayUsingMobilePreview;
-	bool bPlayUsingVulkanPreview;
-	FString PlayUsingMobilePreviewTargetDevice;
+	/**
+	 * Cancel playing via the Launcher
+	 */
+	void CancelPlayingViaLauncher();
 
-	/** The platform to run on (as selected in dreop down) */
-	FString PlayUsingLauncherDeviceId;
-	FString PlayUsingLauncherDeviceName;
-	bool bPlayUsingLauncherHasCode;
-	bool bPlayUsingLauncherBuild;
+	/**
+	* Returns the last device that was used via the Launcher.
+	*/
+	FString GetPlayOnTargetPlatformName() const;
 
+	/**
+	* Launch a standalone instance using the Launcher to process required
+	* cooking, building, and deployment to a (possibly) remote device for testing.
+	*/
+	UE_DEPRECATED(4.25, "Use RequestPlaySession with parameters configured to provide a launcher target instead.")
+	void PlayUsingLauncher();
+
+	/** 
+	* Cancel Play using Launcher on error 
+	* 
+	* if the physical device is not authorized to be launched to, we need to pop an error instead of trying to launch
+	*/
+	void CancelPlayUsingLauncher();
+
+protected:
+	struct FLauncherCachedInfo
+	{
+		FString PlayUsingLauncherDeviceName;
+		bool bPlayUsingLauncherHasCode;
+		bool bPlayUsingLauncherBuild;
+	};
+
+	/** 
+	* If set, there is a active attempt to use the Launcher. Launching is an async process so we store the request info for the duration.
+	* of the launch. This is cleared on cancel (or on finish).
+	*/
+	TOptional<FLauncherCachedInfo> LauncherSessionInfo;
+
+	/** 
+	* The last platform we ran on (as selected by the drop down or via automation)
+	* Stored outside of the Session Info as the UI needs access to this at all times.
+	*/
+	FString LastPlayUsingLauncherDeviceId;
+	
+	/** Start a Launcher session with the given parameters. Called by StartQueuedPlaySessionRequestImpl based on request settings. */
+	virtual void StartPlayUsingLauncherSession(FRequestPlaySessionParams& InRequestParams);
+
+	// This chunk is for Play in Editor
+public:
+	/** @return true if the editor is able to launch PIE with online platform support */
+	bool SupportsOnlinePIE() const;
+
+	/** @return true if there are active PIE instances logged into an online platform */
+	bool IsPlayingWithOnlinePIE() const
+	{
+		return PlayInEditorSessionInfo.IsSet() && PlayInEditorSessionInfo->bUsingOnlinePlatform;
+	}
+
+protected:
+
+	/** Specifies whether or not there has been a Play Session requested to be kicked off on the next Tick. */
+	TOptional<FRequestPlaySessionParams> PlaySessionRequest;
+	
+	/** If set, transient information about the current Play in Editor session which persists until the session stops. */
+	TOptional<FPlayInEditorSessionInfo> PlayInEditorSessionInfo;
+	
 	/** Used to prevent reentrant calls to EndPlayMap(). */
 	bool bIsEndingPlay;
 
-	/** List of files we are deferring adding to source control */
-	TArray<FString> DeferredFilesToAddToSourceControl;
+protected:
+	DECLARE_DELEGATE(FPIEInstanceWindowSwitch);
 
+	/** Sets the delegate for when the focused PIE window is changed */
+	void SetPIEInstanceWindowSwitchDelegate(FPIEInstanceWindowSwitch PIEInstanceWindowSwitchDelegate);
+
+	/** Gets the scene viewport for a viewport client */
+	FSceneViewport* GetGameSceneViewport(UGameViewportClient* ViewportClient) const;
+
+	/** Common init shared by CreatePIEWorldByDuplication and CreatePIEWorldBySavingToTemp */
+	void PostCreatePIEWorld(UWorld *InWorld);
+
+	/**
+	 * Toggles PIE to SIE or vice-versa
+	 */
+	void ToggleBetweenPIEandSIE(bool bNewSession = false);
+
+	/**
+	 * Hack to switch worlds for the PIE window before and after a slate event
+	 *
+	 * @param WorldID	The id of the world to restore or -1 if no world
+	 * @return The ID of the world to restore later or -1 if no world to restore
+	 */
+	int32 OnSwitchWorldForSlatePieWindow(int32 WorldID);
+
+	/**
+	 * Called via a delegate to toggle between the editor and pie world
+	 */
+	void OnSwitchWorldsForPIE(bool bSwitchToPieWorld, UWorld* OverrideWorld = nullptr);
+
+	/**
+	 * Gives focus to the server or first PIE client viewport
+	 */
+	void GiveFocusToLastClientPIEViewport();
+
+	/**
+	 * Delegate called as each PIE instance login is complete, continues creating the PIE world for a single instance
+	 *
+	 * @param LocalUserNum local user id, for PIE is going to be 0 (there is no splitscreen)
+	 * @param bWasSuccessful was the login successful
+	 * @param ErrorString descriptive error when applicable
+	 * @param DataStruct data required to continue PIE creation, set at login time
+	 */
+	virtual void OnLoginPIEComplete(int32 LocalUserNum, bool bWasSuccessful, const FString& ErrorString, FPieLoginStruct DataStruct);
+
+	/** Above function but called a frame later, to stop PIE login from happening from a network callback */
+	virtual void OnLoginPIEComplete_Deferred(int32 LocalUserNum, bool bWasSuccessful, FString ErrorString, FPieLoginStruct DataStruct);
+
+	/** Called when all PIE instances have been successfully logged in */
+	virtual void OnAllPIEInstancesStarted();
+	
+	/** Transfers the current Editor Selection to their SIE counterparts. */
+	void TransferEditorSelectionToSIEInstances();
+	
+	/** Create a new GameInstance for PIE with the specified parameters. */
+	virtual UGameInstance* CreateInnerProcessPIEGameInstance(FRequestPlaySessionParams& InParams, const FGameInstancePIEParameters& InPIEParameters, int32 InPIEInstanceIndex);
+	
+	/** Creates a SPIEViewport for the specified instance. This is used for Play in New Window, but not for instances that live in the level editor viewport. */
+	TSharedRef<class SPIEViewport> GeneratePIEViewportWindow(const FRequestPlaySessionParams& InSessionParams, int32 InViewportIndex, const FWorldContext& InWorldContext, EPlayNetMode InNetMode, UGameViewportClient* InViewportClient, FSlatePlayInEditorInfo& InSlateInfo);
+
+	/** Stores the position of the window for that instance index. Doesn't save it to the CDO until the play session ends. */
+	void StoreWindowSizeAndPositionForInstanceIndex(const int32 InInstanceIndex, const FIntPoint& InSize, const FIntPoint& InPosition);
+	/** Returns the stored position at that index. */
+	void GetWindowSizeAndPositionForInstanceIndex(ULevelEditorPlaySettings& InEditorPlaySettings, const int32 InInstanceIndex, FIntPoint& OutSize, FIntPoint& OutPosition);
+
+	/** Start the queued Play Session Request. After this is called the queued play session request will be cleared. */
+	virtual void StartQueuedPlaySessionRequestImpl();
+	
+	/** Start a Play in Editor session with the given parameters. Called by StartQueuedPlaySessionRequestImpl based on request settings. */
+	virtual void StartPlayInEditorSession(FRequestPlaySessionParams& InRequestParams);
+
+	/** 
+	* Creates a new Play in Editor instance (which may be in a new process if not running under one process.
+	* This reads the current session state to start the next instance needed.
+	* @param	InRequestParams			- Mostly used for global settings. Probably shouldn't vary per instance.
+	* @param	bInDedicatedInstance	- If true, this will create a dedicated server instance. Shouldn't be used if we're using multiple processes.
+	* @param	InNetMode				- The net mode to launch this instance with. Not pulled from the RequestParams because some net modes (such as ListenServer)
+										  need different net modes for subsequent instances.
+	*/
+	virtual void CreateNewPlayInEditorInstance(FRequestPlaySessionParams &InRequestParams, const bool bInDedicatedInstance, const EPlayNetMode InNetMode);
+
+
+	/** Asks the player to save dirty maps, if this fails it will return false. */
+	bool SaveMapsForPlaySession();
+	
+protected:
+
 	FPIEInstanceWindowSwitch PIEInstanceWindowSwitchDelegate;
 
-	/** Number of currently running instances logged into an online platform */
-	int32 NumOnlinePIEInstances;
-
 	/** Cached version of the view location at the point the PIE session was ended */
 	FVector LastViewLocation;
-	
+
 	/** Cached version of the view location at the point the PIE session was ended */
 	FRotator LastViewRotation;
-	
+
 	/** Are the lastview/rotation variables valid */
 	bool bLastViewAndLocationValid;
 
+	/** The output log -> message log redirector for use during PIE */
+	TSharedPtr<class FOutputLogErrorsToMessageLogProxy> OutputLogErrorsToMessageLogProxyPtr;
+
+private:
+	/** List of files we are deferring adding to source control */
+	TArray<FString> DeferredFilesToAddToSourceControl;
+
 protected:
 
-	/**
-	 * Launch a standalone instance on this PC.
-	 *
-	 * @param	MapNameOverride		Map name override
-	 * @param	WindowPos			Position we want to put the window we create.
-	 * @param	PIENum				PIE instance count
-	 * @param	bIsServer			Is this instance a server.
-	 */
-	void PlayStandaloneLocalPc(FString MapNameOverride = FString(), FIntPoint* WindowPos = NULL, int32 PIENum = 0, bool bIsServer = false);
-
-	void PlayUsingLauncher();
-
-	/** 
-	* Cancel Play using Launcher on error 
-	* 
-	* if the physical device is not authorized to be launched to, we need to pop an error instead of trying to launch
-	*/
-	void CancelPlayUsingLauncher();
-
 	/** Called when Matinee is opened */
 	virtual void OnOpenMatinee(){};
@@ -3111,6 +3098,7 @@
 	TSharedPtr<class ILauncherWorker> LauncherWorker;
 	
 	/** Function to run the Play On command for automation testing. */
+	UE_DEPRECATED(4.25, "Use RequestPlaySession and StartQueuedPlaySessionRequest instead.")
 	void AutomationPlayUsingLauncher(const FString& InLauncherDeviceId);	
 
 	virtual void HandleTravelFailure(UWorld* InWorld, ETravelFailure::Type FailureType, const FString& ErrorString);
@@ -3167,5 +3155,14 @@
 	FDelegateHandle ViewportCloseRequestedDelegateHandle;
 
 public:
+	// FGCObject Interface
+	virtual void AddReferencedObjects(FReferenceCollector& Collector) override;
+	virtual FString GetReferencerName() const override
+	{
+		return "EditorEngine";
+	}
+	// ~FGCObject Interface
+
+public:
 	/**
 	 * Get a Subsystem of specified type
@@ -3200,6 +3197,89 @@
 private:
 	FSubsystemCollection<UEditorSubsystem> EditorSubsystemCollection;
 
+	// DEPRECATED VARIABLES ONLY
+public:
+	UE_DEPRECATED(4.25, "Use the Request Parameters to specify this instead. Can be read from the current session if it was set in the request.")
+	/** An optional location for the starting location for "Play From Here"																*/
+	UPROPERTY()
+	FVector PlayWorldLocation;
+
+	UE_DEPRECATED(4.25, "Use the Request Parameters to specify this instead. Can be read from the current session if it was set in the request.")
+	/** An optional rotation for the starting location for "Play From Here"																*/
+	UPROPERTY()
+	FRotator PlayWorldRotation;
+
+	UE_DEPRECATED(4.25, "Use IsPlaySessionQueued() or IsPlaySessionInProgress() instead.")
+	/** Has a request for "Play From Here" been made?													 								*/
+	UPROPERTY()
+	uint32 bIsPlayWorldQueued:1;
+	
+	UE_DEPRECATED(4.25, "Use IsSimulateInEditorQueued() or IsSimulateInEditorInProgress() instead.")
+	/** True if we are requesting to start a simulation-in-editor session */
+	UPROPERTY()
+	uint32 bIsSimulateInEditorQueued:1;
+	
+	/** Did the request include the optional location and rotation?										 								*/
+	UE_DEPRECATED(4.25, "Use FRequestPlaySessionParams::HasPlayWorldPlacement() on the queued/current session instead.")
+	UPROPERTY()
+	uint32 bHasPlayWorldPlacement:1;
+
+	/** True to enable mobile preview mode when launching the game from the editor on PC platform */
+	UE_DEPRECATED(4.25, "Use FRequestPlaySessionParams::SessionPreviewTypeOverride on the queued/current session instead.")
+	UPROPERTY()
+	uint32 bUseMobilePreviewForPlayWorld:1;
+
+	/** True to enable VR preview mode when launching the game from the editor on PC platform */
+	UE_DEPRECATED(4.25, "Use FRequestPlaySessionParams::SessionPreviewTypeOverride on the queued/current session instead.")
+	UPROPERTY()
+	uint32 bUseVRPreviewForPlayWorld:1;
+
+	/** True if we're Simulating In Editor, as opposed to Playing In Editor.  In this mode, simulation takes place right the level editing environment */
+	// UE_DEPRECATED(4.25, "Use IsSimulateInEditorInProgress instead.")
+	UPROPERTY()
+	uint32 bIsSimulatingInEditor:1;
+	
+	/** Viewport the next PlaySession was requested to happen on */
+	UE_DEPRECATED(4.25, "This is stored as part of the FRequestPlaySessionParams now.")
+	TWeakPtr<class IAssetViewport>		RequestedDestinationSlateViewport;
+
+	/** When set to anything other than -1, indicates a specific In-Editor viewport index that PIE should use */
+	UE_DEPRECATED(4.25, "This isn't read and was replaced by RequestedDestinationSlateViewport.")
+	UPROPERTY()
+	int32 PlayInEditorViewportIndex;
+	
+protected:
+
+	/** Count of how many PIE instances are waiting to log in */
+	UE_DEPRECATED(4.25, "This has moved to the current session instead (stored in FPlayInEditorSessionInfo)")
+	int32 PIEInstancesToLogInCount;
+
+	UE_DEPRECATED(4.25, "This has moved to the current session instead (stored in FPlayInEditorSessionInfo)")
+	bool bAtLeastOnePIELoginFailed;
+
+	/* These are parameters that we need to cache for late joining */
+	UE_DEPRECATED(4.25, "This has moved to the current session instead (stored in FPlayInEditorSessionInfo)")
+	FString ServerPrefix;
+	
+	UE_DEPRECATED(4.25, "This has moved to the current session instead (stored in FPlayInEditorSessionInfo)")
+	int32 PIEInstance;
+	
+	UE_DEPRECATED(4.25, "This has moved to the current session instead (stored in FPlayInEditorSessionInfo)")
+	int32 SettingsIndex;
+	
+	UE_DEPRECATED(4.25, "This has moved to the current session instead (stored in FPlayInEditorSessionInfo)")
+	bool bStartLateJoinersInSpectatorMode;
+
+private:
+
+	/** Additional launch options requested for the next PlaySession */
+	UE_DEPRECATED(4.25, "Use FRequestPlaySessionParams::AdditionalStandaloneCommandLineParameters instead.")
+	FString RequestedAdditionalStandaloneLaunchOptions;
+
+protected:
+	UE_DEPRECATED(4.25, "Use FRequestPlaySessionParams::NumOustandingPIELogins instead.")
+	/** Number of currently running instances logged into an online platform */
+	int32 NumOnlinePIEInstances;
 };
 
 //////////////////////////////////////////////////////////////////////////
diff --git a/Engine/Source/Editor/UnrealEd/Classes/Settings/LevelEditorPlaySettings.h b/Engine/Source/Editor/UnrealEd/Classes/Settings/LevelEditorPlaySettings.h
--- a/Engine/Source/Editor/UnrealEd/Classes/Settings/LevelEditorPlaySettings.h
+++ b/Engine/Source/Editor/UnrealEd/Classes/Settings/LevelEditorPlaySettings.h
@@ -85,10 +85,12 @@ enum EPlayModeType
 UENUM()
 enum EPlayNetMode
 {
+	/** A non-networked game will be started. This can be used in combination with bLaunchSeparateServer to test menu -> server flow in your game. */
 	PIE_Standalone UMETA(DisplayName="Play Offline"),
+	/** The editor will act as both a Server and a Client. Additional instances may be opened beyond that depending on the number of clients. */
 	PIE_ListenServer UMETA(DisplayName="Play As Listen Server"),
+	/** The editor will act as a Client. A server will be started for you behind the scenes to connect to. */
 	PIE_Client UMETA(DisplayName="Play As Client"),
-	PIE_StandaloneWithServer UMETA(DisplayName="Play Standalone With Server", ToolTip="Behaves exactly like 'Play Offline', but also launches the dedicated server, allowing you to override the map name"),
 };
 
 
@@ -278,32 +280,43 @@ public:
 	bool bAutoCompileBlueprintsOnLaunch;
 
 	/** A programmatically defined custom PIE window to use */
+	UE_DEPRECATED(4.25, "This variable is no longer read. Specify the Custom Window in the FRequestPlaySessionParams instead during PIE request.")
 	TWeakPtr<SWindow> CustomPIEWindow;
 	
+	/**
+	* This is a rarely used option that will launch a separate server (possibly hidden in-process depending on RunUnderOneProcess) 
+	* even if the net mode does not require a server (such as Standalone). If the net mode requires a server (such as Client) a 
+	* server will be launched for you (regardless of this setting). This allows you to test offline -> server workflows by connecting
+	* ("open 127.0.0.1:<ServerPort>") from the offline game.
+	*/
+	UPROPERTY(config, EditAnywhere, Category="Multiplayer Options")
+	bool bLaunchSeparateServer;
+
 private:
 
 	/** NetMode to use for Play In Editor. */
-	UPROPERTY(config, EditAnywhere, Category=MultiplayerOptions)
+	UPROPERTY(config, EditAnywhere, Category="Multiplayer Options")
 	TEnumAsByte<EPlayNetMode> PlayNetMode;
 
 	/** Spawn multiple player windows in a single instance of UE4. This will load much faster, but has potential to have more issues.  */
-	UPROPERTY(config, EditAnywhere, Category=MultiplayerOptions)
+	UPROPERTY(config, EditAnywhere, Category="Multiplayer Options")
 	bool RunUnderOneProcess;
 
 	/** If checked, a separate dedicated server will be launched. Otherwise the first player will act as a listen server that all other players connect to. */
-	UPROPERTY(config, EditAnywhere, Category=MultiplayerOptions)
+	UE_DEPRECATED(4.25, "This variable is no longer read. Use PlayNetMode = EPlayNetMode::PIE_Client and bLaunchSeparateServer instead.")
+	UPROPERTY(config)
 	bool PlayNetDedicated;
 
-	/** The editor and listen server count as players, a dedicated server will not. Clients make up the remainder. */
-	UPROPERTY(config, EditAnywhere, Category=MultiplayerOptions, meta=(ClampMin = "1", UIMin = "1", UIMax = "64"))
+	/** The number of client windows to open. The first one to open will respect the Play In Editor "Modes" option (PIE, PINW), additional clients respect the RunUnderOneProcess setting. */
+	UPROPERTY(config, EditAnywhere, Category="Multiplayer Options|Client", meta=(ClampMin = "1", UIMin = "1", UIMax = "64"))
 	int32 PlayNumberOfClients;
 
 	/** What port used by the server for simple networking */
-	UPROPERTY(config, EditAnywhere, Category = MultiplayerOptions, meta=(ClampMin="1", UIMin="1", ClampMax="65535"))
+	UPROPERTY(config, EditAnywhere, Category = "Multiplayer Options|Server", meta=(ClampMin="1", UIMin="1", ClampMax="65535", EditCondition = "PlayNetMode != EPlayNetMode::PIE_Standalone"))
 	uint16 ServerPort;
 
 	/** Width to use when spawning additional windows. */
-	UPROPERTY(config, EditAnywhere, Category=MultiplayerOptions, meta=(ClampMin=0))
+	UPROPERTY(config, EditAnywhere, Category="Multiplayer Options|Client", meta=(ClampMin=0))
 	int32 ClientWindowWidth;
 	
 	/**
@@ -311,7 +324,8 @@ private:
 	 *
 	 * If this is checked, the clients will automatically connect to the launched server, if false they will launch into the map and wait
 	 */
-	UPROPERTY(config, EditAnywhere, Category=MultiplayerOptions)
+	UE_DEPRECATED(4.25, "This variable is no longer read. Use PlayNetMode = EPlayNetMode::PIE_Standalone instead to prevent auto-connection.")
+	UPROPERTY(config)
 	bool AutoConnectToServer;
 
 	/**
@@ -321,7 +335,7 @@ private:
 	 *
 	 * If it is checked, the 1st game pad goes the 2nd window. The 1st window can then be controlled by keyboard/mouse, which is convenient if two people are testing on the same computer.
 	 */
-	UPROPERTY(config, EditAnywhere, Category=MultiplayerOptions)
+	UPROPERTY(config, EditAnywhere, Category="Multiplayer Options|Client", meta=(EditCondition = "RunUnderOneProcess"))
 	bool RouteGamepadToSecondWindow;
 
 	/** 
@@ -331,25 +345,31 @@ private:
 	*
 	* Enabling this will allow rendering accurate audio from every player's perspective but will use more CPU. Keep this disabled on lower-perf machines.
 	*/
-	UPROPERTY(config, EditAnywhere, Category = MultiplayerOptions, meta=(EditCondition = "EnableGameSound"))
+	UPROPERTY(config, EditAnywhere, Category = "Multiplayer Options|Client", meta=(EditCondition = "EnableGameSound && RunUnderOneProcess"))
 	bool CreateAudioDeviceForEveryPlayer;
 
 	/** Height to use when spawning additional windows. */
-	UPROPERTY(config, EditAnywhere, Category=MultiplayerOptions, meta=(ClampMin = 0))
+	UPROPERTY(config, EditAnywhere, Category="Multiplayer Options|Client", meta=(ClampMin = 0))
 	int32 ClientWindowHeight;
 
 	/** Override the map launched by the dedicated server (currently only used when in PIE_StandaloneWithServer net mode) */
-	UPROPERTY(config, EditAnywhere, Category = MultiplayerOptions)
+	UPROPERTY(config, EditAnywhere, Category = "Multiplayer Options|Server", meta=(EditCondition = "PlayNetMode != EPlayNetMode::PIE_Standalone"))
 	FString ServerMapNameOverride;
 
 	/** Additional options that will be passed to the server as URL parameters, in the format ?bIsLanMatch=1?listen - any additional command line switches should be passed in the Command Line Arguments field below. */
-	UPROPERTY(config, EditAnywhere, Category=MultiplayerOptions)
+	UPROPERTY(config, EditAnywhere, Category="Multiplayer Options|Server", meta=(EditCondition = "PlayNetMode != EPlayNetMode::PIE_Standalone"))
 	FString AdditionalServerGameOptions;
 
 	/** Additional command line options that will be passed to standalone game instances, for example -debug */
-	UPROPERTY(config, EditAnywhere, Category=MultiplayerOptions)
+	UE_DEPRECATED(4.25, "This variable is no longer read. Use AdditionalServerLaunchParmeters instead to pass specific flags to externally launched servers.")
+	UPROPERTY(config)
 	FString AdditionalLaunchOptions;
 
+public:
+	/** Additional options that will be passed to the server as arguments, for example -debug. Only works with separate process servers. */
+	UPROPERTY(config, EditAnywhere, Category = "Multiplayer Options|Server", meta = (EditCondition = "!RunUnderOneProcess && PlayNetMode == EPlayNetMode::PIE_Client"))
+	FString AdditionalServerLaunchParmeters;
+
 public:
 
 	/** Should network emulation settings be applied or not */
@@ -361,7 +381,7 @@ public:
 	/**
 	 * Customizable settings allowing to emulate latency and packetloss for game network transmissions
 	 */
-	UPROPERTY(config, EditAnywhere, Category = MultiplayerOptions )
+	UPROPERTY(config, EditAnywhere, Category = "Multiplayer Options" )
 	FLevelEditorPlayNetworkEmulationSettings NetworkEmulationSettings;
 
 public:
@@ -376,10 +396,14 @@ public:
 	bool IsRunUnderOneProcessActive() const { return true; }
 	bool GetRunUnderOneProcess( bool &OutRunUnderOneProcess ) const { OutRunUnderOneProcess = RunUnderOneProcess; return IsRunUnderOneProcessActive(); }
 	
-	void SetPlayNetDedicated( const bool InPlayNetDedicated ) { PlayNetDedicated = InPlayNetDedicated; }
-	bool IsPlayNetDedicatedActive() const { return (RunUnderOneProcess ? true : PlayNetMode == PIE_Client || PlayNetMode == PIE_StandaloneWithServer); }
-	bool GetPlayNetDedicated( bool &OutPlayNetDedicated ) const { OutPlayNetDedicated = PlayNetDedicated; return IsPlayNetDedicatedActive(); }
-
+	PRAGMA_DISABLE_DEPRECATION_WARNINGS
+	UE_DEPRECATED(4.25, "Read bLaunchSeparateServer directly instead.")
+	void SetPlayNetDedicated(const bool InPlayNetDedicated) { PlayNetDedicated = InPlayNetDedicated; }
+	UE_DEPRECATED(4.25, "Read bLaunchSeparateServer directly instead.")
+	bool IsPlayNetDedicatedActive() const { return (RunUnderOneProcess ? true : PlayNetMode == PIE_Client); }
+	UE_DEPRECATED(4.25, "Read bLaunchSeparateServer directly instead.")
+	bool GetPlayNetDedicated(bool &OutPlayNetDedicated) const { OutPlayNetDedicated = PlayNetDedicated; return IsPlayNetDedicatedActive(); }
+	
 	void SetPlayNumberOfClients( const int32 InPlayNumberOfClients ) { PlayNumberOfClients = InPlayNumberOfClients; }
 	bool IsPlayNumberOfClientsActive() const { return (PlayNetMode != PIE_Standalone) || RunUnderOneProcess; }
 	bool GetPlayNumberOfClients( int32 &OutPlayNumberOfClients ) const { OutPlayNumberOfClients = PlayNumberOfClients; return IsPlayNumberOfClientsActive(); }
@@ -388,29 +412,36 @@ public:
 	bool IsServerPortActive() const { return (PlayNetMode != PIE_Standalone) || RunUnderOneProcess; }
 	bool GetServerPort(uint16 &OutServerPort) const { OutServerPort = ServerPort; return IsServerPortActive(); }
 	
+	UE_DEPRECATED(4.25, "This feature has been removed. Set PlayNetMode == EPlayNetMode::PIE_Standalone instead.")
 	bool IsAutoConnectToServerActive() const { return PlayNumberOfClients > 1 || PlayNetDedicated; }
+	UE_DEPRECATED(4.25, "This feature has been removed. Set PlayNetMode == EPlayNetMode::PIE_Standalone instead.")
 	bool GetAutoConnectToServer(bool &OutAutoConnectToServer) const { OutAutoConnectToServer = AutoConnectToServer; return IsAutoConnectToServerActive(); }
+	UE_DEPRECATED(4.25, "This feature has been removed. Set PlayNetMode == EPlayNetMode::PIE_Standalone instead.")
 	EVisibility GetAutoConnectToServerVisibility() const { return (RunUnderOneProcess ? EVisibility::Visible : EVisibility::Hidden); }
-
+	
 	bool IsRouteGamepadToSecondWindowActive() const { return PlayNumberOfClients > 1; }
 	bool GetRouteGamepadToSecondWindow( bool &OutRouteGamepadToSecondWindow ) const { OutRouteGamepadToSecondWindow = RouteGamepadToSecondWindow; return IsRouteGamepadToSecondWindowActive(); }
 	EVisibility GetRouteGamepadToSecondWindowVisibility() const { return (RunUnderOneProcess ? EVisibility::Visible : EVisibility::Hidden); }
 
 	EVisibility GetNetworkEmulationVisibility() const { return (PlayNumberOfClients > 1 || PlayNetDedicated) ? EVisibility::Visible : EVisibility::Hidden; }
 
-	bool IsServerMapNameOverrideActive() const { return (PlayNetMode == PIE_StandaloneWithServer); }
+	bool IsServerMapNameOverrideActive() const { return false /*(PlayNetMode == PIE_StandaloneWithServer)*/; }
 	bool GetServerMapNameOverride( FString& OutStandaloneServerMapName ) const { OutStandaloneServerMapName = ServerMapNameOverride; return IsServerMapNameOverrideActive(); }
-	EVisibility GetServerMapNameOverrideVisibility() const { return (PlayNetMode == PIE_StandaloneWithServer ? EVisibility::Visible : EVisibility::Hidden); }
+	EVisibility GetServerMapNameOverrideVisibility() const { return /*(PlayNetMode == PIE_StandaloneWithServer ? */EVisibility::Visible /*: EVisibility::Hidden)*/; }
 
 	bool IsAdditionalServerGameOptionsActive() const { return (PlayNetMode != PIE_Standalone) || RunUnderOneProcess; }
 	bool GetAdditionalServerGameOptions( FString &OutAdditionalServerGameOptions ) const { OutAdditionalServerGameOptions = AdditionalServerGameOptions; return IsAdditionalServerGameOptionsActive(); }
 
+	UE_DEPRECATED(4.25, "Read AdditionalServerLaunchParmeters for a non-zero length instead.")
 	bool IsAdditionalLaunchOptionsActive() const { return true; }
-	bool GetAdditionalLaunchOptions( FString &OutAdditionalLaunchOptions ) const { OutAdditionalLaunchOptions = AdditionalLaunchOptions; return IsAdditionalLaunchOptionsActive(); }
+	UE_DEPRECATED(4.25, "Read AdditionalServerLaunchParmeters directly instead.")
+	bool GetAdditionalLaunchOptions(FString &OutAdditionalLaunchOptions) const { OutAdditionalLaunchOptions = AdditionalLaunchOptions; return IsAdditionalLaunchOptionsActive(); }
+	UE_DEPRECATED(4.25, "Read AdditionalServerLaunchParmeters for a non-zero length instead.")
 	EVisibility GetAdditionalLaunchOptionsVisibility() const { return (RunUnderOneProcess ? EVisibility::Hidden : EVisibility::Visible); }
-	
+	PRAGMA_ENABLE_DEPRECATION_WARNINGS
+
 	void SetClientWindowSize( const FIntPoint InClientWindowSize ) { ClientWindowWidth = InClientWindowSize.X; ClientWindowHeight = InClientWindowSize.Y; }
-	bool IsClientWindowSizeActive() const { return ((PlayNetMode == PIE_Standalone && RunUnderOneProcess) ? false : (PlayNumberOfClients >= 2)); }
+	bool IsClientWindowSizeActive() const { return PlayNumberOfClients > 1; }
 	bool GetClientWindowSize( FIntPoint &OutClientWindowSize ) const { OutClientWindowSize = FIntPoint(ClientWindowWidth, ClientWindowHeight); return IsClientWindowSizeActive(); }
 	EVisibility GetClientWindowSizeVisibility() const { return (RunUnderOneProcess ? EVisibility::Hidden : EVisibility::Visible); }
 	bool IsCreateAudioDeviceForEveryPlayer() const { return CreateAudioDeviceForEveryPlayer; }
diff --git a/Engine/Source/Editor/UnrealEd/Private/EditorEngine.cpp b/Engine/Source/Editor/UnrealEd/Private/EditorEngine.cpp
--- a/Engine/Source/Editor/UnrealEd/Private/EditorEngine.cpp
+++ b/Engine/Source/Editor/UnrealEd/Private/EditorEngine.cpp
@@ -345,13 +345,10 @@ UEditorEngine::UEditorEngine(const FObjectInitializer& ObjectInitializer)
 	}
 
 	DetailMode = DM_MAX;
-	PlayInEditorViewportIndex = -1;
 	CurrentPlayWorldDestination = -1;
 	bDisableDeltaModification = false;
-	bPlayOnLocalPcSession = false;
 	bAllowMultiplePIEWorlds = true;
 	bIsEndingPlay = false;
-	NumOnlinePIEInstances = 0;
 	DefaultWorldFeatureLevel = GMaxRHIFeatureLevel;
 	PreviewPlatform = FPreviewPlatformInfo(DefaultWorldFeatureLevel);
 
@@ -1597,10 +1594,10 @@ void UEditorEngine::Tick( float DeltaSeconds, bool bIdleMode )
 		}
 	}
 
-	// kick off a "Play From Here" if we got one
-	if (bIsPlayWorldQueued)
+	// Kick off a Play Session request if one was queued up during the last frame.
+	if (PlaySessionRequest.IsSet())
 	{
-		StartQueuedPlayMapRequest();
+		StartQueuedPlaySessionRequest();
 	}
 	else if( bIsToggleBetweenPIEandSIEQueued )
 	{
@@ -1972,7 +1969,7 @@ float UEditorEngine::GetMaxTickRate( float DeltaTime, bool bAllowFrameRateSmooth
 	if( !ShouldThrottleCPUUsage() )
 	{
 		// do not limit fps in VR Preview mode
-		if (bUseVRPreviewForPlayWorld)
+		if (IsVRPreviewActive())
 		{
 			return 0.0f;
 		}
@@ -7337,10 +7334,15 @@ void UEditorEngine::AutomationLoadMap(const FString& MapName, FString* OutError)
 		UE_LOG(LogEditor, Log, TEXT("Starting PIE for the automation tests for world, %s"), *GWorld->GetMapName());
 
 		FFailedGameStartHandler FailHandler;
-		FPlayInEditorOverrides Overrides;
-		Overrides.bDedicatedServer = false;
-		Overrides.NumberOfClients = 1;
-		PlayInEditor(GWorld, /*bInSimulateInEditor=*/false, Overrides);
+
+		FRequestPlaySessionParams RequestParams;
+		ULevelEditorPlaySettings* EditorPlaySettings = NewObject<ULevelEditorPlaySettings>();
+		EditorPlaySettings->SetPlayNumberOfClients(1);
+		EditorPlaySettings->bLaunchSeparateServer = false;
+		RequestParams.EditorPlaySettings = EditorPlaySettings;
+
+		StartPlayInEditorSession(RequestParams);
+
 		if (!FailHandler.CanProceed())
 		{
 			*OutError = TEXT("Error encountered.");
@@ -7357,7 +7359,7 @@ void UEditorEngine::AutomationLoadMap(const FString& MapName, FString* OutError)
 bool UEditorEngine::IsHMDTrackingAllowed() const
 {
 	// @todo vreditor: Added GEnableVREditorHacks check below to allow head movement in non-PIE editor; needs revisit
-	return GEnableVREditorHacks || (PlayWorld && (bUseVRPreviewForPlayWorld || GetDefault<ULevelEditorPlaySettings>()->ViewportGetsHMDControl));
+	return GEnableVREditorHacks || (PlayWorld && (IsVRPreviewActive() || GetDefault<ULevelEditorPlaySettings>()->ViewportGetsHMDControl));
 }
 
 void UEditorEngine::OnModuleCompileStarted(bool bIsAsyncCompile)
@@ -7535,4 +7537,13 @@ void UEditorEngine::SaveEditorFeatureLevel()
 	Settings->PostEditChange();
 }
 
+void UEditorEngine::AddReferencedObjects(FReferenceCollector& Collector)
+{
+	if (PlayInEditorSessionInfo.IsSet())
+	{
+		Collector.AddReferencedObject(PlayInEditorSessionInfo->OriginalRequestParams.EditorPlaySettings, this);
+	}
+}
+
+
 #undef LOCTEXT_NAMESPACE 
