Significantly refactored PIE startup  and consolidated a fair bit of the existing code into structures that better reflect the lifespan that they are valid for. This refactoring will make it easier fo

diff --git a/Engine/Source/Editor/UnrealEd/Private/Kismet2/DebuggerCommands.cpp b/Engine/Source/Editor/UnrealEd/Private/Kismet2/DebuggerCommands.cpp
--- a/Engine/Source/Editor/UnrealEd/Private/Kismet2/DebuggerCommands.cpp
+++ b/Engine/Source/Editor/UnrealEd/Private/Kismet2/DebuggerCommands.cpp
@@ -71,6 +71,7 @@
 
 #include "ToolMenus.h"
 #include "SBlueprintEditorToolbar.h"
+#include "SEnumCombobox.h"
 
 #define LOCTEXT_NAMESPACE "DebuggerCommands"
 
@@ -130,14 +131,14 @@ public:
 	static bool PlayInNewProcess_CanExecute();
 	static void PlayInVR_Clicked();
 	static bool PlayInVR_CanExecute();
-	static bool PlayInModeIsChecked( EPlayModeType PlayMode);
+	static bool PlayInModeIsChecked(EPlayModeType PlayMode);
 
 	static void PlayInNewProcessPreviewDevice_Clicked(FString PIEPreviewDeviceName);
 	static bool PlayInModeAndPreviewDeviceIsChecked(FString PIEPreviewDeviceName);
 
-	static bool PlayInLocation_CanExecute( EPlayModeLocations Location );
-	static void PlayInLocation_Clicked( EPlayModeLocations Location );
-	static bool PlayInLocation_IsChecked( EPlayModeLocations Location );
+	static bool PlayInLocation_CanExecute(EPlayModeLocations Location);
+	static void PlayInLocation_Clicked(EPlayModeLocations Location);
+	static bool PlayInLocation_IsChecked(EPlayModeLocations Location);
 
 	static void PlayInSettings_Clicked();
 
@@ -193,8 +194,8 @@ public:
 	static int32 GetNumberOfClients();
 	static void SetNumberOfClients(int32 NumClients, ETextCommit::Type CommitInfo);
 
-	static void OnToggleDedicatedServerPIE();
-	static bool OnIsDedicatedServerPIEEnabled();
+	static int32 GetNetPlayMode();
+	static void SetNetPlayMode(int32 Value, ESelectInfo::Type CommitInfo);
 
 protected:
 
@@ -208,7 +209,7 @@ protected:
 	 * @param TutorialLink A link to an associated tutorial.
 	 * @param DocumentationLink A link to documentation.
 	 */
-	static void AddMessageLog( const FText& Text, const FText& Detail, const FString& TutorialLink, const FString& DocumentationLink);
+	static void AddMessageLog(const FText& Text, const FText& Detail, const FString& TutorialLink, const FString& DocumentationLink);
 
 	/**
 	 * Checks whether the specified platform has a default device that can be launched on.
@@ -235,11 +236,11 @@ protected:
 
 
 /**
- * Called to leave K2 debugging mode                   
+ * Called to leave K2 debugging mode
  */
 static void LeaveDebuggingMode()
 {
-	if( GUnrealEd->PlayWorld != NULL )
+	if (GUnrealEd->PlayWorld != NULL)
 	{
 		GUnrealEd->PlayWorld->bDebugPauseExecution = false;
 	}
@@ -247,14 +248,14 @@ static void LeaveDebuggingMode()
 	// Determine whether or not we are resuming play.
 	const bool bIsResumingPlay = !FKismetDebugUtilities::IsSingleStepping() && !GEditor->ShouldEndPlayMap();
 
-	if( FSlateApplication::Get().InKismetDebuggingMode() && bIsResumingPlay )
+	if (FSlateApplication::Get().InKismetDebuggingMode() && bIsResumingPlay)
 	{
 		// Focus the game view port when resuming from debugging.
 		FModuleManager::GetModuleChecked<FLevelEditorModule>("LevelEditor").FocusPIEViewport();
 	}
 
 	// Tell the application to stop ticking in this stack frame. The parameter controls whether or not to recapture the mouse to the game viewport.
-	FSlateApplication::Get().LeaveDebuggingMode( !bIsResumingPlay );
+	FSlateApplication::Get().LeaveDebuggingMode(!bIsResumingPlay);
 }
 
 
@@ -326,40 +327,39 @@ FPlayWorldCommands::FPlayWorldCommands()
 void FPlayWorldCommands::RegisterCommands()
 {
 	// SIE
-	UI_COMMAND( Simulate, "Simulate", "Start simulating the game", EUserInterfaceActionType::Check, FInputChord( EKeys::S, EModifierKey::Alt ) );
+	UI_COMMAND(Simulate, "Simulate", "Start simulating the game", EUserInterfaceActionType::Check, FInputChord(EKeys::S, EModifierKey::Alt));
 
 	// PIE
-	UI_COMMAND( RepeatLastPlay, "Play", "Launches a game preview session in the same mode as the last game preview session launched from the Game Preview Modes dropdown next to the Play button on the level editor toolbar", EUserInterfaceActionType::Button, FInputChord( EKeys::P, EModifierKey::Alt ) )
-	UI_COMMAND( PlayInViewport, "Selected Viewport", "Play this level in the active level editor viewport", EUserInterfaceActionType::Check, FInputChord() );
-	UI_COMMAND( PlayInEditorFloating, "New Editor Window (PIE)", "Play this level in a new window", EUserInterfaceActionType::Check, FInputChord() );
-	UI_COMMAND( PlayInVR, "VR Preview", "Play this level in VR", EUserInterfaceActionType::Check, FInputChord() );
-	UI_COMMAND( PlayInMobilePreview, "Mobile Preview ES2 (PIE)", "Play this level as a mobile device preview in ES2 mode (runs in its own process)", EUserInterfaceActionType::Check, FInputChord());
-	UI_COMMAND( PlayInVulkanPreview, "Vulkan Mobile Preview (PIE)", "Play this level using mobile Vulkan rendering (runs in its own process)", EUserInterfaceActionType::Check, FInputChord() );
-	UI_COMMAND( PlayInNewProcess, "Standalone Game", "Play this level in a new window that runs in its own process", EUserInterfaceActionType::Check, FInputChord() );
-	UI_COMMAND( PlayInCameraLocation, "Current Camera Location", "Spawn the player at the current camera location", EUserInterfaceActionType::RadioButton, FInputChord() );
-	UI_COMMAND( PlayInDefaultPlayerStart, "Default Player Start", "Spawn the player at the map's default player start", EUserInterfaceActionType::RadioButton, FInputChord() );
-	UI_COMMAND( PlayInNetworkSettings, "Network Settings...", "Open the settings for the 'Play In' feature", EUserInterfaceActionType::Button, FInputChord() );
-	UI_COMMAND( PlayInNetworkDedicatedServer, "Run Dedicated Server", "If checked, a separate dedicated server will be launched. Otherwise the first player will act as a listen server that all other players connect to.", EUserInterfaceActionType::ToggleButton, FInputChord() );
-	UI_COMMAND( PlayInSettings, "Advanced Settings...", "Open the settings for the 'Play In' feature", EUserInterfaceActionType::Button, FInputChord() );
+	UI_COMMAND(RepeatLastPlay, "Play", "Launches a game preview session in the same mode as the last game preview session launched from the Game Preview Modes dropdown next to the Play button on the level editor toolbar", EUserInterfaceActionType::Button, FInputChord(EKeys::P, EModifierKey::Alt))
+		UI_COMMAND(PlayInViewport, "Selected Viewport", "Play this level in the active level editor viewport", EUserInterfaceActionType::Check, FInputChord());
+	UI_COMMAND(PlayInEditorFloating, "New Editor Window (PIE)", "Play this level in a new window", EUserInterfaceActionType::Check, FInputChord());
+	UI_COMMAND(PlayInVR, "VR Preview", "Play this level in VR", EUserInterfaceActionType::Check, FInputChord());
+	UI_COMMAND(PlayInMobilePreview, "Mobile Preview ES2 (PIE)", "Play this level as a mobile device preview in ES2 mode (runs in its own process)", EUserInterfaceActionType::Check, FInputChord());
+	UI_COMMAND(PlayInVulkanPreview, "Vulkan Mobile Preview (PIE)", "Play this level using mobile Vulkan rendering (runs in its own process)", EUserInterfaceActionType::Check, FInputChord());
+	UI_COMMAND(PlayInNewProcess, "Standalone Game", "Play this level in a new window that runs in its own process", EUserInterfaceActionType::Check, FInputChord());
+	UI_COMMAND(PlayInCameraLocation, "Current Camera Location", "Spawn the player at the current camera location", EUserInterfaceActionType::RadioButton, FInputChord());
+	UI_COMMAND(PlayInDefaultPlayerStart, "Default Player Start", "Spawn the player at the map's default player start", EUserInterfaceActionType::RadioButton, FInputChord());
+	UI_COMMAND(PlayInNetworkSettings, "Network Settings...", "Open the settings for the 'Play In' feature", EUserInterfaceActionType::Button, FInputChord());
+	UI_COMMAND(PlayInSettings, "Advanced Settings...", "Open the settings for the 'Play In' feature", EUserInterfaceActionType::Button, FInputChord());
 
 	// SIE & PIE controls
-	UI_COMMAND( StopPlaySession, "Stop", "Stop simulation", EUserInterfaceActionType::Button, FInputChord(EKeys::Escape) );
-	UI_COMMAND( ResumePlaySession, "Resume", "Resume simulation", EUserInterfaceActionType::Button, FInputChord() );
-	UI_COMMAND( PausePlaySession, "Pause", "Pause simulation", EUserInterfaceActionType::Button, FInputChord() );
-	UI_COMMAND( GetMouseControl, "Mouse Control", "Get mouse cursor while in PIE", EUserInterfaceActionType::Button, FInputChord(EModifierKey::Shift, EKeys::F1));
-	UI_COMMAND( LateJoinSession, "Add Client", "Add another client", EUserInterfaceActionType::Button, FInputChord());
-	UI_COMMAND( SingleFrameAdvance, "Frame Skip", "Advances a single frame", EUserInterfaceActionType::Button, FInputChord() );
-	UI_COMMAND( TogglePlayPauseOfPlaySession, "Toggle Play/Pause", "Resume playing if paused, or pause if playing", EUserInterfaceActionType::Button, FInputChord( EKeys::Pause ) );
-	UI_COMMAND( PossessEjectPlayer, "Possess or Eject Player", "Possesses or ejects the player from the camera", EUserInterfaceActionType::Button, FInputChord( EKeys::F8 ) );
-	UI_COMMAND( ShowCurrentStatement, "Find Node", "Show the current node", EUserInterfaceActionType::Button, FInputChord() );
-	UI_COMMAND( StepInto, "Step Into", "Step Into the next node to be executed", EUserInterfaceActionType::Button, PLATFORM_MAC ? FInputChord(EModifierKey::Control, EKeys::F11) : FInputChord(EKeys::F11) );
-	UI_COMMAND( StepOver, "Step Over", "Step to the next node to be executed in the current graph", EUserInterfaceActionType::Button, FInputChord(EKeys::F10) );
-	UI_COMMAND( StepOut, "Step Out", "Step Out to the next node to be executed in the parent graph", EUserInterfaceActionType::Button, FInputChord(EModifierKey::Alt|EModifierKey::Shift, EKeys::F11) );
+	UI_COMMAND(StopPlaySession, "Stop", "Stop simulation", EUserInterfaceActionType::Button, FInputChord(EKeys::Escape));
+	UI_COMMAND(ResumePlaySession, "Resume", "Resume simulation", EUserInterfaceActionType::Button, FInputChord());
+	UI_COMMAND(PausePlaySession, "Pause", "Pause simulation", EUserInterfaceActionType::Button, FInputChord());
+	UI_COMMAND(GetMouseControl, "Mouse Control", "Get mouse cursor while in PIE", EUserInterfaceActionType::Button, FInputChord(EModifierKey::Shift, EKeys::F1));
+	UI_COMMAND(LateJoinSession, "Add Client", "Add another client", EUserInterfaceActionType::Button, FInputChord());
+	UI_COMMAND(SingleFrameAdvance, "Frame Skip", "Advances a single frame", EUserInterfaceActionType::Button, FInputChord());
+	UI_COMMAND(TogglePlayPauseOfPlaySession, "Toggle Play/Pause", "Resume playing if paused, or pause if playing", EUserInterfaceActionType::Button, FInputChord(EKeys::Pause));
+	UI_COMMAND(PossessEjectPlayer, "Possess or Eject Player", "Possesses or ejects the player from the camera", EUserInterfaceActionType::Button, FInputChord(EKeys::F8));
+	UI_COMMAND(ShowCurrentStatement, "Find Node", "Show the current node", EUserInterfaceActionType::Button, FInputChord());
+	UI_COMMAND(StepInto, "Step Into", "Step Into the next node to be executed", EUserInterfaceActionType::Button, PLATFORM_MAC ? FInputChord(EModifierKey::Control, EKeys::F11) : FInputChord(EKeys::F11));
+	UI_COMMAND(StepOver, "Step Over", "Step to the next node to be executed in the current graph", EUserInterfaceActionType::Button, FInputChord(EKeys::F10));
+	UI_COMMAND(StepOut, "Step Out", "Step Out to the next node to be executed in the parent graph", EUserInterfaceActionType::Button, FInputChord(EModifierKey::Alt | EModifierKey::Shift, EKeys::F11));
 
 	// Launch
-	UI_COMMAND( RepeatLastLaunch, "Launch", "Launches the game on the device as the last session launched from the dropdown next to the Play on Device button on the level editor toolbar", EUserInterfaceActionType::Button, FInputChord( EKeys::P, EModifierKey::Alt | EModifierKey::Shift ) )
-	UI_COMMAND( OpenProjectLauncher, "Project Launcher...", "Open the Project Launcher for advanced packaging, deploying and launching of your projects", EUserInterfaceActionType::Button, FInputChord());
-	UI_COMMAND( OpenDeviceManager, "Device Manager...", "View and manage connected devices.", EUserInterfaceActionType::Button, FInputChord() );
+	UI_COMMAND(RepeatLastLaunch, "Launch", "Launches the game on the device as the last session launched from the dropdown next to the Play on Device button on the level editor toolbar", EUserInterfaceActionType::Button, FInputChord(EKeys::P, EModifierKey::Alt | EModifierKey::Shift))
+		UI_COMMAND(OpenProjectLauncher, "Project Launcher...", "Open the Project Launcher for advanced packaging, deploying and launching of your projects", EUserInterfaceActionType::Button, FInputChord());
+	UI_COMMAND(OpenDeviceManager, "Device Manager...", "View and manage connected devices.", EUserInterfaceActionType::Button, FInputChord());
 
 	// PIE mobile preview devices.
 	AddPIEPreviewDeviceCommands();
@@ -395,113 +395,113 @@ void FPlayWorldCommands::AddPIEPreviewDeviceCommands()
 
 void FPlayWorldCommands::BindGlobalPlayWorldCommands()
 {
-	check( !GlobalPlayWorldActions.IsValid() );
+	check(!GlobalPlayWorldActions.IsValid());
 
-	GlobalPlayWorldActions = MakeShareable( new FUICommandList );
+	GlobalPlayWorldActions = MakeShareable(new FUICommandList);
 
 	const FPlayWorldCommands& Commands = FPlayWorldCommands::Get();
 	FUICommandList& ActionList = *GlobalPlayWorldActions;
 
 	// SIE
-	ActionList.MapAction( Commands.Simulate,
-		FExecuteAction::CreateStatic( &FInternalPlayWorldCommandCallbacks::Simulate_Clicked ),
-		FCanExecuteAction::CreateStatic( &FInternalPlayWorldCommandCallbacks::Simulate_CanExecute ),
-		FIsActionChecked::CreateStatic( &FInternalPlayWorldCommandCallbacks::PlayInModeIsChecked, PlayMode_Simulate ),
-		FIsActionButtonVisible::CreateStatic( &FInternalPlayWorldCommandCallbacks::CanShowNonPlayWorldOnlyActions )
-		);
+	ActionList.MapAction(Commands.Simulate,
+		FExecuteAction::CreateStatic(&FInternalPlayWorldCommandCallbacks::Simulate_Clicked),
+		FCanExecuteAction::CreateStatic(&FInternalPlayWorldCommandCallbacks::Simulate_CanExecute),
+		FIsActionChecked::CreateStatic(&FInternalPlayWorldCommandCallbacks::PlayInModeIsChecked, PlayMode_Simulate),
+		FIsActionButtonVisible::CreateStatic(&FInternalPlayWorldCommandCallbacks::CanShowNonPlayWorldOnlyActions)
+	);
 
 	// PIE
-	ActionList.MapAction( Commands.RepeatLastPlay,
-		FExecuteAction::CreateStatic( &FInternalPlayWorldCommandCallbacks::RepeatLastPlay_Clicked ),
-		FCanExecuteAction::CreateStatic( &FInternalPlayWorldCommandCallbacks::RepeatLastPlay_CanExecute ),
+	ActionList.MapAction(Commands.RepeatLastPlay,
+		FExecuteAction::CreateStatic(&FInternalPlayWorldCommandCallbacks::RepeatLastPlay_Clicked),
+		FCanExecuteAction::CreateStatic(&FInternalPlayWorldCommandCallbacks::RepeatLastPlay_CanExecute),
 		FIsActionChecked(),
-		FIsActionButtonVisible::CreateStatic( &FInternalPlayWorldCommandCallbacks::CanShowNonPlayWorldOnlyActions )
-		);
+		FIsActionButtonVisible::CreateStatic(&FInternalPlayWorldCommandCallbacks::CanShowNonPlayWorldOnlyActions)
+	);
 
-	ActionList.MapAction( Commands.PlayInViewport,
-		FExecuteAction::CreateStatic( &FInternalPlayWorldCommandCallbacks::PlayInViewport_Clicked ),
-		FCanExecuteAction::CreateStatic( &FInternalPlayWorldCommandCallbacks::PlayInViewport_CanExecute ),
-		FIsActionChecked::CreateStatic( &FInternalPlayWorldCommandCallbacks::PlayInModeIsChecked, PlayMode_InViewPort ),
-		FIsActionButtonVisible::CreateStatic( &FInternalPlayWorldCommandCallbacks::CanShowNonPlayWorldOnlyActions )
-		);
+	ActionList.MapAction(Commands.PlayInViewport,
+		FExecuteAction::CreateStatic(&FInternalPlayWorldCommandCallbacks::PlayInViewport_Clicked),
+		FCanExecuteAction::CreateStatic(&FInternalPlayWorldCommandCallbacks::PlayInViewport_CanExecute),
+		FIsActionChecked::CreateStatic(&FInternalPlayWorldCommandCallbacks::PlayInModeIsChecked, PlayMode_InViewPort),
+		FIsActionButtonVisible::CreateStatic(&FInternalPlayWorldCommandCallbacks::CanShowNonPlayWorldOnlyActions)
+	);
 
-	ActionList.MapAction( Commands.PlayInEditorFloating,
-		FExecuteAction::CreateStatic( &FInternalPlayWorldCommandCallbacks::PlayInEditorFloating_Clicked ),
-		FCanExecuteAction::CreateStatic( &FInternalPlayWorldCommandCallbacks::PlayInEditorFloating_CanExecute ),
-		FIsActionChecked::CreateStatic( &FInternalPlayWorldCommandCallbacks::PlayInModeIsChecked, PlayMode_InEditorFloating ),
-		FIsActionButtonVisible::CreateStatic( &FInternalPlayWorldCommandCallbacks::CanShowNonPlayWorldOnlyActions )
-		);
+	ActionList.MapAction(Commands.PlayInEditorFloating,
+		FExecuteAction::CreateStatic(&FInternalPlayWorldCommandCallbacks::PlayInEditorFloating_Clicked),
+		FCanExecuteAction::CreateStatic(&FInternalPlayWorldCommandCallbacks::PlayInEditorFloating_CanExecute),
+		FIsActionChecked::CreateStatic(&FInternalPlayWorldCommandCallbacks::PlayInModeIsChecked, PlayMode_InEditorFloating),
+		FIsActionButtonVisible::CreateStatic(&FInternalPlayWorldCommandCallbacks::CanShowNonPlayWorldOnlyActions)
+	);
 
 	ActionList.MapAction(Commands.PlayInVR,
 		FExecuteAction::CreateStatic(&FInternalPlayWorldCommandCallbacks::PlayInVR_Clicked),
 		FCanExecuteAction::CreateStatic(&FInternalPlayWorldCommandCallbacks::PlayInVR_CanExecute),
 		FIsActionChecked::CreateStatic(&FInternalPlayWorldCommandCallbacks::PlayInModeIsChecked, PlayMode_InVR),
 		FIsActionButtonVisible::CreateStatic(&FInternalPlayWorldCommandCallbacks::CanShowVROnlyActions)
-		);
+	);
 
-	ActionList.MapAction( Commands.PlayInMobilePreview,
-		FExecuteAction::CreateStatic( &FInternalPlayWorldCommandCallbacks::PlayInNewProcess_Clicked, PlayMode_InMobilePreview),
-		FCanExecuteAction::CreateStatic( &FInternalPlayWorldCommandCallbacks::PlayInNewProcess_CanExecute ),
-		FIsActionChecked::CreateStatic( &FInternalPlayWorldCommandCallbacks::PlayInModeIsChecked, PlayMode_InMobilePreview ),
-		FIsActionButtonVisible::CreateStatic( &FInternalPlayWorldCommandCallbacks::CanShowNonPlayWorldOnlyActions )
-		);
+	ActionList.MapAction(Commands.PlayInMobilePreview,
+		FExecuteAction::CreateStatic(&FInternalPlayWorldCommandCallbacks::PlayInNewProcess_Clicked, PlayMode_InMobilePreview),
+		FCanExecuteAction::CreateStatic(&FInternalPlayWorldCommandCallbacks::PlayInNewProcess_CanExecute),
+		FIsActionChecked::CreateStatic(&FInternalPlayWorldCommandCallbacks::PlayInModeIsChecked, PlayMode_InMobilePreview),
+		FIsActionButtonVisible::CreateStatic(&FInternalPlayWorldCommandCallbacks::CanShowNonPlayWorldOnlyActions)
+	);
 
 	ActionList.MapAction(Commands.PlayInVulkanPreview,
 		FExecuteAction::CreateStatic(&FInternalPlayWorldCommandCallbacks::PlayInNewProcess_Clicked, PlayMode_InVulkanPreview),
 		FCanExecuteAction::CreateStatic(&FInternalPlayWorldCommandCallbacks::PlayInNewProcess_CanExecute),
 		FIsActionChecked::CreateStatic(&FInternalPlayWorldCommandCallbacks::PlayInModeIsChecked, PlayMode_InVulkanPreview),
 		FIsActionButtonVisible::CreateStatic(&FInternalPlayWorldCommandCallbacks::CanShowVulkanNonPlayWorldOnlyActions)
-		);
+	);
 
 	ActionList.MapAction(Commands.PlayInNewProcess,
-		FExecuteAction::CreateStatic( &FInternalPlayWorldCommandCallbacks::PlayInNewProcess_Clicked, PlayMode_InNewProcess),
-		FCanExecuteAction::CreateStatic( &FInternalPlayWorldCommandCallbacks::PlayInNewProcess_CanExecute ),
-		FIsActionChecked::CreateStatic( &FInternalPlayWorldCommandCallbacks::PlayInModeIsChecked, PlayMode_InNewProcess),
-		FIsActionButtonVisible::CreateStatic( &FInternalPlayWorldCommandCallbacks::CanShowNonPlayWorldOnlyActions )
-		);
+		FExecuteAction::CreateStatic(&FInternalPlayWorldCommandCallbacks::PlayInNewProcess_Clicked, PlayMode_InNewProcess),
+		FCanExecuteAction::CreateStatic(&FInternalPlayWorldCommandCallbacks::PlayInNewProcess_CanExecute),
+		FIsActionChecked::CreateStatic(&FInternalPlayWorldCommandCallbacks::PlayInModeIsChecked, PlayMode_InNewProcess),
+		FIsActionButtonVisible::CreateStatic(&FInternalPlayWorldCommandCallbacks::CanShowNonPlayWorldOnlyActions)
+	);
 
-	ActionList.MapAction( Commands.PlayInCameraLocation,
-		FExecuteAction::CreateStatic( &FInternalPlayWorldCommandCallbacks::PlayInLocation_Clicked, PlayLocation_CurrentCameraLocation ),
-		FCanExecuteAction::CreateStatic( &FInternalPlayWorldCommandCallbacks::PlayInLocation_CanExecute, PlayLocation_CurrentCameraLocation ),
-		FIsActionChecked::CreateStatic( &FInternalPlayWorldCommandCallbacks::PlayInLocation_IsChecked, PlayLocation_CurrentCameraLocation ),
-		FIsActionButtonVisible::CreateStatic( &FInternalPlayWorldCommandCallbacks::CanShowNonPlayWorldOnlyActions )
-		);
+	ActionList.MapAction(Commands.PlayInCameraLocation,
+		FExecuteAction::CreateStatic(&FInternalPlayWorldCommandCallbacks::PlayInLocation_Clicked, PlayLocation_CurrentCameraLocation),
+		FCanExecuteAction::CreateStatic(&FInternalPlayWorldCommandCallbacks::PlayInLocation_CanExecute, PlayLocation_CurrentCameraLocation),
+		FIsActionChecked::CreateStatic(&FInternalPlayWorldCommandCallbacks::PlayInLocation_IsChecked, PlayLocation_CurrentCameraLocation),
+		FIsActionButtonVisible::CreateStatic(&FInternalPlayWorldCommandCallbacks::CanShowNonPlayWorldOnlyActions)
+	);
 
-	ActionList.MapAction( Commands.PlayInDefaultPlayerStart,
-		FExecuteAction::CreateStatic( &FInternalPlayWorldCommandCallbacks::PlayInLocation_Clicked, PlayLocation_DefaultPlayerStart ),
-		FCanExecuteAction::CreateStatic( &FInternalPlayWorldCommandCallbacks::PlayInLocation_CanExecute, PlayLocation_DefaultPlayerStart ),
-		FIsActionChecked::CreateStatic( &FInternalPlayWorldCommandCallbacks::PlayInLocation_IsChecked, PlayLocation_DefaultPlayerStart ),
-		FIsActionButtonVisible::CreateStatic( &FInternalPlayWorldCommandCallbacks::CanShowNonPlayWorldOnlyActions )
-		);
+	ActionList.MapAction(Commands.PlayInDefaultPlayerStart,
+		FExecuteAction::CreateStatic(&FInternalPlayWorldCommandCallbacks::PlayInLocation_Clicked, PlayLocation_DefaultPlayerStart),
+		FCanExecuteAction::CreateStatic(&FInternalPlayWorldCommandCallbacks::PlayInLocation_CanExecute, PlayLocation_DefaultPlayerStart),
+		FIsActionChecked::CreateStatic(&FInternalPlayWorldCommandCallbacks::PlayInLocation_IsChecked, PlayLocation_DefaultPlayerStart),
+		FIsActionButtonVisible::CreateStatic(&FInternalPlayWorldCommandCallbacks::CanShowNonPlayWorldOnlyActions)
+	);
 
-	ActionList.MapAction( Commands.PlayInSettings,
-		FExecuteAction::CreateStatic( &FInternalPlayWorldCommandCallbacks::PlayInSettings_Clicked )
-		);
+	ActionList.MapAction(Commands.PlayInSettings,
+		FExecuteAction::CreateStatic(&FInternalPlayWorldCommandCallbacks::PlayInSettings_Clicked)
+	);
 
 	// Launch
 	ActionList.MapAction(Commands.OpenProjectLauncher,
 		FExecuteAction::CreateStatic(&FInternalPlayWorldCommandCallbacks::OpenProjectLauncher_Clicked)
-		);
-	
-	ActionList.MapAction( Commands.OpenDeviceManager,
-		FExecuteAction::CreateStatic( &FInternalPlayWorldCommandCallbacks::OpenDeviceManager_Clicked )
-		);
+	);
 
-	ActionList.MapAction( Commands.RepeatLastLaunch,
-		FExecuteAction::CreateStatic( &FInternalPlayWorldCommandCallbacks::RepeatLastLaunch_Clicked ),
-		FCanExecuteAction::CreateStatic( &FInternalPlayWorldCommandCallbacks::RepeatLastLaunch_CanExecute ),
+	ActionList.MapAction(Commands.OpenDeviceManager,
+		FExecuteAction::CreateStatic(&FInternalPlayWorldCommandCallbacks::OpenDeviceManager_Clicked)
+	);
+
+	ActionList.MapAction(Commands.RepeatLastLaunch,
+		FExecuteAction::CreateStatic(&FInternalPlayWorldCommandCallbacks::RepeatLastLaunch_Clicked),
+		FCanExecuteAction::CreateStatic(&FInternalPlayWorldCommandCallbacks::RepeatLastLaunch_CanExecute),
 		FIsActionChecked(),
-		FIsActionButtonVisible::CreateStatic( &FInternalPlayWorldCommandCallbacks::CanShowNonPlayWorldOnlyActions )
-		);
+		FIsActionButtonVisible::CreateStatic(&FInternalPlayWorldCommandCallbacks::CanShowNonPlayWorldOnlyActions)
+	);
 
 
 	// Stop play session
-	ActionList.MapAction( Commands.StopPlaySession,
-		FExecuteAction::CreateStatic( &FInternalPlayWorldCommandCallbacks::StopPlaySession_Clicked ),
-		FCanExecuteAction::CreateStatic( &FInternalPlayWorldCommandCallbacks::HasPlayWorld ),
+	ActionList.MapAction(Commands.StopPlaySession,
+		FExecuteAction::CreateStatic(&FInternalPlayWorldCommandCallbacks::StopPlaySession_Clicked),
+		FCanExecuteAction::CreateStatic(&FInternalPlayWorldCommandCallbacks::HasPlayWorld),
 		FIsActionChecked(),
-		FIsActionButtonVisible::CreateStatic( &FInternalPlayWorldCommandCallbacks::HasPlayWorld )
-		);
+		FIsActionButtonVisible::CreateStatic(&FInternalPlayWorldCommandCallbacks::HasPlayWorld)
+	);
 
 	// Late join session
 	ActionList.MapAction(Commands.LateJoinSession,
@@ -509,87 +509,81 @@ void FPlayWorldCommands::BindGlobalPlayWorldCommands()
 		FCanExecuteAction::CreateStatic(&FInternalPlayWorldCommandCallbacks::CanLateJoin),
 		FIsActionChecked(),
 		FIsActionButtonVisible::CreateStatic(&FInternalPlayWorldCommandCallbacks::CanShowLateJoinButton)
-		);
+	);
 
 	// Play, Pause, Toggle between play and pause
-	ActionList.MapAction( Commands.ResumePlaySession,
-		FExecuteAction::CreateStatic( &FPlayWorldCommandCallbacks::ResumePlaySession_Clicked ),
-		FCanExecuteAction::CreateStatic( &FPlayWorldCommandCallbacks::HasPlayWorldAndPaused ),
+	ActionList.MapAction(Commands.ResumePlaySession,
+		FExecuteAction::CreateStatic(&FPlayWorldCommandCallbacks::ResumePlaySession_Clicked),
+		FCanExecuteAction::CreateStatic(&FPlayWorldCommandCallbacks::HasPlayWorldAndPaused),
 		FIsActionChecked(),
-		FIsActionButtonVisible::CreateStatic( &FPlayWorldCommandCallbacks::HasPlayWorldAndPaused )
-		);
+		FIsActionButtonVisible::CreateStatic(&FPlayWorldCommandCallbacks::HasPlayWorldAndPaused)
+	);
 
-	ActionList.MapAction( Commands.PausePlaySession,
-		FExecuteAction::CreateStatic( &FPlayWorldCommandCallbacks::PausePlaySession_Clicked ),
-		FCanExecuteAction::CreateStatic( &FPlayWorldCommandCallbacks::HasPlayWorldAndRunning ),
+	ActionList.MapAction(Commands.PausePlaySession,
+		FExecuteAction::CreateStatic(&FPlayWorldCommandCallbacks::PausePlaySession_Clicked),
+		FCanExecuteAction::CreateStatic(&FPlayWorldCommandCallbacks::HasPlayWorldAndRunning),
 		FIsActionChecked(),
-		FIsActionButtonVisible::CreateStatic( &FPlayWorldCommandCallbacks::HasPlayWorldAndRunning )
-		);
+		FIsActionButtonVisible::CreateStatic(&FPlayWorldCommandCallbacks::HasPlayWorldAndRunning)
+	);
 
-	ActionList.MapAction( Commands.SingleFrameAdvance,
-		FExecuteAction::CreateStatic( &FInternalPlayWorldCommandCallbacks::SingleFrameAdvance_Clicked ),
-		FCanExecuteAction::CreateStatic( &FPlayWorldCommandCallbacks::HasPlayWorldAndPaused ),
+	ActionList.MapAction(Commands.SingleFrameAdvance,
+		FExecuteAction::CreateStatic(&FInternalPlayWorldCommandCallbacks::SingleFrameAdvance_Clicked),
+		FCanExecuteAction::CreateStatic(&FPlayWorldCommandCallbacks::HasPlayWorldAndPaused),
 		FIsActionChecked(),
-		FIsActionChecked::CreateStatic( &FPlayWorldCommandCallbacks::HasPlayWorldAndPaused )
-		);
+		FIsActionChecked::CreateStatic(&FPlayWorldCommandCallbacks::HasPlayWorldAndPaused)
+	);
 
-	ActionList.MapAction( Commands.TogglePlayPauseOfPlaySession,
-		FExecuteAction::CreateStatic( &FInternalPlayWorldCommandCallbacks::TogglePlayPause_Clicked ),
-		FCanExecuteAction::CreateStatic( &FInternalPlayWorldCommandCallbacks::HasPlayWorld ),
+	ActionList.MapAction(Commands.TogglePlayPauseOfPlaySession,
+		FExecuteAction::CreateStatic(&FInternalPlayWorldCommandCallbacks::TogglePlayPause_Clicked),
+		FCanExecuteAction::CreateStatic(&FInternalPlayWorldCommandCallbacks::HasPlayWorld),
 		FIsActionChecked(),
-		FIsActionButtonVisible::CreateStatic( &FInternalPlayWorldCommandCallbacks::HasPlayWorld )
-		);
+		FIsActionButtonVisible::CreateStatic(&FInternalPlayWorldCommandCallbacks::HasPlayWorld)
+	);
 
 	// Get mouse control from PIE
 	ActionList.MapAction(Commands.GetMouseControl,
-		FExecuteAction::CreateStatic( &FInternalPlayWorldCommandCallbacks::GetMouseControlExecute ),
-		FCanExecuteAction::CreateStatic( &FInternalPlayWorldCommandCallbacks::HasPlayWorld ),
+		FExecuteAction::CreateStatic(&FInternalPlayWorldCommandCallbacks::GetMouseControlExecute),
+		FCanExecuteAction::CreateStatic(&FInternalPlayWorldCommandCallbacks::HasPlayWorld),
 		FIsActionChecked(),
-		FIsActionChecked::CreateStatic(&FInternalPlayWorldCommandCallbacks::HasPlayWorld )
-		);
+		FIsActionChecked::CreateStatic(&FInternalPlayWorldCommandCallbacks::HasPlayWorld)
+	);
 
 	// Toggle PIE/SIE, Eject (PIE->SIE), and Possess (SIE->PIE)
-	ActionList.MapAction( Commands.PossessEjectPlayer,
-		FExecuteAction::CreateStatic( &FInternalPlayWorldCommandCallbacks::PossessEjectPlayer_Clicked ),
-		FCanExecuteAction::CreateStatic( &FInternalPlayWorldCommandCallbacks::CanPossessEjectPlayer ),
+	ActionList.MapAction(Commands.PossessEjectPlayer,
+		FExecuteAction::CreateStatic(&FInternalPlayWorldCommandCallbacks::PossessEjectPlayer_Clicked),
+		FCanExecuteAction::CreateStatic(&FInternalPlayWorldCommandCallbacks::CanPossessEjectPlayer),
 		FIsActionChecked(),
-		FIsActionButtonVisible::CreateStatic( &FInternalPlayWorldCommandCallbacks::CanPossessEjectPlayer )
-		);
+		FIsActionButtonVisible::CreateStatic(&FInternalPlayWorldCommandCallbacks::CanPossessEjectPlayer)
+	);
 
 	// Breakpoint-only commands
-	ActionList.MapAction( Commands.ShowCurrentStatement,
-		FExecuteAction::CreateStatic( &FInternalPlayWorldCommandCallbacks::ShowCurrentStatement_Clicked ),
-		FCanExecuteAction::CreateStatic( &FInternalPlayWorldCommandCallbacks::IsStoppedAtBreakpoint ),
+	ActionList.MapAction(Commands.ShowCurrentStatement,
+		FExecuteAction::CreateStatic(&FInternalPlayWorldCommandCallbacks::ShowCurrentStatement_Clicked),
+		FCanExecuteAction::CreateStatic(&FInternalPlayWorldCommandCallbacks::IsStoppedAtBreakpoint),
 		FIsActionChecked(),
-		FIsActionChecked::CreateStatic( &FInternalPlayWorldCommandCallbacks::IsStoppedAtBreakpoint )
-		);
+		FIsActionChecked::CreateStatic(&FInternalPlayWorldCommandCallbacks::IsStoppedAtBreakpoint)
+	);
 
-	ActionList.MapAction( Commands.StepInto,
-		FExecuteAction::CreateStatic( &FInternalPlayWorldCommandCallbacks::StepInto_Clicked ),
-		FCanExecuteAction::CreateStatic( &FInternalPlayWorldCommandCallbacks::IsStoppedAtBreakpoint ),
+	ActionList.MapAction(Commands.StepInto,
+		FExecuteAction::CreateStatic(&FInternalPlayWorldCommandCallbacks::StepInto_Clicked),
+		FCanExecuteAction::CreateStatic(&FInternalPlayWorldCommandCallbacks::IsStoppedAtBreakpoint),
 		FIsActionChecked(),
-		FIsActionChecked::CreateStatic( &FInternalPlayWorldCommandCallbacks::IsStoppedAtBreakpoint )
-		);
+		FIsActionChecked::CreateStatic(&FInternalPlayWorldCommandCallbacks::IsStoppedAtBreakpoint)
+	);
 
-	ActionList.MapAction( Commands.StepOver,
-		FExecuteAction::CreateStatic( &FInternalPlayWorldCommandCallbacks::StepOver_Clicked ),
-		FCanExecuteAction::CreateStatic( &FInternalPlayWorldCommandCallbacks::IsStoppedAtBreakpoint ),
+	ActionList.MapAction(Commands.StepOver,
+		FExecuteAction::CreateStatic(&FInternalPlayWorldCommandCallbacks::StepOver_Clicked),
+		FCanExecuteAction::CreateStatic(&FInternalPlayWorldCommandCallbacks::IsStoppedAtBreakpoint),
 		FIsActionChecked(),
-		FIsActionChecked::CreateStatic( &FInternalPlayWorldCommandCallbacks::IsStoppedAtBreakpoint )
-		);
+		FIsActionChecked::CreateStatic(&FInternalPlayWorldCommandCallbacks::IsStoppedAtBreakpoint)
+	);
 
-	ActionList.MapAction( Commands.StepOut,
-		FExecuteAction::CreateStatic( &FInternalPlayWorldCommandCallbacks::StepOut_Clicked ),
-		FCanExecuteAction::CreateStatic( &FInternalPlayWorldCommandCallbacks::IsStoppedAtBreakpoint ),
+	ActionList.MapAction(Commands.StepOut,
+		FExecuteAction::CreateStatic(&FInternalPlayWorldCommandCallbacks::StepOut_Clicked),
+		FCanExecuteAction::CreateStatic(&FInternalPlayWorldCommandCallbacks::IsStoppedAtBreakpoint),
 		FIsActionChecked(),
-		FIsActionChecked::CreateStatic( &FInternalPlayWorldCommandCallbacks::IsStoppedAtBreakpoint )
-		);
-
-	ActionList.MapAction( Commands.PlayInNetworkDedicatedServer,
-		FExecuteAction::CreateStatic( &FInternalPlayWorldCommandCallbacks::OnToggleDedicatedServerPIE ),
-		FCanExecuteAction(),
-		FIsActionChecked::CreateStatic( &FInternalPlayWorldCommandCallbacks::OnIsDedicatedServerPIEEnabled ) 
-		);
+		FIsActionChecked::CreateStatic(&FInternalPlayWorldCommandCallbacks::IsStoppedAtBreakpoint)
+	);
 
 	AddPIEPreviewDeviceActions(Commands, ActionList);
 }
@@ -617,24 +611,24 @@ void FPlayWorldCommands::AddPIEPreviewDeviceActions(const FPlayWorldCommands &Co
 void FPlayWorldCommands::BuildToolbar(FToolMenuSection& InSection, bool bIncludeLaunchButtonAndOptions)
 {
 	// Play
-	InSection.AddEntry(FToolMenuEntry::InitToolBarButton( 
-		FPlayWorldCommands::Get().RepeatLastPlay, 
+	InSection.AddEntry(FToolMenuEntry::InitToolBarButton(
+		FPlayWorldCommands::Get().RepeatLastPlay,
 		LOCTEXT("RepeatLastPlay", "Play"),
-		TAttribute< FText >::Create( TAttribute< FText >::FGetter::CreateStatic( &FInternalPlayWorldCommandCallbacks::GetRepeatLastPlayToolTip ) ),
-		TAttribute< FSlateIcon >::Create( TAttribute< FSlateIcon >::FGetter::CreateStatic( &FInternalPlayWorldCommandCallbacks::GetRepeatLastPlayIcon ) ),
+		TAttribute< FText >::Create(TAttribute< FText >::FGetter::CreateStatic(&FInternalPlayWorldCommandCallbacks::GetRepeatLastPlayToolTip)),
+		TAttribute< FSlateIcon >::Create(TAttribute< FSlateIcon >::FGetter::CreateStatic(&FInternalPlayWorldCommandCallbacks::GetRepeatLastPlayIcon)),
 		FName(TEXT("LevelToolbarPlay"))
 	));
 
 	// Play combo box
 	FUIAction SpecialPIEOptionsMenuAction;
-	SpecialPIEOptionsMenuAction.IsActionVisibleDelegate = FIsActionButtonVisible::CreateStatic( &FInternalPlayWorldCommandCallbacks::CanShowNonPlayWorldOnlyActions );
+	SpecialPIEOptionsMenuAction.IsActionVisibleDelegate = FIsActionButtonVisible::CreateStatic(&FInternalPlayWorldCommandCallbacks::CanShowNonPlayWorldOnlyActions);
 
 	InSection.AddEntry(FToolMenuEntry::InitComboButton(
 		"PlayCombo",
 		SpecialPIEOptionsMenuAction,
-		FOnGetContent::CreateStatic( &GeneratePlayMenuContent, GlobalPlayWorldActions.ToSharedRef() ),
-		LOCTEXT( "PlayCombo_Label", "Active Play Mode" ),
-		LOCTEXT( "PIEComboToolTip", "Change Play Mode and Play Settings" ),
+		FOnGetContent::CreateStatic(&GeneratePlayMenuContent, GlobalPlayWorldActions.ToSharedRef()),
+		LOCTEXT("PlayCombo_Label", "Active Play Mode"),
+		LOCTEXT("PIEComboToolTip", "Change Play Mode and Play Settings"),
 		FSlateIcon(FEditorStyle::GetStyleSetName(), "LevelEditor.RepeatLastPlay"),
 		true
 	));
@@ -767,7 +761,7 @@ static void AddAndroidConfigExportMenu(FMenuBuilder& InMenuBuilder)
 	);
 }
 
-static void MakePreviewDeviceMenu(FMenuBuilder& MenuBuilder )
+static void MakePreviewDeviceMenu(FMenuBuilder& MenuBuilder)
 {
 	struct FLocal
 	{
@@ -818,7 +812,7 @@ static void MakePreviewDeviceMenu(FMenuBuilder& MenuBuilder )
 
 	const TArray<TSharedPtr<FUICommandInfo>>& TargetedMobilePreviewDeviceCommands = FPlayWorldCommands::Get().PlayInTargetedMobilePreviewDevices;
 	auto PIEPreviewDeviceModule = FModuleManager::LoadModulePtr<FPIEPreviewDeviceModule>(TEXT("PIEPreviewDeviceProfileSelector"));
-	if(PIEPreviewDeviceModule)
+	if (PIEPreviewDeviceModule)
 	{
 		const FPIEPreviewDeviceContainer& DeviceContainer = PIEPreviewDeviceModule->GetPreviewDeviceContainer();
 		MenuBuilder.BeginSection("LevelEditorPlayModesPreviewDevice", LOCTEXT("PreviewDevicePlayButtonModesSection", "Preview Devices"));
@@ -827,7 +821,7 @@ static void MakePreviewDeviceMenu(FMenuBuilder& MenuBuilder )
 	}
 }
 
-TSharedRef< SWidget > FPlayWorldCommands::GeneratePlayMenuContent( TSharedRef<FUICommandList> InCommandList )
+TSharedRef< SWidget > FPlayWorldCommands::GeneratePlayMenuContent(TSharedRef<FUICommandList> InCommandList)
 {
 	// Get all menu extenders for this context menu from the level editor module
 	FLevelEditorModule& LevelEditorModule = FModuleManager::GetModuleChecked<FLevelEditorModule>(TEXT("LevelEditor"));
@@ -835,11 +829,11 @@ TSharedRef< SWidget > FPlayWorldCommands::GeneratePlayMenuContent( TSharedRef<FU
 
 	struct FLocal
 	{
-		static void AddPlayModeMenuEntry( FMenuBuilder& MenuBuilder, EPlayModeType PlayMode )
+		static void AddPlayModeMenuEntry(FMenuBuilder& MenuBuilder, EPlayModeType PlayMode)
 		{
 			TSharedPtr<FUICommandInfo> PlayModeCommand;
 
-			switch(PlayMode)
+			switch (PlayMode)
 			{
 			case PlayMode_InEditorFloating:
 				PlayModeCommand = FPlayWorldCommands::Get().PlayInEditorFloating;
@@ -878,7 +872,7 @@ TSharedRef< SWidget > FPlayWorldCommands::GeneratePlayMenuContent( TSharedRef<FU
 	};
 
 	const bool bShouldCloseWindowAfterMenuSelection = true;
-	FMenuBuilder MenuBuilder( bShouldCloseWindowAfterMenuSelection, InCommandList, MenuExtender );
+	FMenuBuilder MenuBuilder(bShouldCloseWindowAfterMenuSelection, InCommandList, MenuExtender);
 
 	// play in view port
 	MenuBuilder.BeginSection("LevelEditorPlayModes", LOCTEXT("PlayButtonModesSection", "Modes"));
@@ -893,7 +887,7 @@ TSharedRef< SWidget > FPlayWorldCommands::GeneratePlayMenuContent( TSharedRef<FU
 				LOCTEXT("TargetedMobilePreviewSubMenu_ToolTip", "Play this level using a specified mobile device preview (runs in its own process)"),
 				FNewMenuDelegate::CreateStatic(&MakePreviewDeviceMenu), false,
 				FSlateIcon(FEditorStyle::GetStyleSetName(), "PlayWorld.PlayInMobilePreview")
-				);
+			);
 		}
 
 		FLocal::AddPlayModeMenuEntry(MenuBuilder, PlayMode_InVulkanPreview);
@@ -907,12 +901,12 @@ TSharedRef< SWidget > FPlayWorldCommands::GeneratePlayMenuContent( TSharedRef<FU
 	// tip section
 	MenuBuilder.BeginSection("LevelEditorPlayTip");
 	{
-		MenuBuilder.AddWidget( 
-			SNew( STextBlock )
-				.ColorAndOpacity(FSlateColor::UseSubduedForeground())
-				.Text(LOCTEXT("PlayInTip", "Launching a game preview with a different mode will change your default 'Play' mode in the toolbar"))
-				.WrapTextAt(250),
-			FText::GetEmpty() );
+		MenuBuilder.AddWidget(
+			SNew(STextBlock)
+			.ColorAndOpacity(FSlateColor::UseSubduedForeground())
+			.Text(LOCTEXT("PlayInTip", "Launching a game preview with a different mode will change your default 'Play' mode in the toolbar"))
+			.WrapTextAt(250),
+			FText::GetEmpty());
 	}
 	MenuBuilder.EndSection();
 
@@ -926,25 +920,31 @@ TSharedRef< SWidget > FPlayWorldCommands::GeneratePlayMenuContent( TSharedRef<FU
 
 	// Basic network options
 	const ULevelEditorPlaySettings* PlayInSettings = GetDefault<ULevelEditorPlaySettings>();
-	if ( PlayInSettings->IsPlayNumberOfClientsActive() || PlayInSettings->IsPlayNetDedicatedActive() )
 	{
 		MenuBuilder.BeginSection("LevelEditorPlayInWindowNetwork", LOCTEXT("LevelEditorPlayInWindowNetworkSection", "Multiplayer Options"));
-		if ( PlayInSettings->IsPlayNumberOfClientsActive() )
+		// Num Clients
 		{
 			TSharedRef<SWidget> NumPlayers = SNew(SSpinBox<int32>)	// Copy limits from PlayNumberOfClients meta data
 				.MinValue(1)
 				.MaxValue(64)
 				.MinSliderValue(1)
 				.MaxSliderValue(4)
-				.ToolTipText(LOCTEXT( "NumberOfClientsToolTip", "The editor and listen server count as players, a dedicated server will not. Clients make up the remainder." ))
+				.ToolTipText(LOCTEXT("NumberOfClientsToolTip", "How many client instances do you want to create? The first instance respects the Play Mode location (PIE/PINW) and additional instances respect the RunUnderOneProcess setting."))
 				.Value(FInternalPlayWorldCommandCallbacks::GetNumberOfClients())
 				.OnValueCommitted_Static(&FInternalPlayWorldCommandCallbacks::SetNumberOfClients);
-			
-			MenuBuilder.AddWidget(NumPlayers, LOCTEXT( "NumberOfClientsMenuWidget", "Number of Players" ));
+
+			MenuBuilder.AddWidget(NumPlayers, LOCTEXT("NumberOfClientsMenuWidget", "Number of Players"));
 		}
-		if ( PlayInSettings->IsPlayNetDedicatedActive() )
+		// Net Mode
 		{
-			MenuBuilder.AddMenuEntry( FPlayWorldCommands::Get().PlayInNetworkDedicatedServer );
+			const UEnum* PlayNetModeEnum = FindObject<UEnum>(ANY_PACKAGE, TEXT("EPlayNetMode"));
+			
+			TSharedRef<SWidget> NetMode = SNew(SEnumCombobox, PlayNetModeEnum)
+				.CurrentValue(TAttribute<int32>::Create(TAttribute<int32>::FGetter::CreateStatic(&FInternalPlayWorldCommandCallbacks::GetNetPlayMode)))
+				.OnEnumSelectionChanged(SEnumCombobox::FOnEnumSelectionChanged::CreateStatic(&FInternalPlayWorldCommandCallbacks::SetNetPlayMode))
+				.ToolTipText(LOCTEXT("NetworkModeToolTip", "Which network mode should the clients launch in? A server will automatically be started if needed."));
+
+			MenuBuilder.AddWidget(NetMode, LOCTEXT("NetworkModeMenuWidget", "Net Mode"));
 		}
 		MenuBuilder.EndSection();
 	}
@@ -959,7 +959,7 @@ TSharedRef< SWidget > FPlayWorldCommands::GeneratePlayMenuContent( TSharedRef<FU
 	return MenuBuilder.MakeWidget();
 }
 
-/* 
+/*
  * Create an All_<platform>_devices_on_<host> submenu
  * can be extended to any othe All <Platform> aggregate proxy
 */
@@ -981,7 +981,7 @@ static void MakeAllDevicesSubMenu(FMenuBuilder& InMenuBuilder, const PlatformInf
 		{
 			continue;
 		}
-		
+
 		FString DeviceListStr;
 		bool bVariantHasDevices = false;
 
@@ -1032,7 +1032,7 @@ static void MakeAllDevicesSubMenu(FMenuBuilder& InMenuBuilder, const PlatformInf
 	}
 }
 
-TSharedRef< SWidget > FPlayWorldCommands::GenerateLaunchMenuContent( TSharedRef<FUICommandList> InCommandList )
+TSharedRef< SWidget > FPlayWorldCommands::GenerateLaunchMenuContent(TSharedRef<FUICommandList> InCommandList)
 {
 	const bool bShouldCloseWindowAfterMenuSelection = true;
 	FMenuBuilder MenuBuilder(bShouldCloseWindowAfterMenuSelection, InCommandList);
@@ -1076,7 +1076,7 @@ TSharedRef< SWidget > FPlayWorldCommands::GenerateLaunchMenuContent( TSharedRef<
 				TArray<TSharedPtr<ITargetDeviceProxy>> DeviceProxies;
 				// the list of proxies include the "Al_Android" entry
 				TargetDeviceServicesModule->GetDeviceProxyManager()->GetAllProxies(VanillaPlatform.PlatformInfo->VanillaPlatformName, DeviceProxies);
-					
+
 				// if this platform had no devices, but we want to show an extra option if not installed right
 				if (DeviceProxies.Num() == 0)
 				{
@@ -1113,7 +1113,7 @@ TSharedRef< SWidget > FPlayWorldCommands::GenerateLaunchMenuContent( TSharedRef<
 							FExecuteAction::CreateStatic(&FInternalPlayWorldCommandCallbacks::HandleLaunchOnDeviceActionExecute, DeviceProxy->GetTargetDeviceId(NAME_None), DeviceProxy->GetName()),
 							FCanExecuteAction::CreateStatic(&FInternalPlayWorldCommandCallbacks::HandleLaunchOnDeviceActionCanExecute, DeviceProxy->GetName()),
 							FIsActionChecked::CreateStatic(&FInternalPlayWorldCommandCallbacks::HandleLaunchOnDeviceActionIsChecked, DeviceProxy->GetName())
-							);
+						);
 
 						// ... generate display label...
 						FFormatNamedArguments LabelArguments;
@@ -1145,7 +1145,7 @@ TSharedRef< SWidget > FPlayWorldCommands::GenerateLaunchMenuContent( TSharedRef<
 						}
 
 						FProjectStatus ProjectStatus;
-						if (IProjectManager::Get().QueryStatusForCurrentProject(ProjectStatus) && !ProjectStatus.IsTargetPlatformSupported(VanillaPlatform.PlatformInfo->VanillaPlatformName)) 
+						if (IProjectManager::Get().QueryStatusForCurrentProject(ProjectStatus) && !ProjectStatus.IsTargetPlatformSupported(VanillaPlatform.PlatformInfo->VanillaPlatformName))
 						{
 							FText TooltipLine2 = FText::Format(LOCTEXT("LaunchDevicePlatformWarning", "{DisplayName} is not listed as a target platform for this project, so may not run as expected."), TooltipArguments);
 							Tooltip = FText::Format(FText::FromString(TEXT("{0}\n\n{1}")), Tooltip, TooltipLine2);
@@ -1158,7 +1158,7 @@ TSharedRef< SWidget > FPlayWorldCommands::GenerateLaunchMenuContent( TSharedRef<
 							NAME_None,
 							Tooltip,
 							EUserInterfaceActionType::Check
-							);
+						);
 					}
 				}
 			}
@@ -1175,52 +1175,52 @@ TSharedRef< SWidget > FPlayWorldCommands::GenerateLaunchMenuContent( TSharedRef<
 	MenuBuilder.EndSection();
 
 	MenuBuilder.BeginSection("CookerSettings");
-	
+
 	UCookerSettings* CookerSettings = GetMutableDefault<UCookerSettings>();
 
 	{
 		FUIAction UIAction;
 		UIAction.ExecuteAction = FExecuteAction::CreateLambda([CookerSettings]
-			{
-				CookerSettings->bCookOnTheFlyForLaunchOn = !CookerSettings->bCookOnTheFlyForLaunchOn;
-				CookerSettings->Modify(true);
+		{
+			CookerSettings->bCookOnTheFlyForLaunchOn = !CookerSettings->bCookOnTheFlyForLaunchOn;
+			CookerSettings->Modify(true);
 
-				// Update source control
+			// Update source control
 
-				FString ConfigPath = FPaths::ConvertRelativePathToFull(CookerSettings->GetDefaultConfigFilename());
+			FString ConfigPath = FPaths::ConvertRelativePathToFull(CookerSettings->GetDefaultConfigFilename());
 
-				if (FPlatformFileManager::Get().GetPlatformFile().FileExists(*ConfigPath))
+			if (FPlatformFileManager::Get().GetPlatformFile().FileExists(*ConfigPath))
+			{
+				if (ISourceControlModule::Get().IsEnabled())
 				{
-					if (ISourceControlModule::Get().IsEnabled())
-					{
-						FText ErrorMessage;
+					FText ErrorMessage;
 
-						if (!SourceControlHelpers::CheckoutOrMarkForAdd(ConfigPath, FText::FromString(ConfigPath), NULL, ErrorMessage))
-						{
-							FNotificationInfo Info(ErrorMessage);
-							Info.ExpireDuration = 3.0f;
-							FSlateNotificationManager::Get().AddNotification(Info);
-						}
+					if (!SourceControlHelpers::CheckoutOrMarkForAdd(ConfigPath, FText::FromString(ConfigPath), NULL, ErrorMessage))
+					{
+						FNotificationInfo Info(ErrorMessage);
+						Info.ExpireDuration = 3.0f;
+						FSlateNotificationManager::Get().AddNotification(Info);
 					}
-					else
+				}
+				else
+				{
+					if (!FPlatformFileManager::Get().GetPlatformFile().SetReadOnly(*ConfigPath, false))
 					{
-						if (!FPlatformFileManager::Get().GetPlatformFile().SetReadOnly(*ConfigPath, false))
-						{
-							FNotificationInfo Info(FText::Format(LOCTEXT("FailedToMakeWritable", "Could not make {0} writable."), FText::FromString(ConfigPath)));
-							Info.ExpireDuration = 3.0f;
-							FSlateNotificationManager::Get().AddNotification(Info);
-						}
+						FNotificationInfo Info(FText::Format(LOCTEXT("FailedToMakeWritable", "Could not make {0} writable."), FText::FromString(ConfigPath)));
+						Info.ExpireDuration = 3.0f;
+						FSlateNotificationManager::Get().AddNotification(Info);
 					}
 				}
+			}
 
-				// Save settings
-				CookerSettings->UpdateSinglePropertyInConfigFile(CookerSettings->GetClass()->FindPropertyByName(GET_MEMBER_NAME_CHECKED(UCookerSettings, bCookOnTheFlyForLaunchOn)), CookerSettings->GetDefaultConfigFilename());
-			});
+			// Save settings
+			CookerSettings->UpdateSinglePropertyInConfigFile(CookerSettings->GetClass()->FindPropertyByName(GET_MEMBER_NAME_CHECKED(UCookerSettings, bCookOnTheFlyForLaunchOn)), CookerSettings->GetDefaultConfigFilename());
+		});
 
 		UIAction.GetActionCheckState = FGetActionCheckState::CreateLambda([CookerSettings]
-			{
-				return CookerSettings->bCookOnTheFlyForLaunchOn ? ECheckBoxState::Checked : ECheckBoxState::Unchecked;
-			});
+		{
+			return CookerSettings->bCookOnTheFlyForLaunchOn ? ECheckBoxState::Checked : ECheckBoxState::Unchecked;
+		});
 
 		MenuBuilder.AddMenuEntry(
 			LOCTEXT("CookOnTheFlyOnLaunch", "Enable cooking on the fly"),
@@ -1252,7 +1252,7 @@ TSharedRef< SWidget > FPlayWorldCommands::GenerateLaunchMenuContent( TSharedRef<
 				FUIAction NoDeviceAction(
 					FExecuteAction::CreateStatic(&FInternalPlayWorldCommandCallbacks::HandleNoDeviceFoundActionExecute),
 					FCanExecuteAction::CreateStatic(&FInternalPlayWorldCommandCallbacks::HandleNoDeviceFoundActionCanExecute)
-					);
+				);
 
 				// ... generate tooltip text
 				FFormatNamedArguments TooltipArguments;
@@ -1261,12 +1261,12 @@ TSharedRef< SWidget > FPlayWorldCommands::GenerateLaunchMenuContent( TSharedRef<
 
 				// ... and add a menu entry
 				MenuBuilder.AddMenuEntry(
-					NoDeviceAction, 
-					ProjectTargetPlatformEditorModule.MakePlatformMenuItemWidget(*PlatformInfo, true, Label), 
-					NAME_None, 
-					Tooltip, 
+					NoDeviceAction,
+					ProjectTargetPlatformEditorModule.MakePlatformMenuItemWidget(*PlatformInfo, true, Label),
+					NAME_None,
+					Tooltip,
 					EUserInterfaceActionType::Check
-					);
+				);
 			}
 		}
 		MenuBuilder.EndSection();
@@ -1275,11 +1275,11 @@ TSharedRef< SWidget > FPlayWorldCommands::GenerateLaunchMenuContent( TSharedRef<
 	// tip section
 	MenuBuilder.BeginSection("LevelEditorLaunchHint");
 	{
-		MenuBuilder.AddWidget( 
-			SNew( STextBlock )
-				.ColorAndOpacity(FSlateColor::UseSubduedForeground())
-				.Text(LOCTEXT("ZoomToFitHorizontal", "Launching a game on a different device will change your default 'Launch' device in the toolbar"))
-				.WrapTextAt(250),
+		MenuBuilder.AddWidget(
+			SNew(STextBlock)
+			.ColorAndOpacity(FSlateColor::UseSubduedForeground())
+			.Text(LOCTEXT("ZoomToFitHorizontal", "Launching a game on a different device will change your default 'Launch' device in the toolbar"))
+			.WrapTextAt(250),
 			FText::GetEmpty()
 		);
 	}
@@ -1300,7 +1300,7 @@ TSharedRef< SWidget > FPlayWorldCommands::GenerateLaunchMenuContent( TSharedRef<
 				FText Label = FText::Format(LOCTEXT("LaunchPlatformLabel", "{PlatformName} Support"), LabelArguments);
 
 
-				MenuBuilder.AddMenuEntry( 
+				MenuBuilder.AddMenuEntry(
 					Label,
 					LOCTEXT("PlatformSDK", "Show information on setting up the platform tools"),
 					FSlateIcon(FEditorStyle::GetStyleSetName(), "LevelEditor.BrowseDocumentation"),
@@ -1320,15 +1320,15 @@ TSharedRef< SWidget > FPlayWorldCommands::GenerateLaunchMenuContent( TSharedRef<
 			TAttribute<FText>(),
 			TAttribute<FText>(),
 			FSlateIcon(FEditorStyle::GetStyleSetName(), "Launcher.TabIcon")
-			);
+		);
 
-		MenuBuilder.AddMenuEntry( FPlayWorldCommands::Get().OpenDeviceManager,
+		MenuBuilder.AddMenuEntry(FPlayWorldCommands::Get().OpenDeviceManager,
 			NAME_None,
 			TAttribute<FText>(),
 			TAttribute<FText>(),
 			FSlateIcon(FEditorStyle::GetStyleSetName(), "DeviceDetails.TabIcon")
-			);
-		
+		);
+
 		ProjectTargetPlatformEditorModule.AddOpenProjectTargetPlatformEditorMenuItem(MenuBuilder);
 	}
 	MenuBuilder.EndSection();
@@ -1343,40 +1343,33 @@ TSharedRef< SWidget > FPlayWorldCommands::GenerateLaunchMenuContent( TSharedRef<
 void FPlayWorldCommandCallbacks::StartPlayFromHere()
 {
 	// Is a PIE session already running?  If so we close it first
-	if( GUnrealEd->PlayWorld != NULL )
+	if (GUnrealEd->PlayWorld != NULL)
 	{
 		GUnrealEd->EndPlayMap();
 	}
 
+	FRequestPlaySessionParams SessionParams;
+
 	UClass* const PlayerStartClass = GUnrealEd->PlayFromHerePlayerStartClass ? (UClass*)GUnrealEd->PlayFromHerePlayerStartClass : APlayerStart::StaticClass();
 
 	// Figure out the start location of the player
 	UCapsuleComponent*	DefaultCollisionComponent = CastChecked<UCapsuleComponent>(PlayerStartClass->GetDefaultObject<AActor>()->GetRootComponent());
-	FVector				CollisionExtent = FVector(DefaultCollisionComponent->GetScaledCapsuleRadius(),DefaultCollisionComponent->GetScaledCapsuleRadius(),DefaultCollisionComponent->GetScaledCapsuleHalfHeight());
-	FVector				StartLocation = GEditor->UnsnappedClickLocation + GEditor->ClickPlane * (FVector::BoxPushOut(GEditor->ClickPlane,CollisionExtent) + 0.1f);
+	FVector	CollisionExtent = FVector(DefaultCollisionComponent->GetScaledCapsuleRadius(), DefaultCollisionComponent->GetScaledCapsuleRadius(), DefaultCollisionComponent->GetScaledCapsuleHalfHeight());
+	SessionParams.StartLocation = GEditor->UnsnappedClickLocation + GEditor->ClickPlane * (FVector::BoxPushOut(GEditor->ClickPlane, CollisionExtent) + 0.1f);
 
-	FRotator StartRotation = FRotator::ZeroRotator;
-
-	FLevelEditorModule& LevelEditorModule = FModuleManager::GetModuleChecked<FLevelEditorModule>( "LevelEditor");
+	FLevelEditorModule& LevelEditorModule = FModuleManager::GetModuleChecked<FLevelEditorModule>("LevelEditor");
 
 	TSharedPtr<IAssetViewport> ActiveLevelViewport = LevelEditorModule.GetFirstActiveViewport();
 
-	const bool bSimulateInEditor = false;
-	if( ActiveLevelViewport.IsValid() )
-	{
-		if( ActiveLevelViewport->GetAssetViewportClient().IsPerspective() )
-		{
-			StartRotation = ActiveLevelViewport->GetAssetViewportClient().GetViewRotation();
-		}
 
-		// If there is an active level view port, play the game in it.
-		GUnrealEd->RequestPlaySession(false, ActiveLevelViewport, bSimulateInEditor, &StartLocation, &StartRotation, -1, false );
-	}
-	else
+	if (ActiveLevelViewport.IsValid() && ActiveLevelViewport->GetAssetViewportClient().IsPerspective())
 	{
-		// No active level view port, spawn a new window to play in.
-		GUnrealEd->RequestPlaySession(false, NULL, bSimulateInEditor, &StartLocation, NULL );
+		// If there is no level viewport, a new window will be spawned to play in.
+		SessionParams.DestinationSlateViewport = ActiveLevelViewport;
+		SessionParams.StartRotation = ActiveLevelViewport->GetAssetViewportClient().GetViewRotation();
 	}
+
+	GUnrealEd->RequestPlaySession(SessionParams);
 }
 
 
@@ -1468,45 +1461,45 @@ bool FPlayWorldCommandCallbacks::HasPlayWorldAndRunning()
 
 FText FInternalPlayWorldCommandCallbacks::GetPossessEjectLabel()
 {
-	if ( IsInPIE() )
+	if (IsInPIE())
 	{
-		return LOCTEXT( "EjectLabel", "Eject" );
+		return LOCTEXT("EjectLabel", "Eject");
 	}
-	else if ( IsInSIE() )
+	else if (IsInSIE())
 	{
-		return LOCTEXT( "PossessLabel", "Possess" );
+		return LOCTEXT("PossessLabel", "Possess");
 	}
 	else
 	{
-		return LOCTEXT( "ToggleBetweenPieAndSIELabel", "Toggle Between PIE and SIE" );
+		return LOCTEXT("ToggleBetweenPieAndSIELabel", "Toggle Between PIE and SIE");
 	}
 }
 
 
 FText FInternalPlayWorldCommandCallbacks::GetPossessEjectTooltip()
 {
-	if ( IsInPIE() )
+	if (IsInPIE())
 	{
-		return LOCTEXT( "EjectToolTip", "Detaches from the player controller, allowing regular editor controls" );
+		return LOCTEXT("EjectToolTip", "Detaches from the player controller, allowing regular editor controls");
 	}
-	else if ( IsInSIE() )
+	else if (IsInSIE())
 	{
-		return LOCTEXT( "PossessToolTip", "Attaches to the player controller, allowing normal gameplay controls" );
+		return LOCTEXT("PossessToolTip", "Attaches to the player controller, allowing normal gameplay controls");
 	}
 	else
 	{
-		return LOCTEXT( "ToggleBetweenPieAndSIEToolTip", "Toggles the current play session between play in editor and simulate in editor" );
+		return LOCTEXT("ToggleBetweenPieAndSIEToolTip", "Toggles the current play session between play in editor and simulate in editor");
 	}
 }
 
 
 FSlateIcon FInternalPlayWorldCommandCallbacks::GetPossessEjectImage()
 {
-	if ( IsInPIE() )
+	if (IsInPIE())
 	{
 		return FSlateIcon(FEditorStyle::GetStyleSetName(), "PlayWorld.EjectFromPlayer");
 	}
-	else if ( IsInSIE() )
+	else if (IsInSIE())
 	{
 		return FSlateIcon(FEditorStyle::GetStyleSetName(), "PlayWorld.PossessPlayer");
 	}
@@ -1531,10 +1524,10 @@ void SetLastExecutedPlayMode(EPlayModeType PlayMode)
 {
 	ULevelEditorPlaySettings* PlaySettings = GetMutableDefault<ULevelEditorPlaySettings>();
 	PlaySettings->LastExecutedPlayModeType = PlayMode;
-	
+
 	FPropertyChangedEvent PropChangeEvent(ULevelEditorPlaySettings::StaticClass()->FindPropertyByName(GET_MEMBER_NAME_CHECKED(ULevelEditorPlaySettings, LastExecutedPlayModeType)));
 	PlaySettings->PostEditChangeProperty(PropChangeEvent);
-	
+
 	PlaySettings->SaveConfig();
 }
 
@@ -1542,25 +1535,29 @@ void SetLastExecutedPlayMode(EPlayModeType PlayMode)
 void FInternalPlayWorldCommandCallbacks::Simulate_Clicked()
 {
 	// Is a simulation session already running?  If so, do nothing
-	if( HasPlayWorld() && GUnrealEd->bIsSimulatingInEditor )
+	if (HasPlayWorld() && GUnrealEd->bIsSimulatingInEditor)
 	{
 		return;
 	}
 
-	FLevelEditorModule& LevelEditorModule = FModuleManager::GetModuleChecked<FLevelEditorModule>( TEXT("LevelEditor") );
+	FLevelEditorModule& LevelEditorModule = FModuleManager::GetModuleChecked<FLevelEditorModule>(TEXT("LevelEditor"));
 
 	TSharedPtr<IAssetViewport> ActiveLevelViewport = LevelEditorModule.GetFirstActiveViewport();
-	if( ActiveLevelViewport.IsValid())
+	if (ActiveLevelViewport.IsValid())
 	{
 		// Start a new simulation session!
-		if( !HasPlayWorld() )
+		if (!HasPlayWorld())
 		{
-			if( FEngineAnalytics::IsAvailable() )
+			if (FEngineAnalytics::IsAvailable())
 			{
-				FEngineAnalytics::GetProvider().RecordEvent( TEXT("Editor.Usage.SimulateInEditor") );
+				FEngineAnalytics::GetProvider().RecordEvent(TEXT("Editor.Usage.SimulateInEditor"));
 			}
 			SetLastExecutedPlayMode(PlayMode_Simulate);
-			GUnrealEd->RequestPlaySession(false, ActiveLevelViewport, true/*bSimulateInEditor*/, NULL, NULL, -1, false );
+			FRequestPlaySessionParams SessionParams;
+			SessionParams.WorldType = EPlaySessionWorldType::SimulateInEditor;
+			SessionParams.DestinationSlateViewport = ActiveLevelViewport;
+
+			GUnrealEd->RequestPlaySession(SessionParams);
 		}
 		else if (ActiveLevelViewport->HasPlayInEditorViewport())
 		{
@@ -1590,14 +1587,14 @@ const TSharedRef < FUICommandInfo > GetLastPlaySessionCommand()
 	const FPlayWorldCommands& Commands = FPlayWorldCommands::Get();
 	TSharedRef < FUICommandInfo > Command = Commands.PlayInViewport.ToSharedRef();
 
-	switch( PlaySettings->LastExecutedPlayModeType )
+	switch (PlaySettings->LastExecutedPlayModeType)
 	{
-	case PlayMode_InViewPort:		
-		Command = Commands.PlayInViewport.ToSharedRef();				
+	case PlayMode_InViewPort:
+		Command = Commands.PlayInViewport.ToSharedRef();
 		break;
 
-	case PlayMode_InEditorFloating:			
-		Command = Commands.PlayInEditorFloating.ToSharedRef();			
+	case PlayMode_InEditorFloating:
+		Command = Commands.PlayInEditorFloating.ToSharedRef();
 		break;
 
 	case PlayMode_InMobilePreview:
@@ -1707,11 +1704,11 @@ void RecordLastExecutedPlayMode()
 
 		FEngineAnalytics::GetProvider().RecordEvent(TEXT("Editor.Usage.PIE"), TEXT("PlayLocation"), PlayLocationString);
 		FEngineAnalytics::GetProvider().RecordEvent(TEXT("Editor.Usage.PIE"), TEXT("PlayMode"), PlayModeString);
-	}	
+	}
 }
 
 
-void SetLastExecutedLaunchMode( ELaunchModeType LaunchMode )
+void SetLastExecutedLaunchMode(ELaunchModeType LaunchMode)
 {
 	ULevelEditorPlaySettings* PlaySettings = GetMutableDefault<ULevelEditorPlaySettings>();
 	PlaySettings->LastExecutedLaunchModeType = LaunchMode;
@@ -1738,7 +1735,7 @@ void FInternalPlayWorldCommandCallbacks::RepeatLastPlay_Clicked()
 
 bool FInternalPlayWorldCommandCallbacks::RepeatLastPlay_CanExecute()
 {
-	return FPlayWorldCommands::GlobalPlayWorldActions->CanExecuteAction( GetLastPlaySessionCommand() );
+	return FPlayWorldCommands::GlobalPlayWorldActions->CanExecuteAction(GetLastPlaySessionCommand());
 }
 
 
@@ -1754,57 +1751,48 @@ FSlateIcon FInternalPlayWorldCommandCallbacks::GetRepeatLastPlayIcon()
 }
 
 
-void FInternalPlayWorldCommandCallbacks::PlayInViewport_Clicked( )
+void FInternalPlayWorldCommandCallbacks::PlayInViewport_Clicked()
 {
-	FLevelEditorModule& LevelEditorModule = FModuleManager::GetModuleChecked<FLevelEditorModule>( TEXT("LevelEditor") );
+	FLevelEditorModule& LevelEditorModule = FModuleManager::GetModuleChecked<FLevelEditorModule>(TEXT("LevelEditor"));
 
 	/** Set PlayInViewPort as the last executed play command */
 	const FPlayWorldCommands& Commands = FPlayWorldCommands::Get();
 
-	SetLastExecutedPlayMode( PlayMode_InViewPort );
+	SetLastExecutedPlayMode(PlayMode_InViewPort);
 
 	RecordLastExecutedPlayMode();
 
 	TSharedPtr<IAssetViewport> ActiveLevelViewport = LevelEditorModule.GetFirstActiveViewport();
 
 	const bool bAtPlayerStart = (GetPlayModeLocation() == PlayLocation_DefaultPlayerStart);
-	const bool bSimulateInEditor = false;
+
+	FRequestPlaySessionParams SessionParams;
 
 	// Make sure we can find a path to the view port.  This will fail in cases where the view port widget
 	// is in a backgrounded tab, etc.  We can't currently support starting PIE in a backgrounded tab
 	// due to how PIE manages focus and requires event forwarding from the application.
 	if (ActiveLevelViewport.IsValid() && FSlateApplication::Get().FindWidgetWindow(ActiveLevelViewport->AsWidget()).IsValid())
 	{
-		const FVector* StartLoc = NULL;
-		const FRotator* StartRot = NULL;
-		if( !bAtPlayerStart )
+		SessionParams.DestinationSlateViewport = ActiveLevelViewport;
+		if (!bAtPlayerStart)
 		{
 			// Start the player where the camera is if not forcing from player start
-			StartLoc = &ActiveLevelViewport->GetAssetViewportClient().GetViewLocation();
-			StartRot = &ActiveLevelViewport->GetAssetViewportClient().GetViewRotation();
+			SessionParams.StartLocation = ActiveLevelViewport->GetAssetViewportClient().GetViewLocation();
+			SessionParams.StartRotation = ActiveLevelViewport->GetAssetViewportClient().GetViewRotation();
 		}
+	}
 
-		// @todo UE4: Not supported yet
-		const bool bUseMobilePreview = false;
-		const int32 DestinationConsoleIndex = -1;
-
-		if( !HasPlayWorld() )
-		{
-			// If there is an active level view port, play the game in it.
-			GUnrealEd->RequestPlaySession( bAtPlayerStart, ActiveLevelViewport, bSimulateInEditor, StartLoc, StartRot, DestinationConsoleIndex, bUseMobilePreview );
-		}
-		else
-		{
-			// There is already a play world active which means simulate in editor is happening
-			// Toggle to pie 
-			check( !GIsPlayInEditorWorld );
-			GUnrealEd->RequestToggleBetweenPIEandSIE();
-		}
+	if (!HasPlayWorld())
+	{
+		// If there is an active level view port, play the game in it, otherwise make a new window.
+		GUnrealEd->RequestPlaySession(SessionParams);
 	}
 	else
 	{
-		// No active level view port, spawn a new window to play in.
-		GUnrealEd->RequestPlaySession( bAtPlayerStart, NULL, bSimulateInEditor );
+		// There is already a play world active which means simulate in editor is happening
+		// Toggle to PIE
+		check(!GIsPlayInEditorWorld);
+		GUnrealEd->RequestToggleBetweenPIEandSIE();
 	}
 }
 
@@ -1817,45 +1805,42 @@ bool FInternalPlayWorldCommandCallbacks::PlayInViewport_CanExecute()
 	}
 
 	// Allow PIE if we don't already have a play session or the play session is simulate in editor (which we can toggle to PIE)
-	return (!GEditor->bIsPlayWorldQueued && !HasPlayWorld() && !GEditor->IsLightingBuildCurrentlyRunning() ) || GUnrealEd->bIsSimulatingInEditor;
+	return (!GEditor->IsPlaySessionInProgress() && !HasPlayWorld() && !GEditor->IsLightingBuildCurrentlyRunning()) || GUnrealEd->IsSimulateInEditorInProgress();
 }
 
 
-void FInternalPlayWorldCommandCallbacks::PlayInEditorFloating_Clicked( )
+void FInternalPlayWorldCommandCallbacks::PlayInEditorFloating_Clicked()
 {
-	FLevelEditorModule& LevelEditorModule = FModuleManager::GetModuleChecked<FLevelEditorModule>( TEXT("LevelEditor") );
+	FLevelEditorModule& LevelEditorModule = FModuleManager::GetModuleChecked<FLevelEditorModule>(TEXT("LevelEditor"));
+
+	SetLastExecutedPlayMode(PlayMode_InEditorFloating);
 
-	SetLastExecutedPlayMode( PlayMode_InEditorFloating );
+	FRequestPlaySessionParams SessionParams;
 
 	// Is a PIE session already running?  If not, then we'll kick off a new one
-	if( !HasPlayWorld() )
+	if (!HasPlayWorld())
 	{
 		RecordLastExecutedPlayMode();
 
 		const bool bAtPlayerStart = (GetPlayModeLocation() == PlayLocation_DefaultPlayerStart);
-		const bool bSimulateInEditor = false;
-
-		const FVector* StartLoc = NULL;
-		const FRotator* StartRot = NULL;
-
-		if ( !bAtPlayerStart )
+		if (!bAtPlayerStart)
 		{
 			TSharedPtr<IAssetViewport> ActiveLevelViewport = LevelEditorModule.GetFirstActiveViewport();
 
 			// Make sure we can find a path to the view port.  This will fail in cases where the view port widget
 			// is in a backgrounded tab, etc.  We can't currently support starting PIE in a backgrounded tab
 			// due to how PIE manages focus and requires event forwarding from the application.
-			if( ActiveLevelViewport.IsValid() &&
-				FSlateApplication::Get().FindWidgetWindow( ActiveLevelViewport->AsWidget() ).IsValid() )
+			if (ActiveLevelViewport.IsValid() &&
+				FSlateApplication::Get().FindWidgetWindow(ActiveLevelViewport->AsWidget()).IsValid())
 			{
 				// Start the player where the camera is if not forcing from player start
-				StartLoc = &ActiveLevelViewport->GetAssetViewportClient().GetViewLocation();
-				StartRot = &ActiveLevelViewport->GetAssetViewportClient().GetViewRotation();
+				SessionParams.StartLocation = ActiveLevelViewport->GetAssetViewportClient().GetViewLocation();
+				SessionParams.StartRotation = ActiveLevelViewport->GetAssetViewportClient().GetViewRotation();
 			}
 		}
 
 		// Spawn a new window to play in.
-		GUnrealEd->RequestPlaySession(bAtPlayerStart, NULL, bSimulateInEditor, StartLoc, StartRot );
+		GUnrealEd->RequestPlaySession(SessionParams);
 	}
 	else
 	{
@@ -1875,6 +1860,7 @@ void FInternalPlayWorldCommandCallbacks::PlayInVR_Clicked()
 	FLevelEditorModule& LevelEditorModule = FModuleManager::GetModuleChecked<FLevelEditorModule>(TEXT("LevelEditor"));
 
 	SetLastExecutedPlayMode(PlayMode_InVR);
+	FRequestPlaySessionParams SessionParams;
 
 	// Is a PIE session already running?  If not, then we'll kick off a new one
 	if (!HasPlayWorld())
@@ -1882,11 +1868,6 @@ void FInternalPlayWorldCommandCallbacks::PlayInVR_Clicked()
 		RecordLastExecutedPlayMode();
 
 		const bool bAtPlayerStart = (GetPlayModeLocation() == PlayLocation_DefaultPlayerStart);
-		const bool bSimulateInEditor = false;
-
-		const FVector* StartLoc = NULL;
-		const FRotator* StartRot = NULL;
-
 		if (!bAtPlayerStart)
 		{
 			TSharedPtr<IAssetViewport> ActiveLevelViewport = LevelEditorModule.GetFirstActiveViewport();
@@ -1898,14 +1879,19 @@ void FInternalPlayWorldCommandCallbacks::PlayInVR_Clicked()
 				FSlateApplication::Get().FindWidgetWindow(ActiveLevelViewport->AsWidget()).IsValid())
 			{
 				// Start the player where the camera is if not forcing from player start
-				StartLoc = &ActiveLevelViewport->GetAssetViewportClient().GetViewLocation();
-				StartRot = &ActiveLevelViewport->GetAssetViewportClient().GetViewRotation();
+				SessionParams.StartLocation = ActiveLevelViewport->GetAssetViewportClient().GetViewLocation();
+				SessionParams.StartRotation = ActiveLevelViewport->GetAssetViewportClient().GetViewRotation();
 			}
 		}
 
 		const bool bHMDIsReady = (GEngine && GEngine->XRSystem.IsValid() && GEngine->XRSystem->GetHMDDevice() && GEngine->XRSystem->GetHMDDevice()->IsHMDConnected());
+		if (bHMDIsReady)
+		{
+			SessionParams.SessionPreviewTypeOverride = EPlaySessionPreviewType::VRPreview;
+		}
+
 		// Spawn a new window to play in.
-		GUnrealEd->RequestPlaySession(bAtPlayerStart, NULL, bSimulateInEditor, StartLoc, StartRot, -1, false, bHMDIsReady);
+		GUnrealEd->RequestPlaySession(SessionParams);
 	}
 }
 
@@ -1932,39 +1918,47 @@ void FInternalPlayWorldCommandCallbacks::PlayInNewProcessPreviewDevice_Clicked(F
 
 void FInternalPlayWorldCommandCallbacks::PlayInNewProcess_Clicked(EPlayModeType PlayModeType)
 {
-	check(PlayModeType == PlayMode_InNewProcess || PlayModeType == PlayMode_InMobilePreview 
-		|| PlayModeType == PlayMode_InTargetedMobilePreview || PlayModeType == PlayMode_InVulkanPreview	);
+	check(PlayModeType == PlayMode_InNewProcess || PlayModeType == PlayMode_InMobilePreview
+		|| PlayModeType == PlayMode_InTargetedMobilePreview || PlayModeType == PlayMode_InVulkanPreview);
 
 	SetLastExecutedPlayMode(PlayModeType);
+	FRequestPlaySessionParams SessionParams;
 
-	if( !HasPlayWorld() )
+	if (!HasPlayWorld())
 	{
 		RecordLastExecutedPlayMode();
 
-		const FVector* StartLoc = NULL;
-		const FRotator* StartRot = NULL;
-
 		const bool bAtPlayerStart = (GetPlayModeLocation() == PlayLocation_DefaultPlayerStart);
-
 		if (!bAtPlayerStart)
 		{
 			FLevelEditorModule& LevelEditorModule = FModuleManager::GetModuleChecked<FLevelEditorModule>(TEXT("LevelEditor"));
 			TSharedPtr<IAssetViewport> ActiveLevelViewport = LevelEditorModule.GetFirstActiveViewport();
 
-			if ( ActiveLevelViewport.IsValid() && FSlateApplication::Get().FindWidgetWindow( ActiveLevelViewport->AsWidget() ).IsValid() )
+			if (ActiveLevelViewport.IsValid() && FSlateApplication::Get().FindWidgetWindow(ActiveLevelViewport->AsWidget()).IsValid())
 			{
-				StartLoc = &ActiveLevelViewport->GetAssetViewportClient().GetViewLocation();
-				StartRot = &ActiveLevelViewport->GetAssetViewportClient().GetViewRotation();
+				SessionParams.StartLocation = ActiveLevelViewport->GetAssetViewportClient().GetViewLocation();
+				SessionParams.StartRotation = ActiveLevelViewport->GetAssetViewportClient().GetViewRotation();
 			}
 		}
 
-		FString MobilePreviewTargetDevice;
-		if(PlayModeType == PlayMode_InTargetedMobilePreview)
-		{ 
-			MobilePreviewTargetDevice = GetDefault<ULevelEditorPlaySettings>()->LastExecutedPIEPreviewDevice;
+		if (PlayModeType == PlayMode_InMobilePreview || PlayModeType == PlayMode_InTargetedMobilePreview)
+		{
+			if (PlayModeType == PlayMode_InTargetedMobilePreview)
+			{
+				SessionParams.MobilePreviewTargetDevice = GetDefault<ULevelEditorPlaySettings>()->LastExecutedPIEPreviewDevice;
+			}
+
+			SessionParams.SessionPreviewTypeOverride = EPlaySessionPreviewType::MobilePreview;
 		}
+		else if (PlayModeType == PlayMode_InVulkanPreview)
+		{
+			SessionParams.SessionPreviewTypeOverride = EPlaySessionPreviewType::VulkanPreview;
+		}
+
+		SessionParams.SessionDestination = EPlaySessionDestinationType::NewProcess;
+
 		// Spawn a new window to play in.
-		GUnrealEd->RequestPlaySession(StartLoc, StartRot, PlayModeType == PlayMode_InMobilePreview || PlayModeType == PlayMode_InTargetedMobilePreview, PlayModeType == PlayMode_InVulkanPreview, MobilePreviewTargetDevice);
+		GUnrealEd->RequestPlaySession(SessionParams);
 	}
 	else
 	{
@@ -1990,7 +1984,7 @@ bool FInternalPlayWorldCommandCallbacks::PlayInModeIsChecked(EPlayModeType PlayM
 }
 
 
-bool FInternalPlayWorldCommandCallbacks::PlayInLocation_CanExecute( EPlayModeLocations Location )
+bool FInternalPlayWorldCommandCallbacks::PlayInLocation_CanExecute(EPlayModeLocations Location)
 {
 	switch (Location)
 	{
@@ -2006,7 +2000,7 @@ bool FInternalPlayWorldCommandCallbacks::PlayInLocation_CanExecute( EPlayModeLoc
 }
 
 
-void FInternalPlayWorldCommandCallbacks::PlayInLocation_Clicked( EPlayModeLocations Location )
+void FInternalPlayWorldCommandCallbacks::PlayInLocation_Clicked(EPlayModeLocations Location)
 {
 	ULevelEditorPlaySettings* PlaySettings = GetMutableDefault<ULevelEditorPlaySettings>();
 	PlaySettings->LastExecutedPlayModeLocation = Location;
@@ -2015,7 +2009,7 @@ void FInternalPlayWorldCommandCallbacks::PlayInLocation_Clicked( EPlayModeLocati
 }
 
 
-bool FInternalPlayWorldCommandCallbacks::PlayInLocation_IsChecked( EPlayModeLocations Location )
+bool FInternalPlayWorldCommandCallbacks::PlayInLocation_IsChecked(EPlayModeLocations Location)
 {
 	switch (Location)
 	{
@@ -2060,7 +2054,7 @@ void FInternalPlayWorldCommandCallbacks::RepeatLastLaunch_Clicked()
 
 	default:
 		break;
-	}	
+	}
 }
 
 
@@ -2091,7 +2085,7 @@ FText FInternalPlayWorldCommandCallbacks::GetRepeatLastLaunchToolTip()
 			FFormatNamedArguments Arguments;
 			Arguments.Add(TEXT("DeviceName"), FText::FromString(PlaySettings->LastExecutedLaunchName));
 
-			return FText::Format( LOCTEXT("RepeatLaunchTooltip", "Launch this level on {DeviceName}"), Arguments);
+			return FText::Format(LOCTEXT("RepeatLaunchTooltip", "Launch this level on {DeviceName}"), Arguments);
 		}
 
 		break;
@@ -2117,7 +2111,7 @@ FSlateIcon FInternalPlayWorldCommandCallbacks::GetRepeatLastLaunchIcon()
 
 	static FName RepeatLastLaunchIcon("PlayWorld.RepeatLastLaunch");
 
-	return FSlateIcon(FEditorStyle::GetStyleSetName(), RepeatLastLaunchIcon );
+	return FSlateIcon(FEditorStyle::GetStyleSetName(), RepeatLastLaunchIcon);
 }
 
 bool FInternalPlayWorldCommandCallbacks::IsReadyToLaunchOnDevice(FString DeviceId)
@@ -2149,10 +2143,10 @@ bool FInternalPlayWorldCommandCallbacks::IsReadyToLaunchOnDevice(FString DeviceI
 		FString DocumentationLink;
 		FText CustomizedLogMessage;
 
-	    EBuildConfiguration BuildConfiguration = GetDefault<ULevelEditorPlaySettings>()->GetLaunchBuildConfiguration();
-	    bool bEnableAssetNativization = false;
+		EBuildConfiguration BuildConfiguration = GetDefault<ULevelEditorPlaySettings>()->GetLaunchBuildConfiguration();
+		bool bEnableAssetNativization = false;
 		int32 Result = Platform->CheckRequirements(bHasCode, BuildConfiguration, bEnableAssetNativization, NotInstalledTutorialLink, DocumentationLink, CustomizedLogMessage);
-		
+
 		// report to analytics
 		FEditorAnalytics::ReportBuildRequirementsFailure(TEXT("Editor.LaunchOn.Failed"), PlatformName, bHasCode, Result);
 
@@ -2220,8 +2214,8 @@ bool FInternalPlayWorldCommandCallbacks::IsReadyToLaunchOnDevice(FString DeviceI
 		}
 
 		if ((Result & ETargetPlatformReadyStatus::RemoveServerNameEmpty) != 0
-				&& (bHasCode || (Result & ETargetPlatformReadyStatus::CodeBuildRequired)
-					|| (!FApp::GetEngineIsPromotedBuild() && !FApp::IsEngineInstalled())))
+			&& (bHasCode || (Result & ETargetPlatformReadyStatus::CodeBuildRequired)
+				|| (!FApp::GetEngineIsPromotedBuild() && !FApp::IsEngineInstalled())))
 		{
 			AddMessageLog(
 				LOCTEXT("RemoveServerNameNotFound", "Remote compiling requires a server name. "),
@@ -2261,7 +2255,7 @@ bool FInternalPlayWorldCommandCallbacks::IsReadyToLaunchOnDevice(FString DeviceI
 	return true;
 }
 
-void FInternalPlayWorldCommandCallbacks::HandleLaunchOnDeviceActionExecute( FString DeviceId, FString DeviceName )
+void FInternalPlayWorldCommandCallbacks::HandleLaunchOnDeviceActionExecute(FString DeviceId, FString DeviceName)
 {
 	if (IsReadyToLaunchOnDevice(DeviceId))
 	{
@@ -2293,7 +2287,7 @@ bool FInternalPlayWorldCommandCallbacks::HandleLaunchOnDeviceActionIsChecked(FSt
 }
 
 
-void FInternalPlayWorldCommandCallbacks::HandleShowSDKTutorial( FString PlatformName, FString NotInstalledDocLink )
+void FInternalPlayWorldCommandCallbacks::HandleShowSDKTutorial(FString PlatformName, FString NotInstalledDocLink)
 {
 	// broadcast this, and assume someone will pick it up
 	IMainFrameModule& MainFrameModule = FModuleManager::GetModuleChecked<IMainFrameModule>(TEXT("MainFrame"));
@@ -2317,11 +2311,11 @@ void FInternalPlayWorldCommandCallbacks::GetMouseControlExecute()
 
 FSlateIcon FInternalPlayWorldCommandCallbacks::GetResumePlaySessionImage()
 {
-	if ( IsInPIE() )
+	if (IsInPIE())
 	{
 		return FSlateIcon(FEditorStyle::GetStyleSetName(), "PlayWorld.ResumePlaySession");
 	}
-	else if ( IsInSIE() )
+	else if (IsInSIE())
 	{
 		return FSlateIcon(FEditorStyle::GetStyleSetName(), "PlayWorld.Simulate");
 	}
@@ -2334,11 +2328,11 @@ FSlateIcon FInternalPlayWorldCommandCallbacks::GetResumePlaySessionImage()
 
 FText FInternalPlayWorldCommandCallbacks::GetResumePlaySessionToolTip()
 {
-	if ( IsInPIE() )
+	if (IsInPIE())
 	{
 		return LOCTEXT("ResumePIE", "Resume play-in-editor session");
 	}
-	else if ( IsInSIE() )
+	else if (IsInSIE())
 	{
 		return LOCTEXT("ResumeSIE", "Resume simulation");
 	}
@@ -2483,24 +2477,22 @@ void FInternalPlayWorldCommandCallbacks::SetNumberOfClients(int32 NumClients, ET
 }
 
 
-void FInternalPlayWorldCommandCallbacks::OnToggleDedicatedServerPIE()
+int32 FInternalPlayWorldCommandCallbacks::GetNetPlayMode()
 {
-	ULevelEditorPlaySettings* PlayInSettings = GetMutableDefault<ULevelEditorPlaySettings>();
-	bool PlayNetDedicated;
-	PlayInSettings->GetPlayNetDedicated(PlayNetDedicated);			// Ignore 'state' of option, as we're toggling it regardless
-	PlayInSettings->SetPlayNetDedicated(!PlayNetDedicated);
+	const ULevelEditorPlaySettings* PlayInSettings = GetDefault<ULevelEditorPlaySettings>();
+	EPlayNetMode NetMode;
+	PlayInSettings->GetPlayNetMode(NetMode);
 
-	PlayInSettings->PostEditChange();
-	PlayInSettings->SaveConfig();
+	return (int32)NetMode;
 }
 
-
-bool FInternalPlayWorldCommandCallbacks::OnIsDedicatedServerPIEEnabled()
+void FInternalPlayWorldCommandCallbacks::SetNetPlayMode(int32 Value, ESelectInfo::Type CommitInfo)
 {
-	const ULevelEditorPlaySettings* PlayInSettings = GetDefault<ULevelEditorPlaySettings>();
-	bool PlayNetDedicated(false);
-	PlayInSettings->GetPlayNetDedicated(PlayNetDedicated);			// Ignore 'state' of option (handled externally)
-	return PlayNetDedicated;
+	ULevelEditorPlaySettings* PlayInSettings = GetMutableDefault<ULevelEditorPlaySettings>();
+	PlayInSettings->SetPlayNetMode((EPlayNetMode)Value);
+
+	PlayInSettings->PostEditChange();
+	PlayInSettings->SaveConfig();
 }
 
 
@@ -2529,7 +2521,7 @@ bool FInternalPlayWorldCommandCallbacks::CanPossessEjectPlayer()
 }
 
 
-void FInternalPlayWorldCommandCallbacks::AddMessageLog( const FText& Text, const FText& Detail, const FString& TutorialLink, const FString& DocumentationLink)
+void FInternalPlayWorldCommandCallbacks::AddMessageLog(const FText& Text, const FText& Detail, const FString& TutorialLink, const FString& DocumentationLink)
 {
 	TSharedRef<FTokenizedMessage> Message = FTokenizedMessage::Create(EMessageSeverity::Error);
 	Message->AddToken(FTextToken::Create(Text));
@@ -2593,7 +2585,7 @@ bool FInternalPlayWorldCommandCallbacks::CanLaunchOnDevice(const FString& Device
 }
 
 
-void FInternalPlayWorldCommandCallbacks::LaunchOnDevice( const FString& DeviceId, const FString& DeviceName )
+void FInternalPlayWorldCommandCallbacks::LaunchOnDevice(const FString& DeviceId, const FString& DeviceName)
 {
 	FTargetDeviceId TargetDeviceId;
 	if (FTargetDeviceId::Parse(DeviceId, TargetDeviceId))
@@ -2613,7 +2605,16 @@ void FInternalPlayWorldCommandCallbacks::LaunchOnDevice( const FString& DeviceId
 		if (FModuleManager::LoadModuleChecked<IProjectTargetPlatformEditorModule>("ProjectTargetPlatformEditor").ShowUnsupportedTargetWarning(*TargetDeviceId.GetPlatformName()))
 		{
 			GUnrealEd->CancelPlayingViaLauncher();
-			GUnrealEd->RequestPlaySession(DeviceId, DeviceName);
+
+			FRequestPlaySessionParams::FLauncherDeviceInfo DeviceInfo;
+			DeviceInfo.DeviceId = DeviceId;
+			DeviceInfo.DeviceName = DeviceName;
+
+			FRequestPlaySessionParams SessionParams;
+			SessionParams.SessionDestination = EPlaySessionDestinationType::Launcher;
+			SessionParams.LauncherTargetDevice = DeviceInfo;
+
+			GUnrealEd->RequestPlaySession(SessionParams);
 		}
 	}
 }
@@ -2622,8 +2623,8 @@ void FInternalPlayWorldCommandCallbacks::LaunchOnDevice( const FString& DeviceId
 EPlayModeLocations FInternalPlayWorldCommandCallbacks::GetPlayModeLocation()
 {
 	// We can't use PlayLocation_DefaultPlayerStart without a player start position
-	return GEditor->CheckForPlayerStart() 
-		? static_cast<EPlayModeLocations>(GetDefault<ULevelEditorPlaySettings>()->LastExecutedPlayModeLocation) 
+	return GEditor->CheckForPlayerStart()
+		? static_cast<EPlayModeLocations>(GetDefault<ULevelEditorPlaySettings>()->LastExecutedPlayModeLocation)
 		: PlayLocation_CurrentCameraLocation;
 }
 
