Significantly refactored PIE startup  and consolidated a fair bit of the existing code into structures that better reflect the lifespan that they are valid for. This refactoring will make it easier fo

diff --git a/Engine/Source/Editor/UnrealEd/Private/PlayLevel.cpp b/Engine/Source/Editor/UnrealEd/Private/PlayLevel.cpp
--- a/Engine/Source/Editor/UnrealEd/Private/PlayLevel.cpp
+++ b/Engine/Source/Editor/UnrealEd/Private/PlayLevel.cpp
@@ -1,5 +1,6 @@
 // Copyright 1998-2019 Epic Games, Inc. All Rights Reserved.
 
+#include "PlayLevel.h"
 #include "CoreMinimal.h"
 #include "Misc/MessageDialog.h"
 #include "Misc/CommandLine.h"
@@ -38,7 +39,6 @@
 #include "Engine/GameInstance.h"
 #include "Engine/RendererSettings.h"
 #include "Engine/World.h"
-#include "DesktopPlatformModule.h"
 #include "Settings/LevelEditorPlaySettings.h"
 #include "AI/NavigationSystemBase.h"
 #include "Editor/EditorEngine.h"
@@ -80,8 +80,6 @@
 #include "Logging/MessageLog.h"
 #include "Misc/UObjectToken.h"
 #include "Misc/MapErrors.h"
-#include "ITargetDeviceServicesModule.h"
-#include "ILauncherServicesModule.h"
 #include "GameProjectGenerationModule.h"
 #include "SourceCodeNavigation.h"
 #include "Physics/PhysicsInterfaceCore.h"
@@ -101,7 +99,6 @@
 #include "Engine/LevelStreaming.h"
 #include "Components/ModelComponent.h"
 #include "GameDelegates.h"
-#include "PlatformInfo.h"
 #include "Net/OnlineEngineInterface.h"
 #include "Kismet2/DebuggerCommands.h"
 #include "Misc/ScopeExit.h"
@@ -112,15 +109,20 @@
 #include "Widgets/Text/STextBlock.h"
 #include "Widgets/SBoxPanel.h"
 #include "Subsystems/AssetEditorSubsystem.h"
-#include "Async/Async.h"
-
+#include "UObject/SoftObjectPath.h"
+#include "IAssetViewport.h"
 
-DEFINE_LOG_CATEGORY_STATIC(LogPlayLevel, Log, All);
+DEFINE_LOG_CATEGORY(LogPlayLevel);
 
 #define LOCTEXT_NAMESPACE "PlayLevel"
 
 const static FName NAME_CategoryPIE("PIE");
 
+// Forward declare local utility functions
+FText GeneratePIEViewportWindowTitle(const EPlayNetMode InNetMode, const ERHIFeatureLevel::Type InFeatureLevel, const int32 ClientIndex = -1);
+bool PromptMatineeClose();
+
+
 // This class listens to output log messages, and forwards warnings and errors to the message log
 class FOutputLogErrorsToMessageLogProxy : public FOutputDevice
 {
@@ -356,9 +358,6 @@ void UEditorEngine::EndPlayMap()
 	}
 	
 	{
-		// Clear out viewport index
-		PlayInEditorViewportIndex = -1; 
-
 		// We could have been toggling back and forth between simulate and pie before ending the play map
 		// Make sure the property windows are cleared of any pie actors
 		GUnrealEd->UpdateFloatingPropertyWindows();
@@ -560,11 +559,34 @@ void UEditorEngine::EndPlayMap()
 		}
 	}
 
+	if (PlayInEditorSessionInfo.IsSet())
+	{
+		// Save the window positions to the CDO object.
+		ULevelEditorPlaySettings* PlaySettingsConfig = GetMutableDefault<ULevelEditorPlaySettings>();
+
+		for (int32 WindowIndex = 0; WindowIndex < PlayInEditorSessionInfo->CachedWindowInfo.Num(); WindowIndex++)
+		{
+			if (WindowIndex < PlaySettingsConfig->MultipleInstancePositions.Num())
+			{
+				PlaySettingsConfig->MultipleInstancePositions[WindowIndex] = PlayInEditorSessionInfo->CachedWindowInfo[WindowIndex].Position;
+			}
+			else
+			{
+				PlaySettingsConfig->MultipleInstancePositions.Add(PlayInEditorSessionInfo->CachedWindowInfo[WindowIndex].Position);
+			}
+		}
+
+		PlaySettingsConfig->PostEditChange();
+		PlaySettingsConfig->SaveConfig();
+
+		// Now that we're shutting down, we'll no longer need the instance of the EditorPlaySettings.
+		PlayInEditorSessionInfo->OriginalRequestParams.EditorPlaySettings->RemoveFromRoot();
+	}
+	PlayInEditorSessionInfo.Reset();
+
 	// no longer queued
 	CancelRequestPlaySession();
-	bIsSimulateInEditorQueued = false;
 	bRequestEndPlayMapQueued = false;
-	bUseVRPreviewForPlayWorld = false;
 
 	// Tear down the output log to message log thunker
 	OutputLogErrorsToMessageLogProxyPtr.Reset();
@@ -576,6 +598,13 @@ void UEditorEngine::EndPlayMap()
 	FMessageLog(NAME_CategoryPIE).Notify(LOCTEXT("PIEErrorsPresent", "Errors/warnings reported while playing in editor."));
 
 	FMessageLog(NAME_CategoryPIE).Open(EMessageSeverity::Warning);
+
+	{
+		// Temporary until the deprecated variable is removed
+		PRAGMA_DISABLE_DEPRECATION_WARNINGS
+		bIsSimulatingInEditor = false;
+		PRAGMA_ENABLE_DEPRECATION_WARNINGS
+	}
 }
 
 void UEditorEngine::CleanupPIEOnlineSessions(TArray<FName> OnlineIdentifiers)
@@ -584,10 +613,7 @@ void UEditorEngine::CleanupPIEOnlineSessions(TArray<FName> OnlineIdentifiers)
 	{
 		UE_LOG(LogPlayLevel, Display, TEXT("Destroying online subsystem %s"), *OnlineIdentifier.ToString());
 		UOnlineEngineInterface::Get()->DestroyOnlineSubsystem(OnlineIdentifier);
-		NumOnlinePIEInstances--;
 	}
-
-	NumOnlinePIEInstances = 0;
 }
 
 void UEditorEngine::TeardownPlaySession(FWorldContext& PieWorldContext)
@@ -615,7 +641,7 @@ void UEditorEngine::TeardownPlaySession(FWorldContext& PieWorldContext)
 			{
 				TSharedPtr<IAssetViewport> Viewport = SlatePlayInEditorSession->DestinationSlateViewport.Pin();
 
-				if( !bIsSimulatingInEditor)
+				if(PlayInEditorSessionInfo->OriginalRequestParams.WorldType == EPlaySessionWorldType::PlayInEditor)
 				{
 					// Set the editor viewport location to match that of Play in Viewport if we aren't simulating in the editor, we have a valid player to get the location from (unless we're going back to VR Editor, in which case we won't teleport the user.)
 					if (bLastViewAndLocationValid == true && !GEngine->IsStereoscopic3D( Viewport->GetActiveViewport() ) )
@@ -685,10 +711,7 @@ void UEditorEngine::TeardownPlaySession(FWorldContext& PieWorldContext)
 	GIsPlayInEditorWorld = true;
 	
 	// Remember Simulating flag so that we know if OnSimulateSessionFinished is required after everything has been cleaned up. 
-	bool bWasSimulatingInEditor = bIsSimulatingInEditor;
-	// Clear Simulating In Editor bit
-	bIsSimulatingInEditor = false;
-
+	bool bWasSimulatingInEditor = PlayInEditorSessionInfo->OriginalRequestParams.WorldType == EPlaySessionWorldType::SimulateInEditor;
 	
 	// Stop all audio and remove references to temp level.
 	if (FAudioDevice* AudioDevice = PlayWorld->GetAudioDevice())
@@ -817,131 +840,153 @@ void UEditorEngine::TeardownPlaySession(FWorldContext& PieWorldContext)
 	}
 }
 
-void UEditorEngine::PlayMap( const FVector* StartLocation, const FRotator* StartRotation, int32 Destination, int32 InPlayInViewportIndex, bool bUseMobilePreview )
+// Deprecated, just format to match our new style.
+void UEditorEngine::PlayMap(const FVector* StartLocation, const FRotator* StartRotation, int32 Destination, int32 InPlayInViewportIndex, bool bUseMobilePreview)
 {
-	// queue up a Play From Here request, this way the load/save won't conflict with the TransBuffer, which doesn't like 
-	// loading and saving to happen during a transaction
+	// Deprecated, replaced with RequestPlaySession.
+	FRequestPlaySessionParams Params;
+
+	if (bUseMobilePreview)
+	{
+		Params.SessionDestination = EPlaySessionDestinationType::NewProcess;
+		Params.SessionPreviewTypeOverride = EPlaySessionPreviewType::MobilePreview;
+	}
 
-	// save the StartLocation if we have one
 	if (StartLocation)
 	{
-		PlayWorldLocation = *StartLocation;
-		PlayWorldRotation = StartRotation ? *StartRotation : FRotator::ZeroRotator;
-		bHasPlayWorldPlacement = true;
+		Params.StartLocation = *StartLocation;
+		Params.StartRotation = StartRotation ? *StartRotation : FRotator::ZeroRotator;
 	}
-	else
+
+	RequestPlaySession(Params);
+}
+
+void UEditorEngine::RequestPlaySession(const FRequestPlaySessionParams& InParams)
+{
+	// Store our Request to be operated on next Tick.
+	PlaySessionRequest = InParams;
+
+	// If they don't want to use a specific set of Editor Play Settings, fall back to the CDO.
+	if (!PlaySessionRequest->EditorPlaySettings)
 	{
-		bHasPlayWorldPlacement = false;
+		PlaySessionRequest->EditorPlaySettings = GetMutableDefault<ULevelEditorPlaySettings>();
 	}
 
-	// remember where to send the play map request
-	PlayWorldDestination = Destination;
+	// Now we duplicate their Editor Play Settings so that we can mutate it as part of startup
+	// to help rule out invalid configuration combinations.
+	FObjectDuplicationParameters DuplicationParams(PlaySessionRequest->EditorPlaySettings, GetTransientPackage());
+	PlaySessionRequest->EditorPlaySettings = CastChecked<ULevelEditorPlaySettings>(StaticDuplicateObjectEx(DuplicationParams));
 
-	// Set whether or not we want to use mobile preview mode (PC platform only)
-	bUseMobilePreviewForPlayWorld = bUseMobilePreview;
-	bUseVRPreviewForPlayWorld = false;
+	// Add the duplicated object to the root, as it's fairly likely an in-process startup will GC it.
+	PlaySessionRequest->EditorPlaySettings->AddToRoot();
 
-	// tell the editor to kick it off next Tick()
-	bIsPlayWorldQueued = true;
+	// ToDo: Allow the CDO for the Game Instance to modify the settings after we copy them
+	// so that they can validate user settings before attempting a launch.
 
-	// Not wanting to simulate
-	bIsSimulateInEditorQueued = false;
+	// Play Sessions can use the Game Mode to determine the default Player Start position
+	// or the start position can be overridden by the incoming launch arguments
+	PRAGMA_DISABLE_DEPRECATION_WARNINGS
+	GIsPIEUsingPlayerStart = !InParams.StartLocation.IsSet();
+	PRAGMA_ENABLE_DEPRECATION_WARNINGS
 
-	// Unless we've been asked to play in a specific viewport window, this index will be -1
-	PlayInEditorViewportIndex = InPlayInViewportIndex;
+	if (InParams.SessionDestination == EPlaySessionDestinationType::Launcher)
+	{
+		check(InParams.LauncherTargetDevice.IsSet());
+	}
 }
 
-
-void UEditorEngine::RequestPlaySession( bool bAtPlayerStart, TSharedPtr<class IAssetViewport> DestinationViewport, bool bInSimulateInEditor, const FVector* StartLocation, const FRotator* StartRotation, int32 DestinationConsole, bool bUseMobilePreview, bool bUseVRPreview, bool bUseVulkanPreview )
+// Deprecated, just format to match our new style.
+void UEditorEngine::RequestPlaySession(bool bAtPlayerStart, TSharedPtr<class IAssetViewport> DestinationViewport, bool bInSimulateInEditor, const FVector* StartLocation, const FRotator* StartRotation, int32 DestinationConsole, bool bUseMobilePreview, bool bUseVRPreview, bool bUseVulkanPreview)
 {
-	// Remember whether or not we were attempting to play from playerstart or from viewport
-	GIsPIEUsingPlayerStart = bAtPlayerStart;
-
-	// queue up a Play From Here request, this way the load/save won't conflict with the TransBuffer, which doesn't like 
-	// loading and saving to happen during a transaction
+	FRequestPlaySessionParams Params;
 
-	// save the StartLocation if we have one
-	if (!bInSimulateInEditor && StartLocation != NULL)
+	if (StartLocation)
 	{
-		PlayWorldLocation = *StartLocation;
-		PlayWorldRotation = StartRotation ? *StartRotation : FRotator::ZeroRotator;
-		bHasPlayWorldPlacement = true;
+		Params.StartLocation = *StartLocation;
+		Params.StartRotation = StartRotation ? *StartRotation : FRotator::ZeroRotator;
 	}
-	else
+	if (DestinationViewport != nullptr)
 	{
-		bHasPlayWorldPlacement = false;
+		Params.DestinationSlateViewport = DestinationViewport;
 	}
 
-	// remember where to send the play map request
-	PlayWorldDestination = DestinationConsole;
-
-	RequestedDestinationSlateViewport = DestinationViewport;
-
-	// Set whether or not we want to use mobile preview mode (PC platform only)
-	bUseMobilePreviewForPlayWorld = bUseMobilePreview;
-
-	bUseVRPreviewForPlayWorld = bUseVRPreview;
-
-	// tell the editor to kick it off next Tick()
-	bIsPlayWorldQueued = true;
+	if (bInSimulateInEditor)
+	{
+		Params.WorldType = EPlaySessionWorldType::SimulateInEditor;
+	}
 
-	// Store whether we want to play in editor, or only simulate in editor
-	bIsSimulateInEditorQueued = bInSimulateInEditor;
+	if (bUseVRPreview)
+	{
+		check(!bUseMobilePreview && !bUseVulkanPreview);
+		Params.SessionPreviewTypeOverride = EPlaySessionPreviewType::VRPreview;
+	}
 
-	// Unless we have been asked to play in a specific viewport window, this index will be -1
-	PlayInEditorViewportIndex = -1;
+	if (bUseVulkanPreview)
+	{
+		check(!bUseMobilePreview && !bUseVRPreview);
+		Params.SessionPreviewTypeOverride = EPlaySessionPreviewType::VulkanPreview;
+		Params.SessionDestination = EPlaySessionDestinationType::NewProcess;
+	}
 
-	bPlayOnLocalPcSession = false;
-	bPlayUsingLauncher = false;
+	if (bUseMobilePreview)
+	{
+		check(!bUseVRPreview && !bUseVulkanPreview);
+		Params.SessionPreviewTypeOverride = EPlaySessionPreviewType::MobilePreview;
+		Params.SessionDestination = EPlaySessionDestinationType::NewProcess;
+	}
+	
+	RequestPlaySession(Params);
 }
 
-void UEditorEngine::RequestPlaySession( const FVector* StartLocation, const FRotator* StartRotation, bool MobilePreview, bool VulkanPreview , const FString& MobilePreviewTargetDevice, FString AdditionalLaunchParameters)
+// Deprecated, forwards request onto the FRequestPlaySessionParams version.
+void UEditorEngine::RequestPlaySession(const FVector* StartLocation, const FRotator* StartRotation, bool MobilePreview, bool VulkanPreview, const FString& MobilePreviewTargetDevice, FString AdditionalLaunchParameters)
 {
-	bPlayOnLocalPcSession = true;
-	bPlayUsingLauncher = false;
-	bPlayUsingMobilePreview = MobilePreview;
-	bPlayUsingVulkanPreview = VulkanPreview;
-	RequestedAdditionalStandaloneLaunchOptions = AdditionalLaunchParameters;
-	PlayUsingMobilePreviewTargetDevice = MobilePreviewTargetDevice;
+	FRequestPlaySessionParams Params;
 
-	if (StartLocation != NULL)
+	if (MobilePreview)
 	{
-		PlayWorldLocation = *StartLocation;
-		PlayWorldRotation = StartRotation ? *StartRotation : FRotator::ZeroRotator;
-		bHasPlayWorldPlacement = true;
+		check(!VulkanPreview);
+		Params.SessionDestination = EPlaySessionDestinationType::NewProcess;
+		Params.SessionPreviewTypeOverride = EPlaySessionPreviewType::MobilePreview;
+		Params.MobilePreviewTargetDevice = MobilePreviewTargetDevice;
+		Params.AdditionalStandaloneCommandLineParameters = AdditionalLaunchParameters;
 	}
-	else
+
+	if (VulkanPreview)
+	{
+		check(!MobilePreview);
+		Params.SessionDestination = EPlaySessionDestinationType::NewProcess;
+		Params.SessionPreviewTypeOverride = EPlaySessionPreviewType::VulkanPreview;
+		Params.AdditionalStandaloneCommandLineParameters = AdditionalLaunchParameters;
+	}
+	
+	if (StartLocation)
 	{
-		bHasPlayWorldPlacement = false;
+		Params.StartLocation = *StartLocation;
+		Params.StartRotation = StartRotation ? *StartRotation : FRotator::ZeroRotator;
 	}
 
-	bIsPlayWorldQueued = true;
+	RequestPlaySession(Params);
 }
 
+// Deprecated, forwards request onto the FRequestPlaySessionParams version.
 void UEditorEngine::RequestPlaySession(const FString& DeviceId, const FString& DeviceName)
 {
-	bPlayOnLocalPcSession = false;
-	bPlayUsingLauncher = true;
-
-	// always use player start on remote devices (for now?)
-	bHasPlayWorldPlacement = false;
+	FRequestPlaySessionParams::FLauncherDeviceInfo DeviceInfo;
+	DeviceInfo.DeviceId = DeviceId;
+	DeviceInfo.DeviceName = DeviceName;
 
-	// remember the platform name to run on
-	PlayUsingLauncherDeviceId = DeviceId;
-	PlayUsingLauncherDeviceName = DeviceName;
+	FRequestPlaySessionParams Params;
+	Params.LauncherTargetDevice = DeviceInfo;
+	Params.SessionDestination = EPlaySessionDestinationType::Launcher;
 
-	bIsPlayWorldQueued = true;
+	RequestPlaySession(Params);
 }
 
 void UEditorEngine::CancelRequestPlaySession()
 {
-	bIsPlayWorldQueued = false;
-	bPlayOnLocalPcSession = false;
-	bPlayUsingLauncher = false;
-	bPlayUsingMobilePreview = false;
-	bPlayUsingVulkanPreview = false;
-	RequestedAdditionalStandaloneLaunchOptions = FString();
-	PlayUsingMobilePreviewTargetDevice.Reset();
+	PlaySessionRequest.Reset();
+	PlayInEditorSessionInfo.Reset();
 }
 
 bool UEditorEngine::SaveMapsForPlaySession()
@@ -952,18 +997,16 @@ bool UEditorEngine::SaveMapsForPlaySession()
 	{
 		if (!FEditorFileUtils::SaveCurrentLevel())
 		{
-			CancelRequestPlaySession();
 			return false;
 		}
 	}
 
 	// Also save dirty packages, this is required because we're going to be launching a session outside of our normal process
-	bool bPromptUserToSave = true;
-	bool bSaveMapPackages = true;
-	bool bSaveContentPackages = true;
+	const bool bPromptUserToSave = true;
+	const bool bSaveMapPackages = true;
+	const bool bSaveContentPackages = true;
 	if (!FEditorFileUtils::SaveDirtyPackages(bPromptUserToSave, bSaveMapPackages, bSaveContentPackages))
 	{
-		CancelRequestPlaySession();
 		return false;
 	}
 
@@ -972,17 +1015,17 @@ bool UEditorEngine::SaveMapsForPlaySession()
 
 void UEditorEngine::PlaySessionPaused()
 {
-	FEditorDelegates::PausePIE.Broadcast(bIsSimulatingInEditor);
+	FEditorDelegates::PausePIE.Broadcast(PlayInEditorSessionInfo->OriginalRequestParams.WorldType == EPlaySessionWorldType::SimulateInEditor);
 }
 
 void UEditorEngine::PlaySessionResumed()
 {
-	FEditorDelegates::ResumePIE.Broadcast(bIsSimulatingInEditor);
+	FEditorDelegates::ResumePIE.Broadcast(PlayInEditorSessionInfo->OriginalRequestParams.WorldType == EPlaySessionWorldType::SimulateInEditor);
 }
 
 void UEditorEngine::PlaySessionSingleStepped()
 {
-	FEditorDelegates::SingleStepPIE.Broadcast(bIsSimulatingInEditor);
+	FEditorDelegates::SingleStepPIE.Broadcast(PlayInEditorSessionInfo->OriginalRequestParams.WorldType == EPlaySessionWorldType::SimulateInEditor);
 }
 
 bool UEditorEngine::ProcessDebuggerCommands(const FKey InKey, const FModifierKeysState ModifierKeyState, EInputEvent EventType )
@@ -995,222 +1038,67 @@ bool UEditorEngine::ProcessDebuggerCommands(const FKey InKey, const FModifierKey
 	return false;
 }
 
-/* fits the window position to make sure it falls within the confines of the desktop */
-void FitWindowPositionToWorkArea(FIntPoint &WinPos, FIntPoint &WinSize, const FMargin &WinPadding)
+// This function is deprecated, just call the non-deprecated version.
+void UEditorEngine::StartQueuedPlayMapRequest()
 {
-	const int32 HorzPad = WinPadding.GetTotalSpaceAlong<Orient_Horizontal>();
-	const int32 VertPad = WinPadding.GetTotalSpaceAlong<Orient_Vertical>();
-	FIntPoint TotalSize( WinSize.X + HorzPad, WinSize.Y + VertPad );
-
-	FDisplayMetrics DisplayMetrics;
-	FSlateApplication::Get().GetCachedDisplayMetrics(DisplayMetrics);
-
-	// Limit the size, to make sure it fits within the desktop area
-	{
-		FIntPoint NewWinSize;
-		NewWinSize.X = FMath::Min(TotalSize.X, DisplayMetrics.VirtualDisplayRect.Right - DisplayMetrics.VirtualDisplayRect.Left );
-		NewWinSize.Y = FMath::Min(TotalSize.Y, DisplayMetrics.VirtualDisplayRect.Bottom - DisplayMetrics.VirtualDisplayRect.Top );
-		if( NewWinSize != TotalSize )
-		{
-			TotalSize = NewWinSize;
-			WinSize.X = NewWinSize.X - HorzPad;
-			WinSize.Y = NewWinSize.Y - VertPad;
-		}
-	}
-
-	const FSlateRect PreferredWorkArea( DisplayMetrics.VirtualDisplayRect.Left, 
-										DisplayMetrics.VirtualDisplayRect.Top, 
-										DisplayMetrics.VirtualDisplayRect.Right - TotalSize.X, 
-										DisplayMetrics.VirtualDisplayRect.Bottom - TotalSize.Y );
-
-	// if no more windows fit horizontally, place them in a new row
-	if (WinPos.X > PreferredWorkArea.Right)
-	{
-		WinPos.X = PreferredWorkArea.Left;
-		WinPos.Y += TotalSize.Y;
-		if (WinPos.Y > PreferredWorkArea.Bottom)
-		{
-			WinPos.Y = PreferredWorkArea.Top;
-		}
-	}
-	
-	// if no more rows fit vertically, stack windows on top of each other
-	else if (WinPos.Y > PreferredWorkArea.Bottom)
-	{
-		WinPos.Y = PreferredWorkArea.Top;
-		WinPos.X += TotalSize.X;
-		if (WinPos.X > PreferredWorkArea.Right)
-		{
-			WinPos.X = PreferredWorkArea.Left;
-		}
-	}
-
-	// Clamp values to make sure they fall within the desktop area
-	WinPos.X = FMath::Clamp(WinPos.X, (int32)PreferredWorkArea.Left, (int32)PreferredWorkArea.Right);
-	WinPos.Y = FMath::Clamp(WinPos.Y, (int32)PreferredWorkArea.Top, (int32)PreferredWorkArea.Bottom);
+	StartQueuedPlaySessionRequest();
 }
 
-/* advances the windows position to the next location and fits */
-void AdvanceWindowPositionForNextPIEWindow(FIntPoint &WinPos, const FIntPoint &WinSize, const FMargin &WinPadding, bool bVertical)
+void UEditorEngine::StartQueuedPlaySessionRequest()
 {
-	const int32 HorzPad = WinPadding.GetTotalSpaceAlong<Orient_Horizontal>();
-	const int32 VertPad = WinPadding.GetTotalSpaceAlong<Orient_Vertical>();
-	const FIntPoint TotalSize( WinSize.X + HorzPad, WinSize.Y + VertPad );
-
-	if(bVertical)
-	{
-		WinPos.Y += TotalSize.Y;
-	}
-	else
+	if (!PlaySessionRequest.IsSet())
 	{
-		WinPos.X += TotalSize.X;
+		UE_LOG(LogPlayLevel, Warning, TEXT("StartQueuedPlaySessionRequest() called whith no request queued. Ignoring..."));
+		return;
 	}
-}
 
-/* returns the size of the window depending on the net mode. */
-void GetWindowSizeForInstanceType(FIntPoint &WindowSize, const ULevelEditorPlaySettings* PlayInSettings)
-{
-	const EPlayNetMode PlayNetMode = [&PlayInSettings]{ EPlayNetMode NetMode(PIE_Standalone); return (PlayInSettings->GetPlayNetMode(NetMode) ? NetMode : PIE_Standalone); }();
-	if (PlayNetMode == PIE_Standalone)
-	{
-		WindowSize.X = PlayInSettings->NewWindowWidth;
-		WindowSize.Y = PlayInSettings->NewWindowHeight;
-	}
-	else
-	{
-		PlayInSettings->GetClientWindowSize(WindowSize);
-	}
-}
+	StartQueuedPlaySessionRequestImpl();
 
-/* sets the size of the window depending on the net mode */
-void SetWindowSizeForInstanceType(const FIntPoint &WindowSize, ULevelEditorPlaySettings* PlayInSettings)
-{
-	const EPlayNetMode PlayNetMode = [&PlayInSettings]{ EPlayNetMode NetMode(PIE_Standalone); return (PlayInSettings->GetPlayNetMode(NetMode) ? NetMode : PIE_Standalone); }();
-	if (PlayNetMode == PIE_Standalone)
-	{
-		PlayInSettings->NewWindowWidth = WindowSize.X;
-		PlayInSettings->NewWindowHeight = WindowSize.Y;
-	}
-	else
-	{
-		PlayInSettings->SetClientWindowSize(WindowSize);
-	}
+	// Ensure the request is always reset after an attempt (which may fail)
+	// so that we don't get stuck in an infinite loop of start attempts.
+	PlaySessionRequest.Reset();
 }
 
-/** 
- * Generate the command line for pie instance. Window position, size etc. 
- *
- * @param	WinPos			Window position. This will contain the X & Y position to use for the next window. (Not changed for dedicated server window).
- * @param	InstanceNum		PIE instance index.
- * @param	bIsDedicatedServer	Is this instance a dedicate server. true if so else false.
- */
-FString GenerateCmdLineForNextPieInstance(FIntPoint &WinPos, int32 &InstanceNum, bool bIsDedicatedServer)
+void UEditorEngine::StartQueuedPlaySessionRequestImpl()
 {
-	const ULevelEditorPlaySettings* PlayInSettings = GetDefault<ULevelEditorPlaySettings>();
-	// Get GameSettings INI override
-	FString GameUserSettingsOverride = GGameUserSettingsIni.Replace(TEXT("GameUserSettings"), *FString::Printf(TEXT("PIEGameUserSettings%d"), InstanceNum++));
-	
-	// Construct parms:
-	//	-Override GameUserSettings.ini
-	//	-Force no steam
-	//	-Allow saving of config files (since we are giving them an override INI)
-	const FString AdditionalLaunchOptions = [&PlayInSettings]{ FString LaunchOptions; return (PlayInSettings->GetAdditionalLaunchOptions(LaunchOptions) ? LaunchOptions : FString()); }();
-	FString CmdLine = FString::Printf(TEXT("GameUserSettingsINI=\"%s\" -MultiprocessSaveConfig %s -MultiprocessOSS "), *GameUserSettingsOverride, *AdditionalLaunchOptions);
-
-	if (bIsDedicatedServer)
+	if (!ensureAlwaysMsgf(PlaySessionRequest.IsSet(), TEXT("StartQueuedPlaySessionRequest should not be called without a request set!")))
 	{
-		// Append dedicated server options
-		CmdLine += TEXT("-server -log ");
+		return;
 	}
-	else
-	{
-		// Default to what we expect the border to be (on windows at least) to prevent it occurring offscreen if TLW call fails
-		FMargin WindowBorderSize(8.0f, 30.0f, 8.0f, 8.0f);
-		TSharedPtr<SWindow> TopLevelWindow = FSlateApplication::Get().GetActiveTopLevelWindow();
-
-		if (TopLevelWindow.IsValid())
-		{
-			WindowBorderSize = TopLevelWindow->GetWindowBorderSize(true);
-		}
-		
-		// Get the size of the window based on the type
-		FIntPoint WinSize(0,0);
-		GetWindowSizeForInstanceType(WinSize, PlayInSettings);
 
-		// Make sure the window is going to fit where we want it
-		FitWindowPositionToWorkArea(WinPos, WinSize, WindowBorderSize);
+	// We want to use the ULevelEditorPlaySettings that come from the Play Session Request.
+	// By the time this function gets called, these settings are a copy of either the CDO, 
+	// or a user provided instance. The settings may have been modified by the game instance
+	// after the request was made, to allow game instances to pre-validate settings.
+	const ULevelEditorPlaySettings* EditorPlaySettings = PlaySessionRequest->EditorPlaySettings;
+	check(EditorPlaySettings);
 
-		// Set the size, in case it was modified
-		SetWindowSizeForInstanceType(WinSize, GetMutableDefault<ULevelEditorPlaySettings>());
+	PlayInEditorSessionInfo = FPlayInEditorSessionInfo();
+	PlayInEditorSessionInfo->PlayRequestStartTime = FPlatformTime::Seconds();
 
-		// Listen server or clients: specify default win position and SAVEWINPOS so the final positions are saved
-		// in order to preserve PIE networking window setup
-		CmdLine += FString::Printf(TEXT("WinX=%d WinY=%d SAVEWINPOS=1"), WinPos.X + (int32)WindowBorderSize.Left, WinPos.Y + (int32)WindowBorderSize.Top);
-
-		// Advance window for next PIE instance...
-		AdvanceWindowPositionForNextPIEWindow(WinPos, WinSize, WindowBorderSize, false);
-	}
-	
-	return CmdLine;
-}
-
-void GetMultipleInstancePositions(int32 index, int32 &LastX, int32 &LastY)
-{
-	ULevelEditorPlaySettings* PlayInSettings = Cast<ULevelEditorPlaySettings>(ULevelEditorPlaySettings::StaticClass()->GetDefaultObject());
-
-	if (PlayInSettings->MultipleInstancePositions.IsValidIndex(index) &&
-		(PlayInSettings->MultipleInstanceLastHeight == PlayInSettings->NewWindowHeight) &&
-		(PlayInSettings->MultipleInstanceLastWidth == PlayInSettings->NewWindowWidth))
-	{
-		PlayInSettings->NewWindowPosition = PlayInSettings->MultipleInstancePositions[index];
+	// Keep a copy of their original request settings for any late
+	// joiners or async processes that need access to the settings after launch.
+	PlayInEditorSessionInfo->OriginalRequestParams = PlaySessionRequest.GetValue();
 
-		LastX = PlayInSettings->NewWindowPosition.X;
-		LastY = PlayInSettings->NewWindowPosition.Y;
-	}
-	else
+	// Load the saved window positions from the EditorPlaySettings object.
+	for (const FIntPoint& Position : EditorPlaySettings->MultipleInstancePositions)
 	{
-		PlayInSettings->NewWindowPosition = FIntPoint(LastX, LastY);
+		FPlayInEditorSessionInfo::FWindowSizeAndPos& NewPos = PlayInEditorSessionInfo->CachedWindowInfo.Add_GetRef(FPlayInEditorSessionInfo::FWindowSizeAndPos());
+		NewPos.Position = Position;
 	}
 
-	FIntPoint WinPos(LastX, LastY);
-
-	// Get the size of the window based on the type
-	FIntPoint WinSize(0, 0);
-	GetWindowSizeForInstanceType(WinSize, PlayInSettings);
-
-	// Advance window and make sure the window is going to fit where we want it
-	const FMargin WinPadding(16, 16);
-	AdvanceWindowPositionForNextPIEWindow(WinPos, WinSize, WinPadding, false);
-	FitWindowPositionToWorkArea(WinPos, WinSize, WinPadding);
-
-	// Set the size, incase it was modified
-	SetWindowSizeForInstanceType(WinSize, PlayInSettings);
-
-	LastX = WinPos.X;
-	LastY = WinPos.Y;
-}
-
-void UEditorEngine::StartQueuedPlayMapRequest()
-{
-	const bool bWantSimulateInEditor = bIsSimulateInEditorQueued;
-
+	// End any previous sessions running in separate processes.
 	EndPlayOnLocalPc();
 
-	ON_SCOPE_EXIT
-	{
-		// note that we no longer have a queued request
-		bIsPlayWorldQueued = false;
-		bIsSimulateInEditorQueued = false;
-	};
+	// If our settings require us to launch a separate process in any form, we require the user to save
+	// their content so that when the new process reads the data from disk it will match what we have in-editor.
+	const bool bIsExternalMemoryProcess = PlaySessionRequest->SessionDestination != EPlaySessionDestinationType::InProcess;
 
-	const ULevelEditorPlaySettings* PlayInSettings = GetDefault<ULevelEditorPlaySettings>();
+	bool bUserWantsInProcess;
+	EditorPlaySettings->GetRunUnderOneProcess(bUserWantsInProcess);
 
-	// Launch multi-player instances if necessary
-	// (note that if you have 'RunUnderOneProcess' checked and do a bPlayOnLocalPcSession (standalone) - play standalone 'wins' - multiple instances will be launched for multiplayer)
-	const EPlayNetMode PlayNetMode = [&PlayInSettings]{ EPlayNetMode NetMode(PIE_Standalone); return (PlayInSettings->GetPlayNetMode(NetMode) ? NetMode : PIE_Standalone); }();
-	const bool CanRunUnderOneProcess = [&PlayInSettings]{ bool RunUnderOneProcess(false); return (PlayInSettings->GetRunUnderOneProcess(RunUnderOneProcess) && RunUnderOneProcess); }();
+	bool bRequestSave = bIsExternalMemoryProcess && !bUserWantsInProcess;
 
-	const bool bRequestSave = bPlayOnLocalPcSession || bPlayUsingLauncher || (!CanRunUnderOneProcess && PlayNetMode != PIE_Standalone);
 	if (bRequestSave && !SaveMapsForPlaySession())
 	{
 		// Maps did not save, print a warning
@@ -1221,86 +1109,24 @@ void UEditorEngine::StartQueuedPlayMapRequest()
 		return;
 	}
 
-	if (PlayNetMode != PIE_Standalone && (!CanRunUnderOneProcess || bPlayOnLocalPcSession) && !bPlayUsingLauncher)
-	{
-		int32 NumClients = 0;
-
-		// If we start to the right of the editor work area, call FitToWorkArea and it will find the next place we can place a new instance window if that's not preferable.
-		const FSlateRect PreferredWorkArea = FSlateApplication::Get().GetPreferredWorkArea();		
-		FIntPoint WinPosition((int32)PreferredWorkArea.Right, (int32)PreferredWorkArea.Top);
-
-		// We'll need to spawn a server if we're playing outside the editor or the editor wants to run as a client
-		if (bPlayOnLocalPcSession || PlayNetMode == PIE_Client || PlayNetMode == PIE_StandaloneWithServer)
-		{			
-			FString MapNameOverride;
-			if (PlayInSettings->IsServerMapNameOverrideActive())
-			{
-				PlayInSettings->GetServerMapNameOverride(MapNameOverride);
-			}
-
-			PlayStandaloneLocalPc(*MapNameOverride, &WinPosition, NumClients, true);
-			
-			const bool CanPlayNetDedicated = [&PlayInSettings]{ bool PlayNetDedicated(false); return (PlayInSettings->GetPlayNetDedicated(PlayNetDedicated) && PlayNetDedicated); }();
-			if (!CanPlayNetDedicated)
-			{
-				// Listen server counts as a client
-				NumClients++;
-			}
-		}
-		
-		// If we're playing in the editor
-		if (!bPlayOnLocalPcSession)
-		{
-			PlayInEditor(GetEditorWorldContext().World(), bWantSimulateInEditor);
-
-			// Editor counts as a client
-			NumClients++;
-		}
-
-		// Build the connection String
-		FString ConnectionAddr;
-
-		// Ignore the user's settings if the autoconnect option is inaccessible due to settings conflicts.
-		const bool WillAutoConnectToServer = [&PlayInSettings, PlayNetMode] { bool AutoConnectToServer(false);
-			return (PlayInSettings->GetAutoConnectToServerVisibility() == EVisibility::Visible) ? 
-				(PlayInSettings->GetAutoConnectToServer(AutoConnectToServer) && AutoConnectToServer) : (PlayNetMode != PIE_StandaloneWithServer); }();
-
-		if (WillAutoConnectToServer)
-		{
-			ConnectionAddr += TEXT("127.0.0.1");
-			uint16 ServerPort = 0;
-			if (PlayInSettings->GetServerPort(ServerPort))
-			{
-				ConnectionAddr += FString::Printf(TEXT(":%hu"), ServerPort);
-			}
-		}
-
-		// Spawn number of clients
-		const int32 PlayNumberOfClients = [&PlayInSettings]{ int32 NumberOfClients(0); return (PlayInSettings->GetPlayNumberOfClients(NumberOfClients) ? NumberOfClients : 0); }();
-		for (int32 i = NumClients; i < PlayNumberOfClients; ++i)
-		{
-			PlayStandaloneLocalPc(*ConnectionAddr, &WinPosition, i, false);
-		}
-	}
-	else
+	// We'll branch primarily based on the Session Destination, because it affects which settings we apply and how.
+	switch (PlaySessionRequest->SessionDestination)
 	{
-		// Launch standalone PIE session
-		if (bPlayOnLocalPcSession)
-		{
-			PlayStandaloneLocalPc();
-		}
-		else if (bPlayUsingLauncher)
-		{
-			PlayUsingLauncher();
-		}
-		else
-		{
-			PlayInEditor( GetEditorWorldContext().World(), bWantSimulateInEditor );
-		}
+	case EPlaySessionDestinationType::InProcess:
+		// Create one-or-more PIE/SIE sessions inside of the current process.
+		StartPlayInEditorSession(PlaySessionRequest.GetValue());
+		break;
+	case EPlaySessionDestinationType::NewProcess:
+		// Create one-or-more PIE session by launching a new process on the local machine.
+		StartPlayInNewProcessSession(PlaySessionRequest.GetValue());
+		break;
+	case EPlaySessionDestinationType::Launcher:
+		// Create a Play Session via the Launcher which may be on a local or remote device.
+		StartPlayUsingLauncherSession(PlaySessionRequest.GetValue());
+		break;
+	default:
+		check(false);
 	}
-
-	// note that we no longer have a queued request
-	CancelRequestPlaySession();
 }
 
 void UEditorEngine::EndPlayOnLocalPc( )
@@ -1320,2835 +1146,2120 @@ void UEditorEngine::EndPlayOnLocalPc( )
 	PlayOnLocalPCSessions.Empty();
 }
 
-void UEditorEngine::PlayStandaloneLocalPc(FString MapNameOverride, FIntPoint* WindowPos, int32 PIENum, bool bIsServer)
-{
-	const ULevelEditorPlaySettings* PlayInSettings = GetDefault<ULevelEditorPlaySettings>();
-	//const ULevelEditorPlaySettings* PlayInSettings = InPlaySettings != NULL ? InPlaySettings : GetDefault<ULevelEditorPlaySettings>();
-	const bool CanPlayNetDedicated = [&PlayInSettings]{ bool PlayNetDedicated(false); return (PlayInSettings->GetPlayNetDedicated(PlayNetDedicated) && PlayNetDedicated); }();
 
-	FString CmdLine;
-	if (WindowPos != NULL)	// If WindowPos == NULL, we're just launching one instance
+int32 FInternalPlayLevelUtils::ResolveDirtyBlueprints(const bool bPromptForCompile, TArray<UBlueprint*>& ErroredBlueprints, const bool bForceLevelScriptRecompile)
+{
+	struct FLocal
 	{
-		CmdLine = GenerateCmdLineForNextPieInstance(*WindowPos, PIENum, bIsServer && CanPlayNetDedicated);
-	}
+		static void OnMessageLogLinkActivated(const class TSharedRef<IMessageToken>& Token)
+		{
+			if (Token->GetType() == EMessageToken::Object)
+			{
+				const TSharedRef<FUObjectToken> UObjectToken = StaticCastSharedRef<FUObjectToken>(Token);
+				if (UObjectToken->GetObject().IsValid())
+				{
+					FKismetEditorUtilities::BringKismetToFocusAttentionOnObject(UObjectToken->GetObject().Get());
+				}
+			}
+		}
 
-	const FString URLParms = bIsServer && !CanPlayNetDedicated ? TEXT("?Listen") : FString();
+		static void AddCompileErrorToLog(UBlueprint* ErroredBlueprint, FMessageLog& BlueprintLog)
+		{
+			FFormatNamedArguments Arguments;
+			Arguments.Add(TEXT("Name"), FText::FromString(ErroredBlueprint->GetName()));
 
-	// select map to play
-	TArray<FString> SavedMapNames;
-	if (MapNameOverride.IsEmpty())
-	{
-		FWorldContext & EditorContext = GetEditorWorldContext();
+			TSharedRef<FTokenizedMessage> Message = FTokenizedMessage::Create(EMessageSeverity::Warning);
+			Message->AddToken(FTextToken::Create(LOCTEXT("BlueprintCompileFailed", "Blueprint failed to compile: ")));
+			Message->AddToken(FUObjectToken::Create(ErroredBlueprint, FText::FromString(ErroredBlueprint->GetName()))
+				->OnMessageTokenActivated(FOnMessageTokenActivated::CreateStatic(&FLocal::OnMessageLogLinkActivated))
+			);
 
-		// Open original map
-		FString MapName = EditorContext.World()->GetOutermost()->GetName();
-		SavedMapNames.Add(MapName);
-	}
-	else
-	{
-		SavedMapNames.Add(MapNameOverride);
-	}
+			BlueprintLog.AddMessage(Message);
+		}
+	};
 
-	if (SavedMapNames.Num() == 0)
-	{
-		return;
-	}
+	const bool bAutoCompile = !bPromptForCompile;
+	FString PromptDirtyList;
 
-	FString GameNameOrProjectFile;
-	if (FPaths::IsProjectFilePathSet())
-	{
-		GameNameOrProjectFile = FString::Printf(TEXT("\"%s\""), *FPaths::GetProjectFilePath());
-	}
-	else
-	{
-		GameNameOrProjectFile = FApp::GetProjectName();
-	}
+	TArray<UBlueprint*> InNeedOfRecompile;
+	ErroredBlueprints.Empty();
 
-	FString AdditionalParameters(TEXT(" -messaging -SessionName=\"Play in Standalone Game\""));
+	FMessageLog BlueprintLog("BlueprintLog");
 
-	const FString PreviewGameLanguage = FTextLocalizationManager::Get().GetConfiguredGameLocalizationPreviewLanguage();
-	if (!PreviewGameLanguage.IsEmpty())
+	double BPRegenStartTime = FPlatformTime::Seconds();
+	for (TObjectIterator<UBlueprint> BlueprintIt; BlueprintIt; ++BlueprintIt)
 	{
-		AdditionalParameters += TEXT(" -culture=");
-		AdditionalParameters += PreviewGameLanguage;
-	}
+		UBlueprint* Blueprint = *BlueprintIt;
 
-	// apply additional settings
-	if (bPlayUsingMobilePreview)
-	{
-		if (PlayUsingMobilePreviewTargetDevice.IsEmpty() == false)
+		// ignore up-to-date BPs
+		if (Blueprint->IsUpToDate())
 		{
-			AdditionalParameters += TEXT(" -MobileTargetDevice=\"") + PlayUsingMobilePreviewTargetDevice + TEXT("\"");
+			continue;
 		}
-		else
+
+		// do not try to recompile BPs that have not changed since they last failed to compile, so don't check Blueprint->IsUpToDate()
+		const bool bIsDirtyAndShouldBeRecompiled = Blueprint->IsPossiblyDirty();
+		if (!FBlueprintEditorUtils::IsDataOnlyBlueprint(Blueprint)
+			&& (bIsDirtyAndShouldBeRecompiled || (FBlueprintEditorUtils::IsLevelScriptBlueprint(Blueprint) && bForceLevelScriptRecompile))
+			&& (Blueprint->Status != BS_Unknown)
+			&& !Blueprint->IsPendingKill())
 		{
-			AdditionalParameters += TEXT(" -featureleveles2");
-		}
+			InNeedOfRecompile.Add(Blueprint);
 
-		if (IsOpenGLPlatform(GShaderPlatformForFeatureLevel[GMaxRHIFeatureLevel]))
+			if (bPromptForCompile)
+			{
+				PromptDirtyList += FString::Printf(TEXT("\n   %s"), *Blueprint->GetName());
+			}
+		}
+		else if (BS_Error == Blueprint->Status && Blueprint->bDisplayCompilePIEWarning)
 		{
-			AdditionalParameters += TEXT(" -opengl");
+			ErroredBlueprints.Add(Blueprint);
+			FLocal::AddCompileErrorToLog(Blueprint, BlueprintLog);
 		}
-		AdditionalParameters += TEXT(" -faketouches");
-
-		AdditionalParameters += TEXT(" -MultiprocessSaveConfig");
-	}
-
-	if (bPlayUsingVulkanPreview)
-	{
-		ensure(!bPlayUsingMobilePreview);
-		AdditionalParameters += TEXT(" -vulkan -faketouches -featureleveles31");
 	}
 
-	// Disable the HMD device in the new process if present. The editor process owns the HMD resource.
-	if (!bPlayUsingMobilePreview && !bPlayUsingVulkanPreview && GEngine->XRSystem.IsValid() && GEngine->XRSystem->GetHMDDevice() && GEngine->XRSystem->GetHMDDevice()->IsHMDConnected())
+	bool bRunCompilation = bAutoCompile;
+	if (bPromptForCompile)
 	{
-		AdditionalParameters += TEXT(" -nohmd");
-		UE_LOG(LogHMD, Warning, TEXT("Standalone game VR not supported, please use VR Preview."));
-	}
+		FFormatNamedArguments Args;
+		Args.Add(TEXT("DirtyBlueprints"), FText::FromString(PromptDirtyList));
+		const FText PromptMsg = FText::Format(NSLOCTEXT("PlayInEditor", "PrePIE_BlueprintsDirty", "One or more blueprints have been modified without being recompiled. Do you want to compile them now? \n{DirtyBlueprints}"), Args);
 
-	if (PlayInSettings->DisableStandaloneSound)
-	{
-		AdditionalParameters += TEXT(" -nosound");
+		EAppReturnType::Type PromptResponse = FMessageDialog::Open(EAppMsgType::YesNo, PromptMsg);
+		bRunCompilation = (PromptResponse == EAppReturnType::Yes);
 	}
+	int32 RecompiledCount = 0;
 
-	if (PlayInSettings->AdditionalLaunchParameters.Len() > 0)
+	if (bRunCompilation && (InNeedOfRecompile.Num() > 0))
 	{
-		AdditionalParameters += TEXT(" ");
-		AdditionalParameters += PlayInSettings->AdditionalLaunchParameters;
-	}
+		const FText LogPageLabel = (bAutoCompile) ? LOCTEXT("BlueprintAutoCompilationPageLabel", "Pre-Play auto-recompile") :
+			LOCTEXT("BlueprintCompilationPageLabel", "Pre-Play recompile");
+		BlueprintLog.NewPage(LogPageLabel);
 
-	if (bPlayUsingMobilePreview)
-	{
-		if (PlayInSettings->AdditionalLaunchParametersForMobile.Len() > 0)
+		TArray<UBlueprint*> CompiledBlueprints;
+		auto OnBlueprintPreCompileLambda = [&CompiledBlueprints](UBlueprint* InBlueprint)
 		{
-			AdditionalParameters += TEXT(" ");
-			AdditionalParameters += PlayInSettings->AdditionalLaunchParametersForMobile;
-		}
-	}
+			check(InBlueprint != nullptr);
 
-	uint16 ServerPort = 0;
-	if (bIsServer && PlayInSettings->GetServerPort(ServerPort))
-	{
-		AdditionalParameters += FString::Printf(TEXT(" -port=%hu"), ServerPort);
-	}
+			if (CompiledBlueprints.Num() == 0)
+			{
+				UE_LOG(LogPlayLevel, Log, TEXT("[PlayLevel] Compiling %s before play..."), *InBlueprint->GetName());
+			}
+			else
+			{
+				UE_LOG(LogPlayLevel, Log, TEXT("[PlayLevel]   Compiling %s as a dependent..."), *InBlueprint->GetName());
+			}
 
-	if (PlayInSettings->IsNetworkEmulationEnabled())
-	{
-		NetworkEmulationTarget CurrentTarget = bIsServer ? NetworkEmulationTarget::Server : NetworkEmulationTarget::Client;
-		if (PlayInSettings->NetworkEmulationSettings.IsEmulationEnabledForTarget(CurrentTarget))
-		{
-			AdditionalParameters += PlayInSettings->NetworkEmulationSettings.BuildPacketSettingsForCmdLine();
-		}
-	}
+			CompiledBlueprints.Add(InBlueprint);
+		};
 
-	// Decide if fullscreen or windowed based on what is specified in the params
-	if (!AdditionalParameters.Contains(TEXT("-fullscreen")) && !AdditionalParameters.Contains(TEXT("-windowed")))
-	{
-		// Nothing specified fallback to window otherwise keep what is specified
-		AdditionalParameters += TEXT(" -windowed");		
-	}
+		// Register compile callback
+		FDelegateHandle PreCompileDelegateHandle = GEditor->OnBlueprintPreCompile().AddLambda(OnBlueprintPreCompileLambda);
 
-	if (RequestedAdditionalStandaloneLaunchOptions.Len() > 0)
-	{
-		AdditionalParameters += TEXT(" ");
-		AdditionalParameters += RequestedAdditionalStandaloneLaunchOptions;
-		//clear it now it's been used
-		RequestedAdditionalStandaloneLaunchOptions = FString();
-	}
+		// Recompile all necessary blueprints in a single loop, saving GC until the end
+		for (auto BlueprintIt = InNeedOfRecompile.CreateIterator(); BlueprintIt; ++BlueprintIt)
+		{
+			UBlueprint* Blueprint = *BlueprintIt;
 
-	FIntPoint WinSize(0, 0);
-	GetWindowSizeForInstanceType(WinSize, PlayInSettings);
+			int32 CurrItIndex = BlueprintIt.GetIndex();
 
-	// Get desktop metrics
-	FDisplayMetrics DisplayMetrics;
-	FSlateApplication::Get().GetCachedDisplayMetrics(DisplayMetrics);
+			// Compile the Blueprint (note: re-instancing may trigger additional compiles for child/dependent Blueprints; see callback above)
+			FKismetEditorUtilities::CompileBlueprint(Blueprint, EBlueprintCompileOptions::SkipGarbageCollection);
 
-	// Force resolution
-	if ((WinSize.X <= 0 || WinSize.X > DisplayMetrics.PrimaryDisplayWidth) || (WinSize.Y <= 0 || WinSize.Y > DisplayMetrics.PrimaryDisplayHeight))
-	{
-		AdditionalParameters += TEXT(" -ForceRes");
-	}
-	
-	// Check if centered
-	FString Params;
-	if (PlayInSettings->CenterNewWindow)
-	{
-		Params = FString::Printf(TEXT("%s %s -game -PIEVIACONSOLE -ResX=%d -ResY=%d %s%s %s"),
-			*GameNameOrProjectFile,
-			*BuildPlayWorldURL(*SavedMapNames[0], false, URLParms),
-			WinSize.X,
-			WinSize.Y,
-			*FCommandLine::GetSubprocessCommandline(),
-			*AdditionalParameters,
-			*CmdLine
-			);
-	}
-	else
-	{
-		FIntPoint WinPos(0, 0);
-
-		Params = FString::Printf(TEXT("%s %s -game -PIEVIACONSOLE -WinX=%d -WinY=%d -ResX=%d -ResY=%d %s%s %s"),
-			*GameNameOrProjectFile,
-			*BuildPlayWorldURL(*SavedMapNames[0], false, URLParms),
-			PlayInSettings->NewWindowPosition.X,
-			PlayInSettings->NewWindowPosition.Y,
-			WinSize.X,
-			WinSize.Y,
-			*FCommandLine::GetSubprocessCommandline(),
-			*AdditionalParameters,
-			*CmdLine
-		);
-	}
+			// Check for errors after compiling
+			for (UBlueprint* CompiledBlueprint : CompiledBlueprints)
+			{
+				if (CompiledBlueprint != Blueprint)
+				{
+					int32 ExistingIndex = InNeedOfRecompile.Find(CompiledBlueprint);
+					// if this dependent blueprint is already set up to compile 
+					// later in this loop, then there is no need to add it to be recompiled again
+					if (ExistingIndex > CurrItIndex)
+					{
+						InNeedOfRecompile.RemoveAt(ExistingIndex);
+					}
+				}
+
+				const bool bHadError = (!CompiledBlueprint->IsUpToDate() && CompiledBlueprint->Status != BS_Unknown);
+
+				// Check if the Blueprint has already been added to the error list to prevent it from being added again
+				if (bHadError && ErroredBlueprints.Find(CompiledBlueprint) == INDEX_NONE)
+				{
+					ErroredBlueprints.Add(CompiledBlueprint);
+					FLocal::AddCompileErrorToLog(CompiledBlueprint, BlueprintLog);
+				}
+
+				++RecompiledCount;
+			}
+
+			// Reset for next pass
+			CompiledBlueprints.Empty();
+		}
 
-	// launch the game process
-	FString GamePath = FPlatformProcess::ExecutablePath();
-	FPlayOnPCInfo *NewSession = new (PlayOnLocalPCSessions) FPlayOnPCInfo();
+		// Now that all Blueprints have been compiled, run a single GC pass to clean up artifacts
+		CollectGarbage(GARBAGE_COLLECTION_KEEPFLAGS);
 
-	uint32 ProcessID = 0;
-	NewSession->ProcessHandle = FPlatformProcess::CreateProc(*GamePath, *Params, true, false, false, &ProcessID, 0, NULL, NULL);
+		// Unregister compile callback
+		GEditor->OnBlueprintPreCompile().Remove(PreCompileDelegateHandle);
 
-	if (!NewSession->ProcessHandle.IsValid())
+		UE_LOG(LogPlayLevel, Log, TEXT("PlayLevel: Blueprint regeneration took %d ms (%i blueprints)"), (int32)((FPlatformTime::Seconds() - BPRegenStartTime) * 1000), RecompiledCount);
+	}
+	else if (bAutoCompile)
 	{
-		UE_LOG(LogPlayLevel, Error, TEXT("Failed to run a copy of the game on this PC."));
+		UE_LOG(LogPlayLevel, Log, TEXT("PlayLevel: No blueprints needed recompiling"));
 	}
 
-	FMargin SafeZoneRatio = PlayInSettings->PIESafeZoneOverride;
-	SafeZoneRatio.Left /= (PlayInSettings->NewWindowPosition.X / 2.0f);
-	SafeZoneRatio.Right /= (PlayInSettings->NewWindowPosition.X / 2.0f);
-	SafeZoneRatio.Bottom /= (PlayInSettings->NewWindowPosition.Y / 2.0f);
-	SafeZoneRatio.Top /= (PlayInSettings->NewWindowPosition.Y / 2.0f);
-	FSlateApplication::Get().OnDebugSafeZoneChanged.Broadcast(SafeZoneRatio, false);
-
-	FEditorDelegates::BeginStandaloneLocalPlay.Broadcast(ProcessID);
+	return RecompiledCount;
 }
 
-static void HandleOutputReceived  (const FString& InMessage)
+void UEditorEngine::RequestEndPlayMap()
 {
-	if (InMessage.Contains(TEXT("Error:")))
-	{
-		UE_LOG(LogPlayLevel, Error, TEXT("%s"), *InMessage);
-	}
-	else if (InMessage.Contains(TEXT("Warning:")))
+	if( PlayWorld )
 	{
-		UE_LOG(LogPlayLevel, Warning, TEXT("%s"), *InMessage);
+		bRequestEndPlayMapQueued = true;
+
+		// Cache the position and rotation of the camera (the controller may be destroyed before we end the pie session and we need them to preserve the camera position)
+		if (bLastViewAndLocationValid == false)
+		{
+			for (int32 WorldIdx = WorldList.Num() - 1; WorldIdx >= 0; --WorldIdx)
+			{
+				FWorldContext &ThisContext = WorldList[WorldIdx];
+				if (ThisContext.WorldType == EWorldType::PIE)
+				{
+					FSlatePlayInEditorInfo* const SlatePlayInEditorSession = SlatePlayInEditorMap.Find(ThisContext.ContextHandle);
+					if ((SlatePlayInEditorSession != nullptr) && (SlatePlayInEditorSession->EditorPlayer.IsValid() == true) )
+					{
+						if( SlatePlayInEditorSession->EditorPlayer.Get()->PlayerController != nullptr )
+						{
+							SlatePlayInEditorSession->EditorPlayer.Get()->PlayerController->GetPlayerViewPoint( LastViewLocation, LastViewRotation );
+							bLastViewAndLocationValid = true;
+							break;
+						}
+					}
+				}
+			}
+		}
 	}
-	else 
+}
+
+FString UEditorEngine::BuildPlayWorldURL(const TCHAR* MapName, bool bSpectatorMode, FString AdditionalURLOptions)
+{
+	// the URL we are building up
+	FString URL(MapName);
+
+	// If we hold down control, start in spectating mode
+	if (bSpectatorMode)
 	{
-		UE_LOG(LogPlayLevel, Log, TEXT("%s"), *InMessage);
+		// Start in spectator mode
+		URL += TEXT("?SpectatorOnly=1");
 	}
+
+	// Add any game-specific options set in the INI file
+	URL += InEditorGameURLOptions;
+
+	// Add any additional options that were specified for this call
+	URL += AdditionalURLOptions;
+
+	return URL;
 }
 
-static void HandleCancelButtonClicked(ILauncherWorkerPtr LauncherWorker)
+bool UEditorEngine::SpawnPlayFromHereStart(UWorld* World, AActor*& PlayerStart)
 {
-	if (LauncherWorker.IsValid())
+	if (PlayInEditorSessionInfo.IsSet() && PlayInEditorSessionInfo->OriginalRequestParams.HasPlayWorldPlacement())
 	{
-		LauncherWorker->Cancel();
+		return SpawnPlayFromHereStart(World, PlayerStart, PlayInEditorSessionInfo->OriginalRequestParams.StartLocation.GetValue(), PlayInEditorSessionInfo->OriginalRequestParams.StartRotation.GetValue());
 	}
-}
 
-/* FMainFrameActionCallbacks callbacks
- *****************************************************************************/
+	// Not having a location set is still considered a success.
+	return true;
+}
 
-class FLauncherNotificationTask
+bool UEditorEngine::SpawnPlayFromHereStart( UWorld* World, AActor*& PlayerStart, const FVector& StartLocation, const FRotator& StartRotation )
 {
-public:
+	// null it out in case we don't need to spawn one, and the caller relies on us setting it
+	PlayerStart = NULL;
+
+	// spawn the PlayerStartPIE in the given world
+	FActorSpawnParameters SpawnParameters;
+	SpawnParameters.OverrideLevel = World->PersistentLevel;
+	PlayerStart = World->SpawnActor<AActor>(PlayFromHerePlayerStartClass, StartLocation, StartRotation, SpawnParameters);
+
+	// make sure we were able to spawn the PlayerStartPIE there
+	if(!PlayerStart)
+	{
+		FMessageDialog::Open( EAppMsgType::Ok, NSLOCTEXT("UnrealEd", "Prompt_22", "Failed to create entry point. Try another location, or you may have to rebuild your level."));
+		return false;
+	}
 
-	FLauncherNotificationTask( TWeakPtr<SNotificationItem> InNotificationItemPtr, SNotificationItem::ECompletionState InCompletionState, const FText& InText )
-		: CompletionState(InCompletionState)
-		, NotificationItemPtr(InNotificationItemPtr)
-		, Text(InText)
-	{ }
+	// tag the start
+	ANavigationObjectBase* NavPlayerStart = Cast<ANavigationObjectBase>(PlayerStart);
+	if (NavPlayerStart)
+	{
+		NavPlayerStart->bIsPIEPlayerStart = true;
+	}
+
+	return true;
+}
 
-	void DoTask( ENamedThreads::Type CurrentThread, const FGraphEventRef& MyCompletionGraphEvent )
+static bool ShowBlueprintErrorDialog( TArray<UBlueprint*> ErroredBlueprints )
+{
+	struct Local
 	{
-		if (NotificationItemPtr.IsValid())
+		static void OnHyperlinkClicked( TWeakObjectPtr<UBlueprint> InBlueprint, TSharedPtr<SCustomDialog> InDialog )
 		{
-			if (CompletionState == SNotificationItem::CS_Fail)
-			{
-				GEditor->PlayEditorSound(TEXT("/Engine/EditorSounds/Notifications/CompileFailed_Cue.CompileFailed_Cue"));
-			}
-			else if (CompletionState == SNotificationItem::CS_Success)
+			if (UBlueprint* BlueprintToEdit = InBlueprint.Get())
 			{
-				GEditor->PlayEditorSound(TEXT("/Engine/EditorSounds/Notifications/CompileSuccess_Cue.CompileSuccess_Cue"));
+				// Open the blueprint
+				GEditor->EditObject( BlueprintToEdit );
 			}
 
-			TSharedPtr<SNotificationItem> NotificationItem = NotificationItemPtr.Pin();
-			NotificationItem->SetText(Text);
-			NotificationItem->SetCompletionState(CompletionState);
-			if (CompletionState == SNotificationItem::CS_Success || CompletionState == SNotificationItem::CS_Fail)
+			if (InDialog.IsValid())
 			{
-				NotificationItem->ExpireAndFadeout();
+				// Opening the blueprint editor above may end up creating an invisible new window on top of the dialog, 
+				// thus making it not interactable, so we have to force the dialog back to the front
+				InDialog->BringToFront(true);
 			}
 		}
-	}
+	};
+
+	TSharedRef<SVerticalBox> DialogContents = SNew(SVerticalBox)
+		+ SVerticalBox::Slot()
+		.Padding(0, 0, 0, 16)
+		[
+			SNew(STextBlock)
+			.Text(NSLOCTEXT("PlayInEditor", "PrePIE_BlueprintErrors", "One or more blueprints has an unresolved compiler error, are you sure you want to Play in Editor?"))
+		];
 
-	static ESubsequentsMode::Type GetSubsequentsMode() { return ESubsequentsMode::TrackSubsequents; }
-	ENamedThreads::Type GetDesiredThread( ) { return ENamedThreads::GameThread; }
-	FORCEINLINE TStatId GetStatId() const
+	TSharedPtr<SCustomDialog> CustomDialog;
+
+	for (UBlueprint* Blueprint : ErroredBlueprints)
 	{
-		RETURN_QUICK_DECLARE_CYCLE_STAT(FLauncherNotificationTask, STATGROUP_TaskGraphTasks);
+		TWeakObjectPtr<UBlueprint> BlueprintPtr = Blueprint;
+
+		DialogContents->AddSlot()
+			.AutoHeight()
+			.HAlign(HAlign_Left)
+			[
+				SNew(SHyperlink)
+				.Style(FEditorStyle::Get(), "Common.GotoBlueprintHyperlink")
+				.OnNavigate(FSimpleDelegate::CreateLambda([BlueprintPtr, &CustomDialog]() { Local::OnHyperlinkClicked(BlueprintPtr, CustomDialog); }))
+				.Text(FText::FromString(Blueprint->GetName()))
+				.ToolTipText(NSLOCTEXT("SourceHyperlink", "EditBlueprint_ToolTip", "Click to edit the blueprint"))
+			];
 	}
 
-private:
+	DialogContents->AddSlot()
+		.Padding(0, 16, 0, 0)
+		[
+			SNew(STextBlock)
+			.Text(NSLOCTEXT("PlayInEditor", "PrePIE_BlueprintErrorsDelayedOpen", "Clicked blueprints will open once this dialog is closed."))
+		];
+
+	FText DialogTitle = NSLOCTEXT("PlayInEditor", "PrePIE_BlueprintErrorsTitle", "Blueprint Compilation Errors");
 
-	SNotificationItem::ECompletionState CompletionState;
-	TWeakPtr<SNotificationItem> NotificationItemPtr;
-	FText Text;
-};
+	FText OKText = NSLOCTEXT("PlayInEditor", "PrePIE_OkText", "Play in Editor");
+	FText CancelText = NSLOCTEXT("Dialogs", "EAppReturnTypeCancel", "Cancel");
+
+	CustomDialog = SNew(SCustomDialog)
+		.Title(DialogTitle)
+		.IconBrush("NotificationList.DefaultMessage")
+		.DialogContent(DialogContents)
+		.Buttons( { SCustomDialog::FButton(OKText), SCustomDialog::FButton(CancelText) } );
+
+	int ButtonPressed = CustomDialog->ShowModal();
+	return ButtonPressed == 0;
+}
+
+FGameInstancePIEResult UEditorEngine::PreCreatePIEServerInstance(const bool bAnyBlueprintErrors, const bool bStartInSpectatorMode, const float PIEStartTime, const bool bSupportsOnlinePIE, int32& InNumOnlinePIEInstances)
+{
+	return FGameInstancePIEResult::Success();
+}
+
+bool UEditorEngine::SupportsOnlinePIE() const
+{
+	return UOnlineEngineInterface::Get()->SupportsOnlinePIE();
+}
 
+void UEditorEngine::OnLoginPIEComplete(int32 LocalUserNum, bool bWasSuccessful, const FString& ErrorString, FPieLoginStruct DataStruct)
+{
+	// This is needed because pie login may change the state of the online objects that called this function. This also enqueues the
+	// callback back onto the main thread at an appropriate time instead of mid-networking callback.
+	GetTimerManager()->SetTimerForNextTick(FTimerDelegate::CreateUObject(this, &UEditorEngine::OnLoginPIEComplete_Deferred, LocalUserNum, bWasSuccessful, ErrorString, DataStruct));
+}
 
-void UEditorEngine::HandleStageStarted(const FString& InStage, TWeakPtr<SNotificationItem> NotificationItemPtr)
+void UEditorEngine::OnLoginPIEComplete_Deferred(int32 LocalUserNum, bool bWasSuccessful, FString ErrorString, FPieLoginStruct DataStruct)
 {
-	bool bSetNotification = true;
-	FFormatNamedArguments Arguments;
-	FText NotificationText;
-	if (InStage.Contains(TEXT("Cooking")) || InStage.Contains(TEXT("Cook Task")))
+	// This function will get called for both Async Online Service Log-in and for local non-async client creation.
+	// This is called once per client that will be created, and may be called many frames after the initial PIE request
+	// was started.
+
+	UE_LOG(LogPlayLevel, Verbose, TEXT("OnLoginPIEComplete LocalUserNum: %d bSuccess: %d %s"), LocalUserNum, bWasSuccessful, *ErrorString);
+	FWorldContext* PieWorldContext = GetWorldContextFromHandle(DataStruct.WorldContextHandle);
+
+	if (!PieWorldContext)
 	{
-		FString PlatformName = PlayUsingLauncherDeviceId.Left(PlayUsingLauncherDeviceId.Find(TEXT("@")));
-		if (PlatformName.Contains(TEXT("NoEditor")))
-		{
-			PlatformName = PlatformName.Left(PlatformName.Find(TEXT("NoEditor")));
-		}
-		Arguments.Add(TEXT("PlatformName"), FText::FromString(PlatformName));
-		NotificationText = FText::Format(LOCTEXT("LauncherTaskProcessingNotification", "Processing Assets for {PlatformName}..."), Arguments);
+		// This will fail if PIE was ended before this callback happened, silently return
+		return;
 	}
-	else if (InStage.Contains(TEXT("Build Task")))
+
+	// Create a new Game Instance for this.
+	UGameInstance* GameInstance = CreateInnerProcessPIEGameInstance(PlayInEditorSessionInfo->OriginalRequestParams, DataStruct.GameInstancePIEParameters, DataStruct.PIEInstanceIndex);
+
+	bool bAnyPIELoginFailed = false;
+	if (GameInstance)
 	{
-		FString PlatformName = PlayUsingLauncherDeviceId.Left(PlayUsingLauncherDeviceId.Find(TEXT("@")));
-		if (PlatformName.Contains(TEXT("NoEditor")))
-		{
-			PlatformName = PlatformName.Left(PlatformName.Find(TEXT("NoEditor")));
-		}
-		Arguments.Add(TEXT("PlatformName"), FText::FromString(PlatformName));
-		if (!bPlayUsingLauncherBuild)
+		GameInstance->GetWorldContext()->bWaitingOnOnlineSubsystem = false;
+
+		// Logging after the create so a new MessageLog Page is created
+		if (bWasSuccessful)
 		{
-			NotificationText = FText::Format(LOCTEXT("LauncherTaskValidateNotification", "Validating Executable for {PlatformName}..."), Arguments);
+			if (DataStruct.GameInstancePIEParameters.NetMode != EPlayNetMode::PIE_Client )
+			{
+				FMessageLog(NAME_CategoryPIE).Info(LOCTEXT("LoggedInServer", "Server logged in"));
+			}
+			else
+			{
+				FMessageLog(NAME_CategoryPIE).Info(LOCTEXT("LoggedInClient", "Client logged in"));
+			}
 		}
 		else
 		{
-			NotificationText = FText::Format(LOCTEXT("LauncherTaskBuildNotification", "Building Executable for {PlatformName}..."), Arguments);
+			bAnyPIELoginFailed = true;
+			if (DataStruct.GameInstancePIEParameters.NetMode != EPlayNetMode::PIE_Client)
+			{
+				FMessageLog(NAME_CategoryPIE).Error(FText::Format(LOCTEXT("LoggedInServerFailure", "Server failed to login. {0}"), FText::FromString(ErrorString)));
+			}
+			else
+			{
+				FMessageLog(NAME_CategoryPIE).Error(FText::Format(LOCTEXT("LoggedInClientFailure", "Client failed to login. {0}"), FText::FromString(ErrorString)));
+			}
 		}
 	}
-	else if (InStage.Contains(TEXT("Deploy Task")))
+
+	PlayInEditorSessionInfo->NumOutstandingPIELogins--;
+	if (!ensureAlwaysMsgf(PlayInEditorSessionInfo->NumOutstandingPIELogins >= 0, TEXT("PIEInstancesToLogInCount was not properly reset at some point.")))
 	{
-		Arguments.Add(TEXT("DeviceName"), FText::FromString(PlayUsingLauncherDeviceName));
-		if (PlayUsingLauncherDeviceName.Len() == 0)
-		{
-			NotificationText = FText::Format(LOCTEXT("LauncherTaskStageNotificationNoDevice", "Deploying Executable and Assets..."), Arguments);
-		}
-		else
-		{
-			NotificationText = FText::Format(LOCTEXT("LauncherTaskStageNotification", "Deploying Executable and Assets to {DeviceName}..."), Arguments);
-		}
+		PlayInEditorSessionInfo->NumOutstandingPIELogins = 0;
 	}
-	else if (InStage.Contains(TEXT("Run Task")))
+
+	// If there are no more instances waiting to log-in then we can do post-launch notifications.
+	if (PlayInEditorSessionInfo->NumOutstandingPIELogins == 0)
 	{
-		Arguments.Add(TEXT("GameName"), FText::FromString(FApp::GetProjectName()));
-		Arguments.Add(TEXT("DeviceName"), FText::FromString(PlayUsingLauncherDeviceName));
-		if (PlayUsingLauncherDeviceName.Len() == 0)
+		if (bAnyPIELoginFailed)
 		{
-			NotificationText = FText::Format(LOCTEXT("LauncherTaskRunNotificationNoDevice", "Running {GameName}..."), Arguments);
+			UE_LOG(LogPlayLevel, Warning, TEXT("At least one of the requested PIE instances failed to start, ending PIE session."));
+			EndPlayMap();
 		}
 		else
 		{
-			NotificationText = FText::Format(LOCTEXT("LauncherTaskRunNotification", "Running {GameName} on {DeviceName}..."), Arguments);
+			OnAllPIEInstancesStarted();
 		}
 	}
-	else
-	{
-		bSetNotification = false;
-	}
-
-	if (bSetNotification)
-	{
-		TGraphTask<FLauncherNotificationTask>::CreateTask().ConstructAndDispatchWhenReady(
-			NotificationItemPtr,
-			SNotificationItem::CS_Pending,
-			NotificationText
-		);
-	}
 }
 
-void UEditorEngine::HandleStageCompleted(const FString& InStage, double StageTime, bool bHasCode, TWeakPtr<SNotificationItem> NotificationItemPtr)
+void UEditorEngine::OnAllPIEInstancesStarted()
 {
-	UE_LOG(LogPlayLevel, Log, TEXT("Completed Launch On Stage: %s, Time: %f"), *InStage, StageTime);
+	GiveFocusToLastClientPIEViewport();
 
-	// analytics for launch on
-	TArray<FAnalyticsEventAttribute> ParamArray;
-	ParamArray.Add(FAnalyticsEventAttribute(TEXT("Time"), StageTime));
-	ParamArray.Add(FAnalyticsEventAttribute(TEXT("StageName"), InStage));
-	FEditorAnalytics::ReportEvent(TEXT( "Editor.LaunchOn.StageComplete" ), PlayUsingLauncherDeviceId.Left(PlayUsingLauncherDeviceId.Find(TEXT("@"))), bHasCode, ParamArray);
+	// Print out a log message stating the overall startup time.
+	{
+		FFormatNamedArguments Arguments;
+		Arguments.Add(TEXT("StartTime"), FPlatformTime::Seconds() - PlayInEditorSessionInfo->PlayRequestStartTime);
+		FMessageLog(NAME_CategoryPIE).Info(FText::Format(LOCTEXT("PIETotalStartTime", "Play in editor total start time {StartTime} seconds."), Arguments));
+	}
 }
 
-void UEditorEngine::HandleLaunchCanceled(double TotalTime, bool bHasCode, TWeakPtr<SNotificationItem> NotificationItemPtr)
+void UEditorEngine::GiveFocusToLastClientPIEViewport()
 {
-	TGraphTask<FLauncherNotificationTask>::CreateTask().ConstructAndDispatchWhenReady(
-		NotificationItemPtr,
-		SNotificationItem::CS_Fail,
-		LOCTEXT("LaunchtaskFailedNotification", "Launch canceled!")
-	);
-
-	// analytics for launch on
-	TArray<FAnalyticsEventAttribute> ParamArray;
-	ParamArray.Add(FAnalyticsEventAttribute(TEXT("Time"), TotalTime));
-	FEditorAnalytics::ReportEvent(TEXT( "Editor.LaunchOn.Canceled" ), PlayUsingLauncherDeviceId.Left(PlayUsingLauncherDeviceId.Find(TEXT("@"))), bHasCode, ParamArray);
-
-	bPlayUsingLauncher = false;	
-}
+	// Find the non-dedicated server or last client window to give focus to. We choose the last one
+	// because when launching additional instances, by applying focus to the first one it immediately
+	// pushes the PINW instances behind the editor and you can't tell that they were launched.
 
-void UEditorEngine::HandleLaunchCompleted(bool Succeeded, double TotalTime, int32 ErrorCode, bool bHasCode, TWeakPtr<SNotificationItem> NotificationItemPtr)
-{
-	const FString DummyIOSDeviceName(FString::Printf(TEXT("All_iOS_On_%s"), FPlatformProcess::ComputerName()));
-	const FString DummyTVOSDeviceName(FString::Printf(TEXT("All_tvOS_On_%s"), FPlatformProcess::ComputerName()));
-	if (Succeeded)
+	int32 HighestPIEInstance = TNumericLimits<int32>::Min();
+	UGameViewportClient* ViewportClient = nullptr;
+	for (const FWorldContext& WorldContext : WorldList)
 	{
-		FText CompletionMsg;
-		if ((PlayUsingLauncherDeviceId.Left(PlayUsingLauncherDeviceId.Find(TEXT("@"))) == TEXT("IOS") && PlayUsingLauncherDeviceName.Contains(DummyIOSDeviceName)) ||
-			(PlayUsingLauncherDeviceId.Left(PlayUsingLauncherDeviceId.Find(TEXT("@"))) == TEXT("TVOS") && PlayUsingLauncherDeviceName.Contains(DummyTVOSDeviceName)))
+		if (WorldContext.WorldType == EWorldType::PIE && !WorldContext.RunAsDedicated)
 		{
-			CompletionMsg = LOCTEXT("DeploymentTaskCompleted", "Deployment complete! Open the app on your device to launch.");
+			if (WorldContext.PIEInstance > HighestPIEInstance)
+			{
+				HighestPIEInstance = WorldContext.PIEInstance;
+				ViewportClient = WorldContext.GameViewport;
+			}
 		}
-		else
+	}
+
+	if (ViewportClient && ViewportClient->GetGameViewportWidget().IsValid())
+	{
+		FSlateApplication::Get().RegisterGameViewport(ViewportClient->GetGameViewportWidget().ToSharedRef());
+	}
+
+	// Make sure to focus the game viewport.
+	{
+		const bool bPreviewTypeIsVR = PlayInEditorSessionInfo->OriginalRequestParams.SessionPreviewTypeOverride.Get(EPlaySessionPreviewType::NoPreview) == EPlaySessionPreviewType::VRPreview;
+		const bool bIsVR = IVREditorModule::Get().IsVREditorEnabled() || (bPreviewTypeIsVR && GEngine->XRSystem.IsValid());
+		const bool bGameGetsMouseControl = PlayInEditorSessionInfo->OriginalRequestParams.EditorPlaySettings->GameGetsMouseControl;
+
+		if (PlayInEditorSessionInfo->OriginalRequestParams.WorldType == EPlaySessionWorldType::PlayInEditor && (bGameGetsMouseControl || bIsVR))
 		{
-			CompletionMsg = LOCTEXT("LauncherTaskCompleted", "Launch complete!!");
+			FSlateApplication::Get().SetAllUserFocusToGameViewport();
 		}
+	}
+}
 
-		TGraphTask<FLauncherNotificationTask>::CreateTask().ConstructAndDispatchWhenReady(
-			NotificationItemPtr,
-			SNotificationItem::CS_Success,
-			CompletionMsg
-		);
-
-		// analytics for launch on
-		TArray<FAnalyticsEventAttribute> ParamArray;
-		ParamArray.Add(FAnalyticsEventAttribute(TEXT("Time"), TotalTime));
-		FEditorAnalytics::ReportEvent(TEXT( "Editor.LaunchOn.Completed" ), PlayUsingLauncherDeviceId.Left(PlayUsingLauncherDeviceId.Find(TEXT("@"))), bHasCode, ParamArray);
+void UEditorEngine::RequestLateJoin()
+{
+	if (!ensureMsgf(PlayInEditorSessionInfo.IsSet(), TEXT("RequestLateJoin shouldn't be called if no session is in progress!")))
+	{
+		return;
+	}
 
-		UE_LOG(LogPlayLevel, Log, TEXT("Launch On Completed. Time: %f"), TotalTime);
+	if (!ensureMsgf(PlayInEditorSessionInfo->OriginalRequestParams.WorldType != EPlaySessionWorldType::SimulateInEditor, TEXT("RequestLateJoin shouldn't be called for SIE!")))
+	{
+		return;
 	}
-	else
+
+	bool bUsePIEOnlineAuthentication = false;
+	if (PlayInEditorSessionInfo->bUsingOnlinePlatform)
 	{
-		FText CompletionMsg;
-		if ((PlayUsingLauncherDeviceId.Left(PlayUsingLauncherDeviceId.Find(TEXT("@"))) == TEXT("IOS") && PlayUsingLauncherDeviceName.Contains(DummyIOSDeviceName)) ||
-			(PlayUsingLauncherDeviceId.Left(PlayUsingLauncherDeviceId.Find(TEXT("@"))) == TEXT("TVOS") && PlayUsingLauncherDeviceName.Contains(DummyTVOSDeviceName)))
+		int32 TotalNumDesiredClients = PlayInEditorSessionInfo->NumClientInstancesCreated + 1;
+		bool bHasRequiredLogins = TotalNumDesiredClients <= UOnlineEngineInterface::Get()->GetNumPIELogins();
+		if (bHasRequiredLogins)
 		{
-			CompletionMsg = LOCTEXT("DeploymentTaskFailed", "Deployment failed!");
+			bUsePIEOnlineAuthentication = true;
 		}
 		else
 		{
-			CompletionMsg = LOCTEXT("LauncherTaskFailed", "Launch failed!");
+			FText ErrorMsg = LOCTEXT("PIELateJoinLoginFailure", "Not enough login credentials to add additional PIE instance, change editor settings.");
+			UE_LOG(LogPlayLevel, Warning, TEXT("%s"), *ErrorMsg.ToString());
+			FMessageLog(NAME_CategoryPIE).Warning(ErrorMsg);
+			return;
 		}
-		
-		AsyncTask(ENamedThreads::GameThread, [=]
-		{
-			FMessageLog MessageLog("PackagingResults");
-
-			MessageLog.Error()
-				->AddToken(FTextToken::Create(CompletionMsg))
-				->AddToken(FTextToken::Create(FText::FromString(FEditorAnalytics::TranslateErrorCode(ErrorCode))));
-
-			// flush log, because it won't be destroyed until the notification popup closes
-			MessageLog.NumMessages(EMessageSeverity::Info);
-		});
+	}
 
-		TGraphTask<FLauncherNotificationTask>::CreateTask().ConstructAndDispatchWhenReady(
-			NotificationItemPtr,
-			SNotificationItem::CS_Fail,
-			CompletionMsg
-		);
+	// By the time of Late Join, a server should already be running if needed. So instead we only need
+	// to launch a new instance following the same flow as all of the clients that already joined.
+	EPlayNetMode NetMode;
+	PlayInEditorSessionInfo->OriginalRequestParams.EditorPlaySettings->GetPlayNetMode(NetMode);
 
-		TArray<FAnalyticsEventAttribute> ParamArray;
-		ParamArray.Add(FAnalyticsEventAttribute(TEXT("Time"), TotalTime));
-		FEditorAnalytics::ReportEvent(TEXT( "Editor.LaunchOn.Failed" ), PlayUsingLauncherDeviceId.Left(PlayUsingLauncherDeviceId.Find(TEXT("@"))), bHasCode, ErrorCode, ParamArray);
+	// Adding another instance to a Listen Server will create clients, not more servers.
+	if (NetMode == EPlayNetMode::PIE_ListenServer)
+	{
+		NetMode = EPlayNetMode::PIE_Client;
 	}
-	bPlayUsingLauncher = false;
-}
 
-static void HandleHyperlinkNavigate()
-{
-	FGlobalTabmanager::Get()->InvokeTab(FName("OutputLog"));
+	CreateNewPlayInEditorInstance(PlayInEditorSessionInfo->OriginalRequestParams, false, NetMode);
 }
 
-struct FInternalPlayLevelUtils
+void UEditorEngine::CreateNewPlayInEditorInstance(FRequestPlaySessionParams &InRequestParams, const bool bInDedicatedInstance, const EPlayNetMode InNetMode)
 {
-	static int32 ResolveDirtyBlueprints(const bool bPromptForCompile, TArray<UBlueprint*>& ErroredBlueprints, const bool bForceLevelScriptRecompile = true)
+	bool bUserWantsSingleProcess;
+	InRequestParams.EditorPlaySettings->GetRunUnderOneProcess(bUserWantsSingleProcess);
+
+	// If they don't want to use a single process, we can only launch one client inside the editor.
+	if (!bUserWantsSingleProcess && PlayInEditorSessionInfo->NumClientInstancesCreated > 0)
 	{
-		struct FLocal
+		if (bInDedicatedInstance)
 		{
-			static void OnMessageLogLinkActivated(const class TSharedRef<IMessageToken>& Token)
-			{
-				if (Token->GetType() == EMessageToken::Object)
-				{
-					const TSharedRef<FUObjectToken> UObjectToken = StaticCastSharedRef<FUObjectToken>(Token);
-					if (UObjectToken->GetObject().IsValid())
-					{
-						FKismetEditorUtilities::BringKismetToFocusAttentionOnObject(UObjectToken->GetObject().Get());
-					}
-				}
-			}
+			// If they needed a server and they need multiple processes, it should
+			// have already been launched by this time (by the block above). This means
+			// this code should only ever be launching clients (but the code below can
+			// launch either clients, listen servers, or dedicated servers.)
+			check(PlayInEditorSessionInfo->bServerWasLaunched);
+		}
 
-			static void AddCompileErrorToLog(UBlueprint* ErroredBlueprint, FMessageLog& BlueprintLog)
-			{
-				FFormatNamedArguments Arguments;
-				Arguments.Add(TEXT("Name"), FText::FromString(ErroredBlueprint->GetName()));
+		const bool bIsDedicatedServer = false;
+		const bool bIsHost = false;
+		LaunchNewProcess(InRequestParams, PlayInEditorSessionInfo->NumClientInstancesCreated, bIsHost, bIsDedicatedServer);
+		PlayInEditorSessionInfo->NumClientInstancesCreated++;
+	}
+	else
+	{
+		FPieLoginStruct PIELoginInfo;
+		PIELoginInfo.PIEStartTime = FPlatformTime::Seconds();
 
-				TSharedRef<FTokenizedMessage> Message = FTokenizedMessage::Create(EMessageSeverity::Warning);
-				Message->AddToken(FTextToken::Create(LOCTEXT("BlueprintCompileFailed", "Blueprint failed to compile: ")));
-				Message->AddToken(FUObjectToken::Create(ErroredBlueprint, FText::FromString(ErroredBlueprint->GetName()))
-					->OnMessageTokenActivated(FOnMessageTokenActivated::CreateStatic(&FLocal::OnMessageLogLinkActivated))
-				);
+		FGameInstancePIEParameters GameInstancePIEParameters;
+		GameInstancePIEParameters.bAnyBlueprintErrors = PlayInEditorSessionInfo->bAnyBlueprintErrors;
+		// If they require a server and one hasn't been launched then it is dedicated. If they're a client or listen server
+		// then it doesn't count as a dedicated server so this can be false (NetMode will handle ListenServer).
+		GameInstancePIEParameters.bRunAsDedicated = bInDedicatedInstance;
+		GameInstancePIEParameters.bSimulateInEditor = InRequestParams.WorldType == EPlaySessionWorldType::SimulateInEditor;
+		GameInstancePIEParameters.bStartInSpectatorMode = PlayInEditorSessionInfo->bStartedInSpectatorMode;
+		GameInstancePIEParameters.EditorPlaySettings = PlayInEditorSessionInfo->OriginalRequestParams.EditorPlaySettings;
+		GameInstancePIEParameters.WorldFeatureLevel = PreviewPlatform.GetEffectivePreviewFeatureLevel();
+		GameInstancePIEParameters.NetMode = InNetMode;
 
-				BlueprintLog.AddMessage(Message);
-			}
-		};
+		PIELoginInfo.GameInstancePIEParameters = GameInstancePIEParameters;
 
-		const bool bAutoCompile = !bPromptForCompile;
-		FString PromptDirtyList;
+		// Create a World Context for our client.
+		FWorldContext &PieWorldContext = CreateNewWorldContext(EWorldType::PIE);
+		PieWorldContext.PIEInstance = PlayInEditorSessionInfo->PIEInstanceCount++;
+		PieWorldContext.bWaitingOnOnlineSubsystem = true;
 
-		TArray<UBlueprint*> InNeedOfRecompile;
-		ErroredBlueprints.Empty();
+		PIELoginInfo.WorldContextHandle = PieWorldContext.ContextHandle;
+		PIELoginInfo.PIEInstanceIndex = PieWorldContext.PIEInstance;
 
-		FMessageLog BlueprintLog("BlueprintLog");
+		// We increment how many PIE instances we think are going to boot up as it is decremented
+		// in OnLoginPIEComplete_Deferred and used to check if all instances have booted up.
+		PlayInEditorSessionInfo->NumOutstandingPIELogins++;
 
-		double BPRegenStartTime = FPlatformTime::Seconds();
-		for (TObjectIterator<UBlueprint> BlueprintIt; BlueprintIt; ++BlueprintIt)
+		// If they are using an Online Subsystem that requires log-in then they require async
+		// log-in so we use a deferred log-in approach.
+		if (PlayInEditorSessionInfo->bUsingOnlinePlatform)
 		{
-			UBlueprint* Blueprint = *BlueprintIt;
-
-			// ignore up-to-date BPs
-			if (Blueprint->IsUpToDate())
-			{
-				continue;
-			}
+			FName OnlineIdentifier = UOnlineEngineInterface::Get()->GetOnlineIdentifier(PieWorldContext);
+			UE_LOG(LogPlayLevel, Display, TEXT("Creating online subsystem for client %s"), *OnlineIdentifier.ToString());
 
-			// do not try to recompile BPs that have not changed since they last failed to compile, so don't check Blueprint->IsUpToDate()
-			const bool bIsDirtyAndShouldBeRecompiled = Blueprint->IsPossiblyDirty();
-			if (!FBlueprintEditorUtils::IsDataOnlyBlueprint(Blueprint)
-				&& (bIsDirtyAndShouldBeRecompiled || (FBlueprintEditorUtils::IsLevelScriptBlueprint(Blueprint) && bForceLevelScriptRecompile))
-				&& (Blueprint->Status != BS_Unknown)
-				&& !Blueprint->IsPendingKill())
+			if (GameInstancePIEParameters.bRunAsDedicated)
 			{
-				InNeedOfRecompile.Add(Blueprint);
+				// Dedicated servers don't use a login
+				UOnlineEngineInterface::Get()->SetForceDedicated(OnlineIdentifier, true);
 
-				if (bPromptForCompile)
-				{
-					PromptDirtyList += FString::Printf(TEXT("\n   %s"), *Blueprint->GetName());
-				}
+				OnLoginPIEComplete_Deferred(0, true, FString(), PIELoginInfo);
 			}
-			else if (BS_Error == Blueprint->Status && Blueprint->bDisplayCompilePIEWarning)
+			else
 			{
-				ErroredBlueprints.Add(Blueprint);
-				FLocal::AddCompileErrorToLog(Blueprint, BlueprintLog);
+				// Login to Online platform before creating world
+				FOnPIELoginComplete OnPIELoginCompleteDelegate;
+				OnPIELoginCompleteDelegate.BindUObject(this, &UEditorEngine::OnLoginPIEComplete, PIELoginInfo);
+
+				// Kick off an async request to the Online Interface to log-in and, on completion, finish client creation.
+				UOnlineEngineInterface::Get()->LoginPIEInstance(OnlineIdentifier, 0, PlayInEditorSessionInfo->NumClientInstancesCreated, OnPIELoginCompleteDelegate);
 			}
 		}
-
-		bool bRunCompilation = bAutoCompile;
-		if (bPromptForCompile)
+		else
 		{
-			FFormatNamedArguments Args;
-			Args.Add(TEXT("DirtyBlueprints"), FText::FromString(PromptDirtyList));
-			const FText PromptMsg = FText::Format(NSLOCTEXT("PlayInEditor", "PrePIE_BlueprintsDirty", "One or more blueprints have been modified without being recompiled. Do you want to compile them now? \n{DirtyBlueprints}"), Args);
-
-			EAppReturnType::Type PromptResponse = FMessageDialog::Open(EAppMsgType::YesNo, PromptMsg);
-			bRunCompilation = (PromptResponse == EAppReturnType::Yes);
+			// Otherwise, we can create the client immediately. We'll emulate the PIE login structure so we can just
+			// use the same flow as the deferred log-in, but skipping the deferred part.
+			OnLoginPIEComplete_Deferred(0, true, FString(), PIELoginInfo);
 		}
-		int32 RecompiledCount = 0;
 
-		if (bRunCompilation && (InNeedOfRecompile.Num() > 0))
+		// Only count non-dedicated instances as clients so that our indexes line up for log-ins.
+		if (!GameInstancePIEParameters.bRunAsDedicated)
 		{
-			const FText LogPageLabel = (bAutoCompile) ? LOCTEXT("BlueprintAutoCompilationPageLabel", "Pre-Play auto-recompile") :
-				LOCTEXT("BlueprintCompilationPageLabel", "Pre-Play recompile");
-			BlueprintLog.NewPage(LogPageLabel);
+			PlayInEditorSessionInfo->NumClientInstancesCreated++;
+		}
+	}
+}
 
-			TArray<UBlueprint*> CompiledBlueprints;
-			auto OnBlueprintPreCompileLambda = [&CompiledBlueprints](UBlueprint* InBlueprint)
-			{
-				check(InBlueprint != nullptr);
+class SPIEViewport : public SViewport
+{
+	SLATE_BEGIN_ARGS(SPIEViewport)
+		: _Content()
+		, _RenderDirectlyToWindow(false)
+		, _EnableStereoRendering(false)
+		, _IgnoreTextureAlpha(true)
+	{
+		_Clipping = EWidgetClipping::ClipToBoundsAlways;
+	}
 
-				if (CompiledBlueprints.Num() == 0)
-				{
-					UE_LOG(LogPlayLevel, Log, TEXT("[PlayLevel] Compiling %s before play..."), *InBlueprint->GetName());
-				}
-				else
-				{
-					UE_LOG(LogPlayLevel, Log, TEXT("[PlayLevel]   Compiling %s as a dependent..."), *InBlueprint->GetName());
-				}
+		SLATE_DEFAULT_SLOT(FArguments, Content)
 
-				CompiledBlueprints.Add(InBlueprint);
-			};
+		/**
+		 * Whether or not to render directly to the window's backbuffer or an offscreen render target that is applied to the window later
+		 * Rendering to an offscreen target is the most common option in the editor where there may be many frames which this viewport's interface may wish to not re-render but use a cached buffer instead
+		 * Rendering directly to the backbuffer is the most common option in the game where you want to update each frame without the cost of writing to an intermediate target first.
+		 */
+		SLATE_ARGUMENT(bool, RenderDirectlyToWindow)
 
-			// Register compile callback
-			FDelegateHandle PreCompileDelegateHandle = GEditor->OnBlueprintPreCompile().AddLambda(OnBlueprintPreCompileLambda);
+		/** Whether or not to enable stereo rendering. */
+		SLATE_ARGUMENT(bool, EnableStereoRendering )
 
-			// Recompile all necessary blueprints in a single loop, saving GC until the end
-			for (auto BlueprintIt = InNeedOfRecompile.CreateIterator(); BlueprintIt; ++BlueprintIt)
-			{
-				UBlueprint* Blueprint = *BlueprintIt;
+		/**
+		 * If true, the viewport's texture alpha is ignored when performing blending.  In this case only the viewport tint opacity is used
+		 * If false, the texture alpha is used during blending
+		 */
+		SLATE_ARGUMENT( bool, IgnoreTextureAlpha )
 
-				int32 CurrItIndex = BlueprintIt.GetIndex();
+	SLATE_END_ARGS()
 
-				// Compile the Blueprint (note: re-instancing may trigger additional compiles for child/dependent Blueprints; see callback above)
-				FKismetEditorUtilities::CompileBlueprint(Blueprint, EBlueprintCompileOptions::SkipGarbageCollection);
+	void Construct(const FArguments& InArgs)
+	{
+		SViewport::Construct(
+			SViewport::FArguments()
+			.EnableGammaCorrection(false) // Gamma correction in the game is handled in post processing in the scene renderer
+			.RenderDirectlyToWindow(InArgs._RenderDirectlyToWindow)
+			.EnableStereoRendering(InArgs._EnableStereoRendering)
+			.IgnoreTextureAlpha(InArgs._IgnoreTextureAlpha)
+			[
+				InArgs._Content.Widget
+			]
+		);
+	}
 
-				// Check for errors after compiling
-				for (UBlueprint* CompiledBlueprint : CompiledBlueprints)
-				{
-					if (CompiledBlueprint != Blueprint)
-					{
-						int32 ExistingIndex = InNeedOfRecompile.Find(CompiledBlueprint);
-						// if this dependent blueprint is already set up to compile 
-						// later in this loop, then there is no need to add it to be recompiled again
-						if (ExistingIndex > CurrItIndex)
-						{
-							InNeedOfRecompile.RemoveAt(ExistingIndex);
-						}
-					}
+	virtual void Tick(const FGeometry& AllottedGeometry, const double InCurrentTime, const float InDeltaTime) override
+	{
+		SViewport::Tick(AllottedGeometry, InCurrentTime, InDeltaTime);
 
-					const bool bHadError = (!CompiledBlueprint->IsUpToDate() && CompiledBlueprint->Status != BS_Unknown);
+		// Rather than binding the attribute we're going to poll it in tick, otherwise we will make this widget volatile, and it therefore
+		// wont be possible to cache it or its children in GSlateEnableGlobalInvalidation mode.
+		SetEnabled(FSlateApplication::Get().GetNormalExecutionAttribute().Get());
+	}
+};
 
-					// Check if the Blueprint has already been added to the error list to prevent it from being added again
-					if (bHadError && ErroredBlueprints.Find(CompiledBlueprint) == INDEX_NONE)
-					{
-						ErroredBlueprints.Add(CompiledBlueprint);
-						FLocal::AddCompileErrorToLog(CompiledBlueprint, BlueprintLog);
-					}
+void UEditorEngine::OnViewportCloseRequested(FViewport* InViewport)
+{
+	RequestEndPlayMap();
+}
 
-					++RecompiledCount;
-				}
+FSceneViewport* UEditorEngine::GetGameSceneViewport(UGameViewportClient* ViewportClient) const
+{
+	return ViewportClient->GetGameViewport();
+}
 
-				// Reset for next pass
-				CompiledBlueprints.Empty();
-			}
+FViewport* UEditorEngine::GetActiveViewport()
+{
+	// Get the Level editor module and request the Active Viewport.
+	FLevelEditorModule& LevelEditorModule = FModuleManager::Get().GetModuleChecked<FLevelEditorModule>( TEXT("LevelEditor") );
 
-			// Now that all Blueprints have been compiled, run a single GC pass to clean up artifacts
-			CollectGarbage(GARBAGE_COLLECTION_KEEPFLAGS);
+	TSharedPtr<IAssetViewport> ActiveLevelViewport = LevelEditorModule.GetFirstActiveViewport();
 
-			// Unregister compile callback
-			GEditor->OnBlueprintPreCompile().Remove(PreCompileDelegateHandle);
+	if ( ActiveLevelViewport.IsValid() )
+	{
+		return ActiveLevelViewport->GetActiveViewport();
+	}
+	
+	return nullptr;
+}
 
-			UE_LOG(LogPlayLevel, Log, TEXT("PlayLevel: Blueprint regeneration took %d ms (%i blueprints)"), (int32)((FPlatformTime::Seconds() - BPRegenStartTime) * 1000), RecompiledCount);
-		}
-		else if (bAutoCompile)
+FViewport* UEditorEngine::GetPIEViewport()
+{
+	// Check both cases where the PIE viewport may be, otherwise return NULL if none are found.
+	if( GameViewport )
+	{
+		return GameViewport->Viewport;
+	}
+	else
+	{
+		for (const FWorldContext& WorldContext : WorldList)
 		{
-			UE_LOG(LogPlayLevel, Log, TEXT("PlayLevel: No blueprints needed recompiling"));
+			if (WorldContext.WorldType == EWorldType::PIE)
+			{
+				// We can't use FindChecked here because when using the dedicated server option we don't initialize this map 
+				//	(we don't use a viewport for the PIE context in this case)
+				FSlatePlayInEditorInfo * SlatePlayInEditorSessionPtr = SlatePlayInEditorMap.Find(WorldContext.ContextHandle);
+				if ( SlatePlayInEditorSessionPtr != nullptr && SlatePlayInEditorSessionPtr->SlatePlayInEditorWindowViewport.IsValid() )
+				{
+					return SlatePlayInEditorSessionPtr->SlatePlayInEditorWindowViewport.Get();
+				}
+			}
 		}
-
-		return RecompiledCount;
 	}
-};
 
-FString UEditorEngine::GetPlayOnTargetPlatformName() const
-{
-	return PlayUsingLauncherDeviceId.Left(PlayUsingLauncherDeviceId.Find(TEXT("@")));
+	return nullptr;
 }
 
-void UEditorEngine::PlayUsingLauncher()
+bool UEditorEngine::GetSimulateInEditorViewTransform(FTransform& OutViewTransform) const
 {
-	if (!PlayUsingLauncherDeviceId.IsEmpty())
+	if (PlayInEditorSessionInfo->OriginalRequestParams.WorldType == EPlaySessionWorldType::SimulateInEditor)
 	{
-		ILauncherServicesModule& LauncherServicesModule = FModuleManager::LoadModuleChecked<ILauncherServicesModule>(TEXT("LauncherServices"));
-		ITargetDeviceServicesModule& TargetDeviceServicesModule = FModuleManager::LoadModuleChecked<ITargetDeviceServicesModule>("TargetDeviceServices");
-
-		//if the device is not authorized to be launched to, we need to pop an error instead of trying to launch
-		FString LaunchPlatformName = PlayUsingLauncherDeviceId.Left(PlayUsingLauncherDeviceId.Find(TEXT("@")));
-		FString LaunchPlatformNameFromID = PlayUsingLauncherDeviceId.Right(PlayUsingLauncherDeviceId.Find(TEXT("@")));
-		ITargetPlatform* LaunchPlatform = GetTargetPlatformManagerRef().FindTargetPlatform(LaunchPlatformName);
-
-		// create a temporary device group and launcher profile
-		ILauncherDeviceGroupRef DeviceGroup = LauncherServicesModule.CreateDeviceGroup(FGuid::NewGuid(), TEXT("PlayOnDevices"));
-		if (LaunchPlatform != nullptr)
+		// The first PIE world context is the one that can toggle between PIE and SIE
+		for (const FWorldContext& WorldContext : WorldList)
 		{
-			if (LaunchPlatformNameFromID.Equals(LaunchPlatformName))
+			if (WorldContext.WorldType == EWorldType::PIE && !WorldContext.RunAsDedicated)
 			{
-				// create a temporary list of devices for the target platform
-				TArray<ITargetDevicePtr> TargetDevices;
-				LaunchPlatform->GetAllDevices(TargetDevices);
-
-				for (const ITargetDevicePtr& PlayDevice : TargetDevices)
+				const FSlatePlayInEditorInfo* SlateInfoPtr = SlatePlayInEditorMap.Find(WorldContext.ContextHandle);
+				if (SlateInfoPtr)
 				{
-					// compose the device id
-					FString PlayDeviceId = LaunchPlatformName + TEXT("@") + PlayDevice.Get()->GetId().GetDeviceName();
-					if (PlayDevice.IsValid() && !PlayDevice->IsAuthorized())
-					{
-						CancelPlayUsingLauncher();
-					}
-					else
+					// This is only supported inside SLevelEditor viewports currently
+					TSharedPtr<IAssetViewport> LevelViewport = SlateInfoPtr->DestinationSlateViewport.Pin();
+					if (LevelViewport.IsValid())
 					{
-						DeviceGroup->AddDevice(PlayDeviceId);
-						UE_LOG(LogPlayLevel, Log, TEXT("Launcher Device ID: %s"), *PlayDeviceId);
+						FEditorViewportClient& EditorViewportClient = LevelViewport->GetAssetViewportClient();
+						OutViewTransform = FTransform(EditorViewportClient.GetViewRotation(), EditorViewportClient.GetViewLocation());
+						return true;
 					}
 				}
-			}
-			else
-			{
-				ITargetDevicePtr PlayDevice = LaunchPlatform->GetDefaultDevice();
-				if (PlayDevice.IsValid() && !PlayDevice->IsAuthorized())
-				{
-					CancelPlayUsingLauncher();
-				}
-				else
-				{
-					
-					DeviceGroup->AddDevice(PlayUsingLauncherDeviceId);
-					UE_LOG(LogPlayLevel, Log, TEXT("Launcher Device ID: %s"), *PlayUsingLauncherDeviceId);
-				}
-			}
-			
-			if (DeviceGroup.Get().GetNumDevices() == 0)
-			{
-				return;
+				break;
 			}
 		}
+	}
+	return false;
+}
 
-		// does the project have any code?
-		FGameProjectGenerationModule& GameProjectModule = FModuleManager::LoadModuleChecked<FGameProjectGenerationModule>(TEXT("GameProjectGeneration"));
-		bPlayUsingLauncherHasCode = GameProjectModule.Get().ProjectHasCodeFiles();
+void UEditorEngine::ToggleBetweenPIEandSIE( bool bNewSession )
+{
+	bIsToggleBetweenPIEandSIEQueued = false;
 
-		const ULevelEditorPlaySettings* PlayInSettings = GetDefault<ULevelEditorPlaySettings>();
-		// Setup launch profile, keep the setting here to a minimum.
-		ILauncherProfileRef LauncherProfile = LauncherServicesModule.CreateProfile(TEXT("Launch On Device"));
-		if(PlayInSettings->BuildGameBeforeLaunch == EPlayOnBuildMode::PlayOnBuild_Always)
-		{
-			bPlayUsingLauncherBuild = true;
-		}
-		else if(PlayInSettings->BuildGameBeforeLaunch == EPlayOnBuildMode::PlayOnBuild_Never)
-		{
-			bPlayUsingLauncherBuild = false;
-		}
-		else if(PlayInSettings->BuildGameBeforeLaunch == EPlayOnBuildMode::PlayOnBuild_Default)
-		{
-			bPlayUsingLauncherBuild = bPlayUsingLauncherHasCode || !FApp::GetEngineIsPromotedBuild();
-		}
-		else if(PlayInSettings->BuildGameBeforeLaunch == EPlayOnBuildMode::PlayOnBuild_IfEditorBuiltLocally)
-		{
-			bPlayUsingLauncherBuild = !FApp::GetEngineIsPromotedBuild();
-		}
-		LauncherProfile->SetBuildGame(bPlayUsingLauncherBuild);
+	FEditorDelegates::OnPreSwitchBeginPIEAndSIE.Broadcast(bIsSimulatingInEditor);
 
-		// set the build/launch configuration 
-		switch (PlayInSettings->LaunchConfiguration)
+	// The first PIE world context is the one that can toggle between PIE and SIE
+	// Network PIE/SIE toggling is not really meant to be supported.
+	FSlatePlayInEditorInfo * SlateInfoPtr = nullptr;
+	for (const FWorldContext& WorldContext : WorldList)
+	{
+		if (WorldContext.WorldType == EWorldType::PIE && !WorldContext.RunAsDedicated)
 		{
-		case LaunchConfig_Debug:
-			LauncherProfile->SetBuildConfiguration(EBuildConfiguration::Debug);
-			break;
-		case LaunchConfig_Development:
-			LauncherProfile->SetBuildConfiguration(EBuildConfiguration::Development);
-			break;
-		case LaunchConfig_Test:
-			LauncherProfile->SetBuildConfiguration(EBuildConfiguration::Test);
-			break;
-		case LaunchConfig_Shipping:
-			LauncherProfile->SetBuildConfiguration(EBuildConfiguration::Shipping);
-			break;
-		default:
-			// same as the running editor
-			LauncherProfile->SetBuildConfiguration(FApp::GetBuildConfiguration());
+			SlateInfoPtr = SlatePlayInEditorMap.Find(WorldContext.ContextHandle);
 			break;
 		}
+	}
 
-		// select the quickest cook mode based on which in editor cook mode is enabled
-		bool bIncrimentalCooking = true;
-		LauncherProfile->AddCookedPlatform(LaunchPlatformName);
-		ELauncherProfileCookModes::Type CurrentLauncherCookMode = ELauncherProfileCookModes::ByTheBook;
-		bool bCanCookByTheBookInEditor = true;
-		bool bCanCookOnTheFlyInEditor = true;
-		for ( const FString& PlatformName : LauncherProfile->GetCookedPlatforms() )
-		{
-			if ( CanCookByTheBookInEditor(PlatformName) == false )
-			{
-				bCanCookByTheBookInEditor = false;
-			}
-			if ( CanCookOnTheFlyInEditor(PlatformName)== false )
-			{
-				bCanCookOnTheFlyInEditor = false;
-			}
-		}
-		if ( bCanCookByTheBookInEditor )
-		{
-			CurrentLauncherCookMode = ELauncherProfileCookModes::ByTheBookInEditor;
-		}
-		if ( bCanCookOnTheFlyInEditor )
-		{
-			CurrentLauncherCookMode = ELauncherProfileCookModes::OnTheFlyInEditor;
-			bIncrimentalCooking = false;
-		}
-		if ( GetDefault<UCookerSettings>()->bCookOnTheFlyForLaunchOn )
-		{
-			CurrentLauncherCookMode = ELauncherProfileCookModes::OnTheFly;
-			bIncrimentalCooking = false;
-		}
-		LauncherProfile->SetCookMode( CurrentLauncherCookMode );
-		LauncherProfile->SetUnversionedCooking(!bIncrimentalCooking);
-		LauncherProfile->SetIncrementalCooking(bIncrimentalCooking);
-		LauncherProfile->SetDeployedDeviceGroup(DeviceGroup);
-		LauncherProfile->SetIncrementalDeploying(bIncrimentalCooking);
-		LauncherProfile->SetEditorExe(FUnrealEdMisc::Get().GetExecutableForCommandlets());
-
-		const FString DummyIOSDeviceName(FString::Printf(TEXT("All_iOS_On_%s"), FPlatformProcess::ComputerName()));
-		const FString DummyTVOSDeviceName(FString::Printf(TEXT("All_tvOS_On_%s"), FPlatformProcess::ComputerName()));
-
-		if ((LaunchPlatformName != TEXT("IOS") && LaunchPlatformName != TEXT("TVOS")) ||
-			(!PlayUsingLauncherDeviceName.Contains(DummyIOSDeviceName) && !PlayUsingLauncherDeviceName.Contains(DummyTVOSDeviceName)))
-		{
-			LauncherProfile->SetLaunchMode(ELauncherProfileLaunchModes::DefaultRole);
-		}
+	if (!SlateInfoPtr)
+	{
+		return;
+	}
 
-		if ( LauncherProfile->GetCookMode() == ELauncherProfileCookModes::OnTheFlyInEditor || LauncherProfile->GetCookMode() == ELauncherProfileCookModes::OnTheFly )
-		{
-			LauncherProfile->SetDeploymentMode(ELauncherProfileDeploymentModes::FileServer);
-		}
+	if( FEngineAnalytics::IsAvailable() && !bNewSession )
+	{
+		FString ToggleType = bIsSimulatingInEditor ? TEXT("SIEtoPIE") : TEXT("PIEtoSIE");
 
-		TArray<UBlueprint*> ErroredBlueprints;
-		FInternalPlayLevelUtils::ResolveDirtyBlueprints(!PlayInSettings->bAutoCompileBlueprintsOnLaunch, ErroredBlueprints, false);
+		FEngineAnalytics::GetProvider().RecordEvent(TEXT("Editor.Usage.PIE"), TEXT("ToggleBetweenPIEandSIE"), ToggleType );
+	}
 
-		TArray<FString> MapNames;
-		FWorldContext & EditorContext = GetEditorWorldContext();
+	FSlatePlayInEditorInfo & SlatePlayInEditorSession = *SlateInfoPtr;
 
-		// Load maps in place as we saved them above
-		FString EditorMapName = EditorContext.World()->GetOutermost()->GetName();
-		MapNames.Add(EditorMapName);
-	
-		FString InitialMapName;
-		if (MapNames.Num() > 0)
+	// This is only supported inside SLevelEditor viewports currently
+	TSharedPtr<IAssetViewport> LevelViewport = SlatePlayInEditorSession.DestinationSlateViewport.Pin();
+	if( ensure(LevelViewport.IsValid()) )
+	{
+		FEditorViewportClient& EditorViewportClient = LevelViewport->GetAssetViewportClient();
+
+		// Toggle to pie if currently simulating
+		if( bIsSimulatingInEditor )
 		{
-			InitialMapName = MapNames[0];
-		}
+			// The undo system may have a reference to a SIE object that is about to be destroyed, so clear the transactions
+			ResetTransaction( NSLOCTEXT("UnrealEd", "ToggleBetweenPIEandSIE", "Toggle Between PIE and SIE") );
 
-		LauncherProfile->GetDefaultLaunchRole()->SetInitialMap(InitialMapName);
+			// The Game's viewport needs to know about the change away from simluate before the PC is (potentially) created
+			GameViewport->GetGameViewport()->SetPlayInEditorIsSimulate(false);
 
-		for (const FString& MapName : MapNames)
-		{
-			LauncherProfile->AddCookedMap(MapName);
-		}
+			// The editor viewport client wont be visible so temporarily disable it being realtime
+			EditorViewportClient.SetRealtime( false, true );
 
-		if ( LauncherProfile->GetCookMode() == ELauncherProfileCookModes::ByTheBookInEditor )
-		{
-			TArray<ITargetPlatform*> TargetPlatforms;
-			for ( const FString& PlatformName : LauncherProfile->GetCookedPlatforms() )
+			if (!SlatePlayInEditorSession.EditorPlayer.IsValid())
 			{
-				ITargetPlatform* TargetPlatform = GetTargetPlatformManager()->FindTargetPlatform(PlatformName);
-				// todo pass in all the target platforms instead of just the single platform
-				// crashes if two requests are inflight but we can support having multiple platforms cooking at once
-				TargetPlatforms.Add( TargetPlatform ); 
-			}
-			const TArray<FString> &CookedMaps = LauncherProfile->GetCookedMaps();
-
-			// const TArray<FString>& CookedMaps = ChainState.Profile->GetCookedMaps();
-			TArray<FString> CookDirectories;
-			TArray<FString> IniMapSections;
+				OnSwitchWorldsForPIE(true);
 
-			StartCookByTheBookInEditor(TargetPlatforms, CookedMaps, CookDirectories, GetDefault<UProjectPackagingSettings>()->CulturesToStage, IniMapSections );
+				UWorld* World = GameViewport->GetWorld();
+				AGameModeBase* AuthGameMode = World->GetAuthGameMode();
+				if (AuthGameMode && GameViewport->GetGameInstance())	// If there is no GameMode, we are probably the client and cannot RestartPlayer.
+				{
+					AuthGameMode->SpawnPlayerFromSimulate(EditorViewportClient.GetViewLocation(), EditorViewportClient.GetViewRotation());
+				}
 
-			FIsCookFinishedDelegate &CookerFinishedDelegate = LauncherProfile->OnIsCookFinished();
+				OnSwitchWorldsForPIE(false);
+			}
 
-			CookerFinishedDelegate.BindUObject(this, &UEditorEngine::IsCookByTheBookInEditorFinished);
+			// A game viewport already exists, tell the level viewport its in to swap to it
+			LevelViewport->SwapViewportsForPlayInEditor();
 
-			FCookCanceledDelegate &CookCancelledDelegate = LauncherProfile->OnCookCanceled();
+			// No longer simulating
+			GameViewport->SetIsSimulateInEditorViewport(false);
+			EditorViewportClient.SetIsSimulateInEditorViewport(false);
 
-			CookCancelledDelegate.BindUObject(this, &UEditorEngine::CancelCookByTheBookInEditor);
+			FEditorModeRegistry::Get().UnregisterMode(FBuiltinEditorModes::EM_Physics);
+			
+			bIsSimulatingInEditor = false;
 		}
+		else
+		{
+			// Swap to simulate from PIE
+			LevelViewport->SwapViewportsForSimulateInEditor();
+	
+			GameViewport->SetIsSimulateInEditorViewport(true);
+			GameViewport->GetGameViewport()->SetPlayInEditorIsSimulate(true);
+			EditorViewportClient.SetIsSimulateInEditorViewport(true);
 
-		ILauncherPtr Launcher = LauncherServicesModule.CreateLauncher();
-		GEditor->LauncherWorker = Launcher->Launch(TargetDeviceServicesModule.GetDeviceProxyManager(), LauncherProfile);
-
-		// create notification item
-		FText LaunchingText = LOCTEXT("LauncherTaskInProgressNotificationNoDevice", "Launching...");
-		FNotificationInfo Info(LaunchingText);
-
-		Info.Image = FEditorStyle::GetBrush(TEXT("MainFrame.CookContent"));
-		Info.bFireAndForget = false;
-		Info.ExpireDuration = 10.0f;
-		Info.Hyperlink = FSimpleDelegate::CreateStatic(HandleHyperlinkNavigate);
-		Info.HyperlinkText = LOCTEXT("ShowOutputLogHyperlink", "Show Output Log");
-		Info.ButtonDetails.Add(
-			FNotificationButtonInfo(
-				LOCTEXT("LauncherTaskCancel", "Cancel"),
-				LOCTEXT("LauncherTaskCancelToolTip", "Cancels execution of this task."),
-				FSimpleDelegate::CreateStatic(HandleCancelButtonClicked, GEditor->LauncherWorker)
-			)
-		);
+		
+			TSharedRef<FPhysicsManipulationEdModeFactory> Factory = MakeShareable(new FPhysicsManipulationEdModeFactory);
+			FEditorModeRegistry::Get().RegisterMode(FBuiltinEditorModes::EM_Physics, Factory);
+			
+			bIsSimulatingInEditor = true;
 
-		TSharedPtr<SNotificationItem> NotificationItem = FSlateNotificationManager::Get().AddNotification(Info);
+			// Make sure the viewport is in real-time mode
+			EditorViewportClient.SetRealtime( true );
 
-		if (!NotificationItem.IsValid())
-		{
-			return;
-		}
+			// The Simulate window should show stats
+			EditorViewportClient.SetShowStats( true );
 
-		// analytics for launch on
-		int32 ErrorCode = 0;
-		FEditorAnalytics::ReportEvent(TEXT( "Editor.LaunchOn.Started" ), PlayUsingLauncherDeviceId.Left(PlayUsingLauncherDeviceId.Find(TEXT("@"))), bPlayUsingLauncherHasCode);
+			if( SlatePlayInEditorSession.EditorPlayer.IsValid() )
+			{
+				// Move the editor camera to where the player was.  
+				FVector ViewLocation;
+				FRotator ViewRotation;
+				SlatePlayInEditorSession.EditorPlayer.Get()->PlayerController->GetPlayerViewPoint( ViewLocation, ViewRotation );
+				EditorViewportClient.SetViewLocation( ViewLocation );
 
-		NotificationItem->SetCompletionState(SNotificationItem::CS_Pending);
-		
-		TWeakPtr<SNotificationItem> NotificationItemPtr(NotificationItem);
-		if (GEditor->LauncherWorker.IsValid() && GEditor->LauncherWorker->GetStatus() != ELauncherWorkerStatus::Completed)
-		{
-			GEditor->PlayEditorSound(TEXT("/Engine/EditorSounds/Notifications/CompileStart_Cue.CompileStart_Cue"));
-			GEditor->LauncherWorker->OnOutputReceived().AddStatic(HandleOutputReceived);
-			GEditor->LauncherWorker->OnStageStarted().AddUObject(this, &UEditorEngine::HandleStageStarted, NotificationItemPtr);
-			GEditor->LauncherWorker->OnStageCompleted().AddUObject(this, &UEditorEngine::HandleStageCompleted, bPlayUsingLauncherHasCode, NotificationItemPtr);
-			GEditor->LauncherWorker->OnCompleted().AddUObject(this, &UEditorEngine::HandleLaunchCompleted, bPlayUsingLauncherHasCode, NotificationItemPtr);
-			GEditor->LauncherWorker->OnCanceled().AddUObject(this, &UEditorEngine::HandleLaunchCanceled, bPlayUsingLauncherHasCode, NotificationItemPtr);
-		}
-		else
-		{
-			GEditor->LauncherWorker.Reset();
-			GEditor->PlayEditorSound(TEXT("/Engine/EditorSounds/Notifications/CompileFailed_Cue.CompileFailed_Cue"));
-
-			NotificationItem->SetText(LOCTEXT("LauncherTaskFailedNotification", "Failed to launch task!"));
-			NotificationItem->SetCompletionState(SNotificationItem::CS_Fail);
-			NotificationItem->ExpireAndFadeout();
-			bPlayUsingLauncher = false;
-
-			// analytics for launch on
-			TArray<FAnalyticsEventAttribute> ParamArray;
-			ParamArray.Add(FAnalyticsEventAttribute(TEXT("Time"), 0.0));
-			FEditorAnalytics::ReportEvent(TEXT( "Editor.LaunchOn.Failed" ), PlayUsingLauncherDeviceId.Left(PlayUsingLauncherDeviceId.Find(TEXT("@"))), bPlayUsingLauncherHasCode, EAnalyticsErrorCodes::LauncherFailed, ParamArray );
+				if( EditorViewportClient.IsPerspective() )
+				{
+					// Rotation only matters for perspective viewports not orthographic
+					EditorViewportClient.SetViewRotation( ViewRotation );
+				}
+			}
 		}
 	}
-}
 
+	// Backup ActorsThatWereSelected as this will be cleared whilst deselecting
+	TArray<TWeakObjectPtr<class AActor> > BackupOfActorsThatWereSelected(ActorsThatWereSelected);
 
-void UEditorEngine::RequestEndPlayMap()
-{
-	if( PlayWorld )
-	{
-		bRequestEndPlayMapQueued = true;
+	// Unselect everything
+	GEditor->SelectNone( true, true, false );
+	GetSelectedActors()->DeselectAll();
+	GetSelectedObjects()->DeselectAll();
 
-		// Cache the position and rotation of the camera (the controller may be destroyed before we end the pie session and we need them to preserve the camera position)
-		if (bLastViewAndLocationValid == false)
+	// restore the backup
+	ActorsThatWereSelected = BackupOfActorsThatWereSelected;
+
+	// make sure each selected actors sim equivalent is selected if we're Simulating but not if we're Playing
+	for ( int32 ActorIndex = 0; ActorIndex < ActorsThatWereSelected.Num(); ++ActorIndex )
+	{
+		TWeakObjectPtr<AActor> Actor = ActorsThatWereSelected[ ActorIndex ].Get();
+		if (Actor.IsValid())
 		{
-			for (int32 WorldIdx = WorldList.Num() - 1; WorldIdx >= 0; --WorldIdx)
+			AActor* SimActor = EditorUtilities::GetSimWorldCounterpartActor(Actor.Get());
+			if (SimActor && !SimActor->bHidden)
 			{
-				FWorldContext &ThisContext = WorldList[WorldIdx];
-				if (ThisContext.WorldType == EWorldType::PIE)
-				{
-					FSlatePlayInEditorInfo* const SlatePlayInEditorSession = SlatePlayInEditorMap.Find(ThisContext.ContextHandle);
-					if ((SlatePlayInEditorSession != nullptr) && (SlatePlayInEditorSession->EditorPlayer.IsValid() == true) )
-					{
-						if( SlatePlayInEditorSession->EditorPlayer.Get()->PlayerController != nullptr )
-						{
-							SlatePlayInEditorSession->EditorPlayer.Get()->PlayerController->GetPlayerViewPoint( LastViewLocation, LastViewRotation );
-							bLastViewAndLocationValid = true;
-							break;
-						}
-					}
-				}
+				SelectActor( SimActor, bIsSimulatingInEditor, false );
 			}
 		}
 	}
+
+	FEditorDelegates::OnSwitchBeginPIEAndSIE.Broadcast( bIsSimulatingInEditor );
 }
 
-FString UEditorEngine::BuildPlayWorldURL(const TCHAR* MapName, bool bSpectatorMode, FString AdditionalURLOptions)
+int32 UEditorEngine::OnSwitchWorldForSlatePieWindow( int32 WorldID )
 {
-	// the URL we are building up
-	FString URL(MapName);
+	static const int32 EditorWorldID = 0;
+	static const int32 PieWorldID = 1;
 
-	// If we hold down control, start in spectating mode
-	if (bSpectatorMode)
+	int32 RestoreID = -1;
+	if( WorldID == -1 && GWorld != PlayWorld && PlayWorld != NULL)
 	{
-		// Start in spectator mode
-		URL += TEXT("?SpectatorOnly=1");
+		// When we have an invalid world id we always switch to the pie world in the PIE window
+		const bool bSwitchToPIE = true; 
+		OnSwitchWorldsForPIE( bSwitchToPIE );
+		// The editor world was active restore it later
+		RestoreID = EditorWorldID;
 	}
-
-	// Add any game-specific options set in the INI file
-	URL += InEditorGameURLOptions;
-
-	// Add any additional options that were specified for this call
-	URL += AdditionalURLOptions;
-
-	// Add any additional options that are set in the Play In Settings menu
-	const ULevelEditorPlaySettings* PlayInSettings = GetDefault<ULevelEditorPlaySettings>();
-	FString AdditionalServerGameOptions;
-	if(PlayInSettings->GetAdditionalServerGameOptions(AdditionalServerGameOptions))
+	else if( WorldID == PieWorldID && GWorld != PlayWorld)
+	{
+		const bool bSwitchToPIE = true;
+		// Want to restore the PIE world and the current world is not already the pie world
+		OnSwitchWorldsForPIE( bSwitchToPIE );
+	}
+	else if( WorldID == EditorWorldID && GWorld != EditorWorld)
+	{
+		const bool bSwitchToPIE = false;
+		// Want to restore the editor world and the current world is not already the editor world
+		OnSwitchWorldsForPIE( bSwitchToPIE );
+	}
+	else
 	{
-		URL += *AdditionalServerGameOptions;
+		// Current world is already the same as the world being switched to (nested calls to this for example)
 	}
 
-	return URL;
+	return RestoreID;
 }
 
-bool UEditorEngine::SpawnPlayFromHereStart( UWorld* World, AActor*& PlayerStart, const FVector& StartLocation, const FRotator& StartRotation )
+void UEditorEngine::OnSwitchWorldsForPIE( bool bSwitchToPieWorld, UWorld* OverrideWorld )
 {
-	// null it out in case we don't need to spawn one, and the caller relies on us setting it
-	PlayerStart = NULL;
-
-	if( bHasPlayWorldPlacement )
+	if( bSwitchToPieWorld )
 	{
-		// spawn the PlayerStartPIE in the given world
-		FActorSpawnParameters SpawnParameters;
-		SpawnParameters.OverrideLevel = World->PersistentLevel;
-		PlayerStart = World->SpawnActor<AActor>(PlayFromHerePlayerStartClass, StartLocation, StartRotation, SpawnParameters);
-
-		// make sure we were able to spawn the PlayerStartPIE there
-		if(!PlayerStart)
-		{
-			FMessageDialog::Open( EAppMsgType::Ok, NSLOCTEXT("UnrealEd", "Prompt_22", "Failed to create entry point. Try another location, or you may have to rebuild your level."));
-			return false;
-		}
-		// tag the start
-		ANavigationObjectBase* NavPlayerStart = Cast<ANavigationObjectBase>(PlayerStart);
-		if (NavPlayerStart)
-		{
-			NavPlayerStart->bIsPIEPlayerStart = true;
-		}
+		SetPlayInEditorWorld( OverrideWorld ? OverrideWorld : PlayWorld );
+	}
+	else
+	{
+		RestoreEditorWorld( OverrideWorld ? OverrideWorld : EditorWorld );
 	}
-	// true means we didn't need to spawn, or we succeeded
-	return true;
 }
 
-static bool ShowBlueprintErrorDialog( TArray<UBlueprint*> ErroredBlueprints )
+UWorld* UEditorEngine::CreatePIEWorldByDuplication(FWorldContext &WorldContext, UWorld* InWorld, FString &PlayWorldMapName)
 {
-	if (FApp::IsUnattended() || GIsRunningUnattendedScript)
-	{
-		// App is running in unattended mode, so we should avoid modal dialogs and proceed
-		return true;
-	}
+	double StartTime = FPlatformTime::Seconds();
+	UPackage* InPackage = InWorld->GetOutermost();
+	UWorld* NewPIEWorld = NULL;
+	
+	const FString WorldPackageName = InPackage->GetName();
+
+	// Preserve the old path keeping EditorWorld name the same
+	PlayWorldMapName = UWorld::ConvertToPIEPackageName(WorldPackageName, WorldContext.PIEInstance);
+
+	// Display a busy cursor while we prepare the PIE world
+	const FScopedBusyCursor BusyCursor;
+
+	// Before loading the map, we need to set these flags to true so that postload will work properly
+	TGuardValue<bool> OverrideIsPlayWorld(GIsPlayInEditorWorld, true);
+
+	const FName PlayWorldMapFName = FName(*PlayWorldMapName);
+	UWorld::WorldTypePreLoadMap.FindOrAdd(PlayWorldMapFName) = EWorldType::PIE;
+
+	// Create a package for the PIE world
+	UE_LOG( LogPlayLevel, Log, TEXT("Creating play world package: %s"),  *PlayWorldMapName );	
+
+	UPackage* PlayWorldPackage = CreatePackage(nullptr,*PlayWorldMapName);
+	PlayWorldPackage->SetPackageFlags(PKG_PlayInEditor);
+	PlayWorldPackage->PIEInstanceID = WorldContext.PIEInstance;
+	PlayWorldPackage->FileName = InPackage->FileName;
+	PlayWorldPackage->SetGuid( InPackage->GetGuid() );
+	PlayWorldPackage->MarkAsFullyLoaded();
+
+	// check(GPlayInEditorID == -1 || GPlayInEditorID == WorldContext.PIEInstance);
+	// Currently GPlayInEditorID is not correctly reset after map loading, so it's not safe to assert here
+	FTemporaryPlayInEditorIDOverride IDHelper(WorldContext.PIEInstance);
 
-	struct Local
 	{
-		static void OnHyperlinkClicked( TWeakObjectPtr<UBlueprint> InBlueprint, TSharedPtr<SCustomDialog> InDialog )
+		double SDOStart = FPlatformTime::Seconds();
+
+		// Reset any GUID fixups with lazy pointers
+		FLazyObjectPtr::ResetPIEFixups();
+
+		// Prepare soft object paths for fixup
+		FSoftObjectPath::AddPIEPackageName(FName(*PlayWorldMapName));
+		for (ULevelStreaming* StreamingLevel : InWorld->GetStreamingLevels())
 		{
-			if (UBlueprint* BlueprintToEdit = InBlueprint.Get())
+			if (StreamingLevel && !StreamingLevel->HasAllFlags(RF_DuplicateTransient))
 			{
-				// Open the blueprint
-				GEditor->EditObject( BlueprintToEdit );
+				FString StreamingLevelPIEName = UWorld::ConvertToPIEPackageName(StreamingLevel->GetWorldAssetPackageName(), WorldContext.PIEInstance);
+				FSoftObjectPath::AddPIEPackageName(FName(*StreamingLevelPIEName));
 			}
+		}
 
-			if (InDialog.IsValid())
+		// NULL GWorld before various PostLoad functions are called, this makes it easier to debug invalid GWorld accesses
+		GWorld = NULL;
+
+		// Duplicate the editor world to create the PIE world
+		NewPIEWorld = CastChecked<UWorld>( StaticDuplicateObject(
+			InWorld,				// Source root
+			PlayWorldPackage,		// Destination root
+			InWorld->GetFName(),	// Name for new object
+			RF_AllFlags,			// FlagMask
+			NULL,					// DestClass
+			EDuplicateMode::PIE
+			) );
+
+		// Store prefix we used to rename this world and streaming levels package names
+		NewPIEWorld->StreamingLevelsPrefix = UWorld::BuildPIEPackagePrefix(WorldContext.PIEInstance);
+		// Fixup model components. The index buffers have been created for the components in the source world and the order
+		// in which components were post-loaded matters. So don't try to guarantee a particular order here, just copy the
+		// elements over.
+		if ( NewPIEWorld->PersistentLevel->Model != NULL
+			&& NewPIEWorld->PersistentLevel->Model == InWorld->PersistentLevel->Model
+			&& NewPIEWorld->PersistentLevel->ModelComponents.Num() == InWorld->PersistentLevel->ModelComponents.Num() )
+		{
+			NewPIEWorld->PersistentLevel->Model->ClearLocalMaterialIndexBuffersData();
+			for (int32 ComponentIndex = 0; ComponentIndex < NewPIEWorld->PersistentLevel->ModelComponents.Num(); ++ComponentIndex)
 			{
-				// Opening the blueprint editor above may end up creating an invisible new window on top of the dialog, 
-				// thus making it not interactable, so we have to force the dialog back to the front
-				InDialog->BringToFront(true);
+				UModelComponent* SrcComponent = InWorld->PersistentLevel->ModelComponents[ComponentIndex];
+				UModelComponent* DestComponent = NewPIEWorld->PersistentLevel->ModelComponents[ComponentIndex];
+				DestComponent->CopyElementsFrom(SrcComponent);
 			}
 		}
-	};
 
-	TSharedRef<SVerticalBox> DialogContents = SNew(SVerticalBox)
-		+ SVerticalBox::Slot()
-		.Padding(0, 0, 0, 16)
-		[
-			SNew(STextBlock)
-			.Text(NSLOCTEXT("PlayInEditor", "PrePIE_BlueprintErrors", "One or more blueprints has an unresolved compiler error, are you sure you want to Play in Editor?"))
-		];
+		UE_LOG(LogPlayLevel, Log, TEXT("PIE: StaticDuplicateObject took: (%fs)"),  float(FPlatformTime::Seconds() - SDOStart));		
+	}
 
-	TSharedPtr<SCustomDialog> CustomDialog;
+	// Clean up the world type list now that PostLoad has occurred
+	UWorld::WorldTypePreLoadMap.Remove(PlayWorldMapFName);
 
-	for (UBlueprint* Blueprint : ErroredBlueprints)
-	{
-		TWeakObjectPtr<UBlueprint> BlueprintPtr = Blueprint;
+	check( NewPIEWorld );
+	NewPIEWorld->FeatureLevel = InWorld->FeatureLevel;
+	PostCreatePIEWorld(NewPIEWorld);
 
-		DialogContents->AddSlot()
-			.AutoHeight()
-			.HAlign(HAlign_Left)
-			[
-				SNew(SHyperlink)
-				.Style(FEditorStyle::Get(), "Common.GotoBlueprintHyperlink")
-				.OnNavigate(FSimpleDelegate::CreateLambda([BlueprintPtr, &CustomDialog]() { Local::OnHyperlinkClicked(BlueprintPtr, CustomDialog); }))
-				.Text(FText::FromString(Blueprint->GetName()))
-				.ToolTipText(NSLOCTEXT("SourceHyperlink", "EditBlueprint_ToolTip", "Click to edit the blueprint"))
-			];
-	}
-
-	DialogContents->AddSlot()
-		.Padding(0, 16, 0, 0)
-		[
-			SNew(STextBlock)
-			.Text(NSLOCTEXT("PlayInEditor", "PrePIE_BlueprintErrorsDelayedOpen", "Clicked blueprints will open once this dialog is closed."))
-		];
-
-	FText DialogTitle = NSLOCTEXT("PlayInEditor", "PrePIE_BlueprintErrorsTitle", "Blueprint Compilation Errors");
-
-	FText OKText = NSLOCTEXT("PlayInEditor", "PrePIE_OkText", "Play in Editor");
-	FText CancelText = NSLOCTEXT("Dialogs", "EAppReturnTypeCancel", "Cancel");
-
-	CustomDialog = SNew(SCustomDialog)
-		.Title(DialogTitle)
-		.IconBrush("NotificationList.DefaultMessage")
-		.DialogContent(DialogContents)
-		.Buttons( { SCustomDialog::FButton(OKText), SCustomDialog::FButton(CancelText) } );
-
-	int ButtonPressed = CustomDialog->ShowModal();
-	return ButtonPressed == 0;
+	
+	UE_LOG(LogPlayLevel, Log, TEXT("PIE: Created PIE world by copying editor world from %s to %s (%fs)"), *InWorld->GetPathName(), *NewPIEWorld->GetPathName(), float(FPlatformTime::Seconds() - StartTime));
+	return NewPIEWorld;
 }
 
-void UEditorEngine::PlayInEditor( UWorld* InWorld, bool bInSimulateInEditor, FPlayInEditorOverrides Overrides )
+void UEditorEngine::PostCreatePIEWorld(UWorld *NewPIEWorld)
 {
-	// Broadcast PreBeginPIE before checks that might block PIE below (BeginPIE is broadcast below after the checks)
-	FEditorDelegates::PreBeginPIE.Broadcast(bInSimulateInEditor);
-
-	double PIEStartTime = FPlatformTime::Seconds();
+	double WorldInitStart = FPlatformTime::Seconds();
+	
+	// Init the PIE world
+	NewPIEWorld->WorldType = EWorldType::PIE;
+	NewPIEWorld->InitWorld();
+	UE_LOG(LogPlayLevel, Log, TEXT("PIE: World Init took: (%fs)"),  float(FPlatformTime::Seconds() - WorldInitStart));
 
-	// Block PIE when there is a transaction recording into the undo buffer
-	if (GEditor->IsTransactionActive())
+	// Tag PlayWorld Actors that also exist in EditorWorld.  At this point, no temporary/run-time actors exist in PlayWorld
+	for( FActorIterator PlayActorIt(NewPIEWorld); PlayActorIt; ++PlayActorIt )
 	{
-		FFormatNamedArguments Args;
-		Args.Add(TEXT("TransactionName"), GEditor->GetTransactionName());
-
-		FText NotificationText;
-		if (bInSimulateInEditor)
-		{
-			NotificationText = FText::Format(NSLOCTEXT("UnrealEd", "SIECantStartDuringTransaction", "Can't Simulate when performing {TransactionName} operation"), Args);
-		}
-		else
-		{
-			NotificationText = FText::Format(NSLOCTEXT("UnrealEd", "PIECantStartDuringTransaction", "Can't Play In Editor when performing {TransactionName} operation"), Args);
-		}
-
-		FNotificationInfo Info(NotificationText);
-		Info.ExpireDuration = 5.0f;
-		Info.bUseLargeFont = true;
-		FSlateNotificationManager::Get().AddNotification(Info);
-		return;
+		GEditor->ObjectsThatExistInEditorWorld.Set(*PlayActorIt);
 	}
+}
 
-	// Prompt the user that Matinee must be closed before PIE can occur.
-	if( GLevelEditorModeTools().IsModeActive(FBuiltinEditorModes::EM_InterpEdit) )
+UWorld* UEditorEngine::CreatePIEWorldFromEntry(FWorldContext &WorldContext, UWorld* InWorld, FString &PlayWorldMapName)
+{
+	double StartTime = FPlatformTime::Seconds();
+
+	// Create the world
+	UWorld *LoadedWorld = UWorld::CreateWorld( EWorldType::PIE, false );
+	check(LoadedWorld);
+	if (LoadedWorld->GetOutermost() != GetTransientPackage())
 	{
-		const bool bContinuePIE = EAppReturnType::Yes == FMessageDialog::Open( EAppMsgType::YesNo, NSLOCTEXT("UnrealEd", "PIENeedsToCloseMatineeQ", "'Play in Editor' must close UnrealMatinee.  Continue?") );
-		if ( !bContinuePIE )
-		{
-			return;
-		}
-		GLevelEditorModeTools().DeactivateMode( FBuiltinEditorModes::EM_InterpEdit );
+		LoadedWorld->GetOutermost()->PIEInstanceID = WorldContext.PIEInstance;
 	}
+	// Force default GameMode class so project specific code doesn't fire off. 
+	// We want this world to truly remain empty while we wait for connect!
+	check(LoadedWorld->GetWorldSettings());
+	LoadedWorld->GetWorldSettings()->DefaultGameMode = AGameModeBase::StaticClass();
 
-	// Make sure there's no outstanding load requests
-	FlushAsyncLoading();
-
-	FBlueprintEditorUtils::FindAndSetDebuggableBlueprintInstances();
+	PlayWorldMapName = UGameMapsSettings::GetGameDefaultMap();
+	return LoadedWorld;
+}
 
-	// Monitoring when PIE corrupts references between the World and the PIE generated World for UE-20486
+bool UEditorEngine::WorldIsPIEInNewViewport(UWorld *InWorld)
+{
+	FWorldContext &WorldContext = GetWorldContextFromWorldChecked(InWorld);
+	if (WorldContext.WorldType == EWorldType::PIE)
 	{
-		TArray<ULevel*> Levels = InWorld->GetLevels();
-
-		for (ULevel* Level : Levels)
+		FSlatePlayInEditorInfo * SlateInfoPtr = SlatePlayInEditorMap.Find(WorldContext.ContextHandle);
+		if (SlateInfoPtr)
 		{
-			TArray<UBlueprint*> LevelBlueprints = Level->GetLevelBlueprints();
-
-			if (LevelBlueprints.Num() > 0)
-			{
-				UBlueprint* LevelScriptBlueprint = LevelBlueprints[0];
-				if (LevelScriptBlueprint && LevelScriptBlueprint->GeneratedClass && LevelScriptBlueprint->GeneratedClass->ClassGeneratedBy)
-				{
-					UE_LOG(LogBlueprintUserMessages, Log, TEXT("Early PlayInEditor Detection: Level '%s' has LevelScriptBlueprint '%s' with GeneratedClass '%s' with ClassGeneratedBy '%s'"), *Level->GetPathName(), *LevelScriptBlueprint->GetPathName(), *LevelScriptBlueprint->GeneratedClass->GetPathName(), *LevelScriptBlueprint->GeneratedClass->ClassGeneratedBy->GetPathName());
-				}
-				else if (LevelScriptBlueprint && LevelScriptBlueprint->GeneratedClass)
-				{
-					UE_LOG(LogBlueprintUserMessages, Log, TEXT("Early PlayInEditor Detection: Level '%s' has LevelScriptBlueprint '%s' with GeneratedClass '%s'"), *Level->GetPathName(), *LevelScriptBlueprint->GetPathName(), *LevelScriptBlueprint->GeneratedClass->GetPathName());
-				}
-				else if (LevelScriptBlueprint)
-				{
-					UE_LOG(LogBlueprintUserMessages, Log, TEXT("Early PlayInEditor Detection: Level '%s' has LevelScriptBlueprint '%s'"), *Level->GetPathName(), *LevelScriptBlueprint->GetPathName());
-				}
-			}
+			return SlateInfoPtr->SlatePlayInEditorWindow.IsValid();
 		}
 	}
+	
+	return false;
+}
 
-	// Broadcast BeginPIE after checks that might block PIE above (PreBeginPIE is broadcast above before the checks)
-	FEditorDelegates::BeginPIE.Broadcast(bInSimulateInEditor);
-
-	// let navigation know PIE starts so it can avoid any blueprint creation/deletion/instantiation affect editor map's navmesh changes
-	FNavigationSystem::OnPIEStart(*InWorld);
-
-	ULevelEditorPlaySettings* PlayInSettings = GetMutableDefault<ULevelEditorPlaySettings>();
-	check(PlayInSettings);
-
-	TArray<UBlueprint*> ErroredBlueprints;
-	FInternalPlayLevelUtils::ResolveDirtyBlueprints(!PlayInSettings->AutoRecompileBlueprints, ErroredBlueprints);
+void UEditorEngine::SetPIEInstanceWindowSwitchDelegate(FPIEInstanceWindowSwitch InSwitchDelegate)
+{
+	PIEInstanceWindowSwitchDelegate = InSwitchDelegate;
+}
 
-	if (ErroredBlueprints.Num() && !GIsDemoMode)
+void UEditorEngine::FocusNextPIEWorld(UWorld *CurrentPieWorld, bool previous)
+{
+	// Get the current world's idx
+	int32 CurrentIdx = 0;
+	for (CurrentIdx = 0; CurrentPieWorld && CurrentIdx < WorldList.Num(); ++CurrentIdx)
 	{
-		// There was at least one blueprint with an error, make sure the user is OK with that.
-		bool bContinuePIE = ShowBlueprintErrorDialog( ErroredBlueprints );
-
-		if ( !bContinuePIE )
-		{
-			FMessageLog("BlueprintLog").Open(EMessageSeverity::Warning);
-
-			FEditorDelegates::EndPIE.Broadcast(bInSimulateInEditor);
-			FNavigationSystem::OnPIEEnd(*InWorld);
-			return;
-		}
-		else
+		if (WorldList[CurrentIdx].World() == CurrentPieWorld)
 		{
-			// The user wants to ignore the compiler errors, mark the Blueprints and do not warn them again unless the Blueprint attempts to compile
-			for (UBlueprint* Blueprint : ErroredBlueprints)
-			{
-				Blueprint->bDisplayCompilePIEWarning = false;
-			}
+			break;
 		}
 	}
 
-
-	const FScopedBusyCursor BusyCursor;
-
-	// If there's level already being played, close it. (This may change GWorld)
-	if(PlayWorld)
-	{
-		// immediately end the playworld
-		EndPlayMap();
-	}
-
-	// Register for log processing so we can promote errors/warnings to the message log
-	if (GetDefault<UEditorStyleSettings>()->bPromoteOutputLogWarningsDuringPIE)
+	// Step through the list to find the next or previous
+	int32 step = previous? -1 : 1;
+	CurrentIdx += (WorldList.Num() + step);
+	
+	while ( CurrentPieWorld && WorldList[ CurrentIdx % WorldList.Num() ].World() != CurrentPieWorld )
 	{
-		OutputLogErrorsToMessageLogProxyPtr = MakeShareable(new FOutputLogErrorsToMessageLogProxy());
-	}
+		FWorldContext &Context = WorldList[CurrentIdx % WorldList.Num()];
+		if (Context.World() && Context.WorldType == EWorldType::PIE && Context.GameViewport != NULL)
+		{
+			break;
+		}
 
-	if (GEngine->XRSystem.IsValid() && !bInSimulateInEditor)
-	{
-		GEngine->XRSystem->OnBeginPlay(*GEngine->GetWorldContextFromWorld(InWorld));
+		CurrentIdx += step;
 	}
-
-	// remember old GWorld
-	EditorWorld = InWorld;
-
-	// Clear any messages from last time
-	GEngine->ClearOnScreenDebugMessages();
-
-	// Start a new PIE log page
+	
+	if (WorldList[CurrentIdx % WorldList.Num()].World())
 	{
-		const FString WorldPackageName = EditorWorld->GetOutermost()->GetName();
-
-		FFormatNamedArguments Arguments;
-		Arguments.Add(TEXT("Package"), FText::FromString(FPackageName::GetLongPackageAssetName(WorldPackageName)));
-		Arguments.Add(TEXT("TimeStamp"), FText::AsDateTime(FDateTime::Now()));
+		// Bring new window to front and activate new viewport
+		FSlatePlayInEditorInfo* SlateInfoPtr = SlatePlayInEditorMap.Find(WorldList[CurrentIdx % WorldList.Num()].ContextHandle);
+		if (SlateInfoPtr && SlateInfoPtr->SlatePlayInEditorWindowViewport.IsValid())
+		{
+			FSceneViewport* SceneViewport = SlateInfoPtr->SlatePlayInEditorWindowViewport.Get();
 
-		FText PIESessionLabel = bInSimulateInEditor ?
-			FText::Format(LOCTEXT("SIESessionLabel", "SIE session: {Package} ({TimeStamp})"), Arguments) :
-			FText::Format(LOCTEXT("PIESessionLabel", "PIE session: {Package} ({TimeStamp})"), Arguments);
+			FSlateApplication& SlateApp = FSlateApplication::Get();
+			TSharedRef<SViewport> ViewportWidget = SceneViewport->GetViewportWidget().Pin().ToSharedRef();
 
-		FMessageLog(NAME_CategoryPIE).NewPage(PIESessionLabel);
-	}
+			TSharedPtr<SWindow> ViewportWindow = SlateApp.FindWidgetWindow(ViewportWidget);
+			check(ViewportWindow.IsValid());
 
-	// Flush all audio sources from the editor world
-	if (FAudioDevice* AudioDevice = EditorWorld->GetAudioDevice())
-	{
-		AudioDevice->Flush(EditorWorld);
-		AudioDevice->ResetInterpolation();
-		AudioDevice->OnBeginPIE(bInSimulateInEditor);
-	}
-	EditorWorld->bAllowAudioPlayback = false;
+			// Force window to front
+			ViewportWindow->BringToFront();
 
-	if (!GEditor->bAllowMultiplePIEWorlds)
-	{
-		PlayInSettings->SetRunUnderOneProcess( false );
+			// Execute notification delegate in case game code has to do anything else
+			PIEInstanceWindowSwitchDelegate.ExecuteIfBound();
+		}
 	}
-
-	EPlayNetMode PlayNetMode( PIE_Standalone );
-	PlayInSettings->GetPlayNetMode(PlayNetMode);	// Ignore disabled state here
-	const EPlayNetMode OrigPlayNetMode( PlayNetMode );
-
-	bool CanRunUnderOneProcess = [&PlayInSettings]{ bool RunUnderOneProcess(false); return (PlayInSettings->GetRunUnderOneProcess(RunUnderOneProcess) && RunUnderOneProcess); }();
-	if (CanRunUnderOneProcess)
+}
+void UEditorEngine::ResetPIEAudioSetting(UWorld *CurrentPieWorld)
+{
+	ULevelEditorPlaySettings* PlayInSettings = GetMutableDefault<ULevelEditorPlaySettings>();
+	if (!PlayInSettings->EnableGameSound)
 	{
-		bool CanPlayNetDedicated = [&PlayInSettings]{ bool PlayNetDedicated(false); return (PlayInSettings->GetPlayNetDedicated(PlayNetDedicated) && PlayNetDedicated); }();
-		CanPlayNetDedicated = Overrides.bDedicatedServer.Get(CanPlayNetDedicated);
-		int32 PlayNumberOfClients = [&PlayInSettings]{ int32 NumberOfClients(0); return (PlayInSettings->GetPlayNumberOfClients(NumberOfClients) ? NumberOfClients : 0); }();
-		PlayNumberOfClients = Overrides.NumberOfClients.Get(PlayNumberOfClients);
-		const bool WillAutoConnectToServer = [&PlayInSettings]{ bool AutoConnectToServer(false); return (PlayInSettings->GetAutoConnectToServer(AutoConnectToServer) && AutoConnectToServer); }();
-
-		if (!CanPlayNetDedicated && (PlayNumberOfClients == 1 || !WillAutoConnectToServer))
-		{
-			// Since we don't expose PlayNetMode as an option when doing RunUnderOnProcess,
-			// we take 1 player and !PlayNetdedicated and being standalone.
-			// If auto connect is off, launch as standalone unless there is a dedicated server
-			PlayNetMode = EPlayNetMode::PIE_Standalone;
-		}
-		else
+		if (FAudioDevice* AudioDevice = CurrentPieWorld->GetAudioDevice())
 		{
-			// We are doing multi-player under one process so make sure the NetMode is ListenServer
-			PlayNetMode = EPlayNetMode::PIE_ListenServer;
+			AudioDevice->SetTransientMasterVolume(0.0f);
 		}
-		PlayInSettings->SetPlayNetMode(PlayNetMode);
 	}
-
-
-	bool bAnyBlueprintErrors = ErroredBlueprints.Num() ? true : false;
-	bool bStartInSpectatorMode = false;
-	bool bSupportsOnlinePIE = false;
-	const int32 PlayNumberOfClients = [&PlayInSettings]{ int32 NumberOfClients(0); return (PlayInSettings->GetPlayNumberOfClients(NumberOfClients) ? NumberOfClients : 0); }();
-
-	// Can't allow realtime viewports whilst in PIE so disable it for ALL viewports here.
-	DisableRealtimeViewports();
-
-	// Online PIE is disabled in SIE
-	if (SupportsOnlinePIE() && !bInSimulateInEditor)
+}
+UGameViewportClient * UEditorEngine::GetNextPIEViewport(UGameViewportClient * CurrentViewport)
+{
+	// Get the current world's idx
+	int32 CurrentIdx = 0;
+	for (CurrentIdx = 0; CurrentViewport && CurrentIdx < WorldList.Num(); ++CurrentIdx)
 	{
-		bool bHasRequiredLogins = PlayNumberOfClients <= UOnlineEngineInterface::Get()->GetNumPIELogins();
-		if (bHasRequiredLogins)
-		{
-			// If we support online PIE use it even if we're standalone
-			bSupportsOnlinePIE = true;
-		}
-		else
+		if (WorldList[CurrentIdx].GameViewport == CurrentViewport)
 		{
-			FText ErrorMsg = LOCTEXT("PIELoginFailure", "Not enough login credentials to launch all PIE instances, change editor settings");
-			UE_LOG(LogPlayLevel, Verbose, TEXT("%s"), *ErrorMsg.ToString());
-			FMessageLog(NAME_CategoryPIE).Warning(ErrorMsg);
+			break;
 		}
 	}
 
-	UOnlineEngineInterface::Get()->SetShouldTryOnlinePIE(bSupportsOnlinePIE);
-
-	FModifierKeysState KeysState = FSlateApplication::Get().GetModifierKeys();
-	if (bInSimulateInEditor || KeysState.IsControlDown())
-	{
-		// if control is pressed, start in spectator mode
-		bStartInSpectatorMode = true;
-	}
+	// Step through the list to find the next or previous
+	int32 step = 1;
+	CurrentIdx += (WorldList.Num() + step);
 
-	CanRunUnderOneProcess = [&PlayInSettings]{ bool RunUnderOneProcess(false); return (PlayInSettings->GetRunUnderOneProcess(RunUnderOneProcess) && RunUnderOneProcess); }();
-	if (bInSimulateInEditor || (PlayNetMode == EPlayNetMode::PIE_Standalone && PlayNumberOfClients <= 1 && !bSupportsOnlinePIE) || !CanRunUnderOneProcess)
+	while ( CurrentViewport && WorldList[ CurrentIdx % WorldList.Num() ].GameViewport != CurrentViewport )
 	{
-		// Only spawning 1 PIE instance under this process, only set the PIEInstance value if we're not connecting to another local instance of the game, otherwise it will run the wrong streaming levels
-		PIEInstance = ( !CanRunUnderOneProcess && PlayNetMode == EPlayNetMode::PIE_Client ) ? INDEX_NONE : 0;
-		UGameInstance* const GameInstance = CreatePIEGameInstance(PIEInstance, bInSimulateInEditor, bAnyBlueprintErrors, bStartInSpectatorMode, false, bUseVRPreviewForPlayWorld, PIEStartTime);
-
-		if (!PlayInSettings->EnableGameSound)
+		FWorldContext &Context = WorldList[CurrentIdx % WorldList.Num()];
+		if (Context.GameViewport && Context.WorldType == EWorldType::PIE)
 		{
-			UWorld* GameInstanceWorld = GameInstance->GetWorld();
-			if (FAudioDevice* GameInstanceAudioDevice = GameInstanceWorld->GetAudioDevice())
-			{
-				GameInstanceAudioDevice->SetTransientMasterVolume(0.0f);
-			}
-		}
-
-		if (bInSimulateInEditor)
-		{
-			ToggleBetweenPIEandSIE( true );
-		}
-	}
-	else
-	{
-		if (bSupportsOnlinePIE)
-		{
-			// Make sure all instances of PIE are logged in before creating/launching worlds
-			LoginPIEInstances(bAnyBlueprintErrors, bStartInSpectatorMode, PIEStartTime);
-		}
-		else
-		{
-			// Normal, non-online creation/launching of worlds
-			SpawnIntraProcessPIEWorlds(bAnyBlueprintErrors, bStartInSpectatorMode);
-		}
-	}
-
-	PlayInSettings->MultipleInstanceLastHeight = PlayInSettings->NewWindowHeight;
-	PlayInSettings->MultipleInstanceLastWidth = PlayInSettings->NewWindowWidth;
-	PlayInSettings->SetPlayNetMode(OrigPlayNetMode);
-
-	// Monitoring when PIE corrupts references between the World and the PIE generated World for UE-20486
-	if (EditorWorld)
-	{
-		TArray<ULevel*> Levels = EditorWorld->GetLevels();
-
-		for (ULevel* Level : Levels)
-		{
-			TArray<UBlueprint*> LevelBlueprints = Level->GetLevelBlueprints();
-
-			if (LevelBlueprints.Num() > 0)
-			{
-				UBlueprint* LevelScriptBlueprint = LevelBlueprints[0];
-				if (LevelScriptBlueprint && LevelScriptBlueprint->GeneratedClass && LevelScriptBlueprint->GeneratedClass->ClassGeneratedBy)
-				{
-					UE_LOG(LogBlueprintUserMessages, Log, TEXT("Late PlayInEditor Detection: Level '%s' has LevelScriptBlueprint '%s' with GeneratedClass '%s' with ClassGeneratedBy '%s'"), *Level->GetPathName(), *LevelScriptBlueprint->GetPathName(), *LevelScriptBlueprint->GeneratedClass->GetPathName(), *LevelScriptBlueprint->GeneratedClass->ClassGeneratedBy->GetPathName());
-				}
-				else if (LevelScriptBlueprint && LevelScriptBlueprint->GeneratedClass)
-				{
-					UE_LOG(LogBlueprintUserMessages, Log, TEXT("Late PlayInEditor Detection: Level '%s' has LevelScriptBlueprint '%s' with GeneratedClass '%s'"), *Level->GetPathName(), *LevelScriptBlueprint->GetPathName(), *LevelScriptBlueprint->GeneratedClass->GetPathName());
-				}
-				else if (LevelScriptBlueprint)
-				{
-					UE_LOG(LogBlueprintUserMessages, Log, TEXT("Late PlayInEditor Detection: Level '%s' has LevelScriptBlueprint '%s'"), *Level->GetPathName(), *LevelScriptBlueprint->GetPathName());
-				}
-			}
-		}
-	}
-
-	// Make sure to focus the game viewport.
-	if (!bInSimulateInEditor && (GetDefault<ULevelEditorPlaySettings>()->GameGetsMouseControl || IVREditorModule::Get().IsVREditorEnabled() || (bUseVRPreviewForPlayWorld && GEngine->XRSystem.IsValid())))
-	{
-		FSlateApplication::Get().SetAllUserFocusToGameViewport();
-	}
-
-	FMargin SafeZoneRatio = PlayInSettings->PIESafeZoneOverride;
-	SafeZoneRatio.Left /= (PlayInSettings->NewWindowWidth / 2.0f);
-	SafeZoneRatio.Right /= (PlayInSettings->NewWindowWidth / 2.0f);
-	SafeZoneRatio.Bottom /= (PlayInSettings->NewWindowHeight / 2.0f);
-	SafeZoneRatio.Top /= (PlayInSettings->NewWindowHeight / 2.0f);
-	FSlateApplication::Get().OnDebugSafeZoneChanged.Broadcast(SafeZoneRatio, false);
-
-	FEditorDelegates::PostPIEStarted.Broadcast( bInSimulateInEditor );
-}
-
-FGameInstancePIEResult UEditorEngine::PreCreatePIEServerInstance(const bool bAnyBlueprintErrors, const bool bStartInSpectatorMode, const float PIEStartTime, const bool bSupportsOnlinePIE, int32& InNumOnlinePIEInstances)
-{
-	return FGameInstancePIEResult::Success();
-}
-
-void UEditorEngine::SpawnIntraProcessPIEWorlds(bool bAnyBlueprintErrors, bool bStartInSpectatorMode)
-{
-	double PIEStartTime = FPlatformTime::Seconds();
-	bStartLateJoinersInSpectatorMode = bStartInSpectatorMode;
-
-	// Has to be false or this function wouldn't be called
-	bool bInSimulateInEditor = false;
-	ULevelEditorPlaySettings* PlayInSettings = Cast<ULevelEditorPlaySettings>(ULevelEditorPlaySettings::StaticClass()->GetDefaultObject());
-
-	// Spawning multiple PIE instances
-	if (PlayInSettings->MultipleInstancePositions.Num() == 0)
-	{
-		PlayInSettings->MultipleInstancePositions.SetNum(1);
-	}
-
-	PlayInSettings->MultipleInstancePositions[0] = PlayInSettings->NewWindowPosition;
-
-	int32 NextX = 0;
-	int32 NextY = 0;
-	SettingsIndex = 1;
-	int32 ClientNum = 0;
-	PIEInstance = 1;
-	const bool WillAutoConnectToServer = [&PlayInSettings]{ bool AutoConnectToServer(false); return (PlayInSettings->GetAutoConnectToServer(AutoConnectToServer) && AutoConnectToServer); }();
-	const bool CanPlayNetDedicated = [&PlayInSettings]{ bool PlayNetDedicated(false); return (PlayInSettings->GetPlayNetDedicated(PlayNetDedicated) && PlayNetDedicated); }();
-
-	// Server
-	if (CanPlayNetDedicated || WillAutoConnectToServer)
-	{
-		FGameInstancePIEResult PreCreateResult = PreCreatePIEServerInstance(bAnyBlueprintErrors, bStartInSpectatorMode, PIEStartTime, false, NumOnlinePIEInstances);
-		if (!PreCreateResult.IsSuccess())
-		{
-			return;
-		}
-
-		PlayInSettings->SetPlayNetMode(EPlayNetMode::PIE_ListenServer);
-
-		if (!CanPlayNetDedicated)
-		{
-			ClientNum++;
-			GetMultipleInstancePositions(SettingsIndex++, NextX, NextY);
-		}
-        
-		UGameInstance* const ServerGameInstance = CreatePIEGameInstance(PIEInstance, bInSimulateInEditor, bAnyBlueprintErrors, bStartInSpectatorMode, CanPlayNetDedicated, (ClientNum == 0 && bUseVRPreviewForPlayWorld), PIEStartTime);
-		if (ServerGameInstance)
-		{
-			ServerPrefix = ServerGameInstance->GetWorldContext()->PIEPrefix;
-		}
-		else
-		{
-			// Failed, abort
-			return;
-		}
-
-		PIEInstance++;
-	}
-
-	// Clients
-	const int32 PlayNumberOfClients = [&PlayInSettings]{ int32 NumberOfClients(0); return (PlayInSettings->GetPlayNumberOfClients(NumberOfClients) ? NumberOfClients : 0); }();
-	for (; ClientNum < PlayNumberOfClients; ++ClientNum)
-	{
-		// Only launch as clients if they should connect
-		if (WillAutoConnectToServer)
-		{
-			PlayInSettings->SetPlayNetMode(EPlayNetMode::PIE_Client);
-		}
-		else
-		{
-			PlayInSettings->SetPlayNetMode(EPlayNetMode::PIE_Standalone);
-		}
-
-		// For legacy reasons, single player PIE uses ULevelEditorPlaySettings::NewWindowPosition as its window position.
-		// Multiple PIE uses the ULevelEditorPlaySettings::MultipleInstancePositions array, starting with index 1.
-		// If this is a single player PIE, with dedicated server, don't set NewWindowPosition from the MultipleInstancePositions array - leave it as is.
-		if (PlayNumberOfClients > 1)
-		{
-			GetMultipleInstancePositions(SettingsIndex++, NextX, NextY);
-		}
-
-		UGameInstance* const ClientGameInstance = CreatePIEGameInstance(PIEInstance, bInSimulateInEditor, bAnyBlueprintErrors, bStartInSpectatorMode, false, (ClientNum == 0 && bUseVRPreviewForPlayWorld), PIEStartTime);
-		if (!ClientGameInstance)
-		{
-			// Failed, abort
-			return;
-		}
-
-		PIEInstance++;
-	}
-
-	// Restore window settings
-	GetMultipleInstancePositions(0, NextX, NextY);	// restore cached settings
-
-	// Give focus to the first client
-	GiveFocusToFirstClientPIEViewport();
-}
-
-bool UEditorEngine::CreatePIEWorldFromLogin(FWorldContext& PieWorldContext, EPlayNetMode PlayNetMode, FPieLoginStruct& DataStruct)
-{
-	ULevelEditorPlaySettings* PlayInSettings = Cast<ULevelEditorPlaySettings>(ULevelEditorPlaySettings::StaticClass()->GetDefaultObject());
-	PlayInSettings->SetPlayNetMode(PlayNetMode);
-
-	// Set window position
-	GetMultipleInstancePositions(DataStruct.SettingsIndex, DataStruct.NextX, DataStruct.NextY);
-	
-	const bool CanPlayNetDedicated = [&PlayInSettings]{ bool PlayNetDedicated(false); return (PlayInSettings->GetPlayNetDedicated(PlayNetDedicated) && PlayNetDedicated); }();
-	const bool ActAsClient = PlayNetMode == EPlayNetMode::PIE_Client || PlayNetMode == EPlayNetMode::PIE_StandaloneWithServer || PlayNetMode == EPlayNetMode::PIE_Standalone;
-	UGameInstance* const GameInstance = CreatePIEGameInstance(PieWorldContext.PIEInstance, false, DataStruct.bAnyBlueprintErrors, DataStruct.bStartInSpectatorMode, ActAsClient ? false : CanPlayNetDedicated, false, DataStruct.PIEStartTime);
-	
-	// Restore window settings
-	GetMultipleInstancePositions(0, DataStruct.NextX, DataStruct.NextY);	// restore cached settings
-
-	if (GameInstance)
-	{
-		GameInstance->GetWorldContext()->bWaitingOnOnlineSubsystem = false;
-
-		return true;
-	}
-	else
-	{
-		return false;
-	}
-}
-
-void UEditorEngine::CancelPlayingViaLauncher()
-{
-	if (LauncherWorker.IsValid())
-	{
-		LauncherWorker->CancelAndWait();
-	}
-}
-
-bool UEditorEngine::SupportsOnlinePIE() const
-{
-	return UOnlineEngineInterface::Get()->SupportsOnlinePIE();
-}
-
-void UEditorEngine::LoginPIEInstances(bool bAnyBlueprintErrors, bool bStartInSpectatorMode, double PIEStartTime)
-{
-	ULevelEditorPlaySettings* PlayInSettings = Cast<ULevelEditorPlaySettings>(ULevelEditorPlaySettings::StaticClass()->GetDefaultObject());
-	SettingsIndex = 0;
-
-	/** Setup the common data values for each login instance */
-	FPieLoginStruct DataStruct;
-	DataStruct.SettingsIndex = ++SettingsIndex;
-	DataStruct.bAnyBlueprintErrors = bAnyBlueprintErrors;
-	DataStruct.bStartInSpectatorMode = bStartInSpectatorMode;
-	DataStruct.PIEStartTime = PIEStartTime;
-
-	ensure(UOnlineEngineInterface::Get()->GetNumPIELogins() > 0);
-
-	int32 ClientNum = 0;
-	PIEInstance = 1;
-	int32 NextX = 0;
-	int32 NextY = 0;
-
-	const EPlayNetMode PlayNetMode = [&PlayInSettings]{ EPlayNetMode NetMode(PIE_Standalone); return (PlayInSettings->GetPlayNetMode(NetMode) ? NetMode : PIE_Standalone); }();
-	const bool CanPlayNetDedicated = [&PlayInSettings]{ bool PlayNetDedicated(false); return (PlayInSettings->GetPlayNetDedicated(PlayNetDedicated) && PlayNetDedicated); }();
-	const bool WillAutoConnectToServer = [&PlayInSettings]{ bool AutoConnectToServer(false); return (PlayInSettings->GetAutoConnectToServer(AutoConnectToServer) && AutoConnectToServer); }();
-	const int32 PlayNumberOfClients = [&PlayInSettings] { int32 NumberOfClients(0); return (PlayInSettings->GetPlayNumberOfClients(NumberOfClients) ? NumberOfClients : 0); }();
-
-	PIEInstancesToLogInCount = PlayNumberOfClients;
-	bAtLeastOnePIELoginFailed = false;
-
-	// Server
-	if (WillAutoConnectToServer || CanPlayNetDedicated)
-	{
-		FGameInstancePIEResult PreCreateResult = PreCreatePIEServerInstance(bAnyBlueprintErrors, bStartInSpectatorMode, PIEStartTime, true, NumOnlinePIEInstances);
-		if (!PreCreateResult.IsSuccess())
-		{
-			return;
-		}
-
-		FWorldContext &PieWorldContext = CreateNewWorldContext(EWorldType::PIE);
-		PieWorldContext.PIEInstance = PIEInstance++;
-		PieWorldContext.RunAsDedicated = CanPlayNetDedicated;
-		PieWorldContext.bWaitingOnOnlineSubsystem = true;
-
-		// Update login struct parameters
-		DataStruct.WorldContextHandle = PieWorldContext.ContextHandle;
-		DataStruct.NetMode = PlayNetMode;
-
-		// Always get the interface (it will create the subsystem regardless)
-		FName OnlineIdentifier = UOnlineEngineInterface::Get()->GetOnlineIdentifier(PieWorldContext);
-		UE_LOG(LogPlayLevel, Display, TEXT("Creating online subsystem for server %s"), *OnlineIdentifier.ToString());
-		NumOnlinePIEInstances++;
-
-		if (!CanPlayNetDedicated)
-		{
-			DataStruct.NextX = NextX;
-			DataStruct.NextY = NextY;
-			GetMultipleInstancePositions(DataStruct.SettingsIndex, NextX, NextY);
-
-			// Login to online platform before creating world
-			FOnPIELoginComplete Delegate;
-			Delegate.BindUObject(this, &UEditorEngine::OnLoginPIEComplete, DataStruct);
-
-			// Login first and continue the flow later
-			UOnlineEngineInterface::Get()->LoginPIEInstance(OnlineIdentifier, 0, ClientNum, Delegate);
-
-			ClientNum++;
-		}
-		else
-		{
-			// Dedicated servers don't use a login
-			UOnlineEngineInterface::Get()->SetForceDedicated(OnlineIdentifier, true);
-			if (CreatePIEWorldFromLogin(PieWorldContext, EPlayNetMode::PIE_ListenServer, DataStruct))
-			{
-				FMessageLog(NAME_CategoryPIE).Info(LOCTEXT("LoggingInDedicated", "Dedicated Server logged in"));
-			}
-			else
-			{
-				// Failed to create world, this creates a dialog elsewhere
-				return;
-			}
-		}
-	}
-
-	// Clients
-	for (; ClientNum < PlayNumberOfClients; ++ClientNum)
-	{
-		PlayInSettings->SetPlayNetMode(PlayNetMode);
-		FWorldContext &PieWorldContext = CreateNewWorldContext(EWorldType::PIE);
-		PieWorldContext.PIEInstance = PIEInstance++;
-		PieWorldContext.bWaitingOnOnlineSubsystem = true;
-
-		// Update login struct parameters
-		DataStruct.WorldContextHandle = PieWorldContext.ContextHandle;
-		DataStruct.SettingsIndex = ++SettingsIndex;
-		DataStruct.NextX = NextX;
-		DataStruct.NextY = NextY;
-		GetMultipleInstancePositions(DataStruct.SettingsIndex, NextX, NextY);
-		DataStruct.NetMode = WillAutoConnectToServer ? EPlayNetMode::PIE_Client : EPlayNetMode::PIE_Standalone;
-
-		FName OnlineIdentifier = UOnlineEngineInterface::Get()->GetOnlineIdentifier(PieWorldContext);
-		UE_LOG(LogPlayLevel, Display, TEXT("Creating online subsystem for client %s"), *OnlineIdentifier.ToString());
-		NumOnlinePIEInstances++;
-
-		// Login to online platform before creating world
-		FOnPIELoginComplete Delegate;
-		Delegate.BindUObject(this, &UEditorEngine::OnLoginPIEComplete, DataStruct);
-
-		UOnlineEngineInterface::Get()->LoginPIEInstance(OnlineIdentifier, 0, ClientNum, Delegate);
-	}
-
-	// Restore window settings
-	GetMultipleInstancePositions(0, NextX, NextY);	// restore cached settings
-}
-
-void UEditorEngine::OnLoginPIEComplete(int32 LocalUserNum, bool bWasSuccessful, const FString& ErrorString, FPieLoginStruct DataStruct)
-{
-	// This is needed because pie login may change the state of the online objects that called this function
-	GetTimerManager()->SetTimerForNextTick(FTimerDelegate::CreateUObject(this, &UEditorEngine::OnLoginPIEComplete_Deferred, LocalUserNum, bWasSuccessful, ErrorString, DataStruct));
-}
-
-void UEditorEngine::OnLoginPIEComplete_Deferred(int32 LocalUserNum, bool bWasSuccessful, FString ErrorString, FPieLoginStruct DataStruct)
-{
-	UE_LOG(LogPlayLevel, Verbose, TEXT("OnLoginPIEComplete LocalUserNum: %d bSuccess: %d %s"), LocalUserNum, bWasSuccessful, *ErrorString);
-	FWorldContext* PieWorldContext = GetWorldContextFromHandle(DataStruct.WorldContextHandle);
-
-	if (!PieWorldContext)
-	{
-		// This will fail if PIE was ended before this callback happened, silently return
-		return;
-	}
-
-	// Create the new world
-	if (CreatePIEWorldFromLogin(*PieWorldContext, DataStruct.NetMode, DataStruct))
-	{
-		// Logging after the create so a new MessageLog Page is created
-		if (bWasSuccessful)
-		{
-			if (DataStruct.NetMode != EPlayNetMode::PIE_Client)
-			{
-				FMessageLog(NAME_CategoryPIE).Info(LOCTEXT("LoggedInServer", "Server logged in"));
-			}
-			else
-			{
-				FMessageLog(NAME_CategoryPIE).Info(LOCTEXT("LoggedInClient", "Client logged in"));
-			}
-		}
-		else
-		{
-			bAtLeastOnePIELoginFailed = true;
-			if (DataStruct.NetMode != EPlayNetMode::PIE_Client)
-			{
-				FMessageLog(NAME_CategoryPIE).Error(FText::Format(LOCTEXT("LoggedInServerFailure", "Server failed to login. {0}"), FText::FromString(ErrorString)));
-			}
-			else
-			{
-				FMessageLog(NAME_CategoryPIE).Error(FText::Format(LOCTEXT("LoggedInClientFailure", "Client failed to login. {0}"), FText::FromString(ErrorString)));
-			}
-		}
-	}
-
-	PIEInstancesToLogInCount--;
-	if (PIEInstancesToLogInCount == 0)
-	{
-		if (bAtLeastOnePIELoginFailed)
-		{
-			EndPlayMap();
-		}
-		else
-		{
-			OnLoginPIEAllComplete();
-		}
-	}
-}
-
-void UEditorEngine::OnLoginPIEAllComplete()
-{
-	GiveFocusToFirstClientPIEViewport();
-}
-
-void UEditorEngine::GiveFocusToFirstClientPIEViewport()
-{
-	// Find the non-dedicated server or first client window to give focus to
-	int32 LowestPIEInstance = TNumericLimits<int32>::Max();
-	UGameViewportClient* ViewportClient = nullptr;
-	for (const FWorldContext& WorldContext : WorldList)
-	{
-		if (WorldContext.WorldType == EWorldType::PIE && !WorldContext.RunAsDedicated)
-		{
-			if (WorldContext.PIEInstance < LowestPIEInstance)
-			{
-				LowestPIEInstance = WorldContext.PIEInstance;
-				ViewportClient = WorldContext.GameViewport;
-			}
-		}
-	}
-
-	// Give focus to the first client
-	if (ViewportClient && ViewportClient->GetGameViewportWidget().IsValid())
-	{
-		FSlateApplication::Get().RegisterGameViewport(ViewportClient->GetGameViewportWidget().ToSharedRef());
-	}
-}
-
-void UEditorEngine::RequestLateJoin()
-{
-	bool bSupportsOnlinePIE = false;
-
-	if (SupportsOnlinePIE())
-	{
-// 		bool bHasRequiredLogins = PlayNumberOfClients <= UOnlineEngineInterface::Get()->GetNumPIELogins();
-// 		if (bHasRequiredLogins)
-		{
-			// If we support online PIE use it even if we're standalone
-			bSupportsOnlinePIE = true;
-		}
-	}
-
-	ULevelEditorPlaySettings* PlayInSettings = Cast<ULevelEditorPlaySettings>(ULevelEditorPlaySettings::StaticClass()->GetDefaultObject());
-	double PIEStartTime = FPlatformTime::Seconds();
-	int32 NextX = 0;
-	int32 NextY = 0;
-
-	PlayInSettings->SetPlayNetMode(EPlayNetMode::PIE_Client);
-
-
-	if (bSupportsOnlinePIE)
-	{
-		FPieLoginStruct DataStruct;
-		DataStruct.bAnyBlueprintErrors = false;
-		DataStruct.bStartInSpectatorMode = bStartLateJoinersInSpectatorMode;
-		DataStruct.PIEStartTime = PIEStartTime;
-
-		FWorldContext &PieWorldContext = CreateNewWorldContext(EWorldType::PIE);
-		PieWorldContext.PIEInstance = PIEInstance++;
-		PieWorldContext.bWaitingOnOnlineSubsystem = true;
-
-		// Update login struct parameters
-		DataStruct.WorldContextHandle = PieWorldContext.ContextHandle;
-		DataStruct.SettingsIndex = ++SettingsIndex;
-		DataStruct.NextX = NextX;
-		DataStruct.NextY = NextY;
-		GetMultipleInstancePositions(DataStruct.SettingsIndex, NextX, NextY);
-		DataStruct.NetMode = EPlayNetMode::PIE_Client;
-
-		FName OnlineIdentifier = UOnlineEngineInterface::Get()->GetOnlineIdentifier(PieWorldContext);
-		UE_LOG(LogPlayLevel, Display, TEXT("Creating online subsystem for client %s"), *OnlineIdentifier.ToString());
-
-		// Login to online platform before creating world
-		FOnPIELoginComplete Delegate;
-		Delegate.BindUObject(this, &UEditorEngine::OnLoginPIEComplete, DataStruct);
-		// The server doesn't get a PIE login but counts as an online instance
-		UOnlineEngineInterface::Get()->LoginPIEInstance(OnlineIdentifier, 0, NumOnlinePIEInstances - 1, Delegate);
-
-		NumOnlinePIEInstances++;
-	}
-	else
-	{
-		GetMultipleInstancePositions(SettingsIndex++, NextX, NextY);
-
-		UGameInstance* const ClientGameInstance = CreatePIEGameInstance(PIEInstance, false, false, bStartLateJoinersInSpectatorMode, false, false, PIEStartTime);
-		if (!ClientGameInstance)
-		{
-			// Failed, abort
-			return;
-		}
-
-		PIEInstance++;
-	}
-
-	GetMultipleInstancePositions(0, NextX, NextY);
-}
-
-class SPIEViewport : public SViewport
-{
-	SLATE_BEGIN_ARGS(SPIEViewport)
-		: _Content()
-		, _RenderDirectlyToWindow(false)
-		, _EnableStereoRendering(false)
-		, _IgnoreTextureAlpha(true)
-	{
-		_Clipping = EWidgetClipping::ClipToBoundsAlways;
-	}
-
-		SLATE_DEFAULT_SLOT(FArguments, Content)
-
-		/**
-		 * Whether or not to render directly to the window's backbuffer or an offscreen render target that is applied to the window later
-		 * Rendering to an offscreen target is the most common option in the editor where there may be many frames which this viewport's interface may wish to not re-render but use a cached buffer instead
-		 * Rendering directly to the backbuffer is the most common option in the game where you want to update each frame without the cost of writing to an intermediate target first.
-		 */
-		SLATE_ARGUMENT(bool, RenderDirectlyToWindow)
-
-		/** Whether or not to enable stereo rendering. */
-		SLATE_ARGUMENT(bool, EnableStereoRendering )
-
-		/**
-		 * If true, the viewport's texture alpha is ignored when performing blending.  In this case only the viewport tint opacity is used
-		 * If false, the texture alpha is used during blending
-		 */
-		SLATE_ARGUMENT( bool, IgnoreTextureAlpha )
-
-	SLATE_END_ARGS()
-
-	void Construct(const FArguments& InArgs)
-	{
-		SViewport::Construct(
-			SViewport::FArguments()
-			.EnableGammaCorrection(false) // Gamma correction in the game is handled in post processing in the scene renderer
-			.RenderDirectlyToWindow(InArgs._RenderDirectlyToWindow)
-			.EnableStereoRendering(InArgs._EnableStereoRendering)
-			.IgnoreTextureAlpha(InArgs._IgnoreTextureAlpha)
-			[
-				InArgs._Content.Widget
-			]
-		);
-	}
-
-	virtual void Tick(const FGeometry& AllottedGeometry, const double InCurrentTime, const float InDeltaTime) override
-	{
-		SViewport::Tick(AllottedGeometry, InCurrentTime, InDeltaTime);
-
-		// Rather than binding the attribute we're going to poll it in tick, otherwise we will make this widget volatile, and it therefore
-		// wont be possible to cache it or its children in GSlateEnableGlobalInvalidation mode.
-		SetEnabled(FSlateApplication::Get().GetNormalExecutionAttribute().Get());
-	}
-};
-
-UGameInstance* UEditorEngine::CreatePIEGameInstance(int32 InPIEInstance, bool bInSimulateInEditor, bool bAnyBlueprintErrors, bool bStartInSpectatorMode, bool bRunAsDedicated, bool bPlayStereoscopic, float PIEStartTime)
-{
-	// create a new GameInstance
-	FSoftClassPath GameInstanceClassName = GetDefault<UGameMapsSettings>()->GameInstanceClass;
-	UClass* GameInstanceClass = (GameInstanceClassName.IsValid() ? LoadObject<UClass>(NULL, *GameInstanceClassName.ToString()) : UGameInstance::StaticClass());
-
-	// If the GameInstance class from the settings cannot be found, fall back to the base class
-	if(GameInstanceClass == nullptr)
-	{
-		GameInstanceClass = UGameInstance::StaticClass();
-	}
-	UGameInstance* GameInstance = NewObject<UGameInstance>(this, GameInstanceClass);
-
-	// We need to temporarily add the GameInstance to the root because the InitPIE call can do garbage collection wiping out the GameInstance
-	GameInstance->AddToRoot();
-
-	FGameInstancePIEParameters GameInstanceParams;
-	GameInstanceParams.bAnyBlueprintErrors = bAnyBlueprintErrors;
-	GameInstanceParams.bSimulateInEditor = bInSimulateInEditor;
-	GameInstanceParams.bStartInSpectatorMode = bStartInSpectatorMode;
-	GameInstanceParams.bRunAsDedicated = bRunAsDedicated;
-	GameInstanceParams.WorldFeatureLevel = PreviewPlatform.GetEffectivePreviewFeatureLevel();
-
-	const FGameInstancePIEResult InitializeResult = GameInstance->InitializeForPlayInEditor(InPIEInstance, GameInstanceParams);
-	if (!InitializeResult.IsSuccess())
-	{
-		FMessageDialog::Open(EAppMsgType::Ok, InitializeResult.FailureReason);
-
-		FEditorDelegates::EndPIE.Broadcast(bInSimulateInEditor);
-
-		FNavigationSystem::OnPIEEnd(*EditorWorld);
-
-		return nullptr;
-	}
-	
-	FWorldContext* const PieWorldContext = GameInstance->GetWorldContext();
-	check(PieWorldContext);
-	PlayWorld = PieWorldContext->World();
-
-	GWorld = PlayWorld;
-	SetPlayInEditorWorld( PlayWorld );
-
-#if PLATFORM_64BITS
-	const FString PlatformBitsString( TEXT( "64" ) );
-#else
-	const FString PlatformBitsString( TEXT( "32" ) );
-#endif
-
-	const FText WindowTitleOverride = GetDefault<UGeneralProjectSettings>()->ProjectDisplayedTitle;
-
-	FFormatNamedArguments Args;
-	Args.Add( TEXT("GameName"), FText::FromString( FString( WindowTitleOverride.IsEmpty() ? FApp::GetProjectName() : WindowTitleOverride.ToString() ) ) );
-	Args.Add( TEXT("PlatformBits"), FText::FromString( PlatformBitsString ) );
-	Args.Add( TEXT("RHIName"), FText::FromName( LegacyShaderPlatformToShaderFormat( GShaderPlatformForFeatureLevel[PreviewPlatform.PreviewFeatureLevel] ) ) );
-
-	const ULevelEditorPlaySettings* PlayInSettings = GetDefault<ULevelEditorPlaySettings>();
-	const EPlayNetMode PlayNetMode = [&PlayInSettings]{ EPlayNetMode NetMode(PIE_Standalone); return (PlayInSettings->GetPlayNetMode(NetMode) ? NetMode : PIE_Standalone); }();
-	if (PlayNetMode == PIE_Client)
-	{
-		Args.Add(TEXT("NetMode"), FText::FromString(FString::Printf(TEXT("Client %d"), PieWorldContext->PIEInstance - 1)));
-	}
-	else if (PlayNetMode == PIE_ListenServer)
-	{
-		Args.Add( TEXT("NetMode"), FText::FromString( TEXT("Server")));
-	}
-	else
-	{
-		Args.Add( TEXT("NetMode"), FText::FromString( TEXT("Standalone")));
-	}
-
-	const FText ViewportName = FText::Format( NSLOCTEXT("UnrealEd", "PlayInEditor_RHI_F", "{GameName} Game Preview {NetMode} ({PlatformBits}-bit/{RHIName})" ), Args );
-
-	// Make a list of all the selected actors
-	TArray<UObject *> SelectedActors;
-	TArray<UObject*> SelectedComponents;
-	for ( FSelectionIterator It( GetSelectedActorIterator() ); It; ++It )
-	{
-		AActor* Actor = static_cast<AActor*>( *It );
-		if (Actor)
-		{
-			checkSlow( Actor->IsA(AActor::StaticClass()) );
-
-			SelectedActors.Add( Actor );
-		}
-	}
-
-
-	// Unselect everything
-	GEditor->SelectNone( true, true, false );
-	GetSelectedActors()->DeselectAll();
-	GetSelectedObjects()->DeselectAll();
-	GetSelectedComponents()->DeselectAll();
-
-	// For every actor that was selected previously, make sure it's sim equivalent is selected
-	for ( int32 ActorIndex = 0; ActorIndex < SelectedActors.Num(); ++ActorIndex )
-	{
-		AActor* Actor = Cast<AActor>( SelectedActors[ ActorIndex ] );
-		if (Actor)
-		{
-			ActorsThatWereSelected.Add( Actor );
-
-			AActor* SimActor = EditorUtilities::GetSimWorldCounterpartActor(Actor);
-			if (SimActor && !SimActor->IsHidden() && bInSimulateInEditor)
-			{
-				SelectActor( SimActor, true, false );
-			}
-		}
-	}
-
-	// For play in editor, this is the viewport widget where the game is being displayed
-	TSharedPtr<SPIEViewport> PieViewportWidget;
-
-	// Initialize the viewport client.
-	UGameViewportClient* ViewportClient = NULL;
-	ULocalPlayer *NewLocalPlayer = NULL;
-
-	if (!PieWorldContext->RunAsDedicated)
-	{
-		bool bCreateNewAudioDevice = PlayInSettings->IsCreateAudioDeviceForEveryPlayer();
-
-		ViewportClient = NewObject<UGameViewportClient>(this, GameViewportClientClass);
-		ViewportClient->Init(*PieWorldContext, GameInstance, bCreateNewAudioDevice);
-
-		GameViewport = ViewportClient;
-		GameViewport->bIsPlayInEditorViewport = true;
-		PieWorldContext->GameViewport = ViewportClient;
-
-		// Add a handler for game client input key
-		ViewportClient->OnGameViewportInputKey().BindUObject(this, &UEditorEngine::ProcessDebuggerCommands);
-
-		// Add a handler for viewport close requests
-		ViewportCloseRequestedDelegateHandle = ViewportClient->OnCloseRequested().AddUObject(this, &UEditorEngine::OnViewportCloseRequested);
-		FSlatePlayInEditorInfo& SlatePlayInEditorSession = SlatePlayInEditorMap.Add(PieWorldContext->ContextHandle, FSlatePlayInEditorInfo());
-		SlatePlayInEditorSession.DestinationSlateViewport = RequestedDestinationSlateViewport;	// Might be invalid depending how pie was launched. Code below handles this.
-		RequestedDestinationSlateViewport = NULL;
-
-		FString Error;
-		NewLocalPlayer = ViewportClient->SetupInitialLocalPlayer(Error);
-		if(!NewLocalPlayer)
-		{
-			FMessageDialog::Open( EAppMsgType::Ok, FText::Format(NSLOCTEXT("UnrealEd", "Error_CouldntSpawnPlayer", "Couldn't spawn player: {0}"), FText::FromString(Error)) );
-			// go back to using the real world as GWorld
-			RestoreEditorWorld( EditorWorld );
-			EndPlayMap();
-			return nullptr;
-		}
-
-		if (!bInSimulateInEditor)
-		{
-			SlatePlayInEditorSession.EditorPlayer = NewLocalPlayer;
-		}
-			
-		// Note: For K2 debugging purposes this MUST be created before beginplay is called because beginplay can trigger breakpoints
-		// and we need to be able to refocus the pie viewport afterwards so it must be created first in order for us to find it
-		{
-			// Only create a separate viewport and window if we aren't playing in a current viewport
-			if( SlatePlayInEditorSession.DestinationSlateViewport.IsValid() )
-			{
-				TSharedPtr<IAssetViewport> LevelViewportRef = SlatePlayInEditorSession.DestinationSlateViewport.Pin();
-
-				LevelViewportRef->StartPlayInEditorSession( ViewportClient, bInSimulateInEditor );
-
-				//SlatePlayInEditorSession.SlatePlayInEditorWindowViewport = MakeShareable<FSceneViewport>((FSceneViewport*)LevelViewportRef->GetActiveViewport());
-			}
-			else
-			{
-				const int32 PlayNumberOfClients = [&PlayInSettings] { int32 NumberOfClients(0); return (PlayInSettings->GetPlayNumberOfClients(NumberOfClients) ? NumberOfClients : 0); }();
-
-				// Create the top level pie window and add it to Slate
-				uint32 NewWindowHeight = PlayInSettings->NewWindowHeight;
-				uint32 NewWindowWidth = PlayInSettings->NewWindowWidth;
-				FIntPoint NewWindowPosition = PlayInSettings->NewWindowPosition;
-				bool CenterNewWindow = PlayInSettings->CenterNewWindow && (PlayNumberOfClients == 1);
-
-				// Setup size for PIE window
-				if ((NewWindowWidth <= 0) || (NewWindowHeight <= 0))
-				{
-					// Get desktop metrics
-					FDisplayMetrics DisplayMetrics;
-					FSlateApplication::Get().GetCachedDisplayMetrics( DisplayMetrics );
-
-					const FVector2D DisplaySize(
-						DisplayMetrics.PrimaryDisplayWorkAreaRect.Right - DisplayMetrics.PrimaryDisplayWorkAreaRect.Left,
-						DisplayMetrics.PrimaryDisplayWorkAreaRect.Bottom - DisplayMetrics.PrimaryDisplayWorkAreaRect.Top
-					);
-
-					// Use a centered window at the default window size
-					NewWindowPosition.X = 0;
-					NewWindowPosition.Y = 0;
-					NewWindowWidth = 0.75 * DisplaySize.X;
-					NewWindowHeight = 0.75 * DisplaySize.Y;
-					CenterNewWindow = true;
-				}
-
-				bool bUseOSWndBorder = false;
-				bool bRenderDirectlyToWindow = false;
-				bool bEnableStereoRendering = false;
-				if (bPlayStereoscopic)	// @todo vreditor: Is not having an OS window border a problem?  We could spawn a dedicated VR window if so.  What about true fullscreen in VR?
-				{
-					// modify window and viewport properties for VR.
-					bUseOSWndBorder = true;
-					bRenderDirectlyToWindow = true;
-					bEnableStereoRendering = true;
-					CenterNewWindow = true;
-				}
-
-				TSharedPtr<SWindow> PieWindow = PlayInSettings->CustomPIEWindow.Pin();
-
-				const bool bHasCustomWindow = PieWindow.IsValid();
-				if (!bHasCustomWindow)
-				{
-					int32 PosX = NewWindowPosition.X;
-					int32 PosY = NewWindowPosition.Y;
-					const bool CanPlayNetDedicated = [&PlayInSettings] { bool PlayNetDedicated(false); return (PlayInSettings->GetPlayNetDedicated(PlayNetDedicated) && PlayNetDedicated); }();
-					int32 idx = PieWorldContext->PIEInstance - (CanPlayNetDedicated ? 1 : 0);
-					if (idx > 0)
-					{
-						ULevelEditorPlaySettings* LevelEditorPlaySettings = ULevelEditorPlaySettings::StaticClass()->GetDefaultObject<ULevelEditorPlaySettings>();
-						if (idx < LevelEditorPlaySettings->MultipleInstancePositions.Num())
-						{
-							PosX = LevelEditorPlaySettings->MultipleInstancePositions[idx].X;
-							PosY = LevelEditorPlaySettings->MultipleInstancePositions[idx].Y;
-						}
-					}
-
-					PieWindow = SNew(SWindow)
-						.Title(ViewportName)
-						.ScreenPosition(FVector2D(PosX, PosY))
-						.ClientSize(FVector2D(NewWindowWidth, NewWindowHeight))
-						.AutoCenter(CenterNewWindow ? EAutoCenter::PreferredWorkArea : EAutoCenter::None)
-						.UseOSWindowBorder(bUseOSWndBorder)
-						.SaneWindowPlacement(!CenterNewWindow)
-						.SizingRule(ESizingRule::UserSized);
-
-					PieWindow->SetAllowFastUpdate(true);
-				}
-
-
-				// Setup a delegate for switching to the play world on slate input events, drawing and ticking
-				FOnSwitchWorldHack OnWorldSwitch = FOnSwitchWorldHack::CreateUObject( this, &UEditorEngine::OnSwitchWorldForSlatePieWindow );
-				PieWindow->SetOnWorldSwitchHack( OnWorldSwitch );
-
-				if (!bHasCustomWindow)
-				{
-				// Mac does not support parenting, do not keep on top
-#if PLATFORM_MAC
-					FSlateApplication::Get().AddWindow(PieWindow.ToSharedRef());
-#else
-				TSharedRef<SWindow, ESPMode::Fast> MainWindow = FModuleManager::LoadModuleChecked<IMainFrameModule>(TEXT("MainFrame")).GetParentWindow().ToSharedRef();
-				if (PlayInSettings->PIEAlwaysOnTop)
-				{
-						FSlateApplication::Get().AddWindowAsNativeChild(PieWindow.ToSharedRef(), MainWindow, true);
-				}
-				else
-				{
-						FSlateApplication::Get().AddWindow(PieWindow.ToSharedRef());
-				}
-#endif
-				}
-
-				TSharedRef<SOverlay> ViewportOverlayWidgetRef = SNew(SOverlay);
-
-				TSharedRef<SGameLayerManager> GameLayerManagerRef = SNew(SGameLayerManager)
-					.SceneViewport_UObject(this, &UEditorEngine::GetGameSceneViewport, ViewportClient)
-					[
-						ViewportOverlayWidgetRef
-					];
-
-				static const auto CVarPropagateAlpha = IConsoleManager::Get().FindTConsoleVariableDataInt(TEXT("r.PostProcessing.PropagateAlpha"));
-				const EAlphaChannelMode::Type PropagateAlpha = EAlphaChannelMode::FromInt(CVarPropagateAlpha->GetValueOnGameThread());
-				const bool bIgnoreTextureAlpha = (PropagateAlpha != EAlphaChannelMode::AllowThroughTonemapper);
-
-				PieViewportWidget = 
-					SNew( SPIEViewport )
-						.RenderDirectlyToWindow( bRenderDirectlyToWindow )
-						.EnableStereoRendering( bEnableStereoRendering )
-						.IgnoreTextureAlpha(bIgnoreTextureAlpha)
-						[
-							GameLayerManagerRef
-						];
+			return Context.GameViewport;
+		}
 
-				// Create a wrapper widget for PIE viewport to process play world actions
-				TSharedRef<SGlobalPlayWorldActions> GlobalPlayWorldActionsWidgetRef = SNew(SGlobalPlayWorldActions)
-					[
-						PieViewportWidget.ToSharedRef()
-					];
+		CurrentIdx += step;
+	}
 
-				PieWindow->SetContent(GlobalPlayWorldActionsWidgetRef);
+	return NULL;
+}
 
-				if (!bHasCustomWindow)
-				{
-					// Ensure the PIE window appears does not appear behind other windows.
-					PieWindow->BringToFront();
-				}
+void UEditorEngine::RemapGamepadControllerIdForPIE(class UGameViewportClient* InGameViewport, int32 &ControllerId)
+{
+	// Increment the controller id if we are the focused window, and RouteGamepadToSecondWindow is true (and we are running multiple clients).
+	// This cause the focused window to NOT handle the input, decrement controllerID, and pass it to the next window.
+	const ULevelEditorPlaySettings* PlayInSettings = GetDefault<ULevelEditorPlaySettings>();
+	const bool CanRouteGamepadToSecondWindow = [&PlayInSettings]{ bool RouteGamepadToSecondWindow(false); return (PlayInSettings->GetRouteGamepadToSecondWindow(RouteGamepadToSecondWindow) && RouteGamepadToSecondWindow); }();
+	const bool CanRunUnderOneProcess = [&PlayInSettings]{ bool RunUnderOneProcess(false); return (PlayInSettings->GetRunUnderOneProcess(RunUnderOneProcess) && RunUnderOneProcess); }();
+	if ( CanRouteGamepadToSecondWindow && CanRunUnderOneProcess && InGameViewport->GetWindow().IsValid() && InGameViewport->GetWindow()->HasFocusedDescendants())
+	{
+		ControllerId++;
+	}
+}
 
-				ViewportClient->SetViewportOverlayWidget( PieWindow, ViewportOverlayWidgetRef );
-				ViewportClient->SetGameLayerManager(GameLayerManagerRef);
+void UEditorEngine::StartPlayInEditorSession(FRequestPlaySessionParams& InRequestParams)
+{
+	// This reflects that the user has tried to launch a PIE session, but it may still
+	// create one-or-more new processes depending on multiplayer settings.
+	check(InRequestParams.SessionDestination == EPlaySessionDestinationType::InProcess);
 
-				bool bShouldMinimizeRootWindow = bPlayStereoscopic && GEngine->XRSystem.IsValid() && GetDefault<ULevelEditorPlaySettings>()->ShouldMinimizeEditorOnVRPIE;
-				// Set up a notification when the window is closed so we can clean up PIE
-				{
-					struct FLocal
-					{
-						static void OnPIEWindowClosed( const TSharedRef< SWindow >& WindowBeingClosed, TWeakPtr< SViewport > PIEViewportWidget, int32 index, bool bRestoreRootWindow )
-						{
-							// Save off the window position
-							FVector2D PIEWindowPos = WindowBeingClosed->GetPositionInScreen();
-							const float DPIScale = FPlatformApplicationMisc::GetDPIScaleFactorAtPoint(PIEWindowPos.X, PIEWindowPos.Y);
-							PIEWindowPos /= DPIScale;
+	// Broadcast PreBeginPIE before checks that might block PIE below (BeginPIE is broadcast below after the checks)
+	FEditorDelegates::PreBeginPIE.Broadcast(InRequestParams.WorldType == EPlaySessionWorldType::SimulateInEditor);
+	double PIEStartTime = FPlatformTime::Seconds();
+	const FScopedBusyCursor BusyCursor;
 
-							ULevelEditorPlaySettings* LevelEditorPlaySettings = ULevelEditorPlaySettings::StaticClass()->GetDefaultObject<ULevelEditorPlaySettings>();
+	// Block PIE when there is a transaction recording into the undo buffer. This is generally avoided
+	// because we buffer the request for the PIE session until the start of the next frame, but sometimes
+	// transactions can get stuck open due to implementation errors so it's important that we check.
+	if (GEditor->IsTransactionActive())
+	{
+		FFormatNamedArguments Args;
+		Args.Add(TEXT("TransactionName"), GEditor->GetTransactionName());
 
-							if (index <= 0)
-							{
-								// only override the window position if the window isn't being centered
-								if (!LevelEditorPlaySettings->CenterNewWindow)
-								{
-									LevelEditorPlaySettings->NewWindowPosition.X = FPlatformMath::RoundToInt(PIEWindowPos.X);
-									LevelEditorPlaySettings->NewWindowPosition.Y = FPlatformMath::RoundToInt(PIEWindowPos.Y);
-								}
-							}
-							else
-							{
-								if (index >= LevelEditorPlaySettings->MultipleInstancePositions.Num())
-								{
-									LevelEditorPlaySettings->MultipleInstancePositions.SetNum(index + 1);
-								}
+		FText NotificationText;
+		if (InRequestParams.WorldType == EPlaySessionWorldType::SimulateInEditor)
+		{
+			NotificationText = FText::Format(NSLOCTEXT("UnrealEd", "SIECantStartDuringTransaction", "Can't Simulate when performing {TransactionName} operation"), Args);
+		}
+		else
+		{
+			NotificationText = FText::Format(NSLOCTEXT("UnrealEd", "PIECantStartDuringTransaction", "Can't Play In Editor when performing {TransactionName} operation"), Args);
+		}
 
-								LevelEditorPlaySettings->MultipleInstancePositions[index] = FIntPoint(PIEWindowPos.X, PIEWindowPos.Y); 
-							}
+		FNotificationInfo Info(NotificationText);
+		Info.ExpireDuration = 5.0f;
+		Info.bUseLargeFont = true;
+		FSlateNotificationManager::Get().AddNotification(Info);
+		return;
+	}
 
-							LevelEditorPlaySettings->PostEditChange();
-							LevelEditorPlaySettings->SaveConfig();
+	// Prompt the user that Matinee must be closed before PIE can occur. If they don't want
+	// to close Matinee, we can't PIE.
+	if (!PromptMatineeClose())
+	{
+		return;
+	}
 
-							// Route the callback
-							PIEViewportWidget.Pin()->OnWindowClosed( WindowBeingClosed );
+	// Make sure there's no outstanding load requests
+	FlushAsyncLoading();
 
-							if (bRestoreRootWindow)
-							{
-								// restore previously minimized root window.
-								TSharedPtr<SWindow> RootWindow = FGlobalTabmanager::Get()->GetRootWindow();
-								if (RootWindow.IsValid() && RootWindow->IsWindowMinimized())
-								{
-									RootWindow->Restore();
-								}
-							}
-						}
-					};
-				
-					const bool CanPlayNetDedicated = [&PlayInSettings]{ bool PlayNetDedicated(false); return (PlayInSettings->GetPlayNetDedicated(PlayNetDedicated) && PlayNetDedicated); }();
-					PieWindow->SetOnWindowClosed(FOnWindowClosed::CreateStatic(&FLocal::OnPIEWindowClosed, TWeakPtr<SViewport>(PieViewportWidget), 
-						(PlayNumberOfClients == 1) ? 0 : PieWorldContext->PIEInstance - (CanPlayNetDedicated ? 1 : 0), bShouldMinimizeRootWindow));
-				}
+	// Update the Blueprint Debugger 
+	FBlueprintEditorUtils::FindAndSetDebuggableBlueprintInstances();
 
-				// Create a new viewport that the viewport widget will use to render the game
-				SlatePlayInEditorSession.SlatePlayInEditorWindowViewport = MakeShareable( new FSceneViewport( ViewportClient, PieViewportWidget ) );
+	// Broadcast BeginPIE after checks that might block PIE above (PreBeginPIE is broadcast above before the checks)
+	// ToDo: Shouldn't this move below the early-out for Error'd Blueprints?
+	FEditorDelegates::BeginPIE.Broadcast(InRequestParams.WorldType == EPlaySessionWorldType::SimulateInEditor);
 
-				GameLayerManagerRef->SetSceneViewport(SlatePlayInEditorSession.SlatePlayInEditorWindowViewport.Get());
+	UWorld* InWorld = GetEditorWorldContext().World();
 
-				const bool bShouldGameGetMouseControl = GetDefault<ULevelEditorPlaySettings>()->GameGetsMouseControl || (bPlayStereoscopic && GEngine->XRSystem.IsValid());
-				SlatePlayInEditorSession.SlatePlayInEditorWindowViewport->SetPlayInEditorGetsMouseControl(bShouldGameGetMouseControl);
-				PieViewportWidget->SetViewportInterface( SlatePlayInEditorSession.SlatePlayInEditorWindowViewport.ToSharedRef() );
+	// Let navigation know PIE is starting so it can avoid any blueprint creation/deletion/instantiation affect editor map's navmesh changes
+	FNavigationSystem::OnPIEStart(*InWorld);
 
-				FSlateApplication::Get().RegisterViewport(PieViewportWidget.ToSharedRef());
-				
-				SlatePlayInEditorSession.SlatePlayInEditorWindow = PieWindow;
+	ULevelEditorPlaySettings* EditorPlaySettings = InRequestParams.EditorPlaySettings;
+	check(EditorPlaySettings);
 
-				// Let the viewport client know what viewport is using it.  We need to set the Viewport Frame as 
-				// well (which in turn sets the viewport) so that SetRes command will work.
-				ViewportClient->SetViewportFrame(SlatePlayInEditorSession.SlatePlayInEditorWindowViewport.Get());
-				// Mark the viewport as PIE viewport
-				ViewportClient->Viewport->SetPlayInEditorViewport( ViewportClient->bIsPlayInEditorViewport );
+	// Auto-compile dirty blueprints (if needed) and let the user cancel PIE if there are 
+	TArray<UBlueprint*> ErroredBlueprints;
+	FInternalPlayLevelUtils::ResolveDirtyBlueprints(!EditorPlaySettings->AutoRecompileBlueprints, ErroredBlueprints);
 
-				// Change the system resolution to match our window, to make sure game and slate window are kept syncronised
-				FSystemResolution::RequestResolutionChange(NewWindowWidth, NewWindowHeight, EWindowMode::Windowed);
+	// Don't show the dialog if we're in the middle of a demo, just assume they'll work.
+	if (ErroredBlueprints.Num() && !GIsDemoMode)
+	{
+		// There was at least one blueprint with an error, make sure the user is OK with that.
+		bool bContinuePIE = ShowBlueprintErrorDialog(ErroredBlueprints);
 
-				if (bPlayStereoscopic)
-				{
-					GEngine->StereoRenderingDevice->EnableStereo(true);
+		if (!bContinuePIE)
+		{
+			FMessageLog("BlueprintLog").Open(EMessageSeverity::Warning);
 
-					// minimize the root window to provide max performance for the preview.
-					TSharedPtr<SWindow> RootWindow = FGlobalTabmanager::Get()->GetRootWindow();
-					if (RootWindow.IsValid() && bShouldMinimizeRootWindow)
-					{
-						RootWindow->Minimize();
-					}
-				}
+			FEditorDelegates::EndPIE.Broadcast(InRequestParams.WorldType == EPlaySessionWorldType::SimulateInEditor);
+			FNavigationSystem::OnPIEEnd(*InWorld);
+			return;
+		}
+		else
+		{
+			// The user wants to ignore the compiler errors, mark the Blueprints and do not warn them again unless the Blueprint attempts to compile
+			for (UBlueprint* Blueprint : ErroredBlueprints)
+			{
+				Blueprint->bDisplayCompilePIEWarning = false;
 			}
-
-			UGameViewportClient::OnViewportCreated().Broadcast();
 		}
 	}
 
-	if ( GameViewport != NULL && GameViewport->Viewport != NULL )
+	// If there's level already being played, close it. (This may change GWorld). 
+	// This also invalidates PlaySessionRequest so don't use it below.
+	if (PlayWorld)
 	{
-		// Set the game viewport that was just created as a pie viewport.
-		GameViewport->Viewport->SetPlayInEditorViewport( true );
+		// Immediately end the play world.
+		EndPlayMap();
 	}
 
-	// Disable the screensaver when PIE is running.
-	EnableScreenSaver( false );
+	// Register for log processing so we can promote errors/warnings to the message log
+	if (GetDefault<UEditorStyleSettings>()->bPromoteOutputLogWarningsDuringPIE)
+	{
+		OutputLogErrorsToMessageLogProxyPtr = MakeShareable(new FOutputLogErrorsToMessageLogProxy());
+	}
 
+	// Notify the XRSystem that it needs to BeginPlay.
+	if (GEngine->XRSystem.IsValid() && InRequestParams.WorldType == EPlaySessionWorldType::PlayInEditor)
+	{
+		GEngine->XRSystem->OnBeginPlay(*GEngine->GetWorldContextFromWorld(InWorld));
+	}
 
-	EditorWorld->TransferBlueprintDebugReferences(PlayWorld);
+	// Remember old GWorld
+	EditorWorld = InWorld;
 
-	// By this point it is safe to remove the GameInstance from the root and allow it to garbage collected as per usual
-	GameInstance->RemoveFromRoot();
+	// Clear any messages from last time
+	GEngine->ClearOnScreenDebugMessages();
 
-	// Start the game instance, make sure to set the PIE instance global as this is basically a tick
-	FGameInstancePIEResult StartResult = FGameInstancePIEResult::Success();
+	// Start a new PIE log page
 	{
-		FTemporaryPlayInEditorIDOverride OverrideIDHelper(InPIEInstance);
-		StartResult = GameInstance->StartPlayInEditorGameInstance(NewLocalPlayer, GameInstanceParams);
+		const FString WorldPackageName = EditorWorld->GetOutermost()->GetName();
+
+		FFormatNamedArguments Arguments;
+		Arguments.Add(TEXT("Package"), FText::FromString(FPackageName::GetLongPackageAssetName(WorldPackageName)));
+		Arguments.Add(TEXT("TimeStamp"), FText::AsDateTime(FDateTime::Now()));
+
+		FText PIESessionLabel = InRequestParams.WorldType == EPlaySessionWorldType::SimulateInEditor ?
+			FText::Format(LOCTEXT("SIESessionLabel", "SIE session: {Package} ({TimeStamp})"), Arguments) :
+			FText::Format(LOCTEXT("PIESessionLabel", "PIE session: {Package} ({TimeStamp})"), Arguments);
+
+		FMessageLog(NAME_CategoryPIE).NewPage(PIESessionLabel);
 	}
 
-	if (!StartResult.IsSuccess())
+	// Flush all audio sources from the editor world
+	if (FAudioDevice* AudioDevice = EditorWorld->GetAudioDevice())
 	{
-		FMessageDialog::Open(EAppMsgType::Ok, StartResult.FailureReason);
-		RestoreEditorWorld( EditorWorld );
-		EndPlayMap();
-		return nullptr;
+		AudioDevice->Flush(EditorWorld);
+		AudioDevice->ResetInterpolation();
+		AudioDevice->OnBeginPIE(InRequestParams.WorldType == EPlaySessionWorldType::SimulateInEditor);
 	}
 
-	// Set up a delegate to be called in Slate when GWorld needs to change.  Slate does not have direct access to the playworld to switch itself
-	FScopedConditionalWorldSwitcher::SwitchWorldForPIEDelegate = FOnSwitchWorldForPIE::CreateUObject( this, &UEditorEngine::OnSwitchWorldsForPIE );
+	// Mute the editor world so no sounds come from it.
+	EditorWorld->bAllowAudioPlayback = false;
 
-	if( PieViewportWidget.IsValid() )
+	// Can't allow realtime viewports whilst in PIE so disable it for ALL viewports here.
+	DisableRealtimeViewports();
+
+	// Allow the global config to override our ability to create multiple PIE worlds.
+	if (!GEditor->bAllowMultiplePIEWorlds)
 	{
-		// Register the new viewport widget with Slate for viewport specific message routing.
-		FSlateApplication::Get().RegisterGameViewport( PieViewportWidget.ToSharedRef() );
+		EditorPlaySettings->SetRunUnderOneProcess(false);
 	}
 
-	// go back to using the real world as GWorld
-	RestoreEditorWorld( EditorWorld );
+	// If they're Simulating in Editor, we just override them to only one instance (as it'll use the editor world and not duplicate)
+	// We also override them to Offline play mode as networking isn't supported with SIE anyways.
+	if (InRequestParams.WorldType == EPlaySessionWorldType::SimulateInEditor)
+	{
+		EditorPlaySettings->SetPlayNetMode(EPlayNetMode::PIE_Standalone);
+		EditorPlaySettings->SetPlayNumberOfClients(1);
+	}
 
+	// If they need to use the online services, validate that they have provided 
+	// enough pie credentials to launch the desired number of clients.
+	bool bUseOnlineSubsystemForLogin = false;
 	{
-		FFormatNamedArguments Arguments;
-		Arguments.Add(TEXT("MapName"), FText::FromString(GameInstance->PIEMapName));
-		Arguments.Add(TEXT("StartTime"), FPlatformTime::Seconds() - PIEStartTime);
-		FMessageLog(NAME_CategoryPIE).Info(FText::Format(LOCTEXT("PIEStartTime", "Play in editor start time for {MapName} {StartTime}"), Arguments));
+		int32 DesiredNumberOfClients;
+		EditorPlaySettings->GetPlayNumberOfClients(DesiredNumberOfClients);
+		if (SupportsOnlinePIE() && InRequestParams.WorldType != EPlaySessionWorldType::SimulateInEditor)
+		{
+			bool bHasRequiredLogins = DesiredNumberOfClients <= UOnlineEngineInterface::Get()->GetNumPIELogins();
+			if (bHasRequiredLogins)
+			{
+				// If we support online PIE use it even if we're standalone
+				bUseOnlineSubsystemForLogin = true;
+			}
+			else
+			{
+				FText ErrorMsg = LOCTEXT("PIELoginFailure", "Not enough login credentials to launch all PIE instances, change editor settings");
+				UE_LOG(LogPlayLevel, Verbose, TEXT("%s"), *ErrorMsg.ToString());
+				FMessageLog(NAME_CategoryPIE).Warning(ErrorMsg);
+			}
+		}
+
+		UOnlineEngineInterface::Get()->SetShouldTryOnlinePIE(bUseOnlineSubsystemForLogin);
 	}
 
-	// Update the details window with the actors we have just selected
-	GUnrealEd->UpdateFloatingPropertyWindows();
+	PlayInEditorSessionInfo->bUsingOnlinePlatform = bUseOnlineSubsystemForLogin;
+	PlayInEditorSessionInfo->bAnyBlueprintErrors = ErroredBlueprints.Num() > 0;
 
-	// Clean up any editor actors being referenced 
-	GEngine->BroadcastLevelActorListChanged();
+	// Now that we've gotten all of the editor house-keeping out of the way we can finally
+	// start creating world instances and multi player clients!
+	{
+		// First, we handle starting a dedicated server. This can exist as either a separate 
+		// process, or as an internal world.
+		bool bUserWantsSingleProcess;
+		InRequestParams.EditorPlaySettings->GetRunUnderOneProcess(bUserWantsSingleProcess);
 
-	// Set an undo barrier so that transactions prior to PIE can't be undone
-	GUnrealEd->Trans->SetUndoBarrier();
+		EPlayNetMode NetMode;
+		InRequestParams.EditorPlaySettings->GetPlayNetMode(NetMode);
 
-	return GameInstance;
-}
+		// Standalone requires no server, and ListenServer doesn't require a separate server.
+		const bool bNetModeRequiresSeparateServer = NetMode == EPlayNetMode::PIE_Client;
+		const bool bLaunchExtraServerAnyways = InRequestParams.EditorPlaySettings->bLaunchSeparateServer;
+		const bool bNeedsServer = bNetModeRequiresSeparateServer || bLaunchExtraServerAnyways;
+		
 
-void UEditorEngine::OnViewportCloseRequested(FViewport* InViewport)
-{
-	RequestEndPlayMap();
-}
+		// If they require a separate server we'll give the EditorEngine a chance to handle any additional prep-work.
+		if (bNeedsServer)
+		{
+			// Allow the engine to cancel the server request if needed.
+			FGameInstancePIEResult PreCreateResult = PreCreatePIEServerInstance(
+			ErroredBlueprints.Num() > 0, false /*bStartInSpectorMode*/, PIEStartTime, true, PlayInEditorSessionInfo->NumOutstandingPIELogins);
+			if (!PreCreateResult.IsSuccess())
+			{
+				// ToDo: This will skip client creation as well right now. Probably OK though.
+				UE_LOG(LogPlayLevel, Warning, TEXT("PlayInEditor Session Server failed Pre-Create and will not be started."));
+				return;
+			}
 
-FSceneViewport* UEditorEngine::GetGameSceneViewport(UGameViewportClient* ViewportClient) const
-{
-	return ViewportClient->GetGameViewport();
-}
+			// If they don't want single process we launch the server as a separate process. If they do
+			// want single process, it will get handled below as part of client startup.
+			if (!bUserWantsSingleProcess)
+			{
+				const bool bIsDedicatedServer = true;
+				const bool bIsHost = true;
+				const int32 InstanceIndex = 0;
+				LaunchNewProcess(InRequestParams, InstanceIndex, bIsHost, bIsDedicatedServer);
 
-FViewport* UEditorEngine::GetActiveViewport()
-{
-	// Get the Level editor module and request the Active Viewport.
-	FLevelEditorModule& LevelEditorModule = FModuleManager::Get().GetModuleChecked<FLevelEditorModule>( TEXT("LevelEditor") );
+				PlayInEditorSessionInfo->bServerWasLaunched = true;
+			}
+		}
 
-	TSharedPtr<IAssetViewport> ActiveLevelViewport = LevelEditorModule.GetFirstActiveViewport();
+		// If control is pressed, start in spectator mode
+		FModifierKeysState KeysState = FSlateApplication::Get().GetModifierKeys();
+		PlayInEditorSessionInfo->bStartedInSpectatorMode = InRequestParams.WorldType == EPlaySessionWorldType::SimulateInEditor || KeysState.IsControlDown();
 
-	if ( ActiveLevelViewport.IsValid() )
-	{
-		return ActiveLevelViewport->GetActiveViewport();
-	}
-	
-	return nullptr;
-}
+		// Now that the dedicated server was (optionally) started, we'll start as many requested clients as we can.
+		// Because the user indicated they wanted PIE/PINW we'll put the first client in the editor respecting that
+		// setting. Any additional clients will either be in-process new windows, or separate processes based on settings.
+		int32 NumClients;
+		InRequestParams.EditorPlaySettings->GetPlayNumberOfClients(NumClients);
 
-FViewport* UEditorEngine::GetPIEViewport()
-{
-	// Check both cases where the PIE viewport may be, otherwise return NULL if none are found.
-	if( GameViewport )
-	{
-		return GameViewport->Viewport;
-	}
-	else
-	{
-		for (const FWorldContext& WorldContext : WorldList)
+		// If the have a net mode that requires a server but they didn't create (or couldn't create due to single-process
+		// limitations) a dedicated one, then we launch an extra world context acting as a server in-process.
+		const bool bRequiresExtraListenServer = bNeedsServer && !PlayInEditorSessionInfo->bServerWasLaunched;
+		int32 NumRequestedInstances = FMath::Max(NumClients, 1);
+		if (bRequiresExtraListenServer)
 		{
-			if (WorldContext.WorldType == EWorldType::PIE)
+			NumRequestedInstances++;
+		}
+
+		for (int32 InstanceIndex = 0; InstanceIndex < NumRequestedInstances; InstanceIndex++)
+		{
+			// If they are running single-process and they need a server, the first instance  will be the server.
+			const bool bClientIsServer = (InstanceIndex == 0) && (NetMode == EPlayNetMode::PIE_ListenServer || bRequiresExtraListenServer);
+
+			EPlayNetMode LocalNetMode = NetMode;
+
+			// If they're the server, we want to override them to be a ListenServer. This will get ignored if they're secretly a dedicated
+			// server so it's okay.
+			if (bClientIsServer)
 			{
-				// We can't use FindChecked here because when using the dedicated server option we don't initialize this map 
-				//	(we don't use a viewport for the PIE context in this case)
-				FSlatePlayInEditorInfo * SlatePlayInEditorSessionPtr = SlatePlayInEditorMap.Find(WorldContext.ContextHandle);
-				if ( SlatePlayInEditorSessionPtr != nullptr && SlatePlayInEditorSessionPtr->SlatePlayInEditorWindowViewport.IsValid() )
-				{
-					return SlatePlayInEditorSessionPtr->SlatePlayInEditorWindowViewport.Get();
-				}
+				LocalNetMode = EPlayNetMode::PIE_ListenServer;
 			}
+
+			// If they want to launch a Listen Server and have multiple clients, the subsequent clients need to be
+			// treated as Clients so they connect to the listen server instead of launching multiple Listen Servers.
+			if (NetMode == EPlayNetMode::PIE_ListenServer && InstanceIndex > 0)
+			{
+				LocalNetMode = EPlayNetMode::PIE_Client;
+			}
+
+			bool bRunAsDedicated = bClientIsServer && bRequiresExtraListenServer;
+
+			// Create the instance. This can end up creating separate processes if needed based on settings.
+			// This code is separated out of here so it can be re-used by the Late Join flow.
+			CreateNewPlayInEditorInstance(InRequestParams, bRunAsDedicated, LocalNetMode);
+		}
+
+		// If they requested a SIE, we immediately convert the PIE into a SIE. This finds and executes on the first
+		// world context, and doesn't support networking.
+		if (!bUseOnlineSubsystemForLogin && InRequestParams.WorldType == EPlaySessionWorldType::SimulateInEditor)
+		{
+			const bool bNewSession = true;
+			ToggleBetweenPIEandSIE(bNewSession);
 		}
 	}
 
-	return nullptr;
-}
+	// Disable the screen saver when PIE is running
+	EnableScreenSaver(false);
 
-bool UEditorEngine::GetSimulateInEditorViewTransform(FTransform& OutViewTransform) const
-{
-	if (bIsSimulatingInEditor)
+	// Transfer the Blueprint Debug references to the first client that was created.
+	EditorWorld->TransferBlueprintDebugReferences(PlayWorld);
+
+	// If we have actors selected in the Editor, save that information and then re-select
+	// the equivalent actors in the SIE session.
+	if (InRequestParams.WorldType == EPlaySessionWorldType::SimulateInEditor)
 	{
-		// The first PIE world context is the one that can toggle between PIE and SIE
-		for (const FWorldContext& WorldContext : WorldList)
-		{
-			if (WorldContext.WorldType == EWorldType::PIE && !WorldContext.RunAsDedicated)
-			{
-				const FSlatePlayInEditorInfo* SlateInfoPtr = SlatePlayInEditorMap.Find(WorldContext.ContextHandle);
-				if (SlateInfoPtr)
-				{
-					// This is only supported inside SLevelEditor viewports currently
-					TSharedPtr<IAssetViewport> LevelViewport = SlateInfoPtr->DestinationSlateViewport.Pin();
-					if (LevelViewport.IsValid())
-					{
-						FEditorViewportClient& EditorViewportClient = LevelViewport->GetAssetViewportClient();
-						OutViewTransform = FTransform(EditorViewportClient.GetViewRotation(), EditorViewportClient.GetViewLocation());
-						return true;
-					}
-				}
-				break;
-			}
-		}
+		TransferEditorSelectionToSIEInstances();
 	}
-	return false;
-}
 
-void UEditorEngine::ToggleBetweenPIEandSIE( bool bNewSession )
-{
-	bIsToggleBetweenPIEandSIEQueued = false;
+	// Update the details window with the actors we have just selected
+	GUnrealEd->UpdateFloatingPropertyWindows();
 
-	FEditorDelegates::OnPreSwitchBeginPIEAndSIE.Broadcast(bIsSimulatingInEditor);
+	// Clean up any editor actors being referenced 
+	GEngine->BroadcastLevelActorListChanged();
 
-	// The first PIE world context is the one that can toggle between PIE and SIE
-	// Network PIE/SIE toggling is not really meant to be supported.
-	FSlatePlayInEditorInfo * SlateInfoPtr = nullptr;
-	for (const FWorldContext& WorldContext : WorldList)
+	// Set an undo barrier so that transactions prior to PIE can't be undone
+	GUnrealEd->Trans->SetUndoBarrier();
+
+	// Notify that we've changed safe zone ratios for visualization.
 	{
-		if (WorldContext.WorldType == EWorldType::PIE && !WorldContext.RunAsDedicated)
-		{
-			SlateInfoPtr = SlatePlayInEditorMap.Find(WorldContext.ContextHandle);
-			break;
-		}
+		FMargin SafeZoneRatio = EditorPlaySettings->PIESafeZoneOverride;
+		SafeZoneRatio.Left /= (EditorPlaySettings->NewWindowWidth / 2.0f);
+		SafeZoneRatio.Right /= (EditorPlaySettings->NewWindowWidth / 2.0f);
+		SafeZoneRatio.Bottom /= (EditorPlaySettings->NewWindowHeight / 2.0f);
+		SafeZoneRatio.Top /= (EditorPlaySettings->NewWindowHeight / 2.0f);
+		FSlateApplication::Get().OnDebugSafeZoneChanged.Broadcast(SafeZoneRatio, false);
 	}
 
-	if (!SlateInfoPtr)
 	{
-		return;
+		// Temporarily set this information until the deprecated variables are removed.
+		PRAGMA_DISABLE_DEPRECATION_WARNINGS
+		bIsSimulatingInEditor = InRequestParams.WorldType == EPlaySessionWorldType::SimulateInEditor;
+		PRAGMA_ENABLE_DEPRECATION_WARNINGS
+	}
+
+	FEditorDelegates::PostPIEStarted.Broadcast(InRequestParams.WorldType == EPlaySessionWorldType::SimulateInEditor);
+}
+
+/** Creates an GameInstance with the given settings. A window is created if this isn't a server. */
+UGameInstance* UEditorEngine::CreateInnerProcessPIEGameInstance(FRequestPlaySessionParams& InParams, const FGameInstancePIEParameters& InPIEParameters, int32 InPIEInstanceIndex)
+{
+	// Create a GameInstance for this new instance.
+	FSoftClassPath GameInstanceClassName = GetDefault<UGameMapsSettings>()->GameInstanceClass;
+	UClass* GameInstanceClass = GameInstanceClassName.TryLoadClass<UGameInstance>();
+
+	// If an invalid class type was specified we fall back to the default.
+	if (!GameInstanceClass)
+	{
+		GameInstanceClass = UGameInstance::StaticClass();
 	}
 
-	if( FEngineAnalytics::IsAvailable() && !bNewSession )
+	UGameInstance* GameInstance = NewObject<UGameInstance>(this, GameInstanceClass);
+
+	// We need to temporarily add the GameInstance to the root because the InitializeForPlayInEditor
+	// call can do garbage collection wiping out the GameInstance
+	GameInstance->AddToRoot();
+
+	// Attempt to initialize the GameInstance. This will construct the world.
+	const FGameInstancePIEResult InitializeResult = GameInstance->InitializeForPlayInEditor(InPIEInstanceIndex, InPIEParameters);
+	if (!InitializeResult.IsSuccess())
 	{
-		FString ToggleType = bIsSimulatingInEditor ? TEXT("SIEtoPIE") : TEXT("PIEtoSIE");
+		FMessageDialog::Open(EAppMsgType::Ok, InitializeResult.FailureReason);
+		FEditorDelegates::EndPIE.Broadcast(InPIEParameters.bSimulateInEditor);
+		FNavigationSystem::OnPIEEnd(*EditorWorld);
 
-		FEngineAnalytics::GetProvider().RecordEvent(TEXT("Editor.Usage.PIE"), TEXT("ToggleBetweenPIEandSIE"), ToggleType );
+		GameInstance->RemoveFromRoot();
+		return nullptr;
 	}
 
-	FSlatePlayInEditorInfo & SlatePlayInEditorSession = *SlateInfoPtr;
+	// Our game instance was successfully created
+	FWorldContext* const PieWorldContext = GameInstance->GetWorldContext();
+	check(PieWorldContext);
+	PlayWorld = PieWorldContext->World();
 
-	// This is only supported inside SLevelEditor viewports currently
-	TSharedPtr<IAssetViewport> LevelViewport = SlatePlayInEditorSession.DestinationSlateViewport.Pin();
-	if( ensure(LevelViewport.IsValid()) )
-	{
-		FEditorViewportClient& EditorViewportClient = LevelViewport->GetAssetViewportClient();
+	// Temporarily set GWorld to our newly created world. This utility function
+	// also sets GIsPlayInEditorWorld so that users can know if GWorld is actually
+	// a PIE world or not.
+	SetPlayInEditorWorld(PlayWorld);
 
-		// Toggle to pie if currently simulating
-		if( bIsSimulatingInEditor )
-		{
-			// The undo system may have a reference to a SIE object that is about to be destroyed, so clear the transactions
-			ResetTransaction( NSLOCTEXT("UnrealEd", "ToggleBetweenPIEandSIE", "Toggle Between PIE and SIE") );
+	// Initialize a local player and viewport client for non-dedicated server instances.
+	UGameViewportClient* ViewportClient = nullptr;
+	ULocalPlayer *NewLocalPlayer = nullptr;
+	TSharedPtr<SPIEViewport> PIEViewport = nullptr;
 
-			// The Game's viewport needs to know about the change away from simluate before the PC is (potentially) created
-			GameViewport->GetGameViewport()->SetPlayInEditorIsSimulate(false);
+	if (!InPIEParameters.bRunAsDedicated)
+	{
+		bool bCreateNewAudioDevice = InParams.EditorPlaySettings->IsCreateAudioDeviceForEveryPlayer();
 
-			// The editor viewport client wont be visible so temporarily disable it being realtime
-			EditorViewportClient.SetRealtime( false, true );
+		// Create an instance of the Game Viewport Client, with the class specified by the Engine.
+		ViewportClient = NewObject<UGameViewportClient>(this, GameViewportClientClass);
+		ViewportClient->Init(*PieWorldContext, GameInstance, bCreateNewAudioDevice);
 
-			if (!SlatePlayInEditorSession.EditorPlayer.IsValid())
+		if (!InParams.EditorPlaySettings->EnableGameSound)
+		{
+			if (FAudioDevice* GameInstanceAudioDevice = GameInstance->GetWorld()->GetAudioDevice())
 			{
-				OnSwitchWorldsForPIE(true);
+				GameInstanceAudioDevice->SetTransientMasterVolume(0.0f);
+			}
+		}
 
-				UWorld* World = GameViewport->GetWorld();
-				AGameModeBase* AuthGameMode = World->GetAuthGameMode();
-				if (AuthGameMode && GameViewport->GetGameInstance())	// If there is no GameMode, we are probably the client and cannot RestartPlayer.
-				{
-					AuthGameMode->SpawnPlayerFromSimulate(EditorViewportClient.GetViewLocation(), EditorViewportClient.GetViewRotation());
-				}
+		GameViewport = ViewportClient;
+		GameViewport->bIsPlayInEditorViewport = true;
 
-				OnSwitchWorldsForPIE(false);
-			}
+		// Update our World Context to know which Viewport Client is associated.
+		PieWorldContext->GameViewport = ViewportClient;
 
-			// A game viewport already exists, tell the level viewport its in to swap to it
-			LevelViewport->SwapViewportsForPlayInEditor();
+		// Add a callback for Game Input that isn't absorbed by the Game Viewport. This allows us to
+		// make editor commands work (such as Shift F1, etc.) from within PIE.
+		ViewportClient->OnGameViewportInputKey().BindUObject(this, &UEditorEngine::ProcessDebuggerCommands);
 
-			// No longer simulating
-			GameViewport->SetIsSimulateInEditorViewport(false);
-			EditorViewportClient.SetIsSimulateInEditorViewport(false);
+		// Listen for when the viewport is closed, so we can see about shutting down PIE.
+		ViewportCloseRequestedDelegateHandle = ViewportClient->OnCloseRequested().AddUObject(this, &UEditorEngine::OnViewportCloseRequested);
+		FSlatePlayInEditorInfo& SlatePlayInEditorSession = SlatePlayInEditorMap.Add(PieWorldContext->ContextHandle, FSlatePlayInEditorInfo());
 
-			FEditorModeRegistry::Get().UnregisterMode(FBuiltinEditorModes::EM_Physics);
+		// Might be invalid depending how pie was launched. Code below handles this
+		if (InParams.DestinationSlateViewport.Get(nullptr).IsValid())
+		{
+			SlatePlayInEditorSession.DestinationSlateViewport = InParams.DestinationSlateViewport.GetValue();
 			
-			bIsSimulatingInEditor = false;
+			// Only one PIE Instance can live in a given viewport, so we'll null it out so that we create
+			// windows instead for the remaining clients.
+			InParams.DestinationSlateViewport = nullptr;
 		}
-		else
-		{
-			// Swap to simulate from PIE
-			LevelViewport->SwapViewportsForSimulateInEditor();
-	
-			GameViewport->SetIsSimulateInEditorViewport(true);
-			GameViewport->GetGameViewport()->SetPlayInEditorIsSimulate(true);
-			EditorViewportClient.SetIsSimulateInEditorViewport(true);
 
-		
-			TSharedRef<FPhysicsManipulationEdModeFactory> Factory = MakeShareable(new FPhysicsManipulationEdModeFactory);
-			FEditorModeRegistry::Get().RegisterMode(FBuiltinEditorModes::EM_Physics, Factory);
-			
-			bIsSimulatingInEditor = true;
+		// Attempt to initialize a Local Player.
+		FString Error;
+		NewLocalPlayer = ViewportClient->SetupInitialLocalPlayer(Error);
+		if (!NewLocalPlayer)
+		{
+			FMessageDialog::Open(EAppMsgType::Ok, FText::Format(NSLOCTEXT("UnrealEd", "Error_CouldntSpawnPlayer", "Couldn't spawn player: {0}"), FText::FromString(Error)));
+			// go back to using the real world as GWorld
+			RestoreEditorWorld(EditorWorld);
+			EndPlayMap();
+			GameInstance->RemoveFromRoot();
+			return nullptr;
+		}
 
-			// Make sure the viewport is in real-time mode
-			EditorViewportClient.SetRealtime( true );
+		// A Local Player gets created even in SIE (which is different than a Player Controller), but we only
+		// store a reference if we're PIE for the UI to know where to restore our viewport location after PIE closes.
+		if (!InPIEParameters.bSimulateInEditor)
+		{
+			SlatePlayInEditorSession.EditorPlayer = NewLocalPlayer;
+		}
 
-			// The Simulate window should show stats
-			EditorViewportClient.SetShowStats( true );
+		// Note: For K2 debugging purposes this MUST be created before beginplay is called because beginplay can trigger breakpoints
+		// and we need to be able to refocus the pie viewport afterwards so it must be created first in order for us to find it
+		{
+			// If the original request provided a Slate Viewport, we'll use that for our output.
+			if (SlatePlayInEditorSession.DestinationSlateViewport.IsValid())
+			{
+				TSharedPtr<IAssetViewport> LevelViewportRef = SlatePlayInEditorSession.DestinationSlateViewport.Pin();
+				LevelViewportRef->StartPlayInEditorSession(ViewportClient, InParams.WorldType == EPlaySessionWorldType::SimulateInEditor);
 
-			if( SlatePlayInEditorSession.EditorPlayer.IsValid() )
+				// We count this as a viewport being created so that subsequent clients won't think they're the 'first' and use the wrong setting.
+				PlayInEditorSessionInfo->NumViewportInstancesCreated++;
+			}
+			else
 			{
-				// Move the editor camera to where the player was.  
-				FVector ViewLocation;
-				FRotator ViewRotation;
-				SlatePlayInEditorSession.EditorPlayer.Get()->PlayerController->GetPlayerViewPoint( ViewLocation, ViewRotation );
-				EditorViewportClient.SetViewLocation( ViewLocation );
+				// Generate a new Window to put this instance in.
+				EPlayNetMode NetMode;
+				InParams.EditorPlaySettings->GetPlayNetMode(NetMode);
+				PIEViewport = GeneratePIEViewportWindow(InParams, PlayInEditorSessionInfo->NumViewportInstancesCreated, *PieWorldContext, NetMode, ViewportClient, SlatePlayInEditorSession);
 
-				if( EditorViewportClient.IsPerspective() )
-				{
-					// Rotation only matters for perspective viewports not orthographic
-					EditorViewportClient.SetViewRotation( ViewRotation );
-				}
+				// Increment for each viewport so that the window titles get correct numbers and it uses the right save/load setting. Non-visible
+				// servers won't be bumping this number as it's used for saving/restoring window positions.
+				PlayInEditorSessionInfo->NumViewportInstancesCreated++;
 			}
+
+
+			// Broadcast that the Viewport has been successfully created.
+			UGameViewportClient::OnViewportCreated().Broadcast();
+		}
+
+		// Mark the Viewport as a PIE Viewport
+		if (GameViewport && GameViewport->Viewport)
+		{
+			GameViewport->Viewport->SetPlayInEditorViewport(true);
 		}
 	}
 
-	// Backup ActorsThatWereSelected as this will be cleared whilst deselecting
-	TArray<TWeakObjectPtr<class AActor> > BackupOfActorsThatWereSelected(ActorsThatWereSelected);
+	// By this point it is safe to remove the GameInstance from the root and allow it to garbage collected as per usual
+	GameInstance->RemoveFromRoot();
 
-	// Unselect everything
-	GEditor->SelectNone( true, true, false );
-	GetSelectedActors()->DeselectAll();
-	GetSelectedObjects()->DeselectAll();
+	FGameInstancePIEResult StartResult = FGameInstancePIEResult::Success();
+	{
+		FTemporaryPlayInEditorIDOverride OverrideIDHelper(InPIEInstanceIndex);
+		StartResult = GameInstance->StartPlayInEditorGameInstance(NewLocalPlayer, InPIEParameters);
+	}
 
-	// restore the backup
-	ActorsThatWereSelected = BackupOfActorsThatWereSelected;
+	if (!StartResult.IsSuccess())
+	{
+		FMessageDialog::Open(EAppMsgType::Ok, StartResult.FailureReason);
+		RestoreEditorWorld(EditorWorld);
+		EndPlayMap();
+		return nullptr;
+	}
 
-	// make sure each selected actors sim equivalent is selected if we're Simulating but not if we're Playing
-	for ( int32 ActorIndex = 0; ActorIndex < ActorsThatWereSelected.Num(); ++ActorIndex )
+	// Set up a delegate to be called in Slate when GWorld needs to change.  Slate does not have direct access to the playworld to switch itself
+	FScopedConditionalWorldSwitcher::SwitchWorldForPIEDelegate = FOnSwitchWorldForPIE::CreateUObject(this, &UEditorEngine::OnSwitchWorldsForPIE);
+
+	if (PIEViewport.IsValid())
 	{
-		TWeakObjectPtr<AActor> Actor = ActorsThatWereSelected[ ActorIndex ].Get();
-		if (Actor.IsValid())
-		{
-			AActor* SimActor = EditorUtilities::GetSimWorldCounterpartActor(Actor.Get());
-			if (SimActor && !SimActor->IsHidden())
-			{
-				SelectActor( SimActor, bIsSimulatingInEditor, false );
-			}
-		}
+		// Register the new viewport widget with Slate for viewport specific message routing.
+		FSlateApplication::Get().RegisterGameViewport(PIEViewport.ToSharedRef());
 	}
 
-	FEditorDelegates::OnSwitchBeginPIEAndSIE.Broadcast( bIsSimulatingInEditor );
+	// Go back to using the editor world as GWorld.
+	RestoreEditorWorld(EditorWorld);
+
+	return GameInstance;
 }
 
-int32 UEditorEngine::OnSwitchWorldForSlatePieWindow( int32 WorldID )
+FText GeneratePIEViewportWindowTitle(const EPlayNetMode InNetMode, const ERHIFeatureLevel::Type InFeatureLevel, const int32 ClientIndex)
 {
-	static const int32 EditorWorldID = 0;
-	static const int32 PieWorldID = 1;
+#if PLATFORM_64BITS
+	const FString PlatformBitsString(TEXT("64"));
+#else
+	const FString PlatformBitsString(TEXT("32"));
+#endif
 
-	int32 RestoreID = -1;
-	if( WorldID == -1 && GWorld != PlayWorld && PlayWorld != NULL)
-	{
-		// When we have an invalid world id we always switch to the pie world in the PIE window
-		const bool bSwitchToPIE = true; 
-		OnSwitchWorldsForPIE( bSwitchToPIE );
-		// The editor world was active restore it later
-		RestoreID = EditorWorldID;
-	}
-	else if( WorldID == PieWorldID && GWorld != PlayWorld)
+	const FText WindowTitleOverride = GetDefault<UGeneralProjectSettings>()->ProjectDisplayedTitle;
+
+	FFormatNamedArguments Args;
+	Args.Add(TEXT("GameName"), FText::FromString(FString(WindowTitleOverride.IsEmpty() ? FApp::GetProjectName() : WindowTitleOverride.ToString())));
+	Args.Add(TEXT("PlatformBits"), FText::FromString(PlatformBitsString));
+	Args.Add(TEXT("RHIName"), FText::FromName(ShaderPlatformToPlatformName(GetFeatureLevelShaderPlatform(InFeatureLevel))));
+
+	if (InNetMode == PIE_Client)
 	{
-		const bool bSwitchToPIE = true;
-		// Want to restore the PIE world and the current world is not already the pie world
-		OnSwitchWorldsForPIE( bSwitchToPIE );
+		Args.Add(TEXT("NetMode"), FText::FromString(FString::Printf(TEXT("Client %d"), ClientIndex + 1)));
 	}
-	else if( WorldID == EditorWorldID && GWorld != EditorWorld)
+	else if (InNetMode == PIE_ListenServer)
 	{
-		const bool bSwitchToPIE = false;
-		// Want to restore the editor world and the current world is not already the editor world
-		OnSwitchWorldsForPIE( bSwitchToPIE );
+		Args.Add(TEXT("NetMode"), FText::FromString(TEXT("Server")));
 	}
 	else
 	{
-		// Current world is already the same as the world being switched to (nested calls to this for example)
+		Args.Add(TEXT("NetMode"), FText::FromString(TEXT("Standalone")));
 	}
 
-	return RestoreID;
+	return FText::Format(NSLOCTEXT("UnrealEd", "PlayInEditor_WindowTitleFormat", "{GameName} Preview [NetMode: {NetMode}] ({PlatformBits}-bit/{RHIName})"), Args);
 }
 
-void UEditorEngine::OnSwitchWorldsForPIE( bool bSwitchToPieWorld, UWorld* OverrideWorld )
+void UEditorEngine::TransferEditorSelectionToSIEInstances()
 {
-	if( bSwitchToPieWorld )
+	// Make a list of all the selected actors
+	TArray<UObject *> SelectedActors;
+	TArray<UObject*> SelectedComponents;
+	for (FSelectionIterator It(GetSelectedActorIterator()); It; ++It)
 	{
-		SetPlayInEditorWorld( OverrideWorld ? OverrideWorld : PlayWorld );
+		AActor* Actor = static_cast<AActor*>(*It);
+		if (Actor)
+		{
+			checkSlow(Actor->IsA(AActor::StaticClass()));
+
+			SelectedActors.Add(Actor);
+		}
 	}
-	else
+
+
+	// Unselect everything
+	GEditor->SelectNone(true, true, false);
+	GetSelectedActors()->DeselectAll();
+	GetSelectedObjects()->DeselectAll();
+	GetSelectedComponents()->DeselectAll();
+
+	ActorsThatWereSelected.Empty();
+
+	// For every actor that was selected previously, make sure it's sim equivalent is selected
+	for (int32 ActorIndex = 0; ActorIndex < SelectedActors.Num(); ++ActorIndex)
 	{
-		RestoreEditorWorld( OverrideWorld ? OverrideWorld : EditorWorld );
+		AActor* Actor = Cast<AActor>(SelectedActors[ActorIndex]);
+		if (Actor)
+		{
+			ActorsThatWereSelected.Add(Actor);
+
+			AActor* SimActor = EditorUtilities::GetSimWorldCounterpartActor(Actor);
+			if (SimActor && !SimActor->bHidden)
+			{
+				SelectActor(SimActor, true, false);
+			}
+		}
 	}
 }
 
-UWorld* UEditorEngine::CreatePIEWorldByDuplication(FWorldContext &WorldContext, UWorld* InWorld, FString &PlayWorldMapName)
+TSharedRef<SPIEViewport> UEditorEngine::GeneratePIEViewportWindow(const FRequestPlaySessionParams& InSessionParams, int32 InViewportIndex, const FWorldContext& InWorldContext, EPlayNetMode InNetMode, UGameViewportClient* InViewportClient, FSlatePlayInEditorInfo& InSlateInfo)
 {
-	double StartTime = FPlatformTime::Seconds();
-	UPackage* InPackage = InWorld->GetOutermost();
-	UWorld* NewPIEWorld = NULL;
+	FIntPoint WindowSize, WindowPosition;
+	GetWindowSizeAndPositionForInstanceIndex(*InSessionParams.EditorPlaySettings, InViewportIndex, WindowSize, WindowPosition);
+	bool bCenterNewWindowOverride = false;
 	
-	const FString WorldPackageName = InPackage->GetName();
+	// VR Preview overrides window location.
+	const bool bVRPreview = InSessionParams.SessionPreviewTypeOverride.Get(EPlaySessionPreviewType::NoPreview) == EPlaySessionPreviewType::VRPreview;
+	bool bUseOSWndBorder = bVRPreview;
+	if (bVRPreview)
+	{
+		bCenterNewWindowOverride = true;
+	}
 
-	// Preserve the old path keeping EditorWorld name the same
-	PlayWorldMapName = UWorld::ConvertToPIEPackageName(WorldPackageName, WorldContext.PIEInstance);
+	// Check to see if they've provided a custom SWindow for us to place our Session in.
+	TSharedPtr<SWindow> PieWindow = InSessionParams.CustomPIEWindow.Pin();
+	const bool bHasCustomWindow = PieWindow.IsValid();
 
-	// Display a busy cursor while we prepare the PIE world
-	const FScopedBusyCursor BusyCursor;
+	// If they haven't provided a Slate Window (common), we will create one.
+	if (!bHasCustomWindow)
+	{
 
-	// Before loading the map, we need to set these flags to true so that postload will work properly
-	TGuardValue<bool> OverrideIsPlayWorld(GIsPlayInEditorWorld, true);
+		FText ViewportName = GeneratePIEViewportWindowTitle(InNetMode, PreviewPlatform.GetEffectivePreviewFeatureLevel(), InViewportIndex);
+		PieWindow = SNew(SWindow)
+			.Title(ViewportName)
+			.ScreenPosition(FVector2D(WindowPosition.X, WindowPosition.Y))
+			.ClientSize(FVector2D(WindowSize.X, WindowSize.Y))	
+			.AutoCenter(bCenterNewWindowOverride ? EAutoCenter::PreferredWorkArea : EAutoCenter::None)
+			.UseOSWindowBorder(bUseOSWndBorder)
+			.SaneWindowPlacement(!bCenterNewWindowOverride)
+			.SizingRule(ESizingRule::UserSized);
 
-	const FName PlayWorldMapFName = FName(*PlayWorldMapName);
-	UWorld::WorldTypePreLoadMap.FindOrAdd(PlayWorldMapFName) = EWorldType::PIE;
+		PieWindow->SetAllowFastUpdate(true);
+	}
 
-	// Create a package for the PIE world
-	UE_LOG( LogPlayLevel, Log, TEXT("Creating play world package: %s"),  *PlayWorldMapName );	
 
-	UPackage* PlayWorldPackage = CreatePackage(nullptr,*PlayWorldMapName);
-	PlayWorldPackage->SetPackageFlags(PKG_PlayInEditor);
-	PlayWorldPackage->PIEInstanceID = WorldContext.PIEInstance;
-	PlayWorldPackage->FileName = InPackage->FileName;
-	PlayWorldPackage->SetGuid( InPackage->GetGuid() );
-	PlayWorldPackage->MarkAsFullyLoaded();
+	// Setup a delegate for switching to the play world on slate input events, drawing and ticking
+	FOnSwitchWorldHack OnWorldSwitch = FOnSwitchWorldHack::CreateUObject(this, &UEditorEngine::OnSwitchWorldForSlatePieWindow);
+	PieWindow->SetOnWorldSwitchHack(OnWorldSwitch);
+
+	if (!bHasCustomWindow)
+	{
+		// Mac does not support parenting, do not keep on top
+#if PLATFORM_MAC
+		FSlateApplication::Get().AddWindow(PieWindow.ToSharedRef());
+#else
+		TSharedRef<SWindow, ESPMode::Fast> MainWindow = FModuleManager::LoadModuleChecked<IMainFrameModule>(TEXT("MainFrame")).GetParentWindow().ToSharedRef();
+		if (InSessionParams.EditorPlaySettings->PIEAlwaysOnTop)
+		{
+			FSlateApplication::Get().AddWindowAsNativeChild(PieWindow.ToSharedRef(), MainWindow, true);
+		}
+		else
+		{
+			FSlateApplication::Get().AddWindow(PieWindow.ToSharedRef());
+		}
+#endif
+	}
+
+	TSharedRef<SOverlay> ViewportOverlayWidgetRef = SNew(SOverlay);
+
+	TSharedRef<SGameLayerManager> GameLayerManagerRef = SNew(SGameLayerManager)
+		.SceneViewport_UObject(this, &UEditorEngine::GetGameSceneViewport, InViewportClient)
+		[
+			ViewportOverlayWidgetRef
+		];
+
+
+	bool bRenderDirectlyToWindow = bVRPreview;
+	bool bEnableStereoRendering = bVRPreview;
+
+	static const auto CVarPropagateAlpha = IConsoleManager::Get().FindTConsoleVariableDataInt(TEXT("r.PostProcessing.PropagateAlpha"));
+	const EAlphaChannelMode::Type PropagateAlpha = EAlphaChannelMode::FromInt(CVarPropagateAlpha->GetValueOnGameThread());
+	const bool bIgnoreTextureAlpha = (PropagateAlpha != EAlphaChannelMode::AllowThroughTonemapper);
+
+	TSharedRef<SPIEViewport> PieViewportWidget =
+		SNew(SPIEViewport)
+		.RenderDirectlyToWindow(bRenderDirectlyToWindow)
+		.EnableStereoRendering(bEnableStereoRendering)
+		.IgnoreTextureAlpha(bIgnoreTextureAlpha)
+		[
+			GameLayerManagerRef
+		];
 
-	// check(GPlayInEditorID == -1 || GPlayInEditorID == WorldContext.PIEInstance);
-	// Currently GPlayInEditorID is not correctly reset after map loading, so it's not safe to assert here
-	FTemporaryPlayInEditorIDOverride IDHelper(WorldContext.PIEInstance);
+	// Create a wrapper widget for PIE viewport to process play world actions
+	TSharedRef<SGlobalPlayWorldActions> GlobalPlayWorldActionsWidgetRef =
+		SNew(SGlobalPlayWorldActions)
+		[
+			PieViewportWidget
+		];
 
+	PieWindow->SetContent(GlobalPlayWorldActionsWidgetRef);
+
+	if (!bHasCustomWindow)
 	{
-		double SDOStart = FPlatformTime::Seconds();
+		// Ensure the PIE window appears does not appear behind other windows.
+		PieWindow->BringToFront();
+	}
 
-		// Reset any GUID fixups with lazy pointers
-		FLazyObjectPtr::ResetPIEFixups();
+	InViewportClient->SetViewportOverlayWidget(PieWindow, ViewportOverlayWidgetRef);
+	InViewportClient->SetGameLayerManager(GameLayerManagerRef);
 
-		// Prepare soft object paths for fixup
-		FSoftObjectPath::AddPIEPackageName(FName(*PlayWorldMapName));
-		for (ULevelStreaming* StreamingLevel : InWorld->GetStreamingLevels())
+	bool bShouldMinimizeRootWindow = bVRPreview && GEngine->XRSystem.IsValid() && InSessionParams.EditorPlaySettings->ShouldMinimizeEditorOnVRPIE;
+	// Set up a notification when the window is closed so we can clean up PIE
+	{
+		struct FLocal
 		{
-			if (StreamingLevel && !StreamingLevel->HasAllFlags(RF_DuplicateTransient))
+			static void OnPIEWindowClosed(const TSharedRef<SWindow>& WindowBeingClosed, TWeakPtr<SViewport> PIEViewportWidget, TWeakObjectPtr<UEditorEngine> OwningEditorEngine, int32 InInstanceIndex, bool bRestoreRootWindow)
 			{
-				FString StreamingLevelPIEName = UWorld::ConvertToPIEPackageName(StreamingLevel->GetWorldAssetPackageName(), WorldContext.PIEInstance);
-				FSoftObjectPath::AddPIEPackageName(FName(*StreamingLevelPIEName));
-			}
-		}
+				// Save off the window position
+				FVector2D PIEWindowPos = WindowBeingClosed->GetPositionInScreen();
+				const float DPIScale = FPlatformApplicationMisc::GetDPIScaleFactorAtPoint(PIEWindowPos.X, PIEWindowPos.Y);
+				PIEWindowPos /= DPIScale;
 
-		// NULL GWorld before various PostLoad functions are called, this makes it easier to debug invalid GWorld accesses
-		GWorld = NULL;
+				FIntPoint WindowSize = FIntPoint(WindowBeingClosed->GetClientSizeInScreen().X, WindowBeingClosed->GetClientSizeInScreen().Y);
+				FIntPoint WindowPosition = FIntPoint(PIEWindowPos.X, PIEWindowPos.Y);
 
-		// Duplicate the editor world to create the PIE world
-		NewPIEWorld = CastChecked<UWorld>( StaticDuplicateObject(
-			InWorld,				// Source root
-			PlayWorldPackage,		// Destination root
-			InWorld->GetFName(),	// Name for new object
-			RF_AllFlags,			// FlagMask
-			NULL,					// DestClass
-			EDuplicateMode::PIE
-			) );
+				if (OwningEditorEngine.IsValid())
+				{
+					OwningEditorEngine->StoreWindowSizeAndPositionForInstanceIndex(InInstanceIndex, WindowSize, WindowPosition);
+				}
 
-		// Store prefix we used to rename this world and streaming levels package names
-		NewPIEWorld->StreamingLevelsPrefix = UWorld::BuildPIEPackagePrefix(WorldContext.PIEInstance);
-		// Fixup model components. The index buffers have been created for the components in the source world and the order
-		// in which components were post-loaded matters. So don't try to guarantee a particular order here, just copy the
-		// elements over.
-		if ( NewPIEWorld->PersistentLevel->Model != NULL
-			&& NewPIEWorld->PersistentLevel->Model == InWorld->PersistentLevel->Model
-			&& NewPIEWorld->PersistentLevel->ModelComponents.Num() == InWorld->PersistentLevel->ModelComponents.Num() )
-		{
-			NewPIEWorld->PersistentLevel->Model->ClearLocalMaterialIndexBuffersData();
-			for (int32 ComponentIndex = 0; ComponentIndex < NewPIEWorld->PersistentLevel->ModelComponents.Num(); ++ComponentIndex)
-			{
-				UModelComponent* SrcComponent = InWorld->PersistentLevel->ModelComponents[ComponentIndex];
-				UModelComponent* DestComponent = NewPIEWorld->PersistentLevel->ModelComponents[ComponentIndex];
-				DestComponent->CopyElementsFrom(SrcComponent);
+				// Route the callback
+				PIEViewportWidget.Pin()->OnWindowClosed(WindowBeingClosed);
+
+				if (bRestoreRootWindow)
+				{
+					// restore previously minimized root window.
+					TSharedPtr<SWindow> RootWindow = FGlobalTabmanager::Get()->GetRootWindow();
+					if (RootWindow.IsValid() && RootWindow->IsWindowMinimized())
+					{
+						RootWindow->Restore();
+					}
+				}
 			}
-		}
+		};
 
-		UE_LOG(LogPlayLevel, Log, TEXT("PIE: StaticDuplicateObject took: (%fs)"),  float(FPlatformTime::Seconds() - SDOStart));		
+		const bool bHasServer = InNetMode == EPlayNetMode::PIE_Client;
+		int32 InstanceIndex = bHasServer ? InWorldContext.PIEInstance - 1 : InWorldContext.PIEInstance;
+
+		PieWindow->SetOnWindowClosed(FOnWindowClosed::CreateStatic(&FLocal::OnPIEWindowClosed, TWeakPtr<SViewport>(PieViewportWidget), TWeakObjectPtr<UEditorEngine>(this),
+			InstanceIndex, bShouldMinimizeRootWindow));
 	}
 
-	// Clean up the world type list now that PostLoad has occurred
-	UWorld::WorldTypePreLoadMap.Remove(PlayWorldMapFName);
+	// Create a new viewport that the viewport widget will use to render the game
+	InSlateInfo.SlatePlayInEditorWindowViewport = MakeShared<FSceneViewport>(InViewportClient, PieViewportWidget);
 
-	check( NewPIEWorld );
-	NewPIEWorld->FeatureLevel = InWorld->FeatureLevel;
-	PostCreatePIEWorld(NewPIEWorld);
+	GameLayerManagerRef->SetSceneViewport(InSlateInfo.SlatePlayInEditorWindowViewport.Get());
 
-	UE_LOG(LogPlayLevel, Log, TEXT("PIE: Created PIE world by copying editor world from %s to %s (%fs)"), *InWorld->GetPathName(), *NewPIEWorld->GetPathName(), float(FPlatformTime::Seconds() - StartTime));
-	return NewPIEWorld;
-}
+	const bool bShouldGameGetMouseControl = InSessionParams.EditorPlaySettings->GameGetsMouseControl || (bEnableStereoRendering && GEngine->XRSystem.IsValid());
+	InSlateInfo.SlatePlayInEditorWindowViewport->SetPlayInEditorGetsMouseControl(bShouldGameGetMouseControl);
+	PieViewportWidget->SetViewportInterface(InSlateInfo.SlatePlayInEditorWindowViewport.ToSharedRef());
 
-void UEditorEngine::PostCreatePIEWorld(UWorld *NewPIEWorld)
-{
-	double WorldInitStart = FPlatformTime::Seconds();
-	
-	// Init the PIE world
-	NewPIEWorld->WorldType = EWorldType::PIE;
-	NewPIEWorld->InitWorld();
-	UE_LOG(LogPlayLevel, Log, TEXT("PIE: World Init took: (%fs)"),  float(FPlatformTime::Seconds() - WorldInitStart));
+	FSlateApplication::Get().RegisterViewport(PieViewportWidget);
 
-	// Tag PlayWorld Actors that also exist in EditorWorld.  At this point, no temporary/run-time actors exist in PlayWorld
-	for( FActorIterator PlayActorIt(NewPIEWorld); PlayActorIt; ++PlayActorIt )
-	{
-		GEditor->ObjectsThatExistInEditorWorld.Set(*PlayActorIt);
-	}
-}
+	InSlateInfo.SlatePlayInEditorWindow = PieWindow;
 
-UWorld* UEditorEngine::CreatePIEWorldFromEntry(FWorldContext &WorldContext, UWorld* InWorld, FString &PlayWorldMapName)
-{
-	double StartTime = FPlatformTime::Seconds();
+	// Let the viewport client know what viewport is using it.  We need to set the Viewport Frame as 
+	// well (which in turn sets the viewport) so that SetRes command will work.
+	InViewportClient->SetViewportFrame(InSlateInfo.SlatePlayInEditorWindowViewport.Get());
+	// Mark the viewport as PIE viewport
+	InViewportClient->Viewport->SetPlayInEditorViewport(InViewportClient->bIsPlayInEditorViewport);
 
-	// Create the world
-	UWorld *LoadedWorld = UWorld::CreateWorld( EWorldType::PIE, false );
-	check(LoadedWorld);
-	if (LoadedWorld->GetOutermost() != GetTransientPackage())
+	// Change the system resolution to match our window, to make sure game and slate window are kept synchronized
+	FSystemResolution::RequestResolutionChange(WindowSize.X, WindowSize.Y, EWindowMode::Windowed);
+
+	if (bVRPreview)
 	{
-		LoadedWorld->GetOutermost()->PIEInstanceID = WorldContext.PIEInstance;
+		GEngine->StereoRenderingDevice->EnableStereo(true);
+
+		// minimize the root window to provide max performance for the preview.
+		TSharedPtr<SWindow> RootWindow = FGlobalTabmanager::Get()->GetRootWindow();
+		if (RootWindow.IsValid() && bShouldMinimizeRootWindow)
+		{
+			RootWindow->Minimize();
+		}
 	}
-	// Force default GameMode class so project specific code doesn't fire off. 
-	// We want this world to truly remain empty while we wait for connect!
-	check(LoadedWorld->GetWorldSettings());
-	LoadedWorld->GetWorldSettings()->DefaultGameMode = AGameModeBase::StaticClass();
 
-	PlayWorldMapName = UGameMapsSettings::GetGameDefaultMap();
-	return LoadedWorld;
+	return PieViewportWidget;
 }
 
-bool UEditorEngine::WorldIsPIEInNewViewport(UWorld *InWorld)
+
+/* fits the window position to make sure it falls within the confines of the desktop */
+void FitWindowPositionToWorkArea(FIntPoint &WinPos, FIntPoint &WinSize, const FMargin &WinPadding)
 {
-	FWorldContext &WorldContext = GetWorldContextFromWorldChecked(InWorld);
-	if (WorldContext.WorldType == EWorldType::PIE)
+	const int32 HorzPad = WinPadding.GetTotalSpaceAlong<Orient_Horizontal>();
+	const int32 VertPad = WinPadding.GetTotalSpaceAlong<Orient_Vertical>();
+	FIntPoint TotalSize(WinSize.X + HorzPad, WinSize.Y + VertPad);
+
+	FDisplayMetrics DisplayMetrics;
+	FSlateApplication::Get().GetCachedDisplayMetrics(DisplayMetrics);
+
+	// Limit the size, to make sure it fits within the desktop area
 	{
-		FSlatePlayInEditorInfo * SlateInfoPtr = SlatePlayInEditorMap.Find(WorldContext.ContextHandle);
-		if (SlateInfoPtr)
+		FIntPoint NewWinSize;
+		NewWinSize.X = FMath::Min(TotalSize.X, DisplayMetrics.VirtualDisplayRect.Right - DisplayMetrics.VirtualDisplayRect.Left);
+		NewWinSize.Y = FMath::Min(TotalSize.Y, DisplayMetrics.VirtualDisplayRect.Bottom - DisplayMetrics.VirtualDisplayRect.Top);
+		if (NewWinSize != TotalSize)
 		{
-			return SlateInfoPtr->SlatePlayInEditorWindow.IsValid();
+			TotalSize = NewWinSize;
+			WinSize.X = NewWinSize.X - HorzPad;
+			WinSize.Y = NewWinSize.Y - VertPad;
 		}
 	}
-	
-	return false;
-}
 
-void UEditorEngine::SetPIEInstanceWindowSwitchDelegate(FPIEInstanceWindowSwitch InSwitchDelegate)
-{
-	PIEInstanceWindowSwitchDelegate = InSwitchDelegate;
-}
+	const FSlateRect PreferredWorkArea(DisplayMetrics.VirtualDisplayRect.Left,
+		DisplayMetrics.VirtualDisplayRect.Top,
+		DisplayMetrics.VirtualDisplayRect.Right - TotalSize.X,
+		DisplayMetrics.VirtualDisplayRect.Bottom - TotalSize.Y);
 
-void UEditorEngine::FocusNextPIEWorld(UWorld *CurrentPieWorld, bool previous)
-{
-	// Get the current world's idx
-	int32 CurrentIdx = 0;
-	for (CurrentIdx = 0; CurrentPieWorld && CurrentIdx < WorldList.Num(); ++CurrentIdx)
+	// if no more windows fit horizontally, place them in a new row
+	if (WinPos.X > PreferredWorkArea.Right)
 	{
-		if (WorldList[CurrentIdx].World() == CurrentPieWorld)
+		WinPos.X = PreferredWorkArea.Left;
+		WinPos.Y += TotalSize.Y;
+		if (WinPos.Y > PreferredWorkArea.Bottom)
 		{
-			break;
+			WinPos.Y = PreferredWorkArea.Top;
 		}
 	}
 
-	// Step through the list to find the next or previous
-	int32 step = previous? -1 : 1;
-	CurrentIdx += (WorldList.Num() + step);
-	
-	while ( CurrentPieWorld && WorldList[ CurrentIdx % WorldList.Num() ].World() != CurrentPieWorld )
+	// if no more rows fit vertically, stack windows on top of each other
+	else if (WinPos.Y > PreferredWorkArea.Bottom)
 	{
-		FWorldContext &Context = WorldList[CurrentIdx % WorldList.Num()];
-		if (Context.World() && Context.WorldType == EWorldType::PIE && Context.GameViewport != NULL)
+		WinPos.Y = PreferredWorkArea.Top;
+		WinPos.X += TotalSize.X;
+		if (WinPos.X > PreferredWorkArea.Right)
 		{
-			break;
+			WinPos.X = PreferredWorkArea.Left;
 		}
+	}
 
-		CurrentIdx += step;
+	// Clamp values to make sure they fall within the desktop area
+	WinPos.X = FMath::Clamp(WinPos.X, (int32)PreferredWorkArea.Left, (int32)PreferredWorkArea.Right);
+	WinPos.Y = FMath::Clamp(WinPos.Y, (int32)PreferredWorkArea.Top, (int32)PreferredWorkArea.Bottom);
+}
+
+void UEditorEngine::GetWindowSizeAndPositionForInstanceIndex(ULevelEditorPlaySettings& InEditorPlaySettings, const int32 InInstanceIndex, FIntPoint& OutSize, FIntPoint& OutPosition)
+{
+	if (!ensureMsgf(PlayInEditorSessionInfo.IsSet(), TEXT("Cannot get saved Window Size/Position if a session has not been started.")))
+	{
+		OutSize = FIntPoint(1280, 720);
+		OutPosition = FIntPoint(0, 0);
+		return;
 	}
-	
-	if (WorldList[CurrentIdx % WorldList.Num()].World())
+
+	FMargin WindowBorderSize(8.0f, 30.0f, 8.0f, 8.0f);
+	TSharedPtr<SWindow> TopLevelWindow = FSlateApplication::Get().GetActiveTopLevelWindow();
+
+	if (TopLevelWindow.IsValid())
 	{
-		// Bring new window to front and activate new viewport
-		FSlatePlayInEditorInfo* SlateInfoPtr = SlatePlayInEditorMap.Find(WorldList[CurrentIdx % WorldList.Num()].ContextHandle);
-		if (SlateInfoPtr && SlateInfoPtr->SlatePlayInEditorWindowViewport.IsValid())
-		{
-			FSceneViewport* SceneViewport = SlateInfoPtr->SlatePlayInEditorWindowViewport.Get();
+		WindowBorderSize = TopLevelWindow->GetWindowBorderSize(true);
+	}
 
-			FSlateApplication& SlateApp = FSlateApplication::Get();
-			TSharedRef<SViewport> ViewportWidget = SceneViewport->GetViewportWidget().Pin().ToSharedRef();
+	// Alright, they don't have a saved position or don't want to load from the saved position. First, figure out
+	// how big the window should be. If it is the first client, it uses a different resolution source than additional.
+	if (InInstanceIndex == 0)
+	{
+		OutSize = FIntPoint(InEditorPlaySettings.NewWindowWidth, InEditorPlaySettings.NewWindowHeight);
+	}
+	else
+	{
+		// Use the size for additional client windows.
+		InEditorPlaySettings.GetClientWindowSize(OutSize);
+	}
 
-			TSharedPtr<SWindow> ViewportWindow = SlateApp.FindWidgetWindow(ViewportWidget);
-			check(ViewportWindow.IsValid());
+	// Figure out how big the users resolution is
+	FDisplayMetrics DisplayMetrics;
+	FSlateApplication::Get().GetCachedDisplayMetrics(DisplayMetrics);
 
-			// Force window to front
-			ViewportWindow->BringToFront();
+	const FVector2D DisplaySize = FVector2D(
+		DisplayMetrics.PrimaryDisplayWorkAreaRect.Right - DisplayMetrics.PrimaryDisplayWorkAreaRect.Left,
+		DisplayMetrics.PrimaryDisplayWorkAreaRect.Bottom - DisplayMetrics.PrimaryDisplayWorkAreaRect.Top
+	);
 
-			// Execute notification delegate in case game code has to do anything else
-			PIEInstanceWindowSwitchDelegate.ExecuteIfBound();
-		}
+	// If the size is zero then they want us to auto-detect the resolution.
+	if (OutSize.X <= 0 || OutSize.Y <= 0)
+	{
+		OutSize.X = FMath::RoundToInt(0.75f * DisplaySize.X);
+		OutSize.Y = FMath::RoundToInt(0.75f * DisplaySize.Y);
 	}
-}
-void UEditorEngine::ResetPIEAudioSetting(UWorld *CurrentPieWorld)
-{
-	ULevelEditorPlaySettings* PlayInSettings = GetMutableDefault<ULevelEditorPlaySettings>();
-	if (!PlayInSettings->EnableGameSound)
+
+
+	// Now we can position the window. If it is the first window, we can respect the center window flag.
+	if (InInstanceIndex == 0)
 	{
-		if (FAudioDevice* AudioDevice = CurrentPieWorld->GetAudioDevice())
+		if (InEditorPlaySettings.CenterNewWindow)
 		{
-			AudioDevice->SetTransientMasterVolume(0.0f);
+			// We don't store the last window position in this case, because we want additional windows
+			// to open starting at the top left of the monitor.
+			OutPosition.X = FMath::RoundToInt((DisplaySize.X / 2.f) - (OutSize.X / 2));
+			OutPosition.Y = FMath::RoundToInt((DisplaySize.Y / 2.f) - (OutSize.Y / 2));
 		}
-	}
-}
-UGameViewportClient * UEditorEngine::GetNextPIEViewport(UGameViewportClient * CurrentViewport)
-{
-	// Get the current world's idx
-	int32 CurrentIdx = 0;
-	for (CurrentIdx = 0; CurrentViewport && CurrentIdx < WorldList.Num(); ++CurrentIdx)
-	{
-		if (WorldList[CurrentIdx].GameViewport == CurrentViewport)
+		else
 		{
-			break;
+			OutPosition = InEditorPlaySettings.NewWindowPosition;
 		}
 	}
-
-	// Step through the list to find the next or previous
-	int32 step = 1;
-	CurrentIdx += (WorldList.Num() + step);
-
-	while ( CurrentViewport && WorldList[ CurrentIdx % WorldList.Num() ].GameViewport != CurrentViewport )
+	else
 	{
-		FWorldContext &Context = WorldList[CurrentIdx % WorldList.Num()];
-		if (Context.GameViewport && Context.WorldType == EWorldType::PIE)
+		if (InInstanceIndex < PlayInEditorSessionInfo->CachedWindowInfo.Num())
 		{
-			return Context.GameViewport;
+			OutPosition = PlayInEditorSessionInfo->CachedWindowInfo[InInstanceIndex].Position;
+			FitWindowPositionToWorkArea(OutPosition, OutSize, WindowBorderSize);
 		}
+		// Add a new entry.
+		else
+		{
+			// We bump the position to go to the right, and the clamp will auto-wrap it for us if it falls off screen.
+			OutPosition = PlayInEditorSessionInfo->LastOpenedWindowInfo.Position + FIntPoint(PlayInEditorSessionInfo->LastOpenedWindowInfo.Size.X, 0);
+			// We're opening multiple windows. We're going to calculate a new position (opening them
+			FitWindowPositionToWorkArea(OutPosition, OutSize, WindowBorderSize);
 
-		CurrentIdx += step;
+			// Store this position as the 'last opened' position. This means additional windows will start to
+			// the right of this, unless they run out of room at which point they'll start over on the next row
+			PlayInEditorSessionInfo->LastOpenedWindowInfo.Size = OutSize;
+			PlayInEditorSessionInfo->LastOpenedWindowInfo.Position = OutPosition;
+		}
 	}
-
-	return NULL;
+	// Store this Size/Position for this duration
+	StoreWindowSizeAndPositionForInstanceIndex(InInstanceIndex, OutSize, OutPosition);
 }
-
-void UEditorEngine::RemapGamepadControllerIdForPIE(class UGameViewportClient* InGameViewport, int32 &ControllerId)
+void UEditorEngine::StoreWindowSizeAndPositionForInstanceIndex(const int32 InInstanceIndex, const FIntPoint& InSize, const FIntPoint& InPosition)
 {
-	// Increment the controller id if we are the focused window, and RouteGamepadToSecondWindow is true (and we are running multiple clients).
-	// This cause the focused window to NOT handle the input, decrement controllerID, and pass it to the next window.
-	const ULevelEditorPlaySettings* PlayInSettings = GetDefault<ULevelEditorPlaySettings>();
-	const bool CanRouteGamepadToSecondWindow = [&PlayInSettings]{ bool RouteGamepadToSecondWindow(false); return (PlayInSettings->GetRouteGamepadToSecondWindow(RouteGamepadToSecondWindow) && RouteGamepadToSecondWindow); }();
-	const bool CanRunUnderOneProcess = [&PlayInSettings]{ bool RunUnderOneProcess(false); return (PlayInSettings->GetRunUnderOneProcess(RunUnderOneProcess) && RunUnderOneProcess); }();
-	if ( CanRouteGamepadToSecondWindow && CanRunUnderOneProcess && InGameViewport->GetWindow().IsValid() && InGameViewport->GetWindow()->HasFocusedDescendants())
+	// Overwrite an existing one if we have it
+	if (InInstanceIndex < PlayInEditorSessionInfo->CachedWindowInfo.Num())
 	{
-		ControllerId++;
+		PlayInEditorSessionInfo->CachedWindowInfo[InInstanceIndex].Size = InSize;
+		PlayInEditorSessionInfo->CachedWindowInfo[InInstanceIndex].Position = InPosition;
+	}
+	else
+	{
+		FPlayInEditorSessionInfo::FWindowSizeAndPos& NewInfo = PlayInEditorSessionInfo->CachedWindowInfo.Add_GetRef(FPlayInEditorSessionInfo::FWindowSizeAndPos());
+		NewInfo.Size = InSize;
+		NewInfo.Position = InPosition;
 	}
 }
 
-void UEditorEngine::AutomationPlayUsingLauncher(const FString& InLauncherDeviceId)
+bool PromptMatineeClose()
 {
-	PlayUsingLauncherDeviceId = InLauncherDeviceId;
-	PlayUsingLauncherDeviceName = PlayUsingLauncherDeviceId.Right(PlayUsingLauncherDeviceId.Find(TEXT("@")));
-	PlayUsingLauncher();
+	if (GLevelEditorModeTools().IsModeActive(FBuiltinEditorModes::EM_InterpEdit))
+	{
+		const bool bContinuePIE = EAppReturnType::Yes == FMessageDialog::Open(EAppMsgType::YesNo, NSLOCTEXT("UnrealEd", "PIENeedsToCloseMatineeQ", "'Play in Editor' must close UnrealMatinee.  Continue?"));
+		if (!bContinuePIE)
+		{
+			return false;
+		}
+		GLevelEditorModeTools().DeactivateMode(FBuiltinEditorModes::EM_InterpEdit);
+	}
+
+	return true;
 }
 
-/** 
-* Cancel Play using Launcher on error 
-* 
-* if the physical device is not authorized to be launched to, we need to pop an error instead of trying to launch
-*/
-void UEditorEngine::CancelPlayUsingLauncher()
+// Deprecated Stubs
+PRAGMA_DISABLE_DEPRECATION_WARNINGS
+UGameInstance* UEditorEngine::CreatePIEGameInstance(int32 InPIEInstance, bool bInSimulateInEditor, bool bAnyBlueprintErrors, bool bStartInSpectatorMode, bool bPlayNetDedicated, bool bPlayStereoscopic, float PIEStartTime)
 {
-	CancelRequestPlaySession();
-
-	FText LaunchingText = LOCTEXT("LauncherTaskInProgressNotificationNotAuthorized", "Cannot launch to this device until this computer is authorized from the device");
-	FNotificationInfo Info(LaunchingText);
-	Info.ExpireDuration = 5.0f;
-	TSharedPtr<SNotificationItem> Notification = FSlateNotificationManager::Get().AddNotification(Info);
-	if (Notification.IsValid())
-	{
-		Notification->SetCompletionState(SNotificationItem::CS_Fail);
-		Notification->ExpireAndFadeout();
-	}
+	return nullptr;
 }
+
+void UEditorEngine::LoginPIEInstances(bool bAnyBlueprintErrors, bool bStartInSpectatorMode, double PIEStartTime)
+{}
+
+void UEditorEngine::OnLoginPIEAllComplete()
+{}
+
+void UEditorEngine::PlayInEditor(UWorld* InWorld, bool bInSimulateInEditor, FPlayInEditorOverrides Overrides /* = FPlayInEditorOverrides() */)
+{}
+PRAGMA_ENABLE_DEPRECATION_WARNINGS
 #undef LOCTEXT_NAMESPACE
