Significantly refactored PIE startup  and consolidated a fair bit of the existing code into structures that better reflect the lifespan that they are valid for. This refactoring will make it easier fo

diff --git a/Engine/Source/Editor/UnrealEd/Private/Settings/LevelEditorPlaySettingsCustomization.h b/Engine/Source/Editor/UnrealEd/Private/Settings/LevelEditorPlaySettingsCustomization.h
--- a/Engine/Source/Editor/UnrealEd/Private/Settings/LevelEditorPlaySettingsCustomization.h
+++ b/Engine/Source/Editor/UnrealEd/Private/Settings/LevelEditorPlaySettingsCustomization.h
@@ -649,28 +649,17 @@ public:
 		}
 
 		// multi-player options
-		IDetailCategoryBuilder& NetworkCategory = LayoutBuilder.EditCategory("MultiplayerOptions");
-		{
-			// Number of players
-			NetworkCategory.AddProperty("PlayNumberOfClients")
-				.DisplayName(LOCTEXT("NumberOfPlayersLabel", "Number of Players"))
-				.IsEnabled(TAttribute<bool>::Create(TAttribute<bool>::FGetter::CreateSP(this, &FLevelEditorPlaySettingsCustomization::HandlePlayNumberOfClientsIsEnabled)));
-
-			NetworkCategory.AddProperty("ServerPort")
-				.DisplayName(LOCTEXT("ServerPortLabel", "Server Port"))
-				.IsEnabled(TAttribute<bool>::Create(TAttribute<bool>::FGetter::CreateSP(this, &FLevelEditorPlaySettingsCustomization::HandleServerPortIsEnabled)));
-
-			NetworkCategory.AddProperty("AdditionalServerGameOptions")
-				.DisplayName(LOCTEXT("ServerGameOptionsLabel", "Server Game Options"))
-				.IsEnabled(TAttribute<bool>::Create(TAttribute<bool>::FGetter::CreateSP(this, &FLevelEditorPlaySettingsCustomization::HandleGameOptionsIsEnabled)));
+		IDetailCategoryBuilder& NetworkCategory = LayoutBuilder.EditCategory("Multiplayer Options");
+		TArray<TSharedRef<IPropertyHandle>> AllProperties;
+		NetworkCategory.GetDefaultProperties(AllProperties);
 
-			NetworkCategory.AddProperty("PlayNetDedicated")
-				.DisplayName(LOCTEXT("RunDedicatedServerLabel", "Run Dedicated Server"))
-				.IsEnabled(TAttribute<bool>::Create(TAttribute<bool>::FGetter::CreateSP(this, &FLevelEditorPlaySettingsCustomization::HandlePlayNetDedicatedPropertyIsEnabled)));
-
-			NetworkCategory.AddProperty("ServerMapNameOverride")
-				.DisplayName(LOCTEXT("ServerMapNameOverrideLabel", "Server Map Name Override"))
-				.Visibility(TAttribute<EVisibility>::Create(TAttribute<EVisibility>::FGetter::CreateSP(this, &FLevelEditorPlaySettingsCustomization::HandleServerMapNameOverrideVisibility)));
+		{
+			// Add all the default properties in before we add custom rows at the bottom.
+			for (TSharedRef<IPropertyHandle> DefaultProperty : AllProperties)
+			{
+				NetworkCategory.AddProperty(DefaultProperty);
+			}
+			
 
 			// client window size
 			TSharedRef<IPropertyHandle> WindowHeightHandle = LayoutBuilder.GetProperty("ClientWindowHeight");
@@ -679,54 +668,19 @@ public:
 			WindowHeightHandle->MarkHiddenByCustomization();
 			WindowWidthHandle->MarkHiddenByCustomization();
 
-			NetworkCategory.AddProperty("AutoConnectToServer")
-				.DisplayName(LOCTEXT("AutoConnectToServerLabel", "Auto Connect To Server"))
-				.IsEnabled(TAttribute<bool>::Create(TAttribute<bool>::FGetter::CreateSP(this, &FLevelEditorPlaySettingsCustomization::HandleAutoConnectToServerEnabled)))
-				.Visibility(TAttribute<EVisibility>::Create(TAttribute<EVisibility>::FGetter::CreateSP(this, &FLevelEditorPlaySettingsCustomization::HandleAutoConnectToServerVisibility)));
-
-			NetworkCategory.AddProperty("RouteGamepadToSecondWindow")
-				.DisplayName(LOCTEXT("RouteGamepadToSecondWindowLabel", "Route 1st Gamepad to 2nd Client"))
-				.IsEnabled(TAttribute<bool>::Create(TAttribute<bool>::FGetter::CreateSP(this, &FLevelEditorPlaySettingsCustomization::HandleRerouteInputToSecondWindowEnabled)))
-				.Visibility(TAttribute<EVisibility>::Create(TAttribute<EVisibility>::FGetter::CreateSP(this, &FLevelEditorPlaySettingsCustomization::HandleRerouteInputToSecondWindowVisibility)));
 			
-			// Run under one instance
-			if (GEditor && GEditor->bAllowMultiplePIEWorlds)
-			{
-				NetworkCategory.AddProperty("RunUnderOneProcess")
-					.DisplayName(LOCTEXT("RunUnderOneProcessEnabledLabel", "Use Single Process"));
-			}
-			else
-			{
-				NetworkCategory.AddProperty("RunUnderOneProcess")
-					.DisplayName( LOCTEXT("RunUnderOneProcessDisabledLabel", "Run Under One Process is disabled.") )
-					.Visibility( EVisibility::Collapsed )
-					.IsEnabled( false );
-			}
-
-			// Net Mode
-			NetworkCategory.AddProperty("PlayNetMode")
-				.Visibility(TAttribute<EVisibility>::Create(TAttribute<EVisibility>::FGetter::CreateSP(this, &FLevelEditorPlaySettingsCustomization::HandlePlayNetModeVisibility)))
-				.DisplayName(LOCTEXT("PlayNetModeLabel", "Editor Multiplayer Mode"));
 
-			NetworkCategory.AddProperty("AdditionalLaunchOptions")
-				.DisplayName(LOCTEXT("AdditionalLaunchOptionsLabel", "Command Line Arguments"))
-				.Visibility(TAttribute<EVisibility>::Create(TAttribute<EVisibility>::FGetter::CreateSP(this, &FLevelEditorPlaySettingsCustomization::HandleCmdLineVisibility)));
-
-			NetworkCategory.AddProperty("NetworkEmulationSettings")
-				.IsEnabled(TAttribute<bool>(this, &FLevelEditorPlaySettingsCustomization::HandleNetworkEmulationIsEnabled));
-
-			NetworkCategory.AddCustomRow(LOCTEXT("PlayInNetworkViewportSize", "Multiplayer Viewport Size"), false)
+			NetworkCategory.AddCustomRow(LOCTEXT("PlayInNetworkViewportSize", "Viewport Size\n(Additional Clients)"), false)
 				.NameContent()
 				[
-					WindowHeightHandle->CreatePropertyNameWidget(LOCTEXT("ClientViewportSizeName", "Multiplayer Viewport Size (in pixels)"), LOCTEXT("ClientWindowSizeTooltip", "Width and Height to use when spawning additional windows."))
+					WindowHeightHandle->CreatePropertyNameWidget(LOCTEXT("ClientViewportSizeName", "Multiplayer Viewport Size (in pixels)"), LOCTEXT("ClientWindowSizeTooltip", "Width and Height to use when spawning additional clients. Useful when you need multiple clients connected but only interact with one window."))
 				]
 				.ValueContent()
 				.MaxDesiredWidth(MaxPropertyWidth)
 				[
 					SNew(SScreenResolutionCustomization, &LayoutBuilder, WindowHeightHandle, WindowWidthHandle)
 				]
-				.IsEnabled(TAttribute<bool>(this, &FLevelEditorPlaySettingsCustomization::HandleClientWindowSizePropertyIsEnabled))
-				.Visibility(TAttribute<EVisibility>::Create(TAttribute<EVisibility>::FGetter::CreateSP(this, &FLevelEditorPlaySettingsCustomization::HandleClientWindowSizePropertyVisibility)));
+				.IsEnabled(TAttribute<bool>(this, &FLevelEditorPlaySettingsCustomization::HandleClientWindowSizePropertyIsEnabled));
 				
 			NetworkCategory.AddCustomRow(LOCTEXT("AdditionalMultiplayerDetails", "Additional Options"), true)
 				.NameContent()
@@ -767,17 +721,17 @@ private:
 	FText HandleMultiplayerOptionsDescription( ) const
 	{
 		const ULevelEditorPlaySettings* PlayInSettings = GetDefault<ULevelEditorPlaySettings>();
+		const bool bLaunchSeparateServer = PlayInSettings->bLaunchSeparateServer;
 		const bool CanRunUnderOneProcess = [&PlayInSettings]{ bool RunUnderOneProcess(false); return (PlayInSettings->GetRunUnderOneProcess(RunUnderOneProcess) && RunUnderOneProcess); }();
-		const bool CanPlayNetDedicated = [&PlayInSettings]{ bool PlayNetDedicated(false); return (PlayInSettings->GetPlayNetDedicated(PlayNetDedicated) && PlayNetDedicated); }();
-		const int32 PlayNumberOfClients = [&PlayInSettings]{ int32 NumberOfClients(0); return (PlayInSettings->GetPlayNumberOfClients(NumberOfClients) ? NumberOfClients : 0); }();
+		const int32 PlayNumberOfClients = [&PlayInSettings] { int32 NumClients(0); return (PlayInSettings->GetPlayNumberOfClients(NumClients) ? NumClients : 0); }();
 		const EPlayNetMode PlayNetMode = [&PlayInSettings]{ EPlayNetMode NetMode(PIE_Standalone); return (PlayInSettings->GetPlayNetMode(NetMode) ? NetMode : PIE_Standalone); }();
 		FString Desc;
 		if (CanRunUnderOneProcess)
 		{
 			Desc += LOCTEXT("MultiplayerDescription_OneProcess", "The following will all run under one UE4 instance:\n").ToString();
-			if (CanPlayNetDedicated)
+			if (PlayNetMode == EPlayNetMode::PIE_Client)
 			{
-				Desc += LOCTEXT("MultiplayerDescription_DedicatedServer", "A dedicated server will open in a new window. ").ToString();
+				Desc += LOCTEXT("MultiplayerDescription_DedicatedServer", "A hidden dedicated server instance will run in editor. ").ToString();
 				if (PlayNumberOfClients == 1)
 				{
 					Desc += LOCTEXT("MultiplayerDescription_EditorClient", "The editor will connect as a client. ").ToString();
@@ -787,7 +741,7 @@ private:
 					Desc += FText::Format(LOCTEXT("MultiplayerDescription_EditorAndClients", "The editor will connect as a client and {0} additional client window(s) will also connect. "), FText::AsNumber(PlayNumberOfClients-1)).ToString();
 				}
 			}
-			else
+			else if(PlayNetMode == EPlayNetMode::PIE_ListenServer)
 			{
 				if (PlayNumberOfClients == 1)
 				{
@@ -798,13 +752,41 @@ private:
 					Desc += FText::Format(LOCTEXT("MultiplayerDescription_EditorListenServerAndClients", "The editor will run as a listen server and {0} additional client window(s) will also connect to it. "), FText::AsNumber(PlayNumberOfClients-1)).ToString();
 				}
 			}
+			else
+			{
+				if (PlayNumberOfClients == 1)
+				{
+					Desc += LOCTEXT("MultiplayerDescription_EditorStandalone", "The editor will run in offline mode. ").ToString();
+				}
+				else
+				{
+					Desc += FText::Format(LOCTEXT("MultiplayerDescription_StandaloneAndClients", "The editor will run offline and {0} additional offline mode window(s) will also connect to it. "), FText::AsNumber(PlayNumberOfClients - 1)).ToString();
+				}
+
+				if (bLaunchSeparateServer)
+				{
+					Desc += LOCTEXT("MultiplayerDescription_StandaloneSeparateServer", "\nAn additional server instance will be launched but not connected to. Use \"open 127.0.0.1:<port>\" to connect. ").ToString();
+				}
+			}
 		}
 		else
 		{
 			Desc += LOCTEXT("MultiplayerDescription_MultiProcess", "The following will run with multiple UE4 instances:\n").ToString();
 			if (PlayNetMode == PIE_Standalone)
 			{
-				Desc += LOCTEXT("MultiplayerDescription_EditorOffline", "The editor will run offline. ").ToString();
+				if (PlayNumberOfClients == 1)
+				{
+					Desc += LOCTEXT("MultiplayerDescription_EditorStandalone", "The editor will run in offline mode. ").ToString();
+				}
+				else
+				{
+					Desc += FText::Format(LOCTEXT("MultiplayerDescription_StandaloneAndClients", "The editor will run offline and {0} additional offline mode window(s) will also connect to it. "), FText::AsNumber(PlayNumberOfClients - 1)).ToString();
+				}
+
+				if (bLaunchSeparateServer)
+				{
+					Desc += LOCTEXT("MultiplayerDescription_StandaloneSeparateServer", "\nAn additional server instance will be launched but not connected to. Use \"open 127.0.0.1:<port>\" to connect. ").ToString();
+				}
 			}
 			else if (PlayNetMode == PIE_ListenServer)
 			{
@@ -819,29 +801,17 @@ private:
 			}
 			else
 			{
-				if (CanPlayNetDedicated)
+				// Client requires additional dedicated server instance
+				Desc += LOCTEXT("MultiplayerDescription_DedicatedServer", "A dedicated server will open in a new window. ").ToString();
+				if (PlayNumberOfClients == 1)
 				{
-					Desc += LOCTEXT("MultiplayerDescription_DedicatedServer", "A dedicated server will open in a new window. ").ToString();
-					if (PlayNumberOfClients == 1)
-					{
-						Desc += LOCTEXT("MultiplayerDescription_EditorClient", "The editor will connect as a client. ").ToString();
-					}
-					else
-					{
-						Desc += FText::Format(LOCTEXT("MultiplayerDescription_EditorAndClients", "The editor will connect as a client and {0} additional client window(s) will also connect. "), FText::AsNumber(PlayNumberOfClients-1)).ToString();
-					}
+					Desc += LOCTEXT("MultiplayerDescription_EditorClient", "The editor will connect as a client. ").ToString();
 				}
 				else
 				{
-					if (PlayNumberOfClients <= 2)
-					{
-						Desc += LOCTEXT("MultiplayerDescription_EditorClientAndListenServer", "A listen server will open in a new window and the editor will connect to it. ").ToString();
-					}
-					else
-					{
-						Desc += FText::Format(LOCTEXT("MultiplayerDescription_EditorClientAndListenServerClients", "A listen server will open in a new window and the editor will connect as a client and {0} additional client window(s) will also connect to it. "), FText::AsNumber(FMath::Max(0, PlayNumberOfClients-2))).ToString(); 
-					}
+					Desc += FText::Format(LOCTEXT("MultiplayerDescription_EditorAndClients", "The editor will connect as a client and {0} additional client window(s) will also connect. "), FText::AsNumber(PlayNumberOfClients-1)).ToString();
 				}
+
 			}
 		}
 		return FText::FromString(Desc);
@@ -853,48 +823,12 @@ private:
 		return GetDefault<ULevelEditorPlaySettings>()->IsClientWindowSizeActive();
 	}
 
-	// Callback for getting the visibility of the ClientWindowHeight and ClientWindowWidth properties.
-	EVisibility HandleClientWindowSizePropertyVisibility() const
-	{
-		return GetDefault<ULevelEditorPlaySettings>()->GetClientWindowSizeVisibility();
-	}
-
-	// Callback for checking whether the PlayNetDedicated is enabled.
-	bool HandlePlayNetDedicatedPropertyIsEnabled( ) const
-	{		
-		return GetDefault<ULevelEditorPlaySettings>()->IsPlayNetDedicatedActive();
-	}
-
-	// Callback for checking whether the PlayNumberOfClients is enabled.
-	bool HandlePlayNumberOfClientsIsEnabled( ) const
-	{
-		return GetDefault<ULevelEditorPlaySettings>()->IsPlayNumberOfClientsActive();
-	}
-
-	// Callback for checking whether the ServerPort is enabled.
-	bool HandleServerPortIsEnabled() const
-	{
-		return GetDefault<ULevelEditorPlaySettings>()->IsServerPortActive();
-	}
-
 	// Callback for checking whether the AdditionalServerGameOptions is enabled.
 	bool HandleGameOptionsIsEnabled( ) const
 	{
 		return GetDefault<ULevelEditorPlaySettings>()->IsAdditionalServerGameOptionsActive();
 	}
 
-	// Callback for getting the enabled state of the AutoConnectToServer property.
-	bool HandleAutoConnectToServerEnabled() const
-	{
-		return GetDefault<ULevelEditorPlaySettings>()->IsAutoConnectToServerActive();
-	}
-
-	// Callback for getting the visibility of the RerouteInputToSecondWindow property.
-	EVisibility HandleAutoConnectToServerVisibility() const
-	{
-		return GetDefault<ULevelEditorPlaySettings>()->GetAutoConnectToServerVisibility();
-	}
-
 	// Callback for getting the enabled state of the RerouteInputToSecondWindow property.
 	bool HandleRerouteInputToSecondWindowEnabled( ) const
 	{
@@ -907,30 +841,6 @@ private:
 		return GetDefault<ULevelEditorPlaySettings>()->GetRouteGamepadToSecondWindowVisibility();
 	}
 
-	// Callback for getting the visibility of the PlayNetMode property.
-	EVisibility HandlePlayNetModeVisibility( ) const
-	{
-		return GetDefault<ULevelEditorPlaySettings>()->GetPlayNetModeVisibility();
-	}
-
-	// Callback for checking if the Network Emulation can be used
-	bool HandleNetworkEmulationIsEnabled() const
-	{
-		return GetDefault<ULevelEditorPlaySettings>()->GetNetworkEmulationVisibility() == EVisibility::Visible;
-	}
-
-	// Callback for getting the visibility of the StandaloneServerMapName property.
-	EVisibility HandleServerMapNameOverrideVisibility() const
-	{
-		return GetDefault<ULevelEditorPlaySettings>()->GetServerMapNameOverrideVisibility();
-	}
-
-	// Callback for getting the visibility of the AdditionalLaunchOptions property.
-	EVisibility HandleCmdLineVisibility( ) const
-	{
-		return GetDefault<ULevelEditorPlaySettings>()->GetAdditionalLaunchOptionsVisibility();
-	}
-
 	void HandleQualityLevelComboBoxOpening()
 	{
 		const UAudioSettings* AudioSettings = GetDefault<UAudioSettings>();
diff --git a/Engine/Source/Editor/UnrealEd/Private/Settings/SettingsClasses.cpp b/Engine/Source/Editor/UnrealEd/Private/Settings/SettingsClasses.cpp
--- a/Engine/Source/Editor/UnrealEd/Private/Settings/SettingsClasses.cpp
+++ b/Engine/Source/Editor/UnrealEd/Private/Settings/SettingsClasses.cpp
@@ -461,12 +461,16 @@ ULevelEditorPlaySettings::ULevelEditorPlaySettings( const FObjectInitializer& Ob
 {
 	ClientWindowWidth = 640;
 	ClientWindowHeight = 480;
+	PlayNetMode = EPlayNetMode::PIE_Standalone;
+	bLaunchSeparateServer = false;
 	PlayNumberOfClients = 1;
 	ServerPort = 17777;
+PRAGMA_DISABLE_DEPRECATION_WARNINGS
 	PlayNetDedicated = false;
+	AutoConnectToServer = true;
+PRAGMA_ENABLE_DEPRECATION_WARNINGS
 	RunUnderOneProcess = true;
 	RouteGamepadToSecondWindow = false;
-	AutoConnectToServer = true;
 	BuildGameBeforeLaunch = EPlayOnBuildMode::PlayOnBuild_Default;
 	LaunchConfiguration = EPlayOnLaunchConfiguration::LaunchConfig_Default;
 	bAutoCompileBlueprintsOnLaunch = true;
diff --git a/Engine/Source/Editor/UnrealEd/Private/Tests/AutomationEditorCommon.cpp b/Engine/Source/Editor/UnrealEd/Private/Tests/AutomationEditorCommon.cpp
--- a/Engine/Source/Editor/UnrealEd/Private/Tests/AutomationEditorCommon.cpp
+++ b/Engine/Source/Editor/UnrealEd/Private/Tests/AutomationEditorCommon.cpp
@@ -773,9 +773,15 @@ bool FCloseAllAssetEditorsCommand::Update()
 bool FStartPIECommand::Update()
 {
 	FLevelEditorModule& LevelEditorModule = FModuleManager::Get().GetModuleChecked<FLevelEditorModule>(TEXT("LevelEditor"));
-	TSharedPtr<class IAssetViewport> ActiveLevelViewport = LevelEditorModule.GetFirstActiveViewport();
 
-	GUnrealEd->RequestPlaySession(false, ActiveLevelViewport, bSimulateInEditor, NULL, NULL, -1, false);
+	FRequestPlaySessionParams Params;
+	Params.DestinationSlateViewport = LevelEditorModule.GetFirstActiveViewport();
+	if (bSimulateInEditor)
+	{
+		Params.WorldType = EPlaySessionWorldType::SimulateInEditor;
+	}
+
+	GUnrealEd->RequestPlaySession(Params);
 	return true;
 }
 
@@ -933,7 +939,18 @@ bool FSaveLevelCommand::Update()
 
 bool FLaunchOnCommand::Update()
 {
-	GUnrealEd->AutomationPlayUsingLauncher(InLauncherDeviceID);
+	FRequestPlaySessionParams::FLauncherDeviceInfo LaunchedDeviceInfo;
+	LaunchedDeviceInfo.DeviceId = InLauncherDeviceID;
+	LaunchedDeviceInfo.DeviceName = InLauncherDeviceID.Right(InLauncherDeviceID.Find(TEXT("@")));
+
+	FRequestPlaySessionParams Params;
+	Params.LauncherTargetDevice = LaunchedDeviceInfo;
+
+	GUnrealEd->RequestPlaySession(Params);
+
+	// Immediately start our requested play session
+	GUnrealEd->StartQueuedPlaySessionRequest();
+
 	return true;
 }
 
diff --git a/Engine/Source/Editor/UnrealEd/Private/Tests/AutomationEditorPromotionCommon.cpp b/Engine/Source/Editor/UnrealEd/Private/Tests/AutomationEditorPromotionCommon.cpp
--- a/Engine/Source/Editor/UnrealEd/Private/Tests/AutomationEditorPromotionCommon.cpp
+++ b/Engine/Source/Editor/UnrealEd/Private/Tests/AutomationEditorPromotionCommon.cpp
@@ -235,7 +235,14 @@ void FEditorPromotionTestUtilities::StartPIE(bool bSimulateInEditor)
 	FLevelEditorModule& LevelEditorModule = FModuleManager::Get().GetModuleChecked<FLevelEditorModule>(TEXT("LevelEditor"));
 	TSharedPtr<class IAssetViewport> ActiveLevelViewport = LevelEditorModule.GetFirstActiveViewport();
 
-	GUnrealEd->RequestPlaySession(false, ActiveLevelViewport, bSimulateInEditor, NULL, NULL, -1, false);
+	FRequestPlaySessionParams SessionParams;
+	SessionParams.DestinationSlateViewport = ActiveLevelViewport;
+	if (bSimulateInEditor)
+	{
+		SessionParams.WorldType = EPlaySessionWorldType::SimulateInEditor;
+	}
+
+	GUnrealEd->RequestPlaySession(SessionParams);
 }
 
 /**
diff --git a/Engine/Source/Editor/UnrealEd/Public/Kismet2/DebuggerCommands.h b/Engine/Source/Editor/UnrealEd/Public/Kismet2/DebuggerCommands.h
--- a/Engine/Source/Editor/UnrealEd/Public/Kismet2/DebuggerCommands.h
+++ b/Engine/Source/Editor/UnrealEd/Public/Kismet2/DebuggerCommands.h
@@ -91,7 +91,6 @@ public:
 	TSharedPtr<FUICommandInfo> PlayInDefaultPlayerStart;
 	TSharedPtr<FUICommandInfo> PlayInSettings;
 	TSharedPtr<FUICommandInfo> PlayInNetworkSettings;
-	TSharedPtr<FUICommandInfo> PlayInNetworkDedicatedServer;
 
 	TArray< TSharedPtr< FUICommandInfo > > PlayInTargetedMobilePreviewDevices;
 
diff --git a/Engine/Source/Editor/VREditor/VREditorMode.cpp b/Engine/Source/Editor/VREditor/VREditorMode.cpp
--- a/Engine/Source/Editor/VREditor/VREditorMode.cpp
+++ b/Engine/Source/Editor/VREditor/VREditorMode.cpp
@@ -897,9 +897,11 @@ void UVREditorMode::ToggleSIEAndVREditor()
 {
 	if (GEditor->EditorWorld == nullptr && !GEditor->bIsSimulatingInEditor)
 	{
-		const FVector* StartLoc = NULL;
-		const FRotator* StartRot = NULL;
-		GEditor->RequestPlaySession(false, VREditorLevelViewportWeakPtr.Pin(), true /*bSimulateInEditor*/, StartLoc, StartRot, -1);
+		FRequestPlaySessionParams SessionParams;
+		SessionParams.DestinationSlateViewport = VREditorLevelViewportWeakPtr;
+		SessionParams.WorldType = EPlaySessionWorldType::SimulateInEditor;
+
+		GEditor->RequestPlaySession(SessionParams);
 	}
 	else if (GEditor->PlayWorld != nullptr && GEditor->bIsSimulatingInEditor)
 	{
@@ -912,10 +914,16 @@ void UVREditorMode::TogglePIEAndVREditor()
 	bool bRequestedPIE = false;
 	if (GEditor->EditorWorld == nullptr && GEditor->PlayWorld == nullptr && !GEditor->bIsSimulatingInEditor)
 	{
-		const FVector* StartLoc = NULL;
-		const FRotator* StartRot = NULL;
+		FRequestPlaySessionParams SessionParams;
+		SessionParams.DestinationSlateViewport = VREditorLevelViewportWeakPtr;
+		SessionParams.WorldType = EPlaySessionWorldType::SimulateInEditor;
+		
 		const bool bHMDIsReady = (GEngine && GEngine->XRSystem.IsValid() && GEngine->XRSystem->GetHMDDevice() && GEngine->XRSystem->GetHMDDevice()->IsHMDConnected());
-		GEditor->RequestPlaySession(true, VREditorLevelViewportWeakPtr.Pin(), false /*bSimulateInEditor*/, StartLoc, StartRot, -1, false, bHMDIsReady);
+		if (bHMDIsReady)
+		{
+			SessionParams.SessionPreviewTypeOverride = EPlaySessionPreviewType::VRPreview;
+		}
+		GEditor->RequestPlaySession(SessionParams);
 		bRequestedPIE = true;
 	}
 	else if (GEditor->PlayWorld != nullptr)
diff --git a/Engine/Source/Runtime/Core/Public/CoreGlobals.h b/Engine/Source/Runtime/Core/Public/CoreGlobals.h
--- a/Engine/Source/Runtime/Core/Public/CoreGlobals.h
+++ b/Engine/Source/Runtime/Core/Public/CoreGlobals.h
@@ -346,6 +346,7 @@ extern CORE_API bool GIsPlayInEditorWorld;
 extern CORE_API int32 GPlayInEditorID;
 
 /** Whether or not PIE was attempting to play from PlayerStart */
+UE_DEPRECATED(4.25, "This variable is unused.")
 extern CORE_API bool GIsPIEUsingPlayerStart;
 
 /** true if the runtime needs textures to be powers of two */
diff --git a/Engine/Source/Runtime/Engine/Classes/Engine/GameInstance.h b/Engine/Source/Runtime/Engine/Classes/Engine/GameInstance.h
--- a/Engine/Source/Runtime/Engine/Classes/Engine/GameInstance.h
+++ b/Engine/Source/Runtime/Engine/Classes/Engine/GameInstance.h
@@ -14,6 +14,10 @@
 #include "Subsystems/GameInstanceSubsystem.h"
 #include "Subsystems/SubsystemCollection.h"
 
+#if WITH_EDITOR
+#include "Settings/LevelEditorPlaySettings.h"
+#endif 
+
 #include "GameInstance.generated.h"
 
 class AGameModeBase;
@@ -24,6 +28,7 @@ class FUniqueNetId;
 class ULocalPlayer;
 class UOnlineSession;
 struct FLatentActionManager;
+class ULevelEditorPlaySettings;
 
 // 
 // 	EWelcomeScreen, 	//initial screen.  Used for platforms where we may not have a signed in user yet.
@@ -90,6 +95,16 @@ private:
 //@TODO: Some of these are really mutually exclusive and should be refactored (put into a struct to make this easier in the future)
 struct FGameInstancePIEParameters
 {
+	FGameInstancePIEParameters()
+		: bSimulateInEditor(false)
+		, bAnyBlueprintErrors(false)
+		, bStartInSpectatorMode(false)
+		, bRunAsDedicated(false)
+		, WorldFeatureLevel(ERHIFeatureLevel::Num)
+		, EditorPlaySettings(nullptr)
+		, NetMode(EPlayNetMode::PIE_Standalone)
+	{}
+
 	// Are we doing SIE instead of PIE?
 	bool bSimulateInEditor;
 
@@ -103,7 +118,16 @@ struct FGameInstancePIEParameters
 	bool bRunAsDedicated;
 
 	// The feature level that PIE world should use
-	ERHIFeatureLevel::Type WorldFeatureLevel = ERHIFeatureLevel::Num;
+	ERHIFeatureLevel::Type WorldFeatureLevel;
+
+	// Kept alive externally.
+	ULevelEditorPlaySettings* EditorPlaySettings;
+
+	// Which net mode should this PIE instance start in? Affects which maps are loaded.
+	EPlayNetMode NetMode;
+
+	// The map we should force the game to load instead of the one currently running in the editor. Blank for no override
+	FString OverrideMapURL;
 };
 
 #endif
diff --git a/Engine/Source/Runtime/Engine/Private/AudioDevice.cpp b/Engine/Source/Runtime/Engine/Private/AudioDevice.cpp
--- a/Engine/Source/Runtime/Engine/Private/AudioDevice.cpp
+++ b/Engine/Source/Runtime/Engine/Private/AudioDevice.cpp
@@ -5998,7 +5998,7 @@ bool FAudioDevice::CanUseVRAudioDevice()
 	if (GIsEditor)
 	{
 		UEditorEngine* EdEngine = Cast<UEditorEngine>(GEngine);
-		return EdEngine->bUseVRPreviewForPlayWorld;
+		return EdEngine->IsVRPreviewActive();
 	}
 	else
 #endif
diff --git a/Engine/Source/Runtime/Engine/Private/GameInstance.cpp b/Engine/Source/Runtime/Engine/Private/GameInstance.cpp
--- a/Engine/Source/Runtime/Engine/Private/GameInstance.cpp
+++ b/Engine/Source/Runtime/Engine/Private/GameInstance.cpp
@@ -229,15 +229,13 @@ FGameInstancePIEResult UGameInstance::InitializeForPlayInEditor(int32 PIEInstanc
 	WorldContext->LastURL.Map = WorldPackageName;
 	WorldContext->PIEPrefix = WorldContext->PIEInstance != INDEX_NONE ? UWorld::BuildPIEPackagePrefix(WorldContext->PIEInstance) : FString();
 
-	const ULevelEditorPlaySettings* PlayInSettings = GetDefault<ULevelEditorPlaySettings>();
-
+	
 	// We always need to create a new PIE world unless we're using the editor world for SIE
 	UWorld* NewWorld = nullptr;
 
 	bool bNeedsGarbageCollection = false;
-	const EPlayNetMode PlayNetMode = [&PlayInSettings]{ EPlayNetMode NetMode(PIE_Standalone); return (PlayInSettings->GetPlayNetMode(NetMode) ? NetMode : PIE_Standalone); }();
-	const bool CanRunUnderOneProcess = [&PlayInSettings]{ bool RunUnderOneProcess(false); return (PlayInSettings->GetRunUnderOneProcess(RunUnderOneProcess) && RunUnderOneProcess); }();
-	if (PlayNetMode == PIE_Client)
+
+	if (Params.NetMode == EPlayNetMode::PIE_Client)
 	{
 		// We are going to connect, so just load an empty world
 		NewWorld = EditorEngine->CreatePIEWorldFromEntry(*WorldContext, EditorEngine->EditorWorld, PIEMapName);
@@ -300,28 +298,25 @@ FGameInstancePIEResult UGameInstance::StartPlayInEditorGameInstance(ULocalPlayer
 	OnStart();
 
 	UEditorEngine* const EditorEngine = CastChecked<UEditorEngine>(GetEngine());
-	ULevelEditorPlaySettings const* PlayInSettings = GetDefault<ULevelEditorPlaySettings>();
-
-	const EPlayNetMode PlayNetMode = [&PlayInSettings]{ EPlayNetMode NetMode(PIE_Standalone); return (PlayInSettings->GetPlayNetMode(NetMode) ? NetMode : PIE_Standalone); }();
 
 	// for clients, just connect to the server
-	if (PlayNetMode == PIE_Client)
+	if (Params.NetMode == PIE_Client)
 	{
 		FString Error;
 		FURL BaseURL = WorldContext->LastURL;
 
 		FString URLString(TEXT("127.0.0.1"));
 		uint16 ServerPort = 0;
-		if (PlayInSettings->GetServerPort(ServerPort))
+		if (Params.EditorPlaySettings->GetServerPort(ServerPort))
 		{
 			URLString += FString::Printf(TEXT(":%hu"), ServerPort);
 		}
 
-		if (PlayInSettings->IsNetworkEmulationEnabled())
+		if (Params.EditorPlaySettings->IsNetworkEmulationEnabled())
 		{
-			if (PlayInSettings->NetworkEmulationSettings.IsEmulationEnabledForTarget(NetworkEmulationTarget::Client))
+			if (Params.EditorPlaySettings->NetworkEmulationSettings.IsEmulationEnabledForTarget(NetworkEmulationTarget::Client))
 			{
-				URLString += PlayInSettings->NetworkEmulationSettings.BuildPacketSettingsForURL();
+				URLString += Params.EditorPlaySettings->NetworkEmulationSettings.BuildPacketSettingsForURL();
 			}
 		}
 
@@ -340,21 +335,21 @@ FGameInstancePIEResult UGameInstance::StartPlayInEditorGameInstance(ULocalPlayer
 		UWorld* const PlayWorld = GetWorld();
 
 		FString ExtraURLOptions;
-		if (PlayInSettings->IsNetworkEmulationEnabled())
+		if (Params.EditorPlaySettings->IsNetworkEmulationEnabled())
 		{
-			NetworkEmulationTarget CurrentTarget = PlayNetMode == PIE_ListenServer ? NetworkEmulationTarget::Server : NetworkEmulationTarget::Client;
-			if (PlayInSettings->NetworkEmulationSettings.IsEmulationEnabledForTarget(CurrentTarget))
+			NetworkEmulationTarget CurrentTarget = Params.NetMode == PIE_ListenServer ? NetworkEmulationTarget::Server : NetworkEmulationTarget::Client;
+			if (Params.EditorPlaySettings->NetworkEmulationSettings.IsEmulationEnabledForTarget(CurrentTarget))
 			{
-				ExtraURLOptions += PlayInSettings->NetworkEmulationSettings.BuildPacketSettingsForURL();
+				ExtraURLOptions += Params.EditorPlaySettings->NetworkEmulationSettings.BuildPacketSettingsForURL();
 			}
 		}
 
 		// make a URL
 		FURL URL;
 		// If the user wants to start in spectator mode, do not use the custom play world for now
-		if (EditorEngine->UserEditedPlayWorldURL.Len() > 0 && !Params.bStartInSpectatorMode)
+		if (EditorEngine->UserEditedPlayWorldURL.Len() > 0 || Params.OverrideMapURL.Len() > 0)
 		{
-			FString UserURL = EditorEngine->UserEditedPlayWorldURL;
+			FString UserURL = EditorEngine->UserEditedPlayWorldURL.Len() > 0 ? EditorEngine->UserEditedPlayWorldURL : Params.OverrideMapURL;
 			UserURL += ExtraURLOptions;
 
 			// If the user edited the play world url. Verify that the map name is the same as the currently loaded map.
@@ -373,7 +368,7 @@ FGameInstancePIEResult UGameInstance::StartPlayInEditorGameInstance(ULocalPlayer
 
 		// If a start location is specified, spawn a temporary PlayerStartPIE actor at the start location and use it as the portal.
 		AActor* PlayerStart = NULL;
-		if (EditorEngine->SpawnPlayFromHereStart(PlayWorld, PlayerStart, EditorEngine->PlayWorldLocation, EditorEngine->PlayWorldRotation) == false)
+		if (!EditorEngine->SpawnPlayFromHereStart(PlayWorld, PlayerStart))
 		{
 			// failed to create "play from here" playerstart
 			return FGameInstancePIEResult::Failure(NSLOCTEXT("UnrealEd", "Error_FailedCreatePlayFromHerePlayerStart", "Failed to create PlayerStart at desired starting location."));
@@ -417,11 +412,11 @@ FGameInstancePIEResult UGameInstance::StartPlayInEditorGameInstance(ULocalPlayer
 			GEngine->BlockTillLevelStreamingCompleted(PlayWorld);
 		}
 		
-		if (PlayNetMode == PIE_ListenServer)
+		if (Params.NetMode == PIE_ListenServer)
 		{
 			// Add port
 			uint16 ServerPort = 0;
-			if (PlayInSettings->GetServerPort(ServerPort))
+			if (Params.EditorPlaySettings->GetServerPort(ServerPort))
 			{
 				URL.Port = ServerPort;
 			}
