Fix compile
#rnx

diff --git a/Engine/Source/Editor/UnrealEd/Private/PlayLevelNewProcess.cpp b/Engine/Source/Editor/UnrealEd/Private/PlayLevelNewProcess.cpp
--- a/Engine/Source/Editor/UnrealEd/Private/PlayLevelNewProcess.cpp
+++ b/Engine/Source/Editor/UnrealEd/Private/PlayLevelNewProcess.cpp
@@ -3,6 +3,7 @@
 #include "Editor/EditorEngine.h"
 #include "PlayLevel.h"
 #include "HeadMountedDisplayTypes.h"
+#include "Editor.h"
 
 void UEditorEngine::StartPlayInNewProcessSession(FRequestPlaySessionParams& InRequestParams)
 {
diff --git a/Engine/Source/Editor/UnrealEd/Private/PlayLevelViaLauncher.cpp b/Engine/Source/Editor/UnrealEd/Private/PlayLevelViaLauncher.cpp
--- a/Engine/Source/Editor/UnrealEd/Private/PlayLevelViaLauncher.cpp
+++ b/Engine/Source/Editor/UnrealEd/Private/PlayLevelViaLauncher.cpp
@@ -16,6 +16,7 @@
 #include "Framework/Notifications/NotificationManager.h"
 #include "PlayLevel.h"
 #include "Async/Async.h"
+#include "Logging/MessageLog.h"
 
 #define LOCTEXT_NAMESPACE "PlayLevel"
 
