Linux/clang CIS fixes.

#rnx

diff --git a/Engine/Source/Editor/UnrealEd/Classes/Editor/EditorEngine.h b/Engine/Source/Editor/UnrealEd/Classes/Editor/EditorEngine.h
--- a/Engine/Source/Editor/UnrealEd/Classes/Editor/EditorEngine.h
+++ b/Engine/Source/Editor/UnrealEd/Classes/Editor/EditorEngine.h
@@ -1733,8 +1733,8 @@ public:
 	/** Returns true if the user has queued a Play in Editor (in VR Mode) or one is currently running. */
 	bool IsVRPreviewActive() const
 	{
-		return PlaySessionRequest.IsSet() && PlaySessionRequest->SessionPreviewTypeOverride == EPlaySessionPreviewType::VRPreview || 
-			PlayInEditorSessionInfo.IsSet() && PlayInEditorSessionInfo->OriginalRequestParams.SessionPreviewTypeOverride == EPlaySessionPreviewType::VRPreview;
+		return (PlaySessionRequest.IsSet() && PlaySessionRequest->SessionPreviewTypeOverride == EPlaySessionPreviewType::VRPreview) ||
+			(PlayInEditorSessionInfo.IsSet() && PlayInEditorSessionInfo->OriginalRequestParams.SessionPreviewTypeOverride == EPlaySessionPreviewType::VRPreview);
 	}
 
 	/** Called when game client received input key */
