Fixed various PIE crashes when different parts of the PIE startup were canceled due to error checking.

Fixed using PIE with Online Subsystem causing a crash with a null PlayWorld.

#rnx

diff --git a/Engine/Source/Editor/UnrealEd/Private/PlayLevel.cpp b/Engine/Source/Editor/UnrealEd/Private/PlayLevel.cpp
--- a/Engine/Source/Editor/UnrealEd/Private/PlayLevel.cpp
+++ b/Engine/Source/Editor/UnrealEd/Private/PlayLevel.cpp
@@ -1536,29 +1536,39 @@ void UEditorEngine::OnLoginPIEComplete_Deferred(int32 LocalUserNum, bool bWasSuc
 		}
 	}
 
-	PlayInEditorSessionInfo->NumOutstandingPIELogins--;
-	if (!ensureAlwaysMsgf(PlayInEditorSessionInfo->NumOutstandingPIELogins >= 0, TEXT("PIEInstancesToLogInCount was not properly reset at some point.")))
-	{
-		PlayInEditorSessionInfo->NumOutstandingPIELogins = 0;
-	}
-
-	// If there are no more instances waiting to log-in then we can do post-launch notifications.
-	if (PlayInEditorSessionInfo->NumOutstandingPIELogins == 0)
+	// If there was a startup failure EndPlayMap might have already been called and unsetting this.
+	if (PlayInEditorSessionInfo.IsSet())
 	{
-		if (bAnyPIELoginFailed)
+		PlayInEditorSessionInfo->NumOutstandingPIELogins--;
+		if (!ensureAlwaysMsgf(PlayInEditorSessionInfo->NumOutstandingPIELogins >= 0, TEXT("PIEInstancesToLogInCount was not properly reset at some point.")))
 		{
-			UE_LOG(LogPlayLevel, Warning, TEXT("At least one of the requested PIE instances failed to start, ending PIE session."));
-			EndPlayMap();
+			PlayInEditorSessionInfo->NumOutstandingPIELogins = 0;
 		}
-		else
+
+		// If there are no more instances waiting to log-in then we can do post-launch notifications.
+		if (PlayInEditorSessionInfo->NumOutstandingPIELogins == 0)
 		{
-			OnAllPIEInstancesStarted();
+			if (bAnyPIELoginFailed)
+			{
+				UE_LOG(LogPlayLevel, Warning, TEXT("At least one of the requested PIE instances failed to start, ending PIE session."));
+				EndPlayMap();
+			}
+			else
+			{
+				OnAllPIEInstancesStarted();
+			}
 		}
 	}
 }
 
 void UEditorEngine::OnAllPIEInstancesStarted()
 {
+	// Transfer the Blueprint Debug references to the last client that was created (since that's all we track).
+	if (PlayWorld)
+	{
+		EditorWorld->TransferBlueprintDebugReferences(PlayWorld);
+	}
+
 	GiveFocusToLastClientPIEViewport();
 
 	// Print out a log message stating the overall startup time.
@@ -2594,6 +2604,12 @@ void UEditorEngine::StartPlayInEditorSession(FRequestPlaySessionParams& InReques
 			// Create the instance. This can end up creating separate processes if needed based on settings.
 			// This code is separated out of here so it can be re-used by the Late Join flow.
 			CreateNewPlayInEditorInstance(InRequestParams, bRunAsDedicated, LocalNetMode);
+
+			// If there was an error creating an instance it will call EndPlay which invalidates the session info.
+			if (!PlayInEditorSessionInfo.IsSet())
+			{
+				return;
+			}
 		}
 
 		// If they requested a SIE, we immediately convert the PIE into a SIE. This finds and executes on the first
@@ -2602,22 +2618,16 @@ void UEditorEngine::StartPlayInEditorSession(FRequestPlaySessionParams& InReques
 		{
 			const bool bNewSession = true;
 			ToggleBetweenPIEandSIE(bNewSession);
+
+			// If we have actors selected in the Editor, save that information and then re-select
+			// the equivalent actors in the SIE session.
+			TransferEditorSelectionToSIEInstances();
 		}
 	}
 
 	// Disable the screen saver when PIE is running
 	EnableScreenSaver(false);
 
-	// Transfer the Blueprint Debug references to the first client that was created.
-	EditorWorld->TransferBlueprintDebugReferences(PlayWorld);
-
-	// If we have actors selected in the Editor, save that information and then re-select
-	// the equivalent actors in the SIE session.
-	if (InRequestParams.WorldType == EPlaySessionWorldType::SimulateInEditor)
-	{
-		TransferEditorSelectionToSIEInstances();
-	}
-
 	// Update the details window with the actors we have just selected
 	GUnrealEd->UpdateFloatingPropertyWindows();
 
