Fixed a regression that removed CL 9473628, and 9324846.

#rnx

diff --git a/Engine/Source/Editor/UnrealEd/Private/PlayLevel.cpp b/Engine/Source/Editor/UnrealEd/Private/PlayLevel.cpp
--- a/Engine/Source/Editor/UnrealEd/Private/PlayLevel.cpp
+++ b/Engine/Source/Editor/UnrealEd/Private/PlayLevel.cpp
@@ -1403,6 +1403,12 @@ bool UEditorEngine::SpawnPlayFromHereStart( UWorld* World, AActor*& PlayerStart,
 
 static bool ShowBlueprintErrorDialog( TArray<UBlueprint*> ErroredBlueprints )
 {
+	if (FApp::IsUnattended() || GIsRunningUnattendedScript)
+	{
+		// App is running in unattended mode, so we should avoid modal dialogs and proceed
+		return true;
+	}
+
 	struct Local
 	{
 		static void OnHyperlinkClicked( TWeakObjectPtr<UBlueprint> InBlueprint, TSharedPtr<SCustomDialog> InDialog )
@@ -2018,7 +2024,7 @@ void UEditorEngine::ToggleBetweenPIEandSIE( bool bNewSession )
 		if (Actor.IsValid())
 		{
 			AActor* SimActor = EditorUtilities::GetSimWorldCounterpartActor(Actor.Get());
-			if (SimActor && !SimActor->bHidden)
+			if (SimActor && !SimActor->IsHidden())
 			{
 				SelectActor( SimActor, bIsSimulatingInEditor, false );
 			}
@@ -2896,7 +2902,7 @@ void UEditorEngine::TransferEditorSelectionToSIEInstances()
 			ActorsThatWereSelected.Add(Actor);
 
 			AActor* SimActor = EditorUtilities::GetSimWorldCounterpartActor(Actor);
-			if (SimActor && !SimActor->bHidden)
+			if (SimActor && !SimActor->IsHidden())
 			{
 				SelectActor(SimActor, true, false);
 			}
