Converted UWorld::AttemptDeriveFromPlayInSettings to use a cached net mode instead specified when the world is created - that is the only time there is enough context to infer the actual meaning of the world. Note that this is still incorrect if they change net modes after PIE starts, but this matches the existing behavior.

diff --git a/Engine/Source/Editor/UnrealEd/Classes/Settings/LevelEditorPlaySettings.h b/Engine/Source/Editor/UnrealEd/Classes/Settings/LevelEditorPlaySettings.h
--- a/Engine/Source/Editor/UnrealEd/Classes/Settings/LevelEditorPlaySettings.h
+++ b/Engine/Source/Editor/UnrealEd/Classes/Settings/LevelEditorPlaySettings.h
@@ -312,7 +312,7 @@ private:
 	int32 PlayNumberOfClients;
 
 	/** What port used by the server for simple networking */
-	UPROPERTY(config, EditAnywhere, Category = "Multiplayer Options|Server", meta=(ClampMin="1", UIMin="1", ClampMax="65535", EditCondition = "PlayNetMode != EPlayNetMode::PIE_Standalone"))
+	UPROPERTY(config, EditAnywhere, Category = "Multiplayer Options|Server", meta=(ClampMin="1", UIMin="1", ClampMax="65535", EditCondition = "PlayNetMode != EPlayNetMode::PIE_Standalone || bLaunchSeparateServer"))
 	uint16 ServerPort;
 
 	/** Width to use when spawning additional windows. */
@@ -353,11 +353,11 @@ private:
 	int32 ClientWindowHeight;
 
 	/** Override the map launched by the dedicated server (currently only used when in PIE_StandaloneWithServer net mode) */
-	UPROPERTY(config, EditAnywhere, Category = "Multiplayer Options|Server", meta=(EditCondition = "PlayNetMode != EPlayNetMode::PIE_Standalone"))
+	UPROPERTY(config, EditAnywhere, Category = "Multiplayer Options|Server", meta=(EditCondition = "PlayNetMode != EPlayNetMode::PIE_Standalone || bLaunchSeparateServer"))
 	FString ServerMapNameOverride;
 
 	/** Additional options that will be passed to the server as URL parameters, in the format ?bIsLanMatch=1?listen - any additional command line switches should be passed in the Command Line Arguments field below. */
-	UPROPERTY(config, EditAnywhere, Category="Multiplayer Options|Server", meta=(EditCondition = "PlayNetMode != EPlayNetMode::PIE_Standalone"))
+	UPROPERTY(config, EditAnywhere, Category="Multiplayer Options|Server", meta=(EditCondition = "PlayNetMode != EPlayNetMode::PIE_Standalone || bLaunchSeparateServer"))
 	FString AdditionalServerGameOptions;
 
 	/** Additional command line options that will be passed to standalone game instances, for example -debug */
@@ -367,8 +367,8 @@ private:
 
 public:
 	/** Additional options that will be passed to the server as arguments, for example -debug. Only works with separate process servers. */
-	UPROPERTY(config, EditAnywhere, Category = "Multiplayer Options|Server", meta = (EditCondition = "!RunUnderOneProcess && PlayNetMode == EPlayNetMode::PIE_Client"))
-	FString AdditionalServerLaunchParmeters;
+	UPROPERTY(config, EditAnywhere, Category = "Multiplayer Options|Server")
+	FString AdditionalServerLaunchParameters;
 
 public:
 
@@ -405,7 +405,7 @@ public:
 	bool GetPlayNetDedicated(bool &OutPlayNetDedicated) const { OutPlayNetDedicated = PlayNetDedicated; return IsPlayNetDedicatedActive(); }
 	
 	void SetPlayNumberOfClients( const int32 InPlayNumberOfClients ) { PlayNumberOfClients = InPlayNumberOfClients; }
-	bool IsPlayNumberOfClientsActive() const { return (PlayNetMode != PIE_Standalone) || RunUnderOneProcess; }
+	bool IsPlayNumberOfClientsActive() const { return true; }
 	bool GetPlayNumberOfClients( int32 &OutPlayNumberOfClients ) const { OutPlayNumberOfClients = PlayNumberOfClients; return IsPlayNumberOfClientsActive(); }
 
 	void SetServerPort(const uint16 InServerPort) { ServerPort = InServerPort; }
@@ -539,6 +539,7 @@ public:
 	// UObject interface
 	virtual void PostEditChangeProperty(struct FPropertyChangedEvent& PropertyChangedEvent) override;
 	virtual void PostInitProperties() override;
+	virtual bool CanEditChange(const UProperty* InProperty) const override;
 	// End of UObject interface
 
 #if WITH_EDITOR
diff --git a/Engine/Source/Editor/UnrealEd/Private/PlayLevelNewProcess.cpp b/Engine/Source/Editor/UnrealEd/Private/PlayLevelNewProcess.cpp
--- a/Engine/Source/Editor/UnrealEd/Private/PlayLevelNewProcess.cpp
+++ b/Engine/Source/Editor/UnrealEd/Private/PlayLevelNewProcess.cpp
@@ -89,7 +89,7 @@ void UEditorEngine::LaunchNewProcess(const FRequestPlaySessionParams& InParams,
 	if (bIsDedicatedServer)
 	{
 		// Allow server specific launch parameters. Only works with separate process standalone servers.
-		CommandLine += FString::Printf(TEXT(" %s"), *InParams.EditorPlaySettings->AdditionalServerLaunchParmeters);
+		CommandLine += FString::Printf(TEXT(" %s"), *InParams.EditorPlaySettings->AdditionalServerLaunchParameters);
 	}
 
 	// If they're not a host, configure the URL Params to connect to the server (instead of a specifying a map later)
diff --git a/Engine/Source/Editor/UnrealEd/Private/Settings/LevelEditorPlaySettingsCustomization.h b/Engine/Source/Editor/UnrealEd/Private/Settings/LevelEditorPlaySettingsCustomization.h
--- a/Engine/Source/Editor/UnrealEd/Private/Settings/LevelEditorPlaySettingsCustomization.h
+++ b/Engine/Source/Editor/UnrealEd/Private/Settings/LevelEditorPlaySettingsCustomization.h
@@ -760,7 +760,7 @@ private:
 				}
 				else
 				{
-					Desc += FText::Format(LOCTEXT("MultiplayerDescription_StandaloneAndClients", "The editor will run offline and {0} additional offline mode window(s) will also connect to it. "), FText::AsNumber(PlayNumberOfClients - 1)).ToString();
+					Desc += FText::Format(LOCTEXT("MultiplayerDescription_StandaloneAndClients", "The editor will run offline and {0} additional offline mode window(s) will also open. "), FText::AsNumber(PlayNumberOfClients - 1)).ToString();
 				}
 
 				if (bLaunchSeparateServer)
@@ -780,7 +780,7 @@ private:
 				}
 				else
 				{
-					Desc += FText::Format(LOCTEXT("MultiplayerDescription_StandaloneAndClients", "The editor will run offline and {0} additional offline mode window(s) will also connect to it. "), FText::AsNumber(PlayNumberOfClients - 1)).ToString();
+					Desc += FText::Format(LOCTEXT("MultiplayerDescription_StandaloneAndClients", "The editor will run offline and {0} additional offline mode window(s) will also open. "), FText::AsNumber(PlayNumberOfClients - 1)).ToString();
 				}
 
 				if (bLaunchSeparateServer)
diff --git a/Engine/Source/Runtime/Engine/Classes/Engine/World.h b/Engine/Source/Runtime/Engine/Classes/Engine/World.h
--- a/Engine/Source/Runtime/Engine/Classes/Engine/World.h
+++ b/Engine/Source/Runtime/Engine/Classes/Engine/World.h
@@ -3215,16 +3215,23 @@ private:
 	/** Private version without inlining that does *not* check Dedicated server build flags (which should already have been done). */
 	ENetMode InternalGetNetMode() const;
 
-#if WITH_EDITOR
-	/** Attempts to derive the net mode from PlayInSettings for PIE*/
-	ENetMode AttemptDeriveFromPlayInSettings() const;
-#endif
-
 	/** Attempts to derive the net mode from URL */
 	ENetMode AttemptDeriveFromURL() const;
 
 	APhysicsVolume* InternalGetDefaultPhysicsVolume() const;
 
+#if WITH_EDITOR
+public:
+	void SetPlayInEditorInitialNetMode(ENetMode InNetMode)
+	{
+		PlayInEditorNetMode = InNetMode;
+	}
+
+private:
+	/** In PIE, what Net Mode was this world started in? Fallback for not having a NetDriver */
+	ENetMode PlayInEditorNetMode;
+#endif
+
 public:
 
 	/**
diff --git a/Engine/Source/Runtime/Engine/Private/GameInstance.cpp b/Engine/Source/Runtime/Engine/Private/GameInstance.cpp
--- a/Engine/Source/Runtime/Engine/Private/GameInstance.cpp
+++ b/Engine/Source/Runtime/Engine/Private/GameInstance.cpp
@@ -202,6 +202,35 @@ void UGameInstance::CreateMinimalNetRPCWorld(const FName InPackageName, UPackage
 }
 
 #if WITH_EDITOR
+static ENetMode GetNetModeFromPlayNetMode(const EPlayNetMode InPlayNetMode, const bool bInDedicatedServer)
+{
+	switch (InPlayNetMode)
+	{
+		case EPlayNetMode::PIE_Client:
+		{
+			return NM_Client;
+		}
+		case EPlayNetMode::PIE_ListenServer:
+		{
+			if (bInDedicatedServer)
+			{
+				return NM_DedicatedServer;
+			}
+
+			return NM_ListenServer;
+		}
+		case EPlayNetMode::PIE_Standalone:
+		{
+			return NM_Standalone;
+		}
+		default:
+		{
+			break;
+		}
+	}
+
+	return NM_Standalone;
+}
 
 FGameInstancePIEResult UGameInstance::InitializeForPlayInEditor(int32 PIEInstanceIndex, const FGameInstancePIEParameters& Params)
 {
@@ -255,6 +284,7 @@ FGameInstancePIEResult UGameInstance::InitializeForPlayInEditor(int32 PIEInstanc
 		return FGameInstancePIEResult::Failure(NSLOCTEXT("UnrealEd", "Error_FailedCreateEditorPreviewWorld", "Failed to create editor preview world."));
 	}
 
+	NewWorld->SetPlayInEditorInitialNetMode(GetNetModeFromPlayNetMode(Params.NetMode, Params.bRunAsDedicated));
 	NewWorld->SetGameInstance(this);
 	WorldContext->SetCurrentWorld(NewWorld);
 	WorldContext->AddRef(EditorEngine->PlayWorld);	// Tie this context to this UEngine::PlayWorld*		// @fixme, needed still?
diff --git a/Engine/Source/Runtime/Engine/Private/World.cpp b/Engine/Source/Runtime/Engine/Private/World.cpp
--- a/Engine/Source/Runtime/Engine/Private/World.cpp
+++ b/Engine/Source/Runtime/Engine/Private/World.cpp
@@ -357,6 +357,7 @@ UWorld::UWorld( const FObjectInitializer& ObjectInitializer )
 {
 	TimerManager = new FTimerManager();
 #if WITH_EDITOR
+	SetPlayInEditorInitialNetMode(ENetMode::NM_Standalone);
 	bBroadcastSelectionChange = true; //Ed Only
 	EditorViews.SetNum(ELevelViewportType::LVT_MAX);
 #endif // WITH_EDITOR
@@ -6840,12 +6841,10 @@ ENetMode UWorld::InternalGetNetMode() const
 #if WITH_EDITOR
 	if (WorldType == EWorldType::PIE)
 	{
-		// PIE: NetDriver is not initialized so use PlayInSettings
-		// to determine the Net Mode
-
-		// This function only works for PIE, do not use for -game/-server standalone editor builds
-		// otherwise it will always return NM_Standalone which is wrong once we have a PendingNetGame
-		return AttemptDeriveFromPlayInSettings();
+		// Return the cached NetMode that we were started with. We can't derive this from the Play Settings,
+		// because those can be modified during PIE startup, plus it is not easy for us to tell if we are
+		// a listen server or a client.
+		return PlayInEditorNetMode;
 	}
 #endif
 	return AttemptDeriveFromURL();
@@ -6869,58 +6868,6 @@ bool UWorld::IsPlayingClientReplay() const
 	return (DemoNetDriver != nullptr && DemoNetDriver->IsPlayingClientReplay());
 }
 
-#if WITH_EDITOR
-ENetMode UWorld::AttemptDeriveFromPlayInSettings() const
-{
-	if (ensure(WorldType == EWorldType::PIE))
-	{
-		const ULevelEditorPlaySettings* PlayInSettings = GetDefault<ULevelEditorPlaySettings>();
-		if (PlayInSettings)
-		{
-			EPlayNetMode PlayNetMode;
-			PlayInSettings->GetPlayNetMode(PlayNetMode);
-
-			switch (PlayNetMode)
-			{
-			case EPlayNetMode::PIE_Client:
-			{
-				int32 NumberOfClients = 0;
-				PlayInSettings->GetPlayNumberOfClients(NumberOfClients);
-
-				bool bAutoConnectToServer = false;
-				PlayInSettings->GetAutoConnectToServer(bAutoConnectToServer);
-
-				// Playing as client without listen server in single process,
-				// or as a client not going to connect to a server
-				if(NumberOfClients == 1 || bAutoConnectToServer == false)
-				{
-					return NM_Standalone;
-				}
-				return NM_Client;
-			}
-			case EPlayNetMode::PIE_ListenServer:
-			{
-				bool bDedicatedServer = false;
-				PlayInSettings->GetPlayNetDedicated(bDedicatedServer);
-
-				if(bDedicatedServer == true)
-				{
-					return NM_DedicatedServer;
-				}
-
-				return NM_ListenServer;
-			}
-			case EPlayNetMode::PIE_Standalone:
-				return NM_Standalone;
-			default:
-				break;
-			}
-		}
-	}
-	return NM_Standalone;
-}
-#endif
-
 ENetMode UWorld::AttemptDeriveFromURL() const
 {
 	if (GEngine != nullptr)
