Fixed missing function compile error.

diff --git a/Engine/Source/Editor/UnrealEd/Private/Settings/SettingsClasses.cpp b/Engine/Source/Editor/UnrealEd/Private/Settings/SettingsClasses.cpp
--- a/Engine/Source/Editor/UnrealEd/Private/Settings/SettingsClasses.cpp
+++ b/Engine/Source/Editor/UnrealEd/Private/Settings/SettingsClasses.cpp
@@ -516,6 +516,19 @@ void ULevelEditorPlaySettings::PostInitProperties()
 #endif
 }
 
+bool ULevelEditorPlaySettings::CanEditChange(const UProperty* InProperty) const
+{
+	const bool ParentVal = Super::CanEditChange(InProperty);
+	FName PropertyName = InProperty->GetFName();
+
+	if (PropertyName == GET_MEMBER_NAME_CHECKED(ULevelEditorPlaySettings, AdditionalServerLaunchParameters))
+	{
+		return ParentVal && (!RunUnderOneProcess && (PlayNetMode == EPlayNetMode::PIE_Client || bLaunchSeparateServer));
+	}
+
+	return ParentVal;
+}
+
 #if WITH_EDITOR
 void ULevelEditorPlaySettings::SwapSafeZoneTypes()
 {
