Updates to Viewport Layout system to support lamba application from host editor.
Fix to monolithic link error

diff --git a/Engine/Source/Editor/UnrealEd/Private/AssetEditorViewportLayout.cpp b/Engine/Source/Editor/UnrealEd/Private/AssetEditorViewportLayout.cpp
--- a/Engine/Source/Editor/UnrealEd/Private/AssetEditorViewportLayout.cpp
+++ b/Engine/Source/Editor/UnrealEd/Private/AssetEditorViewportLayout.cpp
@@ -45,19 +45,19 @@ namespace ViewportLayoutDefs
 }
 
 
-// SViewportsOverlay ////////////////////////////////////////////////
+// AssetEditorViewport ////////////////////////////////////////////////
 
 /**
 * Overlay wrapper class so that we can cache the size of the widget
 * It will also store the ViewportLayout data because that data can't be stored
 * per app; it must be stored per viewport overlay in case the app that made it closes.
 */
-class SViewportsOverlay : public SCompoundWidget
+class SAssetEditorViewportsOverlay : public SCompoundWidget
 {
 
 public:
 
-	SLATE_BEGIN_ARGS( SViewportsOverlay ){}
+	SLATE_BEGIN_ARGS( SAssetEditorViewportsOverlay ){}
 	SLATE_DEFAULT_SLOT( FArguments, Content )
 		SLATE_ARGUMENT( TSharedPtr<FViewportTabContent>, ViewportTab )
 		SLATE_END_ARGS()
@@ -65,7 +65,7 @@ public:
 		void Construct( const FArguments& InArgs );
 
 	/** Default constructor */
-	SViewportsOverlay()
+	SAssetEditorViewportsOverlay()
 		: CachedSize( FVector2D::ZeroVector )
 	{}
 
@@ -101,7 +101,7 @@ private:
 };
 
 
-void SViewportsOverlay::Construct( const FArguments& InArgs )
+void SAssetEditorViewportsOverlay::Construct( const FArguments& InArgs )
 {
 	const TSharedRef<SWidget>& ContentWidget = InArgs._Content.Widget;
 	ViewportTab = InArgs._ViewportTab;
@@ -116,27 +116,27 @@ void SViewportsOverlay::Construct( const FArguments& InArgs )
 		];
 }
 
-void SViewportsOverlay::Tick( const FGeometry& AllottedGeometry, const double InCurrentTime, const float InDeltaTime )
+void SAssetEditorViewportsOverlay::Tick( const FGeometry& AllottedGeometry, const double InCurrentTime, const float InDeltaTime )
 {
 	CachedSize = AllottedGeometry.Size;
 }
 
-SOverlay::FOverlaySlot& SViewportsOverlay::AddSlot()
+SOverlay::FOverlaySlot& SAssetEditorViewportsOverlay::AddSlot()
 {
 	return OverlayWidget->AddSlot();
 }
 
-void SViewportsOverlay::RemoveSlot()
+void SAssetEditorViewportsOverlay::RemoveSlot()
 {
 	return OverlayWidget->RemoveSlot();
 }
 
-const FVector2D& SViewportsOverlay::GetCachedSize() const
+const FVector2D& SAssetEditorViewportsOverlay::GetCachedSize() const
 {
 	return CachedSize;
 }
 
-TSharedPtr<FViewportTabContent> SViewportsOverlay::GetViewportTab() const
+TSharedPtr<FViewportTabContent> SAssetEditorViewportsOverlay::GetViewportTab() const
 {
 	return ViewportTab;
 }
@@ -173,8 +173,8 @@ TSharedRef<SWidget> FAssetEditorViewportLayout::BuildViewportLayout(TFunction<TS
 
 	// We use an overlay so that we can draw a maximized viewport on top of the other viewports
 	TSharedPtr<SBorder> ViewportsBorder;
-	TSharedRef<SViewportsOverlay> ViewportsOverlay =
-		SNew(SViewportsOverlay)
+	TSharedRef<SAssetEditorViewportsOverlay> ViewportsOverlay =
+		SNew(SAssetEditorViewportsOverlay)
 		.ViewportTab(InParentTab)
 		[
 			SAssignNew(ViewportsBorder, SBorder)
@@ -191,6 +191,16 @@ TSharedRef<SWidget> FAssetEditorViewportLayout::BuildViewportLayout(TFunction<TS
 
 }
 
+FString FAssetEditorViewportLayout::GetTypeSpecificLayoutString(const FString& LayoutString) const
+{
+	if (LayoutString.IsEmpty())
+	{
+		return LayoutString;
+	}
+	return FString::Printf(TEXT("%s.%s"), *GetLayoutTypeName().ToString(), *LayoutString);
+}
+
+
 
 #undef LOCTEXT_NAMESPACE
 
diff --git a/Engine/Source/Editor/UnrealEd/Private/EditorViewportLayout2x2.cpp b/Engine/Source/Editor/UnrealEd/Private/EditorViewportLayout2x2.cpp
--- a/Engine/Source/Editor/UnrealEd/Private/EditorViewportLayout2x2.cpp
+++ b/Engine/Source/Editor/UnrealEd/Private/EditorViewportLayout2x2.cpp
@@ -18,7 +18,7 @@ namespace ViewportLayout2x2Defs
 
 TSharedRef<SWidget> FEditorViewportLayout2x2::MakeViewportLayout(TFunction<TSharedRef<SEditorViewport>(void)> &Func, const FString& LayoutString)
 {
-// 	FString SpecificLayoutString = GetTypeSpecificLayoutString(LayoutString);
+ 	FString SpecificLayoutString = GetTypeSpecificLayoutString(LayoutString);
 
 	FString TopLeftKey, BottomLeftKey, TopRightKey, BottomRightKey;
 
@@ -26,6 +26,34 @@ TSharedRef<SWidget> FEditorViewportLayout2x2::MakeViewportLayout(TFunction<TShar
 
 	TArray<FVector2D> SplitterPercentages;
 	
+ 	if (!SpecificLayoutString.IsEmpty())
+	{
+		// The Layout String only holds the unique ID of the Additional Layout Configs to use
+		const FString& IniSection = FLayoutSaveRestore::GetAdditionalLayoutConfigIni();
+
+		TopLeftKey = SpecificLayoutString + TEXT(".Viewport0");
+		BottomLeftKey = SpecificLayoutString + TEXT(".Viewport1");
+		TopRightKey = SpecificLayoutString + TEXT(".Viewport2");
+		BottomRightKey = SpecificLayoutString + TEXT(".Viewport3");
+
+		GConfig->GetString(*IniSection, *(TopLeftKey + TEXT(".TypeWithinLayout")), TopLeftType, GEditorPerProjectIni);
+		GConfig->GetString(*IniSection, *(TopRightKey + TEXT(".TypeWithinLayout")), TopRightType, GEditorPerProjectIni);
+		GConfig->GetString(*IniSection, *(BottomLeftKey + TEXT(".TypeWithinLayout")), BottomLeftType, GEditorPerProjectIni);
+		GConfig->GetString(*IniSection, *(BottomRightKey + TEXT(".TypeWithinLayout")), BottomRightType, GEditorPerProjectIni);
+
+		for (int32 i = 0; i < 4; ++i)
+		{
+			FString PercentageString;
+			FVector2D NewPercentage = ViewportLayout2x2Defs::DefaultSplitterPercentages;
+			if (GConfig->GetString(*IniSection, *(SpecificLayoutString + FString::Printf(TEXT(".Percentages%i"), i)), PercentageString, GEditorPerProjectIni))
+			{
+				NewPercentage.InitFromString(PercentageString);
+			}
+			SplitterPercentages.Add(NewPercentage);
+		}
+	}
+
+
 	// Set up the viewports
 	FAssetEditorViewportConstructionArgs Args;
  	Args.ParentLayout = AsShared();
diff --git a/Engine/Source/Editor/UnrealEd/Private/EditorViewportLayoutFourPanes.cpp b/Engine/Source/Editor/UnrealEd/Private/EditorViewportLayoutFourPanes.cpp
--- a/Engine/Source/Editor/UnrealEd/Private/EditorViewportLayoutFourPanes.cpp
+++ b/Engine/Source/Editor/UnrealEd/Private/EditorViewportLayoutFourPanes.cpp
@@ -20,6 +20,8 @@ namespace ViewportLayoutFourPanesDefs
 
 TSharedRef<SWidget> FEditorViewportLayoutFourPanes::MakeViewportLayout(TFunction<TSharedRef<SEditorViewport>(void)> &Func, const FString& LayoutString)
 {
+	FString SpecificLayoutString = GetTypeSpecificLayoutString(LayoutString);
+
 	FString ViewportKey0, ViewportKey1, ViewportKey2, ViewportKey3;
 	FString ViewportType0 = TEXT("Default"), ViewportType1 = TEXT("Default"), ViewportType2 = TEXT("Default"), ViewportType3 = TEXT("Default");
 
@@ -27,6 +29,36 @@ TSharedRef<SWidget> FEditorViewportLayoutFourPanes::MakeViewportLayout(TFunction
 	float SecondarySplitterPercentage0 = ViewportLayoutFourPanesDefs::DefaultSecondarySplitterPercentage;
 	float SecondarySplitterPercentage1 = ViewportLayoutFourPanesDefs::DefaultSecondarySplitterPercentage;
 
+	if (!SpecificLayoutString.IsEmpty())
+	{
+		// The Layout String only holds the unique ID of the Additional Layout Configs to use
+		const FString& IniSection = FLayoutSaveRestore::GetAdditionalLayoutConfigIni();
+
+		ViewportKey0 = SpecificLayoutString + TEXT(".Viewport0");
+		ViewportKey1 = SpecificLayoutString + TEXT(".Viewport1");
+		ViewportKey2 = SpecificLayoutString + TEXT(".Viewport2");
+		ViewportKey3 = SpecificLayoutString + TEXT(".Viewport3");
+
+		GConfig->GetString(*IniSection, *(ViewportKey0 + TEXT(".TypeWithinLayout")), ViewportType0, GEditorPerProjectIni);
+		GConfig->GetString(*IniSection, *(ViewportKey1 + TEXT(".TypeWithinLayout")), ViewportType1, GEditorPerProjectIni);
+		GConfig->GetString(*IniSection, *(ViewportKey2 + TEXT(".TypeWithinLayout")), ViewportType2, GEditorPerProjectIni);
+		GConfig->GetString(*IniSection, *(ViewportKey3 + TEXT(".TypeWithinLayout")), ViewportType3, GEditorPerProjectIni);
+
+		FString PercentageString;
+		if (GConfig->GetString(*IniSection, *(SpecificLayoutString + TEXT(".Percentage0")), PercentageString, GEditorPerProjectIni))
+		{
+			TTypeFromString<float>::FromString(PrimarySplitterPercentage, *PercentageString);
+		}
+		if (GConfig->GetString(*IniSection, *(SpecificLayoutString + TEXT(".Percentage1")), PercentageString, GEditorPerProjectIni))
+		{
+			TTypeFromString<float>::FromString(SecondarySplitterPercentage0, *PercentageString);
+		}
+		if (GConfig->GetString(*IniSection, *(SpecificLayoutString + TEXT(".Percentage2")), PercentageString, GEditorPerProjectIni))
+		{
+			TTypeFromString<float>::FromString(SecondarySplitterPercentage1, *PercentageString);
+		}
+	}
+
 	// Set up the viewports
 	FAssetEditorViewportConstructionArgs Args;
 	Args.ParentLayout = AsShared();
diff --git a/Engine/Source/Editor/UnrealEd/Private/EditorViewportLayoutOnePane.cpp b/Engine/Source/Editor/UnrealEd/Private/EditorViewportLayoutOnePane.cpp
--- a/Engine/Source/Editor/UnrealEd/Private/EditorViewportLayoutOnePane.cpp
+++ b/Engine/Source/Editor/UnrealEd/Private/EditorViewportLayoutOnePane.cpp
@@ -13,8 +13,25 @@
 
 TSharedRef<SWidget> FEditorViewportLayoutOnePane::MakeViewportLayout(TFunction<TSharedRef<SEditorViewport>(void)> &Func, const FString& LayoutString)
 {
+	FString SpecificLayoutString = GetTypeSpecificLayoutString(LayoutString);
+
+
+	FEngineShowFlags OrthoShowFlags(ESFIM_Editor);
+	ApplyViewMode(VMI_BrushWireframe, false, OrthoShowFlags);
+
+	FEngineShowFlags PerspectiveShowFlags(ESFIM_Editor);
+	ApplyViewMode(VMI_Lit, true, PerspectiveShowFlags);
+
 	FString ViewportKey, ViewportType;
 
+	if (!SpecificLayoutString.IsEmpty())
+	{
+		const FString& IniSection = FLayoutSaveRestore::GetAdditionalLayoutConfigIni();
+
+		ViewportKey = SpecificLayoutString + TEXT(".Viewport0");
+		GConfig->GetString(*IniSection, *(ViewportKey + TEXT(".TypeWithinLayout")), ViewportType, GEditorPerProjectIni);
+	}
+
 	// Set up the viewport
 	FAssetEditorViewportConstructionArgs Args;
  	Args.ParentLayout = AsShared();
diff --git a/Engine/Source/Editor/UnrealEd/Private/EditorViewportLayoutThreePanes.cpp b/Engine/Source/Editor/UnrealEd/Private/EditorViewportLayoutThreePanes.cpp
--- a/Engine/Source/Editor/UnrealEd/Private/EditorViewportLayoutThreePanes.cpp
+++ b/Engine/Source/Editor/UnrealEd/Private/EditorViewportLayoutThreePanes.cpp
@@ -19,6 +19,8 @@ namespace ViewportLayoutThreePanesDefs
 
 TSharedRef<SWidget> FEditorViewportLayoutThreePanes::MakeViewportLayout(TFunction<TSharedRef<SEditorViewport>(void)> &Func, const FString& LayoutString)
 {
+	FString SpecificLayoutString = GetTypeSpecificLayoutString(LayoutString);
+
 	FEngineShowFlags OrthoShowFlags(ESFIM_Editor);	
 	ApplyViewMode(VMI_BrushWireframe, false, OrthoShowFlags);
 
@@ -30,6 +32,30 @@ TSharedRef<SWidget> FEditorViewportLayoutThreePanes::MakeViewportLayout(TFunctio
 	float PrimarySplitterPercentage = ViewportLayoutThreePanesDefs::DefaultSplitterPercentage;
 	float SecondarySplitterPercentage = ViewportLayoutThreePanesDefs::DefaultSplitterPercentage;
 
+	if (!SpecificLayoutString.IsEmpty())
+	{
+		// The Layout String only holds the unique ID of the Additional Layout Configs to use
+		const FString& IniSection = FLayoutSaveRestore::GetAdditionalLayoutConfigIni();
+
+		ViewportKey0 = SpecificLayoutString + TEXT(".Viewport0");
+		ViewportKey1 = SpecificLayoutString + TEXT(".Viewport1");
+		ViewportKey2 = SpecificLayoutString + TEXT(".Viewport2");
+
+		GConfig->GetString(*IniSection, *(ViewportKey0 + TEXT(".TypeWithinLayout")), ViewportType0, GEditorPerProjectIni);
+		GConfig->GetString(*IniSection, *(ViewportKey1 + TEXT(".TypeWithinLayout")), ViewportType1, GEditorPerProjectIni);
+		GConfig->GetString(*IniSection, *(ViewportKey2 + TEXT(".TypeWithinLayout")), ViewportType2, GEditorPerProjectIni);
+
+		FString PercentageString;
+		if (GConfig->GetString(*IniSection, *(SpecificLayoutString + TEXT(".Percentage0")), PercentageString, GEditorPerProjectIni))
+		{
+			TTypeFromString<float>::FromString(PrimarySplitterPercentage, *PercentageString);
+		}
+		if (GConfig->GetString(*IniSection, *(SpecificLayoutString + TEXT(".Percentage1")), PercentageString, GEditorPerProjectIni))
+		{
+			TTypeFromString<float>::FromString(SecondarySplitterPercentage, *PercentageString);
+		}
+	}
+
 	// Set up the viewports
 	FAssetEditorViewportConstructionArgs Args;
 	Args.ParentLayout = AsShared();
diff --git a/Engine/Source/Editor/UnrealEd/Private/EditorViewportLayoutTwoPanes.cpp b/Engine/Source/Editor/UnrealEd/Private/EditorViewportLayoutTwoPanes.cpp
--- a/Engine/Source/Editor/UnrealEd/Private/EditorViewportLayoutTwoPanes.cpp
+++ b/Engine/Source/Editor/UnrealEd/Private/EditorViewportLayoutTwoPanes.cpp
@@ -18,10 +18,29 @@ namespace ViewportLayoutTwoPanesDefs
 template <EOrientation TOrientation>
 TSharedRef<SWidget> TEditorViewportLayoutTwoPanes<TOrientation>::MakeViewportLayout(TFunction<TSharedRef<SEditorViewport>(void)> &Func, const FString& LayoutString)
 {
+	FString SpecificLayoutString = GetTypeSpecificLayoutString(LayoutString);
+
 	FString ViewportKey0, ViewportKey1;
 	FString ViewportType0, ViewportType1;
-	float PrimarySplitterPercentage = ViewportLayoutTwoPanesDefs::DefaultSplitterPercentage;
+	float SplitterPercentage = ViewportLayoutTwoPanesDefs::DefaultSplitterPercentage;
+
+	if (!SpecificLayoutString.IsEmpty())
+	{
+		// The Layout String only holds the unique ID of the Additional Layout Configs to use
+		const FString& IniSection = FLayoutSaveRestore::GetAdditionalLayoutConfigIni();
+
+		ViewportKey0 = SpecificLayoutString + TEXT(".Viewport0");
+		ViewportKey1 = SpecificLayoutString + TEXT(".Viewport1");
+
+		GConfig->GetString(*IniSection, *(ViewportKey0 + TEXT(".TypeWithinLayout")), ViewportType0, GEditorPerProjectIni);
+		GConfig->GetString(*IniSection, *(ViewportKey1 + TEXT(".TypeWithinLayout")), ViewportType1, GEditorPerProjectIni);
 
+		FString PercentageString;
+		if (GConfig->GetString(*IniSection, *(SpecificLayoutString + TEXT(".Percentage")), PercentageString, GEditorPerProjectIni))
+		{
+			TTypeFromString<float>::FromString(SplitterPercentage, *PercentageString);
+		}
+	}
 	// Set up the viewports
 	FAssetEditorViewportConstructionArgs Args;
 	Args.ParentLayout = AsShared();
@@ -44,12 +63,12 @@ TSharedRef<SWidget> TEditorViewportLayoutTwoPanes<TOrientation>::MakeViewportLay
 		SNew(SSplitter)
 		.Orientation(TOrientation)
 		+ SSplitter::Slot()
-		.Value(PrimarySplitterPercentage)
+		.Value(SplitterPercentage)
 		[
 			Viewport0->AsWidget()
 		]
 	+ SSplitter::Slot()
-		.Value(1.0f - PrimarySplitterPercentage)
+		.Value(1.0f - SplitterPercentage)
 		[
 			Viewport1->AsWidget()
 		];
diff --git a/Engine/Source/Editor/UnrealEd/Private/EditorViewportTabContent.cpp b/Engine/Source/Editor/UnrealEd/Private/EditorViewportTabContent.cpp
--- a/Engine/Source/Editor/UnrealEd/Private/EditorViewportTabContent.cpp
+++ b/Engine/Source/Editor/UnrealEd/Private/EditorViewportTabContent.cpp
@@ -37,7 +37,8 @@ TSharedPtr< class FEditorViewportLayout > FEditorViewportTabContent::ConstructVi
 
 void FEditorViewportTabContent::Initialize(TFunction<TSharedRef<SEditorViewport>(void)> Func, TSharedPtr<SDockTab> InParentTab, const FString& InLayoutString)
 {
-	ViewportCreationFunc = Func;
+	check(!InLayoutString.IsEmpty());
+
 	ParentTab = InParentTab;
 	LayoutString = InLayoutString;
 
@@ -47,41 +48,70 @@ void FEditorViewportTabContent::Initialize(TFunction<TSharedRef<SEditorViewport>
 
 void FEditorViewportTabContent::SetViewportConfiguration(TFunction<TSharedRef<SEditorViewport>(void)> &Func, const FName& ConfigurationName)
 {
+	ViewportCreationFunc = Func;
+	SetViewportConfiguration(ConfigurationName);
+}
+
+void FEditorViewportTabContent::SetViewportConfiguration(const FName& ConfigurationName)
+{
+	check(ViewportCreationFunc != nullptr);
+
 	bool bSwitchingLayouts = ActiveViewportLayout.IsValid();
 
 	if (bSwitchingLayouts)
 	{
+		SaveConfig();
 		ActiveViewportLayout.Reset();
 	}
 
 	ActiveViewportLayout = ConstructViewportLayoutByTypeName(ConfigurationName, bSwitchingLayouts);
 	check(ActiveViewportLayout.IsValid());
 
-	UpdateViewportTabWidget(Func);
+	UpdateViewportTabWidget(ViewportCreationFunc);
+
 }
 
-void FEditorViewportTabContent::SetViewportConfiguration(const FName& ConfigurationName)
+void FEditorViewportTabContent::SaveConfig() const
 {
-	bool bSwitchingLayouts = ActiveViewportLayout.IsValid();
-
-	if (bSwitchingLayouts)
+	if (ActiveViewportLayout.IsValid())
 	{
-		ActiveViewportLayout.Reset();
-	}
+		if (!LayoutString.IsEmpty())
+		{
+			FString LayoutTypeString = ActiveViewportLayout->GetLayoutTypeName().ToString();
 
-	ActiveViewportLayout = ConstructViewportLayoutByTypeName(ConfigurationName, bSwitchingLayouts);
-	check(ActiveViewportLayout.IsValid());
+			const FString& IniSection = FLayoutSaveRestore::GetAdditionalLayoutConfigIni();
+			GConfig->SetString(*IniSection, *(LayoutString + TEXT(".LayoutType")), *LayoutTypeString, GEditorPerProjectIni);
+		}
+
+		ActiveViewportLayout->SaveLayoutString(LayoutString);
+	}
+}
 
-  	UpdateViewportTabWidget(ViewportCreationFunc);
+TSharedPtr<SEditorViewport> FEditorViewportTabContent::GetFirstViewport()
+{
+ 	const TMap< FName, TSharedPtr< IEditorViewportLayoutEntity > >& EditorViewports = ActiveViewportLayout->GetViewports();
+ 
+	for (auto& Pair : EditorViewports)
+	{
+		TSharedPtr<SWidget> ViewportWidget = StaticCastSharedPtr<IEditorViewportLayoutEntity>(Pair.Value)->AsWidget();
+		TSharedPtr<SEditorViewport> Viewport = StaticCastSharedPtr<SEditorViewport>(ViewportWidget);
+		if (Viewport.IsValid())
+		{
+			return Viewport;
+			break;
+		}
+	}
 
+	return nullptr;
 }
 
+
 void FEditorViewportTabContent::UpdateViewportTabWidget(TFunction<TSharedRef<SEditorViewport>(void)> &Func)
 {
 	TSharedPtr<SDockTab> ParentTabPinned = ParentTab.Pin();
 	if (ParentTabPinned.IsValid() && ActiveViewportLayout.IsValid())
 	{
- 		TSharedRef<SWidget> LayoutWidget = StaticCastSharedPtr<FAssetEditorViewportLayout>(ActiveViewportLayout)->BuildViewportLayout(Func, ParentTabPinned, SharedThis(this), LayoutString);
+		TSharedRef<SWidget> LayoutWidget = StaticCastSharedPtr<FAssetEditorViewportLayout>(ActiveViewportLayout)->BuildViewportLayout(Func, ParentTabPinned, SharedThis(this), LayoutString);
  		ParentTabPinned->SetContent(LayoutWidget);
 
 		if (PreviouslyFocusedViewport.IsSet())
diff --git a/Engine/Source/Editor/UnrealEd/Private/ViewportTabContent.cpp b/Engine/Source/Editor/UnrealEd/Private/ViewportTabContent.cpp
--- a/Engine/Source/Editor/UnrealEd/Private/ViewportTabContent.cpp
+++ b/Engine/Source/Editor/UnrealEd/Private/ViewportTabContent.cpp
@@ -23,3 +23,12 @@ bool FViewportTabContent::IsViewportConfigurationSet(const FName& ConfigurationN
 	}
 	return false;
 }
+
+void FViewportTabContent::PerformActionOnViewports(TFunction<void(FName Name, TSharedPtr<IEditorViewportLayoutEntity>)> &TFuncPtr)
+{
+	const TMap< FName, TSharedPtr<IEditorViewportLayoutEntity> >& Entities = ActiveViewportLayout->GetViewports();
+	for (auto& Entity : Entities)
+	{
+		TFuncPtr(Entity.Key, Entity.Value);
+	}
+}
diff --git a/Engine/Source/Editor/UnrealEd/Public/AssetEditorViewportLayout.h b/Engine/Source/Editor/UnrealEd/Public/AssetEditorViewportLayout.h
--- a/Engine/Source/Editor/UnrealEd/Public/AssetEditorViewportLayout.h
+++ b/Engine/Source/Editor/UnrealEd/Public/AssetEditorViewportLayout.h
@@ -16,7 +16,7 @@
 #include "SEditorViewport.h"
 
 class FAssetEditorViewportLayout;
-class SViewportsOverlay;
+class SAssetEditorViewportsOverlay;
 class SWindow;
 class SAssetEditorViewport;
 
@@ -94,6 +94,9 @@ public:
 	/** Returns the parent tab content object */
 	TWeakPtr< class FViewportTabContent > GetParentTabContent() const { return ParentTabContent; }
 
+	/** Generates a layout string for persisting settings for this layout based on the runtime type of layout */
+	FString GetTypeSpecificLayoutString(const FString& LayoutString) const;
+
 	/**
 	 * Overridden in derived classes to set up custom layouts  
 	 *
@@ -103,7 +106,7 @@ public:
   	virtual TSharedRef<SWidget> MakeViewportLayout(TFunction<TSharedRef<SEditorViewport>(void)> &Func, const FString& LayoutString) = 0;
 
 	/** The overlay widget that handles what viewports should be on top (non-maximized or maximized) */
-	TWeakPtr< class SViewportsOverlay > ViewportsOverlayPtr;
+	TWeakPtr< class SAssetEditorViewportsOverlay > ViewportsOverlayPtr;
 
 	/** The parent tab content object where this layout resides */
 	TWeakPtr< class FViewportTabContent > ParentTabContent;
diff --git a/Engine/Source/Editor/UnrealEd/Public/EditorViewportLayoutFourPanes.h b/Engine/Source/Editor/UnrealEd/Public/EditorViewportLayoutFourPanes.h
--- a/Engine/Source/Editor/UnrealEd/Public/EditorViewportLayoutFourPanes.h
+++ b/Engine/Source/Editor/UnrealEd/Public/EditorViewportLayoutFourPanes.h
@@ -10,11 +10,6 @@
 
 class FEditorViewportLayoutFourPanes : public FAssetEditorViewportLayout
 {
-public:
-	/**
-	 * Saves viewport layout information between editor sessions
-	 */
-// 	virtual void SaveLayoutString(const FString& LayoutString) const override;
 protected:
 	/**
 	 * Creates the viewports and splitter for the four-pane layout              
diff --git a/Engine/Source/Editor/UnrealEd/Public/EditorViewportTabContent.h b/Engine/Source/Editor/UnrealEd/Public/EditorViewportTabContent.h
--- a/Engine/Source/Editor/UnrealEd/Public/EditorViewportTabContent.h
+++ b/Engine/Source/Editor/UnrealEd/Public/EditorViewportTabContent.h
@@ -25,32 +25,31 @@ public:
 	{
 		return LayoutString;
 	}
-
-
+	
+	const TSharedPtr< class FEditorViewportLayout > GetActiveViewportLayout() const
+	{
+		return ActiveViewportLayout;
+	}
+	   
 	TSharedPtr< class FEditorViewportLayout > ConstructViewportLayoutByTypeName(const FName& TypeName, bool bSwitchingLayouts);
 
 	void Initialize(TFunction<TSharedRef<SEditorViewport>(void)> Func, TSharedPtr<SDockTab> InParentTab, const FString& InLayoutString);
 
-	void SetViewportConfiguration(TFunction<TSharedRef<SEditorViewport>(void)> &Func, const FName& ConfigurationName);
-
-	void UpdateViewportTabWidget(TFunction<TSharedRef<SEditorViewport>(void)> &Func);
-
 	/**
 	* Sets the current layout by changing the contained layout object
 	* 
 	* @param ConfigurationName		The name of the layout (for the names in namespace EditorViewportConfigurationNames)
 	*/
-	void SetViewportConfiguration(const FName& ConfigurationName) override;
-
+	void SetViewportConfiguration(TFunction<TSharedRef<SEditorViewport>(void)> &Func, const FName& ConfigurationName);
 
+	void SetViewportConfiguration(const FName& ConfigurationName) override;
+	TSharedPtr<SEditorViewport> GetFirstViewport();
 
-protected:
-	TWeakPtr<class SDockTab> ParentTab;
+	void UpdateViewportTabWidget(TFunction<TSharedRef<SEditorViewport>(void)> &Func);
 
-	/** Current layout */
-	TSharedPtr< class FEditorViewportLayout > ActiveViewportLayout;
+	/** Save any configuration required to persist state for this viewport layout */
+	void SaveConfig() const;
 
-	TOptional<FName> PreviouslyFocusedViewport;
 
 private:
 	TFunction<TSharedRef<SEditorViewport>(void)> ViewportCreationFunc;
diff --git a/Engine/Source/Editor/UnrealEd/Public/ViewportTabContent.h b/Engine/Source/Editor/UnrealEd/Public/ViewportTabContent.h
--- a/Engine/Source/Editor/UnrealEd/Public/ViewportTabContent.h
+++ b/Engine/Source/Editor/UnrealEd/Public/ViewportTabContent.h
@@ -33,6 +33,8 @@ public:
 
 	virtual void SetViewportConfiguration(const FName& ConfigurationName) {}
 
+	void PerformActionOnViewports(TFunction<void(FName Name, TSharedPtr<IEditorViewportLayoutEntity>)> &TFuncPtr);
+
 protected:
 	TWeakPtr<class SDockTab> ParentTab;
 
