missing header #rb none

diff --git a/Engine/Source/Editor/UnrealEd/Private/EditorViewportTabContent.cpp b/Engine/Source/Editor/UnrealEd/Private/EditorViewportTabContent.cpp
--- a/Engine/Source/Editor/UnrealEd/Private/EditorViewportTabContent.cpp
+++ b/Engine/Source/Editor/UnrealEd/Private/EditorViewportTabContent.cpp
@@ -6,6 +6,7 @@
 #include "EditorViewportLayoutThreePanes.h"
 #include "EditorViewportLayoutFourPanes.h"
 #include "EditorViewportLayout2x2.h"
+#include "Framework/Docking/LayoutService.h"
 
 
 
