Missing globals header #rb none

diff --git a/Engine/Source/Editor/UnrealEd/Private/EditorViewportTabContent.cpp b/Engine/Source/Editor/UnrealEd/Private/EditorViewportTabContent.cpp
--- a/Engine/Source/Editor/UnrealEd/Private/EditorViewportTabContent.cpp
+++ b/Engine/Source/Editor/UnrealEd/Private/EditorViewportTabContent.cpp
@@ -7,6 +7,7 @@
 #include "EditorViewportLayoutFourPanes.h"
 #include "EditorViewportLayout2x2.h"
 #include "Framework/Docking/LayoutService.h"
+#include "CoreGlobals.h"
 
 
 
