header #rb none

diff --git a/Engine/Source/Editor/UnrealEd/Private/EditorViewportTabContent.cpp b/Engine/Source/Editor/UnrealEd/Private/EditorViewportTabContent.cpp
--- a/Engine/Source/Editor/UnrealEd/Private/EditorViewportTabContent.cpp
+++ b/Engine/Source/Editor/UnrealEd/Private/EditorViewportTabContent.cpp
@@ -8,6 +8,7 @@
 #include "EditorViewportLayout2x2.h"
 #include "Framework/Docking/LayoutService.h"
 #include "CoreGlobals.h"
+#include "Misc/ConfigCacheIni.h"
 
 
 
