Added a setting to control whether or not the new modes UI is visible

diff --git a/Engine/Source/Editor/UnrealEd/Classes/Settings/LevelEditorMiscSettings.h b/Engine/Source/Editor/UnrealEd/Classes/Settings/LevelEditorMiscSettings.h
--- a/Engine/Source/Editor/UnrealEd/Classes/Settings/LevelEditorMiscSettings.h
+++ b/Engine/Source/Editor/UnrealEd/Classes/Settings/LevelEditorMiscSettings.h
@@ -43,6 +43,9 @@ public:
 	UPROPERTY(EditAnywhere, config, Category=Editing, meta=( DisplayName = "Preserve Actor Scale on Replace" ))
 	uint32 bReplaceRespectsScale:1;
 
+	/** If enabled, the modes tab will revert to the previous mode UI */
+	UPROPERTY(EditAnywhere, config, AdvancedDisplay, Category = Editing, meta=(ConfigRestartRequired = true))
+	uint32 bEnableLegacyEditorModeUI : 1;
 public:
 	/** If checked audio playing in the editor will continue to play even if the editor is in the background */
 	UPROPERTY(EditAnywhere, config, Category=Sound)
diff --git a/Engine/Source/Editor/UnrealEd/Private/Settings/SettingsClasses.cpp b/Engine/Source/Editor/UnrealEd/Private/Settings/SettingsClasses.cpp
--- a/Engine/Source/Editor/UnrealEd/Private/Settings/SettingsClasses.cpp
+++ b/Engine/Source/Editor/UnrealEd/Private/Settings/SettingsClasses.cpp
@@ -430,6 +430,7 @@ ULevelEditorMiscSettings::ULevelEditorMiscSettings( const FObjectInitializer& Ob
 	PercentageThresholdForPrompt = 20.0f;
 	MinimumBoundsForCheckingSize = FVector(500.0f, 500.0f, 50.0f);
 	bCreateNewAudioDeviceForPlayInEditor = true;
+	bEnableLegacyEditorModeUI = false;
 }
 
 void ULevelEditorMiscSettings::PostEditChangeProperty( struct FPropertyChangedEvent& PropertyChangedEvent )
