Add separate command to snap pivot to vertex

diff --git a/Engine/Source/Editor/LevelEditor/Private/LevelEditorActions.cpp b/Engine/Source/Editor/LevelEditor/Private/LevelEditorActions.cpp
--- a/Engine/Source/Editor/LevelEditor/Private/LevelEditorActions.cpp
+++ b/Engine/Source/Editor/LevelEditor/Private/LevelEditorActions.cpp
@@ -3186,6 +3186,7 @@ void FLevelEditorCommands::RegisterCommands()
 	UI_COMMAND( AttachActorIteractive, "Attach Actor Interactive", "Start an interactive actor picker to let you choose a parent for the currently selected actor", EUserInterfaceActionType::Button, FInputChord(EModifierKey::Alt, EKeys::A) );
 	UI_COMMAND( CreateNewOutlinerFolder, "Create Folder", "Place the selected actors in a new folder", EUserInterfaceActionType::Button, FInputChord() );
 	UI_COMMAND( HoldToEnableVertexSnapping, "Hold to Enable Vertex Snapping", "When the key binding is pressed and held vertex snapping will be enabled", EUserInterfaceActionType::ToggleButton, FInputChord(EKeys::V) );
+	UI_COMMAND( HoldToEnablePivotVertexSnapping, "Hold to Enable Pivot Vertex Snapping", "Hold to enable vertex snapping while dragging a pivot. Alt must be a modifier in this command or it will not work.", EUserInterfaceActionType::ToggleButton, FInputChord(EModifierKey::Alt, EKeys::V) );
 
 	//@ todo Slate better tooltips for pivot options
 	UI_COMMAND( SavePivotToPrePivot, "Set as Pivot Offset", "Sets the current pivot location as the pivot offset for this actor", EUserInterfaceActionType::Button, FInputChord() );
diff --git a/Engine/Source/Editor/LevelEditor/Public/LevelEditorActions.h b/Engine/Source/Editor/LevelEditor/Public/LevelEditorActions.h
--- a/Engine/Source/Editor/LevelEditor/Public/LevelEditorActions.h
+++ b/Engine/Source/Editor/LevelEditor/Public/LevelEditorActions.h
@@ -266,6 +266,7 @@ public:
 	TSharedPtr< FUICommandInfo > CreateNewOutlinerFolder;
 
 	TSharedPtr< FUICommandInfo > HoldToEnableVertexSnapping;
+	TSharedPtr< FUICommandInfo > HoldToEnablePivotVertexSnapping;
 
 	/**
 	 * Brush Commands                   
diff --git a/Engine/Source/Editor/UnrealEd/Private/LevelEditorViewport.cpp b/Engine/Source/Editor/UnrealEd/Private/LevelEditorViewport.cpp
--- a/Engine/Source/Editor/UnrealEd/Private/LevelEditorViewport.cpp
+++ b/Engine/Source/Editor/UnrealEd/Private/LevelEditorViewport.cpp
@@ -2601,7 +2601,7 @@ bool FLevelEditorViewportClient::InputWidgetDelta(FViewport* InViewport, EAxisLi
 				}
 				else
 				{
-					FSnappingUtils::SnapDragLocationToNearestVertex( ModeTools->PivotLocation, Drag, this );
+					FSnappingUtils::SnapDragLocationToNearestVertex( ModeTools->PivotLocation, Drag, this, true );
 					GUnrealEd->SetPivotMovedIndependently(true);
 					bOnlyMovedPivot = true;
 				}
diff --git a/Engine/Source/Editor/UnrealEd/Private/SnappingUtils.cpp b/Engine/Source/Editor/UnrealEd/Private/SnappingUtils.cpp
--- a/Engine/Source/Editor/UnrealEd/Private/SnappingUtils.cpp
+++ b/Engine/Source/Editor/UnrealEd/Private/SnappingUtils.cpp
@@ -42,7 +42,7 @@ public:
 	/**
 	 * @return true if snapping to vertices is enabled
 	 */
-	bool IsSnapToVertexEnabled();
+	bool IsSnapToVertexEnabled(bool bIsPivot = false);
 
 	/**
 	 * @return true if snapping actors to other actors is enabled
@@ -83,7 +83,7 @@ public:
 	 * @param ViewportClient	The viewport client being dragged in.
 	 * @return true if anything was snapped
 	 */
-	bool SnapDragLocationToNearestVertex( const FVector& BaseLocation, FVector& DragDelta, FLevelEditorViewportClient* ViewportClient );
+	bool SnapDragLocationToNearestVertex( const FVector& BaseLocation, FVector& DragDelta, FLevelEditorViewportClient* ViewportClient, bool bIsPivot = false );
 
 	/**
 	 * Snaps a location to the nearest vertex
@@ -123,7 +123,7 @@ bool FEditorViewportSnapping::IsSnapScaleEnabled()
 	return GetDefault<ULevelEditorViewportSettings>()->SnapScaleEnabled;
 }
 
-bool FEditorViewportSnapping::IsSnapToVertexEnabled()
+bool FEditorViewportSnapping::IsSnapToVertexEnabled(bool bIsPivot)
 {
 	if( GetDefault<ULevelEditorViewportSettings>()->bSnapVertices )
 	{
@@ -137,7 +137,7 @@ bool FEditorViewportSnapping::IsSnapToVertexEnabled()
 		for (uint32 i = 0; i < static_cast<uint8>(EMultipleKeyBindingIndex::NumChords); ++i)
 		{
 			EMultipleKeyBindingIndex ChordIndex = static_cast<EMultipleKeyBindingIndex> (i);
-			const FInputChord& Chord = *Commands.HoldToEnableVertexSnapping->GetActiveChord(ChordIndex);
+			const FInputChord& Chord = bIsPivot ? (*Commands.HoldToEnablePivotVertexSnapping->GetActiveChord(ChordIndex)) : (*Commands.HoldToEnableVertexSnapping->GetActiveChord(ChordIndex));
 
 			bIsChordPressed |= (Chord.NeedsControl() == GCurrentLevelEditingViewportClient->IsCtrlPressed())
 				&& (Chord.NeedsAlt() == GCurrentLevelEditingViewportClient->IsAltPressed())
@@ -396,10 +396,10 @@ bool FEditorViewportSnapping::SnapDraggedActorsToNearestVertex( FVector& DragDel
 	return bSnapped;
 }
 
-bool FEditorViewportSnapping::SnapDragLocationToNearestVertex( const FVector& BaseLocation, FVector& DragDelta, FLevelEditorViewportClient* ViewportClient )
+bool FEditorViewportSnapping::SnapDragLocationToNearestVertex( const FVector& BaseLocation, FVector& DragDelta, FLevelEditorViewportClient* ViewportClient, bool bIsPivot )
 {
 	bool bSnapped = false;
-	if( IsSnapToVertexEnabled() && !DragDelta.IsNearlyZero() )
+	if( IsSnapToVertexEnabled(bIsPivot) && !DragDelta.IsNearlyZero() )
 	{
 		bSnapped = VertexSnappingImpl.SnapDragLocationToNearestVertex( BaseLocation, DragDelta, ViewportClient );
 	}
@@ -462,9 +462,9 @@ bool FSnappingUtils::SnapDraggedActorsToNearestVertex(FVector& DragDelta, FLevel
 	return EditorViewportSnapper.IsValid() ? EditorViewportSnapper->SnapDraggedActorsToNearestVertex(DragDelta, ViewportClient) : false;
 }
 
-bool FSnappingUtils::SnapDragLocationToNearestVertex(const FVector& BaseLocation, FVector& DragDelta, FLevelEditorViewportClient* ViewportClient)
+bool FSnappingUtils::SnapDragLocationToNearestVertex(const FVector& BaseLocation, FVector& DragDelta, FLevelEditorViewportClient* ViewportClient, bool bIsPivot)
 {
-	return EditorViewportSnapper.IsValid() ? EditorViewportSnapper->SnapDragLocationToNearestVertex(BaseLocation, DragDelta, ViewportClient) : false;
+	return EditorViewportSnapper.IsValid() ? EditorViewportSnapper->SnapDragLocationToNearestVertex(BaseLocation, DragDelta, ViewportClient, bIsPivot) : false;
 }
 
 bool FSnappingUtils::SnapLocationToNearestVertex(FVector& Location, const FVector2D& MouseLocation, FLevelEditorViewportClient* ViewportClient, FVector& OutVertexNormal, bool bDrawVertHelpers )
diff --git a/Engine/Source/Editor/UnrealEd/Public/SnappingUtils.h b/Engine/Source/Editor/UnrealEd/Public/SnappingUtils.h
--- a/Engine/Source/Editor/UnrealEd/Public/SnappingUtils.h
+++ b/Engine/Source/Editor/UnrealEd/Public/SnappingUtils.h
@@ -66,7 +66,7 @@ public:
 	 * @param ViewportClient	The viewport client being dragged in.
 	 * @return true if anything was snapped
 	 */
-	static bool SnapDragLocationToNearestVertex( const FVector& BaseLocation, FVector& DragDelta, FLevelEditorViewportClient* ViewportClient );
+	static bool SnapDragLocationToNearestVertex( const FVector& BaseLocation, FVector& DragDelta, FLevelEditorViewportClient* ViewportClient, bool bIsPivot = false );
 
 	/**
 	 * Snaps a location to the nearest vertex
