Added the ability to filter the meta-data gather by field and field owner types

This allows you to do things like "gather all properties belonging to actors", or "gather everything but functions".

diff --git a/Engine/Source/Editor/UnrealEd/Classes/Commandlets/GatherTextFromMetadataCommandlet.h b/Engine/Source/Editor/UnrealEd/Classes/Commandlets/GatherTextFromMetadataCommandlet.h
--- a/Engine/Source/Editor/UnrealEd/Classes/Commandlets/GatherTextFromMetadataCommandlet.h
+++ b/Engine/Source/Editor/UnrealEd/Classes/Commandlets/GatherTextFromMetadataCommandlet.h
@@ -31,7 +31,16 @@ private:
 private:
 	void GatherTextFromUObjects(const TArray<FString>& IncludePaths, const TArray<FString>& ExcludePaths, const FGatherParameters& Arguments);
 	void GatherTextFromUObject(UField* const Field, const FGatherParameters& Arguments, const FName InPlatformName);
+	bool ShouldGatherFromField(const UField* Field);
 
 private:
 	bool ShouldGatherFromEditorOnlyData;
+
+	/** Array of field types (eg, UProperty, UFunction, UScriptStruct, etc) that should be included or excluded in the current gather */
+	TArray<const UClass*> FieldTypesToInclude;
+	TArray<const UClass*> FieldTypesToExclude;
+
+	/** Array of field owner types (eg, UMyClass, FMyStruct, etc) that should have fields within them included or excluded in the current gather */
+	TArray<const UStruct*> FieldOwnerTypesToInclude;
+	TArray<const UStruct*> FieldOwnerTypesToExclude;
 };
diff --git a/Engine/Source/Editor/UnrealEd/Private/Commandlets/GatherTextFromMetadataCommandlet.cpp b/Engine/Source/Editor/UnrealEd/Private/Commandlets/GatherTextFromMetadataCommandlet.cpp
--- a/Engine/Source/Editor/UnrealEd/Private/Commandlets/GatherTextFromMetadataCommandlet.cpp
+++ b/Engine/Source/Editor/UnrealEd/Private/Commandlets/GatherTextFromMetadataCommandlet.cpp
@@ -2,6 +2,7 @@
 
 #include "Commandlets/GatherTextFromMetadataCommandlet.h"
 #include "UObject/Class.h"
+#include "UObject/Package.h"
 #include "UObject/UnrealType.h"
 #include "Misc/Paths.h"
 #include "Modules/ModuleManager.h"
@@ -106,6 +107,60 @@ int32 UGatherTextFromMetaDataCommandlet::Main( const FString& Params )
 		ShouldGatherFromEditorOnlyData = false;
 	}
 
+	// FieldTypesToInclude/FieldTypesToExclude
+	{
+		auto GetFieldTypesArrayFromConfig = [this, &SectionName, &GatherTextConfigPath](const TCHAR* InConfigKey, TArray<const UClass*>& OutFieldTypes)
+		{
+			TArray<FString> FieldTypeStrs;
+			GetStringArrayFromConfig(*SectionName, InConfigKey, FieldTypeStrs, GatherTextConfigPath);
+
+			for (const FString& FieldTypeStr : FieldTypeStrs)
+			{
+				const UClass* FieldType = FindObject<UClass>(ANY_PACKAGE, *FieldTypeStr);
+				if (!FieldType)
+				{
+					UE_LOG(LogGatherTextFromMetaDataCommandlet, Warning, TEXT("Field Type %s was not found (from %s in section %s). Did you forget a ModulesToPreload entry?"), *FieldTypeStr, InConfigKey, *SectionName);
+					continue;
+				}
+
+				check(FieldType->IsChildOf<UField>());
+				OutFieldTypes.Add(FieldType);
+				GetDerivedClasses(FieldType, (TArray<UClass*>&)OutFieldTypes);
+			}
+		};
+
+		GetFieldTypesArrayFromConfig(TEXT("FieldTypesToInclude"), FieldTypesToInclude);
+		GetFieldTypesArrayFromConfig(TEXT("FieldTypesToExclude"), FieldTypesToExclude);
+	}
+
+	// FieldOwnerTypesToInclude/FieldOwnerTypesToExclude
+	{
+		auto GetFieldOwnerTypesArrayFromConfig = [this, &SectionName, &GatherTextConfigPath](const TCHAR* InConfigKey, TArray<const UStruct*>& OutFieldOwnerTypes)
+		{
+			TArray<FString> FieldOwnerTypeStrs;
+			GetStringArrayFromConfig(*SectionName, InConfigKey, FieldOwnerTypeStrs, GatherTextConfigPath);
+
+			for (const FString& FieldOwnerTypeStr : FieldOwnerTypeStrs)
+			{
+				const UStruct* FieldOwnerType = FindObject<UStruct>(ANY_PACKAGE, *FieldOwnerTypeStr);
+				if (!FieldOwnerType)
+				{
+					UE_LOG(LogGatherTextFromMetaDataCommandlet, Warning, TEXT("Field Owner Type %s was not found (from %s in section %s). Did you forget a ModulesToPreload entry?"), *FieldOwnerTypeStr, InConfigKey, *SectionName);
+					continue;
+				}
+
+				OutFieldOwnerTypes.Add(FieldOwnerType);
+				if (const UClass* FieldOwnerClass = Cast<UClass>(FieldOwnerType))
+				{
+					GetDerivedClasses(FieldOwnerClass, (TArray<UClass*>&)OutFieldOwnerTypes);
+				}
+			}
+		};
+
+		GetFieldOwnerTypesArrayFromConfig(TEXT("FieldOwnerTypesToInclude"), FieldOwnerTypesToInclude);
+		GetFieldOwnerTypesArrayFromConfig(TEXT("FieldOwnerTypesToExclude"), FieldOwnerTypesToExclude);
+	}
+
 	FGatherParameters Arguments;
 	GetStringArrayFromConfig(*SectionName, TEXT("InputKeys"), Arguments.InputKeys, GatherTextConfigPath);
 	GetStringArrayFromConfig(*SectionName, TEXT("OutputNamespaces"), Arguments.OutputNamespaces, GatherTextConfigPath);
@@ -142,6 +197,12 @@ void UGatherTextFromMetaDataCommandlet::GatherTextFromUObjects(const TArray<FStr
 
 	for(TObjectIterator<UField> It; It; ++It)
 	{
+		// Skip fields excluded by our filter
+		if (!ShouldGatherFromField(*It))
+		{
+			continue;
+		}
+
 		// Skip editor-only properties if we're not gathering for editor-only data.
 		UProperty* Property = Cast<UProperty>(*It);
 		if (Property && !ShouldGatherFromEditorOnlyData && Property->HasAnyPropertyFlags(CPF_EditorOnly))
@@ -238,3 +299,38 @@ void UGatherTextFromMetaDataCommandlet::GatherTextFromUObject(UField* const Fiel
 		}
 	}
 }
+
+bool UGatherTextFromMetaDataCommandlet::ShouldGatherFromField(const UField* Field)
+{
+	auto ShouldGatherFieldByType = [this, Field]()
+	{
+		if (FieldTypesToInclude.Num() == 0 && FieldTypesToExclude.Num() == 0)
+		{
+			return true;
+		}
+
+		const UClass* FieldClass = Field->GetClass();
+		return (FieldTypesToInclude.Num() == 0 || FieldTypesToInclude.Contains(FieldClass))
+			&& (FieldTypesToExclude.Num() == 0 || !FieldTypesToExclude.Contains(FieldClass));
+	};
+
+	auto ShouldGatherFieldByOwnerType = [this, Field]()
+	{
+		if (FieldOwnerTypesToInclude.Num() == 0 && FieldOwnerTypesToExclude.Num() == 0)
+		{
+			return true;
+		}
+
+		const UStruct* FieldOwnerType = Field->GetOwnerStruct();
+		if (FieldOwnerType)
+		{
+			// Only properties and functions will have an owner struct type
+			return (FieldOwnerTypesToInclude.Num() == 0 || FieldOwnerTypesToInclude.Contains(FieldOwnerType))
+				&& (FieldOwnerTypesToExclude.Num() == 0 || !FieldOwnerTypesToExclude.Contains(FieldOwnerType));
+		}
+
+		return true;
+	};
+
+	return ShouldGatherFieldByType() && ShouldGatherFieldByOwnerType();
+}
