Editor: Should no longer trip an ensure when automation tests attempt to load PIE.

#rnx

diff --git a/Engine/Source/Editor/UnrealEd/Private/EditorEngine.cpp b/Engine/Source/Editor/UnrealEd/Private/EditorEngine.cpp
--- a/Engine/Source/Editor/UnrealEd/Private/EditorEngine.cpp
+++ b/Engine/Source/Editor/UnrealEd/Private/EditorEngine.cpp
@@ -7341,7 +7341,10 @@ void UEditorEngine::AutomationLoadMap(const FString& MapName, FString* OutError)
 		EditorPlaySettings->bLaunchSeparateServer = false;
 		RequestParams.EditorPlaySettings = EditorPlaySettings;
 
-		StartPlayInEditorSession(RequestParams);
+		RequestPlaySession(RequestParams);
+
+		// Immediately launch the session 
+		StartQueuedPlaySessionRequest();
 
 		if (!FailHandler.CanProceed())
 		{
