CIS: Fixed deprecation warning.

#rnx

diff --git a/Engine/Source/Editor/UnrealEd/Public/PlayInEditorDataTypes.h b/Engine/Source/Editor/UnrealEd/Public/PlayInEditorDataTypes.h
--- a/Engine/Source/Editor/UnrealEd/Public/PlayInEditorDataTypes.h
+++ b/Engine/Source/Editor/UnrealEd/Public/PlayInEditorDataTypes.h
@@ -79,12 +79,14 @@ struct FPieLoginStruct
 
 	FPieLoginStruct()
 		: WorldContextHandle(NAME_None)
-		, SettingsIndex(0)
 		, NextX(0)
 		, NextY(0)
 		, PIEStartTime(0)
 		, PIEInstanceIndex(-1)
 	{
+		PRAGMA_DISABLE_DEPRECATION_WARNINGS
+		SettingsIndex = 0;
+		PRAGMA_ENABLE_DEPRECATION_WARNINGS
 	}
 };
 
