CIS: Supress deprecated variable warning on Linux created by the default copy constructor.

#rnx

diff --git a/Engine/Source/Editor/UnrealEd/Public/PlayInEditorDataTypes.h b/Engine/Source/Editor/UnrealEd/Public/PlayInEditorDataTypes.h
--- a/Engine/Source/Editor/UnrealEd/Public/PlayInEditorDataTypes.h
+++ b/Engine/Source/Editor/UnrealEd/Public/PlayInEditorDataTypes.h
@@ -76,6 +76,9 @@ struct FPieLoginStruct
 	/** Which index is the instance of this Play in Editor session that will be created. */
 	int32 PIEInstanceIndex;
 
+	PRAGMA_DISABLE_DEPRECATION_WARNINGS
+	FPieLoginStruct& operator = (FPieLoginStruct const&) = default;
+	PRAGMA_ENABLE_DEPRECATION_WARNINGS
 
 	FPieLoginStruct()
 		: WorldContextHandle(NAME_None)
