Fix for infinite loop when changing widget mode on a bone in the AnimBP editor

diff --git a/Engine/Source/Editor/UnrealEd/Private/SEditorViewport.cpp b/Engine/Source/Editor/UnrealEd/Private/SEditorViewport.cpp
--- a/Engine/Source/Editor/UnrealEd/Private/SEditorViewport.cpp
+++ b/Engine/Source/Editor/UnrealEd/Private/SEditorViewport.cpp
@@ -630,6 +630,12 @@ void SEditorViewport::OnCycleWidgetMode()
 {
 	FWidget::EWidgetMode WidgetMode = Client->GetWidgetMode();
 
+	// Can't cycle the widget mode if we don't currently have a widget
+	if (WidgetMode == FWidget::WM_None)
+	{
+		return;
+	}
+
 	int32 WidgetModeAsInt = WidgetMode;
 
 	do
