Removed my own deprecated code during 4.24 for 4.25:
1) After enhancement JIRA UE-77383, ILayers and GEditor->Layers were deprecated.
2) After JIRA UE-75683, FSlateRect GetToolTipForceFieldRect() was deprecated.

diff --git a/Engine/Plugins/Experimental/DataPrepEditor/Source/DataPrepLibraries/Private/DataprepEditingOperations.cpp b/Engine/Plugins/Experimental/DataPrepEditor/Source/DataPrepLibraries/Private/DataprepEditingOperations.cpp
--- a/Engine/Plugins/Experimental/DataPrepEditor/Source/DataPrepLibraries/Private/DataprepEditingOperations.cpp
+++ b/Engine/Plugins/Experimental/DataPrepEditor/Source/DataPrepLibraries/Private/DataprepEditingOperations.cpp
@@ -20,7 +20,6 @@
 #include "GenericPlatform/GenericPlatformTime.h"
 #include "IMeshBuilderModule.h"
 #include "IMeshMergeUtilities.h"
-#include "Layers/ILayers.h"
 #include "LevelSequence.h"
 #include "Materials/MaterialInstance.h"
 #include "MeshMergeModule.h"
diff --git a/Engine/Source/Editor/UnrealEd/Classes/Editor/EditorEngine.h b/Engine/Source/Editor/UnrealEd/Classes/Editor/EditorEngine.h
--- a/Engine/Source/Editor/UnrealEd/Classes/Editor/EditorEngine.h
+++ b/Engine/Source/Editor/UnrealEd/Classes/Editor/EditorEngine.h
@@ -589,10 +589,6 @@ private:
 
 public:
 
-	/** The "manager" of all the layers for the UWorld currently being edited */
-	UE_DEPRECATED(4.24, "ILayers and FLayers (GEditor->Layers) have been deprecated, use ULayersSubsystem (rather than ILayers) and GEditor->GetEditorSubsystem<ULayersSubsystem>() (rather than GEditor->Layers) instead.")
-	TSharedPtr< class ILayers >				Layers;
-
 	/** List of all viewport clients */
 	const TArray<class FEditorViewportClient*>& GetAllViewportClients() { return AllViewportClients; }
 	const TArray<class FEditorViewportClient*>& GetAllViewportClients() const { return AllViewportClients; }
diff --git a/Engine/Source/Editor/UnrealEd/Private/EditorActor.cpp b/Engine/Source/Editor/UnrealEd/Private/EditorActor.cpp
--- a/Engine/Source/Editor/UnrealEd/Private/EditorActor.cpp
+++ b/Engine/Source/Editor/UnrealEd/Private/EditorActor.cpp
@@ -49,7 +49,6 @@
 #include "EditorLevelUtils.h"
 #include "Kismet2/BlueprintEditorUtils.h"
 #include "LevelEditorViewport.h"
-#include "Layers/ILayers.h"
 #include "Layers/LayersSubsystem.h"
 #include "ActorEditorUtils.h"
 #include "Particles/ParticleSystemComponent.h"
diff --git a/Engine/Source/Editor/UnrealEd/Private/EditorEngine.cpp b/Engine/Source/Editor/UnrealEd/Private/EditorEngine.cpp
--- a/Engine/Source/Editor/UnrealEd/Private/EditorEngine.cpp
+++ b/Engine/Source/Editor/UnrealEd/Private/EditorEngine.cpp
@@ -953,12 +953,6 @@ void UEditorEngine::Init(IEngineLoop* InEngineLoop)
 
 	LoadEditorFeatureLevel();
 
-	PRAGMA_DISABLE_DEPRECATION_WARNINGS
-	// ILayers (GEditor->Layers) has been deprecated, use ULayersSubsystem (GEditor->GetEditorSubsystem<ULayersSubsystem>()) instead.
-	// We temporaily assign Layers = GEditor->GetEditorSubsystem<ULayersSubsystem>(), but we will remove Layers in future releases.
-	Layers = MakeShareable(GetEditorSubsystem<ULayersSubsystem>(), [](ULayersSubsystem*) {});
-	PRAGMA_ENABLE_DEPRECATION_WARNINGS
-
 
 	// Init transactioning.
 	Trans = CreateTrans();
diff --git a/Engine/Source/Editor/UnrealEd/Private/Layers/ILayers.cpp b/Engine/Source/Editor/UnrealEd/Private/Layers/ILayers.cpp
deleted file mode 100644
diff --git a/Engine/Source/Editor/UnrealEd/Private/Layers/LayersSubsystem.cpp b/Engine/Source/Editor/UnrealEd/Private/Layers/LayersSubsystem.cpp
--- a/Engine/Source/Editor/UnrealEd/Private/Layers/LayersSubsystem.cpp
+++ b/Engine/Source/Editor/UnrealEd/Private/Layers/LayersSubsystem.cpp
@@ -94,8 +94,6 @@ void FLayersBroadcast::OnEditorRefreshLayerBrowser()
 
 
 
-PRAGMA_DISABLE_DEPRECATION_WARNINGS
-
 void ULayersSubsystem::Initialize(FSubsystemCollectionBase& Collection)
 {
 	Super::Initialize(Collection);
@@ -1421,5 +1419,3 @@ ULayer* ULayersSubsystem::EnsureLayerExists(const FName& LayerName)
 
 	return Layer;
 }
-
-PRAGMA_ENABLE_DEPRECATION_WARNINGS
diff --git a/Engine/Source/Editor/UnrealEd/Public/Layers/ILayers.h b/Engine/Source/Editor/UnrealEd/Public/Layers/ILayers.h
deleted file mode 100644
diff --git a/Engine/Source/Editor/UnrealEd/Public/Layers/LayersSubsystem.h b/Engine/Source/Editor/UnrealEd/Public/Layers/LayersSubsystem.h
--- a/Engine/Source/Editor/UnrealEd/Public/Layers/LayersSubsystem.h
+++ b/Engine/Source/Editor/UnrealEd/Public/Layers/LayersSubsystem.h
@@ -4,7 +4,6 @@
 
 #include "CoreMinimal.h"
 #include "EditorSubsystem.h"
-#include "Layers/ILayers.h"
 
 #include "LayersSubsystem.generated.h"
 
@@ -14,9 +13,32 @@ class UEditorEngine;
 class ULayer;
 class UWorld;
 
-PRAGMA_DISABLE_DEPRECATION_WARNINGS
+namespace ELayersAction
+{
+	enum Type
+	{
+		/**	The specified ChangedLayer is a newly created ULayer, if ChangedLayer is invalid then multiple Layers were added */
+		Add,
+
+		/**
+		 *	The specified ChangedLayer was just modified, if ChangedLayer is invalid then multiple Layers were modified.
+		 *  ChangedProperty specifies what field on the ULayer was changed, if NAME_None then multiple fields were changed
+		 */
+		Modify,
+
+		/**	A ULayer was deleted */
+		Delete,
+
+		/**	The specified ChangedLayer was just renamed */
+		Rename,
+
+		/**	A large amount of changes have occurred to a number of Layers. A full rebind will be required. */
+		Reset,
+	};
+}
+
 UCLASS()
-class UNREALED_API ULayersSubsystem : public UEditorSubsystem, public ILayers
+class UNREALED_API ULayersSubsystem : public UEditorSubsystem
 {
 	GENERATED_BODY()
 
@@ -36,11 +58,11 @@ public:
 	/**
 	 *	Destructor
 	 */
-	virtual ~ULayersSubsystem() override final;
+	virtual ~ULayersSubsystem() final;
 
 	/** Broadcasts whenever one or more Layers are modified*/
-	DECLARE_DERIVED_EVENT(ULayersSubsystem, ILayers::FOnLayersChanged, FOnLayersChanged);
-	virtual FOnLayersChanged& OnLayersChanged() override final { return LayersChanged; }
+	DECLARE_EVENT_ThreeParams(ULayersSubsystem, FOnLayersChanged, const ELayersAction::Type /*Action*/, const TWeakObjectPtr< ULayer >& /*ChangedLayer*/, const FName& /*ChangedProperty*/);
+	virtual FOnLayersChanged& OnLayersChanged() final { return LayersChanged; }
 
 	///////////////////////////////////////////////////////////////////////////////////////////////////////////////////////
 	// Operations on Levels
@@ -51,14 +73,14 @@ public:
 	 *	@param	Level	The process
 	 */
 	UFUNCTION(BlueprintCallable, Category = Layers)
-	virtual void AddLevelLayerInformation(ULevel* Level) override final;
+	virtual void AddLevelLayerInformation(ULevel* Level) final;
 	/**
 	 *	Purges any information regarding layers associated with the level and it contents
 	 *
 	 *	@param	Level	The process
 	 */
 	UFUNCTION(BlueprintCallable, Category = Layers)
-	virtual void RemoveLevelLayerInformation(ULevel* Level) override final;
+	virtual void RemoveLevelLayerInformation(ULevel* Level) final;
 
 	///////////////////////////////////////////////////////////////////////////////////////////////////////////////////////
 	// Operations on an individual actor.
@@ -69,21 +91,21 @@ public:
 	 *	@param	Actor	The actor to validate
 	 */
 	UFUNCTION(BlueprintCallable, Category = Layers)
-	virtual bool IsActorValidForLayer(AActor* Actor) override final;
+	virtual bool IsActorValidForLayer(AActor* Actor) final;
 	/**
 	 *	Synchronizes an newly created Actor's layers with the layer system
 	 *
 	 *	@param	Actor	The actor to initialize
 	 */
 	UFUNCTION(BlueprintCallable, Category = Layers)
-	virtual bool InitializeNewActorLayers(AActor* Actor) override final;
+	virtual bool InitializeNewActorLayers(AActor* Actor) final;
 	/**
 	 *	Disassociates an Actor's layers from the layer system, general used before deleting the Actor
 	 *
 	 *	@param	Actor	The actor to disassociate from the layer system
 	 */
 	UFUNCTION(BlueprintCallable, Category = Layers)
-	virtual bool DisassociateActorFromLayers(AActor* Actor) override final;
+	virtual bool DisassociateActorFromLayers(AActor* Actor) final;
 
 	/**
 	 * Adds the actor to the named layer.
@@ -93,7 +115,7 @@ public:
 	 * @return				true if the actor was added.  false is returned if the actor already belongs to the layer.
 	 */
 	UFUNCTION(BlueprintCallable, Category = Layers)
-	virtual bool AddActorToLayer(AActor* Actor, const FName& LayerName) override final;
+	virtual bool AddActorToLayer(AActor* Actor, const FName& LayerName) final;
 	/**
 	 * Adds the provided actor to the named layers.
 	 *
@@ -102,7 +124,7 @@ public:
 	 * @return				true if the actor was added to at least one of the provided layers.
 	 */
 	UFUNCTION(BlueprintCallable, Category = Layers)
-	virtual bool AddActorToLayers(AActor* Actor, const TArray< FName >& LayerNames) override final;
+	virtual bool AddActorToLayers(AActor* Actor, const TArray< FName >& LayerNames) final;
 
 	/**
 	 * Removes an actor from the specified layer.
@@ -112,7 +134,7 @@ public:
 	 * @return					true if the actor was removed from the layer.  false is returned if the actor already belonged to the layer.
 	 */
 	UFUNCTION(BlueprintCallable, Category = Layers)
-	virtual bool RemoveActorFromLayer(AActor* Actor, const FName& LayerToRemove, const bool bUpdateStats = true) override final;
+	virtual bool RemoveActorFromLayer(AActor* Actor, const FName& LayerToRemove, const bool bUpdateStats = true) final;
 	/**
 	 * Removes the provided actor from the named layers.
 	 *
@@ -121,7 +143,7 @@ public:
 	 * @return				true if the actor was removed from at least one of the provided layers.
 	 */
 	UFUNCTION(BlueprintCallable, Category = Layers)
-	virtual bool RemoveActorFromLayers(AActor* Actor, const TArray< FName >& LayerNames, const bool bUpdateStats = true) override final;
+	virtual bool RemoveActorFromLayers(AActor* Actor, const TArray< FName >& LayerNames, const bool bUpdateStats = true) final;
 
 
 	/////////////////////////////////////////////////
@@ -143,7 +165,7 @@ public:
 	 * @param	LayerName	The name of the layer to add to
 	 * @return				true if at least one actor was added to the layer.  false is returned if all the actors already belonged to the layer.
 	 */
-	virtual bool AddActorsToLayer(const TArray< TWeakObjectPtr< AActor > >& Actors, const FName& LayerName) override final;
+	virtual bool AddActorsToLayer(const TArray< TWeakObjectPtr< AActor > >& Actors, const FName& LayerName) final;
 	/**
 	 * Add the actors to the named layers
 	 *
@@ -160,7 +182,7 @@ public:
 	 * @param	LayerNames	A valid list of layer names.
 	 * @return				true if at least one actor was added to at least one layer.  false is returned if all the actors already belonged to all specified layers.
 	 */
-	virtual bool AddActorsToLayers(const TArray< TWeakObjectPtr< AActor > >& Actors, const TArray< FName >& LayerNames) override final;
+	virtual bool AddActorsToLayers(const TArray< TWeakObjectPtr< AActor > >& Actors, const TArray< FName >& LayerNames) final;
 
 	/**
 	 * Removes the actors from the specified layer.
@@ -178,7 +200,7 @@ public:
 	 * @param	LayerToRemove	The name of the layer to remove the actors from
 	 * @return					true if at least one actor was removed from the layer.  false is returned if all the actors already belonged to the layer.
 	 */
-	virtual bool RemoveActorsFromLayer(const TArray< TWeakObjectPtr< AActor > >& Actors, const FName& LayerName, const bool bUpdateStats = true) override final;
+	virtual bool RemoveActorsFromLayer(const TArray< TWeakObjectPtr< AActor > >& Actors, const FName& LayerName, const bool bUpdateStats = true) final;
 	/**
 	 * Remove the actors to the named layers
 	 *
@@ -195,7 +217,7 @@ public:
 	 * @param	LayerNames	A valid list of layer names.
 	 * @return				true if at least one actor was removed from at least one layer.  false is returned if none of the actors belonged to any of the specified layers.
 	 */
-	virtual bool RemoveActorsFromLayers(const TArray< TWeakObjectPtr< AActor > >& Actors, const TArray< FName >& LayerNames, const bool bUpdateStats = true) override final;
+	virtual bool RemoveActorsFromLayers(const TArray< TWeakObjectPtr< AActor > >& Actors, const TArray< FName >& LayerNames, const bool bUpdateStats = true) final;
 
 	
 	/////////////////////////////////////////////////
@@ -217,7 +239,7 @@ public:
 	 * @return				true if at least one actor was added.  false is returned if all selected actors already belong to the named layer.
 	 */
 	UFUNCTION(BlueprintCallable, Category = Layers)
-	virtual bool AddSelectedActorsToLayer(const FName& LayerName) override final;
+	virtual bool AddSelectedActorsToLayer(const FName& LayerName) final;
 	/**
 	 * Adds selected actors to the named layers.
 	 *
@@ -225,7 +247,7 @@ public:
 	 * @return				true if at least one actor was added.  false is returned if all selected actors already belong to the named layers.
 	 */
 	UFUNCTION(BlueprintCallable, Category = Layers)
-	virtual bool AddSelectedActorsToLayers(const TArray< FName >& LayerNames) override final;
+	virtual bool AddSelectedActorsToLayers(const TArray< FName >& LayerNames) final;
 
 	/**
 	 * Removes the selected actors from the named layer.
@@ -234,7 +256,7 @@ public:
 	 * @return				true if at least one actor was added.  false is returned if all selected actors already belong to the named layer.
 	 */
 	UFUNCTION(BlueprintCallable, Category = Layers)
-	virtual bool RemoveSelectedActorsFromLayer(const FName& LayerName) override final;
+	virtual bool RemoveSelectedActorsFromLayer(const FName& LayerName) final;
 	/**
 	 * Removes selected actors from the named layers.
 	 *
@@ -242,7 +264,7 @@ public:
 	 * @return				true if at least one actor was removed.
 	 */
 	UFUNCTION(BlueprintCallable, Category = Layers)
-	virtual bool RemoveSelectedActorsFromLayers(const TArray< FName >& LayerNames) override final;
+	virtual bool RemoveSelectedActorsFromLayers(const TArray< FName >& LayerNames) final;
 
 
 	/////////////////////////////////////////////////
@@ -269,7 +291,7 @@ public:
 	 * @param	Filter	[optional]				Actor that don't pass the specified filter restrictions won't be selected.
 	 * @return									true if at least one actor was selected/deselected.
 	 */
-	virtual bool SelectActorsInLayer(const FName& LayerName, const bool bSelect, const bool bNotify, const bool bSelectEvenIfHidden, const TSharedPtr< ActorFilter >& Filter) override final;
+	virtual bool SelectActorsInLayer(const FName& LayerName, const bool bSelect, const bool bNotify, const bool bSelectEvenIfHidden, const TSharedPtr< ActorFilter >& Filter) final;
 	/**
 	 * Selects/de-selects actors belonging to the named layers.
 	 *
@@ -291,7 +313,7 @@ public:
 	 * @param	Filter	[optional]				Actor that don't pass the specified filter restrictions won't be selected.
 	 * @return									true if at least one actor was selected/deselected.
 	 */
-	virtual bool SelectActorsInLayers(const TArray< FName >& LayerNames, const bool bSelect, const bool bNotify, const bool bSelectEvenIfHidden, const TSharedPtr< ActorFilter >& Filter) override final;
+	virtual bool SelectActorsInLayers(const TArray< FName >& LayerNames, const bool bSelect, const bool bNotify, const bool bSelectEvenIfHidden, const TSharedPtr< ActorFilter >& Filter) final;
 
 
 	/////////////////////////////////////////////////
@@ -303,14 +325,14 @@ public:
 	 * @param LayerThatChanged  If one layer was changed (toggled in view pop-up, etc), then we only need to modify actors that use that layer.
 	 */
 	UFUNCTION(BlueprintCallable, Category = Layers)
-	virtual void UpdateAllViewVisibility(const FName& LayerThatChanged) override final;
+	virtual void UpdateAllViewVisibility(const FName& LayerThatChanged) final;
 	/**
 	 * Updates the per-view visibility for all actors for the given view
 	 *
 	 * @param ViewportClient				The viewport client to update visibility on
 	 * @param LayerThatChanged [optional]	If one layer was changed (toggled in view pop-up, etc), then we only need to modify actors that use that layer
 	 */
-	virtual void UpdatePerViewVisibility(FLevelEditorViewportClient* ViewportClient, const FName& LayerThatChanged = NAME_Skip) override final;
+	virtual void UpdatePerViewVisibility(FLevelEditorViewportClient* ViewportClient, const FName& LayerThatChanged = NAME_Skip) final;
 
 	/**
 	 * Updates per-view visibility for the given actor in the given view
@@ -319,21 +341,21 @@ public:
 	 * @param Actor								Actor to update
 	 * @param bReregisterIfDirty [optional]		If true, the actor will reregister itself to give the rendering thread updated information
 	 */
-	virtual void UpdateActorViewVisibility(FLevelEditorViewportClient* ViewportClient, AActor* Actor, const bool bReregisterIfDirty = true) override final;
+	virtual void UpdateActorViewVisibility(FLevelEditorViewportClient* ViewportClient, AActor* Actor, const bool bReregisterIfDirty = true) final;
 	/**
 	 * Updates per-view visibility for the given actor for all views
 	 *
 	 * @param Actor		Actor to update
 	 */
 	UFUNCTION(BlueprintCallable, Category = Layers)
-	virtual void UpdateActorAllViewsVisibility(AActor* Actor) override final;
+	virtual void UpdateActorAllViewsVisibility(AActor* Actor) final;
 
 	/**
 	 * Removes the corresponding visibility bit from all actors (slides the later bits down 1)
 	 *
 	 * @param ViewportClient	The viewport client to update visibility on
 	 */
-	virtual void RemoveViewFromActorViewVisibility(FLevelEditorViewportClient* ViewportClient) override final;
+	virtual void RemoveViewFromActorViewVisibility(FLevelEditorViewportClient* ViewportClient) final;
 
 	/**
 	 * Updates the provided actors visibility in the viewports
@@ -345,7 +367,7 @@ public:
 	 * @param	bRedrawViewports			If true the viewports will be redrawn; if false, they will not
 	 */
 	UFUNCTION(BlueprintCallable, Category = Layers)
-	virtual bool UpdateActorVisibility(AActor* Actor, bool& bOutSelectionChanged, bool& bOutActorModified, const bool bNotifySelectionChange, const bool bRedrawViewports) override final;
+	virtual bool UpdateActorVisibility(AActor* Actor, bool& bOutSelectionChanged, bool& bOutActorModified, const bool bNotifySelectionChange, const bool bRedrawViewports) final;
 	/**
 	 * Updates the visibility of all actors in the viewports
 	 *
@@ -353,7 +375,7 @@ public:
 	 * @param	bRedrawViewports			If true the viewports will be redrawn; if false, they will not
 	 */
 	UFUNCTION(BlueprintCallable, Category = Layers)
-	virtual bool UpdateAllActorsVisibility(const bool bNotifySelectionChange, const bool bRedrawViewports) override final;
+	virtual bool UpdateAllActorsVisibility(const bool bNotifySelectionChange, const bool bRedrawViewports) final;
 
 
 	/////////////////////////////////////////////////
@@ -382,7 +404,7 @@ public:
 	 *	@param	InOutActors			The list to append the found actors to.
 	 *  @param	Filter	[optional]	Actor that don't pass the specified filter restrictions won't be selected.
 	 */
-	virtual void AppendActorsFromLayer(const FName& LayerName, TArray< TWeakObjectPtr< AActor > >& InOutActors, const TSharedPtr< ActorFilter >& Filter = TSharedPtr< ActorFilter >(nullptr)) const override final;
+	virtual void AppendActorsFromLayer(const FName& LayerName, TArray< TWeakObjectPtr< AActor > >& InOutActors, const TSharedPtr< ActorFilter >& Filter = TSharedPtr< ActorFilter >(nullptr)) const final;
 	/**
 	 *	Appends all the actors associated with ANY of the specified layers.
 	 *
@@ -406,7 +428,7 @@ public:
 	 *	@param	InOutActors			The list to append the found actors to.
 	 *  @param	Filter	[optional]	Actor that don't pass the specified filter restrictions won't be selected.
 	 */
-	virtual void AppendActorsFromLayers(const TArray< FName >& LayerNames, TArray< TWeakObjectPtr< AActor > >& InOutActors, const TSharedPtr< ActorFilter >& Filter = TSharedPtr< ActorFilter >(nullptr)) const override final;
+	virtual void AppendActorsFromLayers(const TArray< FName >& LayerNames, TArray< TWeakObjectPtr< AActor > >& InOutActors, const TSharedPtr< ActorFilter >& Filter = TSharedPtr< ActorFilter >(nullptr)) const final;
 
 	/**
 	 *	Gets all the actors associated with the specified layer. Analog to AppendActorsFromLayer but it returns rather than appends the actors.
@@ -448,7 +470,7 @@ public:
 	 * @param	bIsVisible	If true the layer will be visible; if false, the layer will not be visible.
 	 */
 	UFUNCTION(BlueprintCallable, Category = Layers)
-	virtual void SetLayerVisibility(const FName& LayerName, const bool bIsVisible) override final;
+	virtual void SetLayerVisibility(const FName& LayerName, const bool bIsVisible) final;
 	/**
 	 * Changes visibility of the named layers to the provided state
 	 *
@@ -456,7 +478,7 @@ public:
 	 * @param	bIsVisible	If true the layers will be visible; if false, the layers will not be visible
 	 */
 	UFUNCTION(BlueprintCallable, Category = Layers)
-	virtual void SetLayersVisibility(const TArray< FName >& LayerNames, const bool bIsVisible) override final;
+	virtual void SetLayersVisibility(const TArray< FName >& LayerNames, const bool bIsVisible) final;
 
 	/**
 	 * Toggles the named layer's visibility
@@ -464,20 +486,20 @@ public:
 	 * @param LayerName	The name of the layer to affect
 	 */
 	UFUNCTION(BlueprintCallable, Category = Layers)
-	virtual void ToggleLayerVisibility(const FName& LayerName) override final;
+	virtual void ToggleLayerVisibility(const FName& LayerName) final;
 	/**
 	 * Toggles the visibility of all of the named layers
 	 *
 	 * @param	LayerNames	The names of the layers to affect
 	 */
 	UFUNCTION(BlueprintCallable, Category = Layers)
-	virtual void ToggleLayersVisibility(const TArray< FName >& LayerNames) override final;
+	virtual void ToggleLayersVisibility(const TArray< FName >& LayerNames) final;
 
 	/**
 	 * Set the visibility of all layers to true
 	 */
 	UFUNCTION(BlueprintCallable, Category = Layers)
-	virtual void MakeAllLayersVisible() override final;
+	virtual void MakeAllLayersVisible() final;
 
 	/**
 	 * Gets the ULayer Object of the named layer
@@ -503,7 +525,7 @@ public:
 	 * @return					If true a valid ULayer Object was found and set to OutLayer; if false, a valid ULayer object was not found and invalid set to OutLayer
 	 */
 	UFUNCTION(BlueprintCallable, Category = Layers)
-	virtual bool TryGetLayer(const FName& LayerName, ULayer*& OutLayer) override final;
+	virtual bool TryGetLayer(const FName& LayerName, ULayer*& OutLayer) final;
 
 	/**
 	 * Gets all known layers and appends their names to the provide array
@@ -511,20 +533,20 @@ public:
 	 * @param OutLayers[OUT] Output array to store all known layers
 	 */
 	UFUNCTION(BlueprintCallable, Category = Layers)
-	virtual void AddAllLayerNamesTo(TArray< FName >& OutLayerNames) const override final;
+	virtual void AddAllLayerNamesTo(TArray< FName >& OutLayerNames) const final;
 	/**
 	 * Gets all known layers and appends them to the provided array
 	 *
 	 * @param OutLayers[OUT] Output array to store all known layers
 	 */
 	UFUNCTION(BlueprintCallable, Category = Layers)
-	virtual void AddAllLayersTo(TArray< ULayer* > & OutLayers) const override final;
+	virtual void AddAllLayersTo(TArray< ULayer* > & OutLayers) const final;
 	/**
 	 * Gets all known layers and appends them to the provided array
 	 *
 	 * @param OutLayers[OUT] Output array to store all known layers
 	 */
-	virtual void AddAllLayersTo(TArray< TWeakObjectPtr< ULayer > >& OutLayers) const override final;
+	virtual void AddAllLayersTo(TArray< TWeakObjectPtr< ULayer > >& OutLayers) const final;
 
 	/**
 	 * Creates a ULayer Object for the named layer
@@ -541,14 +563,14 @@ public:
 	 * @param LayersToDelete	A valid list of layer names.
 	 */
 	UFUNCTION(BlueprintCallable, Category = Layers)
-	virtual void DeleteLayers(const TArray< FName >& LayersToDelete) override final;
+	virtual void DeleteLayers(const TArray< FName >& LayersToDelete) final;
 	/**
 	 * Deletes the provided layer, disassociating all actors from them
 	 *
 	 * @param LayerToDelete		A valid layer name
 	 */
 	UFUNCTION(BlueprintCallable, Category = Layers)
-	virtual void DeleteLayer(const FName& LayerToDelete) override final;
+	virtual void DeleteLayer(const FName& LayerToDelete) final;
 
 	/**
 	 * Renames the provided originally named layer to the provided new name
@@ -557,7 +579,7 @@ public:
 	 * @param	NewLayerName		The new name for the layer to be renamed
 	 */
 	UFUNCTION(BlueprintCallable, Category = Layers)
-	virtual bool RenameLayer(const FName& OriginalLayerName, const FName& NewLayerName) override final;
+	virtual bool RenameLayer(const FName& OriginalLayerName, const FName& NewLayerName) final;
 
 	/**
 	 * Get the current UWorld object.
@@ -599,5 +621,3 @@ private:
 	 */
 	TSharedPtr<class FLayersBroadcast> LayersBroadcast;
 };
-
-PRAGMA_ENABLE_DEPRECATION_WARNINGS
diff --git a/Engine/Source/Runtime/Slate/Private/Framework/Application/MenuStack.cpp b/Engine/Source/Runtime/Slate/Private/Framework/Application/MenuStack.cpp
--- a/Engine/Source/Runtime/Slate/Private/Framework/Application/MenuStack.cpp
+++ b/Engine/Source/Runtime/Slate/Private/Framework/Application/MenuStack.cpp
@@ -785,13 +785,6 @@ bool FMenuStack::GetToolTipForceFieldRect(const TSharedRef<IMenu>& InMenu, const
 	return bWasSolutionFound;
 }
 
-FSlateRect FMenuStack::GetToolTipForceFieldRect(TSharedRef<IMenu> InMenu, const FWidgetPath& PathContainingMenu) const
-{
-	FSlateRect ForceFieldRect(0, 0, 0, 0);
-	GetToolTipForceFieldRect(InMenu, PathContainingMenu, ForceFieldRect);
-	return ForceFieldRect;
-}
-
 bool FMenuStack::HasMenus() const
 {
 	return Stack.Num() > 0;
diff --git a/Engine/Source/Runtime/Slate/Public/Framework/Application/MenuStack.h b/Engine/Source/Runtime/Slate/Public/Framework/Application/MenuStack.h
--- a/Engine/Source/Runtime/Slate/Public/Framework/Application/MenuStack.h
+++ b/Engine/Source/Runtime/Slate/Public/Framework/Application/MenuStack.h
@@ -184,16 +184,6 @@ public:
 	 */
 	bool GetToolTipForceFieldRect(const TSharedRef<IMenu>& InMenu, const FWidgetPath& InPathContainingMenu, FSlateRect& OutSlateRect) const;
 
-	/**
-	 * Called by the application when showing tooltips. It prevents tooltips from drawing over menu items.
-	 *
-	 * @param InMenu				A menu in the stack, used to generate the force field rect
-	 * @param PathContainingMenu	A widget path containing InMenu. This could be generated but the application has the path so it helps performance to pass it in here. 
-	 * @return						A rectangle enclosing the menu's children (created from the clipping rects in the geometry in PathContainingMenu)
-	 */
-	UE_DEPRECATED(4.24, "FSlateRect GetToolTipForceFieldRect(TSharedRef<IMenu>, const FWidgetPath&) has been deprecated to avoid bugs when no solution is found. Instead, use bool GetToolTipForceFieldRect(const TSharedRef<IMenu>&, const FWidgetPath&, FSlateRect&).")
-	FSlateRect GetToolTipForceFieldRect(TSharedRef<IMenu> InMenu, const FWidgetPath& PathContainingMenu) const;
-
 	/**
 	 * @return	Returns the window that is the parent of everything in the stack, if any. If the stack is empty, returns an invalid ptr.
 	 */
