Solving missing includes (resulting in CIS compiling errors) introduced by CL 10061416.

diff --git a/Engine/Source/Editor/UnrealEd/Public/Layers/LayersSubsystem.h b/Engine/Source/Editor/UnrealEd/Public/Layers/LayersSubsystem.h
--- a/Engine/Source/Editor/UnrealEd/Public/Layers/LayersSubsystem.h
+++ b/Engine/Source/Editor/UnrealEd/Public/Layers/LayersSubsystem.h
@@ -7,10 +7,12 @@
 
 #include "LayersSubsystem.generated.h"
 
+class AActor;
 class FLayersBroadcast;
 class FLevelEditorViewportClient;
 class UEditorEngine;
 class ULayer;
+class ULevel;
 class UWorld;
 
 namespace ELayersAction
