Solving missing includes (resulting in CIS compiling errors) introduced by CL 10061416. Part 2/2.

diff --git a/Engine/Source/Editor/UnrealEd/Public/Layers/LayersSubsystem.h b/Engine/Source/Editor/UnrealEd/Public/Layers/LayersSubsystem.h
--- a/Engine/Source/Editor/UnrealEd/Public/Layers/LayersSubsystem.h
+++ b/Engine/Source/Editor/UnrealEd/Public/Layers/LayersSubsystem.h
@@ -14,6 +14,7 @@ class UEditorEngine;
 class ULayer;
 class ULevel;
 class UWorld;
+template< typename TItemType > class IFilter;
 
 namespace ELayersAction
 {
