Moved settings meta-data to the Editor localization target

#rnx

diff --git a/Engine/Config/Localization/Category.ini b/Engine/Config/Localization/Category.ini
--- a/Engine/Config/Localization/Category.ini
+++ b/Engine/Config/Localization/Category.ini
@@ -1,6 +1,7 @@
 ﻿[CommonSettings]
-SourcePath=./Content/Localization/Category
-DestinationPath=./Content/Localization/Category
+ManifestDependencies=Content/Localization/Editor/Editor.manifest
+SourcePath=Content/Localization/Category
+DestinationPath=Content/Localization/Category
 ManifestName=Category.manifest
 ArchiveName=Category.archive
 PortableObjectName=Category.po
diff --git a/Engine/Config/Localization/Editor.ini b/Engine/Config/Localization/Editor.ini
--- a/Engine/Config/Localization/Editor.ini
+++ b/Engine/Config/Localization/Editor.ini
@@ -41,32 +41,63 @@ PackageFileNameFilters=*.uasset
 bFixBroken=false
 ShouldGatherFromEditorOnlyData=true
 
-;Write Manifest
+;Gather text from metadata
 [GatherTextStep2]
+CommandletClass=GatherTextFromMetaData
+ModulesToPreload=StructViewer
+IncludePathFilters=Source/Editor/*
+IncludePathFilters=Source/Runtime/*
+IncludePathFilters=Source/Developer/*
+IncludePathFilters=Plugins/*
+ExcludePathFilters=*/NoRedist/*
+ExcludePathFilters=*/NotForLicensees/*
+FieldTypesToInclude=Property
+; general settings
+FieldOwnerTypesToInclude=DeveloperSettings
+FieldOwnerTypesToInclude=*Settings
+; explict settings missed by the above filters
+FieldOwnerTypesToInclude=AISystem
+FieldOwnerTypesToInclude=CrowdManager
+FieldOwnerTypesToInclude=Engine
+FieldOwnerTypesToInclude=GameplayDebuggerConfig
+FieldOwnerTypesToInclude=InternationalizationSettingsModel
+FieldOwnerTypesToInclude=NavigationSystemV1
+FieldOwnerTypesToInclude=PersonaOptions
+FieldOwnerTypesToInclude=RecastNavMesh
+InputKeys=DisplayName
+OutputNamespaces=UObjectDisplayNames
+OutputKeys="{FieldPath}"
+InputKeys=Category
+OutputNamespaces=UObjectCategory
+OutputKeys="{MetaDataValue}"
+ShouldGatherFromEditorOnlyData=true
+
+;Write Manifest
+[GatherTextStep3]
 CommandletClass=GenerateGatherManifest
 
 ;Write Archives
-[GatherTextStep3]
+[GatherTextStep4]
 CommandletClass=GenerateGatherArchive
 bPurgeOldEmptyEntries=true
 
 ;Import localized PO files
-[GatherTextStep4]
+[GatherTextStep5]
 CommandletClass=InternationalizationExport
 bImportLoc=true
 
 ;Write Localized Text Resource
-[GatherTextStep5]
+[GatherTextStep6]
 CommandletClass=GenerateTextLocalizationResource
 ResourceName=Editor.locres
 
 ;Export to PO files
-[GatherTextStep6]
+[GatherTextStep7]
 CommandletClass=InternationalizationExport
 bExportLoc=true
 
 ;Write Text Localization Report
-[GatherTextStep7]
+[GatherTextStep8]
 CommandletClass=GenerateTextLocalizationReport
 DestinationPath=\\epicgames.net\root\UE3\Localization\WordCounts
 bWordCountReport=true
diff --git a/Engine/Config/Localization/PropertyNames.ini b/Engine/Config/Localization/PropertyNames.ini
--- a/Engine/Config/Localization/PropertyNames.ini
+++ b/Engine/Config/Localization/PropertyNames.ini
@@ -1,4 +1,5 @@
 ﻿[CommonSettings]
+ManifestDependencies=Content/Localization/Editor/Editor.manifest
 SourcePath=Content/Localization/PropertyNames
 DestinationPath=Content/Localization/PropertyNames
 ManifestName=PropertyNames.manifest
diff --git a/Engine/Source/Editor/UnrealEd/Private/Commandlets/GatherTextFromMetadataCommandlet.cpp b/Engine/Source/Editor/UnrealEd/Private/Commandlets/GatherTextFromMetadataCommandlet.cpp
--- a/Engine/Source/Editor/UnrealEd/Private/Commandlets/GatherTextFromMetadataCommandlet.cpp
+++ b/Engine/Source/Editor/UnrealEd/Private/Commandlets/GatherTextFromMetadataCommandlet.cpp
@@ -114,18 +114,36 @@ int32 UGatherTextFromMetaDataCommandlet::Main( const FString& Params )
 			TArray<FString> FieldTypeStrs;
 			GetStringArrayFromConfig(*SectionName, InConfigKey, FieldTypeStrs, GatherTextConfigPath);
 
+			TArray<const UClass*> AllFieldTypes;
+			AllFieldTypes.Add(UField::StaticClass());
+			GetDerivedClasses(UField::StaticClass(), (TArray<UClass*>&)AllFieldTypes);
+
 			for (const FString& FieldTypeStr : FieldTypeStrs)
 			{
-				const UClass* FieldType = FindObject<UClass>(ANY_PACKAGE, *FieldTypeStr);
-				if (!FieldType)
+				const bool bIsWildcard = FieldTypeStr.GetCharArray().Contains(TEXT('*')) || FieldTypeStr.GetCharArray().Contains(TEXT('?'));
+				if (bIsWildcard)
 				{
-					UE_LOG(LogGatherTextFromMetaDataCommandlet, Warning, TEXT("Field Type %s was not found (from %s in section %s). Did you forget a ModulesToPreload entry?"), *FieldTypeStr, InConfigKey, *SectionName);
-					continue;
+					for (const UClass* FieldType : AllFieldTypes)
+					{
+						if (FieldType->GetName().MatchesWildcard(FieldTypeStr))
+						{
+							OutFieldTypes.Add(FieldType);
+						}
+					}
 				}
+				else
+				{
+					const UClass* FieldType = FindObject<UClass>(ANY_PACKAGE, *FieldTypeStr);
+					if (!FieldType)
+					{
+						UE_LOG(LogGatherTextFromMetaDataCommandlet, Warning, TEXT("Field Type %s was not found (from %s in section %s). Did you forget a ModulesToPreload entry?"), *FieldTypeStr, InConfigKey, *SectionName);
+						continue;
+					}
 
-				check(FieldType->IsChildOf<UField>());
-				OutFieldTypes.Add(FieldType);
-				GetDerivedClasses(FieldType, (TArray<UClass*>&)OutFieldTypes);
+					check(FieldType->IsChildOf<UField>());
+					OutFieldTypes.Add(FieldType);
+					GetDerivedClasses(FieldType, (TArray<UClass*>&)OutFieldTypes);
+				}
 			}
 		};
 
@@ -140,19 +158,37 @@ int32 UGatherTextFromMetaDataCommandlet::Main( const FString& Params )
 			TArray<FString> FieldOwnerTypeStrs;
 			GetStringArrayFromConfig(*SectionName, InConfigKey, FieldOwnerTypeStrs, GatherTextConfigPath);
 
+			TArray<const UStruct*> AllFieldOwnerTypes;
+			GetObjectsOfClass(UClass::StaticClass(), (TArray<UObject*>&)AllFieldOwnerTypes, false);
+			GetObjectsOfClass(UScriptStruct::StaticClass(), (TArray<UObject*>&)AllFieldOwnerTypes, false);
+
 			for (const FString& FieldOwnerTypeStr : FieldOwnerTypeStrs)
 			{
-				const UStruct* FieldOwnerType = FindObject<UStruct>(ANY_PACKAGE, *FieldOwnerTypeStr);
-				if (!FieldOwnerType)
+				const bool bIsWildcard = FieldOwnerTypeStr.GetCharArray().Contains(TEXT('*')) || FieldOwnerTypeStr.GetCharArray().Contains(TEXT('?'));
+				if (bIsWildcard)
 				{
-					UE_LOG(LogGatherTextFromMetaDataCommandlet, Warning, TEXT("Field Owner Type %s was not found (from %s in section %s). Did you forget a ModulesToPreload entry?"), *FieldOwnerTypeStr, InConfigKey, *SectionName);
-					continue;
+					for (const UStruct* FieldOwnerType : AllFieldOwnerTypes)
+					{
+						if (FieldOwnerType->GetName().MatchesWildcard(FieldOwnerTypeStr))
+						{
+							OutFieldOwnerTypes.Add(FieldOwnerType);
+						}
+					}
 				}
-
-				OutFieldOwnerTypes.Add(FieldOwnerType);
-				if (const UClass* FieldOwnerClass = Cast<UClass>(FieldOwnerType))
+				else
 				{
-					GetDerivedClasses(FieldOwnerClass, (TArray<UClass*>&)OutFieldOwnerTypes);
+					const UStruct* FieldOwnerType = FindObject<UStruct>(ANY_PACKAGE, *FieldOwnerTypeStr);
+					if (!FieldOwnerType)
+					{
+						UE_LOG(LogGatherTextFromMetaDataCommandlet, Warning, TEXT("Field Owner Type %s was not found (from %s in section %s). Did you forget a ModulesToPreload entry?"), *FieldOwnerTypeStr, InConfigKey, *SectionName);
+						continue;
+					}
+
+					OutFieldOwnerTypes.Add(FieldOwnerType);
+					if (const UClass* FieldOwnerClass = Cast<UClass>(FieldOwnerType))
+					{
+						GetDerivedClasses(FieldOwnerClass, (TArray<UClass*>&)OutFieldOwnerTypes);
+					}
 				}
 			}
 		};
@@ -248,11 +284,12 @@ void UGatherTextFromMetaDataCommandlet::GatherTextFromUObject(UField* const Fiel
 				{
 					PatternArguments.Add( TEXT("MetaDataValue"), FText::FromString(MetaDataValue) );
 
+					const UStruct* FieldOwnerType = Field->GetOwnerStruct();
 					const FString Namespace = Arguments.OutputNamespaces[i];
 					FLocItem LocItem(MetaDataValue);
 					FManifestContext Context;
 					Context.Key = FText::Format(Arguments.OutputKeys[i], PatternArguments).ToString();
-					Context.SourceLocation = FString::Printf(TEXT("From metadata for key %s of member %s in %s"), *Arguments.InputKeys[i], *Field->GetName(), *Field->GetFullGroupName(true));
+					Context.SourceLocation = FString::Printf(TEXT("From metadata for key %s of member %s in %s (type: %s, owner: %s)"), *Arguments.InputKeys[i], *Field->GetName(), *Field->GetFullGroupName(true), *Field->GetClass()->GetName(), FieldOwnerType ? *FieldOwnerType->GetName() : TEXT("<null>"));
 					Context.PlatformName = InPlatformName;
 					GatherManifestHelper->AddSourceText(Namespace, LocItem, Context);
 				}
