This CL solves the regression of JIRA 55372: PIE window isn't returning to its last position after each run. The new functionality depends on the value of the checkbox "Always center window to screen" (in "Editor Preferences" --> "Level Editor" --> "Play"):
1) If enabled, it always sets PIE to the center of the screen, regardless of whether it was dragged in the previous run. This is the default behavior (analog to the current and previous versions of UE).
2) If disabled, it maintains the current functionality of allowing introducing the exact position. However, if the window is dragged, that desired position is updated accordingly when PIE is stopped (regaining the behavior of versions before 4.20).

Note that these settings only apply to the first PIE window, additional PIE instances are always set to their last dragged positions (maintaining the behavior of previous versions).

diff --git a/Engine/Source/Editor/UnrealEd/Private/PlayLevel.cpp b/Engine/Source/Editor/UnrealEd/Private/PlayLevel.cpp
--- a/Engine/Source/Editor/UnrealEd/Private/PlayLevel.cpp
+++ b/Engine/Source/Editor/UnrealEd/Private/PlayLevel.cpp
@@ -3029,6 +3029,17 @@ TSharedRef<SPIEViewport> UEditorEngine::GeneratePIEViewportWindow(const FRequest
 				if (OwningEditorEngine.IsValid())
 				{
 					OwningEditorEngine->StoreWindowSizeAndPositionForInstanceIndex(InInstanceIndex, WindowSize, WindowPosition);
+					// Update the position of the first PIE window (that also updates its displayed value in "Editor Preferences" --> "Level Editor" --> "Play" --> "New Window Position")
+					if (InInstanceIndex == 0)
+					{
+						ULevelEditorPlaySettings* LevelEditorPlaySettings = GetMutableDefault<ULevelEditorPlaySettings>();
+						// Only update if "Always center window to screen" is disabled
+						if (LevelEditorPlaySettings && !LevelEditorPlaySettings->CenterNewWindow)
+						{
+							// Update the value
+							LevelEditorPlaySettings->NewWindowPosition = WindowPosition;
+						}
+					}
 				}
 
 				// Route the callback
