Merged 2 of the places where ULevelEditorPlaySettings CDO is modified inside PlayLevel.cpp. The thought behind it is that we'd like to avoid modifying ULevelEditorPlaySettings CDO directly - while this version doesn't avoid that it at least helps consolidate the places we do that.

diff --git a/Engine/Source/Editor/UnrealEd/Private/PlayLevel.cpp b/Engine/Source/Editor/UnrealEd/Private/PlayLevel.cpp
--- a/Engine/Source/Editor/UnrealEd/Private/PlayLevel.cpp
+++ b/Engine/Source/Editor/UnrealEd/Private/PlayLevel.cpp
@@ -574,6 +574,16 @@ void UEditorEngine::EndPlayMap()
 			{
 				PlaySettingsConfig->MultipleInstancePositions.Add(PlayInEditorSessionInfo->CachedWindowInfo[WindowIndex].Position);
 			}
+
+			// Update the position where the first PIE window will be opened (this also updates its displayed value in "Editor Preferences" --> "Level Editor" --> "Play" --> "New Window Position")
+			if (WindowIndex == 0)
+			{
+				// Only update it if "Always center window to screen" is disabled
+				if (PlaySettingsConfig && !PlaySettingsConfig->CenterNewWindow)
+				{
+					PlaySettingsConfig->NewWindowPosition = PlaySettingsConfig->MultipleInstancePositions[WindowIndex];
+				}
+			}
 		}
 
 		PlaySettingsConfig->PostEditChange();
@@ -3029,17 +3039,6 @@ TSharedRef<SPIEViewport> UEditorEngine::GeneratePIEViewportWindow(const FRequest
 				if (OwningEditorEngine.IsValid())
 				{
 					OwningEditorEngine->StoreWindowSizeAndPositionForInstanceIndex(InInstanceIndex, WindowSize, WindowPosition);
-					// Update the position of the first PIE window (that also updates its displayed value in "Editor Preferences" --> "Level Editor" --> "Play" --> "New Window Position")
-					if (InInstanceIndex == 0)
-					{
-						ULevelEditorPlaySettings* LevelEditorPlaySettings = GetMutableDefault<ULevelEditorPlaySettings>();
-						// Only update if "Always center window to screen" is disabled
-						if (LevelEditorPlaySettings && !LevelEditorPlaySettings->CenterNewWindow)
-						{
-							// Update the value
-							LevelEditorPlaySettings->NewWindowPosition = WindowPosition;
-						}
-					}
 				}
 
 				// Route the callback
