PIE: Mobile Preview (and other separate-process launch methods) now correctly handle launching multiple Listen Server clients and Play Offline not loading the correct map.

#rnx

diff --git a/Engine/Source/Editor/UnrealEd/Classes/Editor/EditorEngine.h b/Engine/Source/Editor/UnrealEd/Classes/Editor/EditorEngine.h
--- a/Engine/Source/Editor/UnrealEd/Classes/Editor/EditorEngine.h
+++ b/Engine/Source/Editor/UnrealEd/Classes/Editor/EditorEngine.h
@@ -2844,11 +2844,11 @@ public:
 	/** 
 	 * Launch a new process with settings pulled from the specified parameters.
 	 * @param InParams 				Shared parameters to launch with.
-	 * @param InInstanceNum 		Which instance index is this? Used to do load settins per-instance if needed.
-	 * @param bIsHost				Is this instance a host (either dedicated or listen).
-	 * @param bIsDedicatedServer 	Is this instance a dedicated server? If true, implies bIsHost.
+	 * @param InInstanceNum 		Which instance index is this? Used to do load settings per-instance if needed.
+	 * @param NetMode				What net mode should this process launch under?
+	 * @param bIsDedicatedServer 	Is this instance a dedicated server? Overrides NetMode.
 	 */
-	virtual void LaunchNewProcess(const FRequestPlaySessionParams& InParams, const int32 InInstanceNum, bool bIsHost, bool bIsDedicatedServer);
+	virtual void LaunchNewProcess(const FRequestPlaySessionParams& InParams, const int32 InInstanceNum, EPlayNetMode NetMode, bool bIsDedicatedServer);
 
 	// This chunk is used for Play In New Process
 protected:
diff --git a/Engine/Source/Editor/UnrealEd/Private/PlayLevel.cpp b/Engine/Source/Editor/UnrealEd/Private/PlayLevel.cpp
--- a/Engine/Source/Editor/UnrealEd/Private/PlayLevel.cpp
+++ b/Engine/Source/Editor/UnrealEd/Private/PlayLevel.cpp
@@ -1695,8 +1695,7 @@ void UEditorEngine::CreateNewPlayInEditorInstance(FRequestPlaySessionParams &InR
 		}
 
 		const bool bIsDedicatedServer = false;
-		const bool bIsHost = false;
-		LaunchNewProcess(InRequestParams, PlayInEditorSessionInfo->NumClientInstancesCreated, bIsHost, bIsDedicatedServer);
+		LaunchNewProcess(InRequestParams, PlayInEditorSessionInfo->NumClientInstancesCreated, InNetMode, bIsDedicatedServer);
 		PlayInEditorSessionInfo->NumClientInstancesCreated++;
 	}
 	else
@@ -2569,7 +2568,7 @@ void UEditorEngine::StartPlayInEditorSession(FRequestPlaySessionParams& InReques
 				const bool bIsDedicatedServer = true;
 				const bool bIsHost = true;
 				const int32 InstanceIndex = 0;
-				LaunchNewProcess(InRequestParams, InstanceIndex, bIsHost, bIsDedicatedServer);
+				LaunchNewProcess(InRequestParams, InstanceIndex, EPlayNetMode::PIE_ListenServer, bIsDedicatedServer);
 
 				PlayInEditorSessionInfo->bServerWasLaunched = true;
 			}
diff --git a/Engine/Source/Editor/UnrealEd/Private/PlayLevelNewProcess.cpp b/Engine/Source/Editor/UnrealEd/Private/PlayLevelNewProcess.cpp
--- a/Engine/Source/Editor/UnrealEd/Private/PlayLevelNewProcess.cpp
+++ b/Engine/Source/Editor/UnrealEd/Private/PlayLevelNewProcess.cpp
@@ -22,8 +22,7 @@ void UEditorEngine::StartPlayInNewProcessSession(FRequestPlaySessionParams& InRe
 	if (bNeedsServer)
 	{
 		const bool bIsDedicatedServer = true;
-		const bool bIsHost = true;
-		LaunchNewProcess(InRequestParams, 0, bIsHost, bIsDedicatedServer);
+		LaunchNewProcess(InRequestParams, 0, EPlayNetMode::PIE_ListenServer, bIsDedicatedServer);
 		
 		bServerWasLaunched = true;
 	}
@@ -36,20 +35,32 @@ void UEditorEngine::StartPlayInNewProcessSession(FRequestPlaySessionParams& InRe
 	int32 NumRequestedInstances = FMath::Max(NumClients, 1);
 	for (int32 InstanceIndex = 0; InstanceIndex < NumRequestedInstances; InstanceIndex++)
 	{
+		EPlayNetMode LocalNetMode = NetMode;
+
+		// If they want to launch a Listen Server and have multiple clients, the subsequent clients need to be
+		// treated as Clients so they connect to the listen server instead of launching multiple Listen Servers.
+		if (NetMode == EPlayNetMode::PIE_ListenServer && InstanceIndex > 0)
+		{
+			LocalNetMode = EPlayNetMode::PIE_Client;
+		}
+
 		// Dedicated servers should have been launched above, so this is only clients + listen servers.
 		const bool bIsDedicatedServer = false;
-		const bool bIsHost = (InstanceIndex == 0) && (NetMode == EPlayNetMode::PIE_ListenServer);
 
-		LaunchNewProcess(InRequestParams, InstanceIndex, bIsHost, bIsDedicatedServer);
+		LaunchNewProcess(InRequestParams, InstanceIndex, LocalNetMode, bIsDedicatedServer);
 	}
+
+	// Now that we've launched the new process, we'll cancel the request so that the UI lets us go into PIE.
+	// This doesn't clear our tracked sessions, so next time PIE is started it will close any standalone instances.
+	CancelRequestPlaySession();
 }
 
-void UEditorEngine::LaunchNewProcess(const FRequestPlaySessionParams& InParams, const int32 InInstanceNum, bool bIsHost, bool bIsDedicatedServer)
+void UEditorEngine::LaunchNewProcess(const FRequestPlaySessionParams& InParams, const int32 InInstanceNum, EPlayNetMode NetMode, bool bIsDedicatedServer)
 {
 	// All dedicated servers should be considered hosts as well.
 	if (bIsDedicatedServer)
 	{
-		bIsHost = true;
+		NetMode = EPlayNetMode::PIE_ListenServer;
 	}
 
 	// Apply various launch arguments based on their settings.
@@ -60,7 +71,7 @@ void UEditorEngine::LaunchNewProcess(const FRequestPlaySessionParams& InParams,
 	{
 		CommandLine += TEXT("-server -log");
 	}
-	else if (bIsHost)
+	else if (NetMode == EPlayNetMode::PIE_ListenServer)
 	{
 		UnrealURLParams += TEXT("?Listen");
 
@@ -93,7 +104,7 @@ void UEditorEngine::LaunchNewProcess(const FRequestPlaySessionParams& InParams,
 	}
 
 	// If they're not a host, configure the URL Params to connect to the server (instead of a specifying a map later)
-	if (!bIsHost)
+	if (NetMode == EPlayNetMode::PIE_Client)
 	{
 		FString ServerIP = TEXT("127.0.0.1");
 		uint16 ServerPort;
@@ -179,7 +190,7 @@ void UEditorEngine::LaunchNewProcess(const FRequestPlaySessionParams& InParams,
 	}
 
 	// Allow servers to override which port they are launched on.
-	if (bIsHost)
+	if (NetMode == EPlayNetMode::PIE_ListenServer)
 	{
 		uint16 ServerPort;
 		InParams.EditorPlaySettings->GetServerPort(ServerPort);
@@ -190,7 +201,7 @@ void UEditorEngine::LaunchNewProcess(const FRequestPlaySessionParams& InParams,
 	// Allow emulating adverse network conditions.
 	if (InParams.EditorPlaySettings->IsNetworkEmulationEnabled())
 	{
-		NetworkEmulationTarget CurrentTarget = bIsHost ? NetworkEmulationTarget::Server : NetworkEmulationTarget::Client;
+		NetworkEmulationTarget CurrentTarget = (NetMode == EPlayNetMode::PIE_ListenServer) ? NetworkEmulationTarget::Server : NetworkEmulationTarget::Client;
 		if (InParams.EditorPlaySettings->NetworkEmulationSettings.IsEmulationEnabledForTarget(CurrentTarget))
 		{
 			CommandLine += InParams.EditorPlaySettings->NetworkEmulationSettings.BuildPacketSettingsForCmdLine();
@@ -248,7 +259,7 @@ void UEditorEngine::LaunchNewProcess(const FRequestPlaySessionParams& InParams,
 	// Launch a new process.
 	TMap<FString, FStringFormatArg> NamedArguments;
 	NamedArguments.Add(TEXT("GameNameOrProjectFile"), GameNameOrProjectFile);
-	if (bIsHost)
+	if (NetMode != EPlayNetMode::PIE_Client)
 	{
 		// If we're not a client, build a PlayWorld URL to load to.
 		FString ServerMapNameOverride;
