Updated "Prompt When Adding To Level Before Checkout" to use cached source control state instead of running a fresh query each time an actor is added

diff --git a/Engine/Source/Editor/UnrealEd/Private/AssetSelection.cpp b/Engine/Source/Editor/UnrealEd/Private/AssetSelection.cpp
--- a/Engine/Source/Editor/UnrealEd/Private/AssetSelection.cpp
+++ b/Engine/Source/Editor/UnrealEd/Private/AssetSelection.cpp
@@ -48,6 +48,7 @@
 #include "Engine/LevelStreaming.h"
 #include "Engine/LevelBounds.h"
 #include "SourceControlHelpers.h"
+#include "ISourceControlModule.h"
 #include "Misc/MessageDialog.h"
 
 
@@ -451,8 +452,8 @@ namespace ActorPlacementUtils
 		{
 			FString FileName = SourceControlHelpers::PackageFilename(InLevel->GetPathName());
 			// Query file state also checks the source control status
-			FSourceControlState SCState = SourceControlHelpers::QueryFileState(FileName, true);
-			if (!InLevel->bLevelOkayForPlacementWhileCheckedIn && !(SCState.bIsCheckedOut || SCState.bIsAdded || SCState.bCanAdd || SCState.bIsUnknown))
+			FSourceControlStatePtr SCState = ISourceControlModule::Get().GetProvider().GetState(FileName, EStateCacheUsage::Use);
+			if (!InLevel->bLevelOkayForPlacementWhileCheckedIn && !(SCState->IsCheckedOut() || SCState->IsAdded() || SCState->CanAdd() || SCState->IsUnknown()))
 			{
 				FText Title = NSLOCTEXT("UnrealEd", "LevelCheckout_Title", "Level Checkout Warning");
 				if (EAppReturnType::Ok != FMessageDialog::Open(EAppMsgType::OkCancel, NSLOCTEXT("UnrealEd","LevelNotCheckedOutMsg", "This actor will be placed in a level that is in source control but not currently checked out. Continue?"), &Title))
