CIS fix

#rnx

diff --git a/Engine/Source/Editor/UnrealEd/Private/AssetSelection.cpp b/Engine/Source/Editor/UnrealEd/Private/AssetSelection.cpp
--- a/Engine/Source/Editor/UnrealEd/Private/AssetSelection.cpp
+++ b/Engine/Source/Editor/UnrealEd/Private/AssetSelection.cpp
@@ -49,6 +49,7 @@
 #include "Engine/LevelBounds.h"
 #include "SourceControlHelpers.h"
 #include "ISourceControlModule.h"
+#include "ISourceControlProvider.h"
 #include "Misc/MessageDialog.h"
 
 
