New feature: Added the capability of adding Radio buttons to the submenu main button (not only to its entries).

This CL is the first of 2 CLs that introduce multiple improvements on the "View Mode" menu of the Viewport, including a new feature request from JIRA UE-81469. However, we isolate this particular new feature into an isolated CL given that it can be applicable anywhere in the Editor, not only for "View Mode".

Issue: If any subcategory of "View Mode" was chosen, the user could not know in which category it was located (unless the user manually looked into each category until he found the selected one).
Solution: Added the checked functionality (with a radio button) to each one of the categories, behaving similarly to any subcategory. This way, both the radio button of the category and the one from the specific subcategory will be checked simulatenously.

More technical details:
Adding this functionality to "View Mode" means adding the functionality of "FUIAction::FIsActionChecked" into "AddSubMenu", which includes modifications to "MultiBox", "SMenuEntryBlock", and "MultiBoxBuilder".

1) We are changing a submenu so it has a radioButton
2) To specify a radio button a developer has to use a different API that includes a delegate (FUIAction) to say if it should be checked or not before this CL
3) The different API for AddSubMenu was not passing along the command list for that child submenu to use
4) The updated code now passes the command list along for the child submenu to use

diff --git a/Engine/Source/Editor/LevelEditor/Private/SLevelViewportToolBar.cpp b/Engine/Source/Editor/LevelEditor/Private/SLevelViewportToolBar.cpp
--- a/Engine/Source/Editor/LevelEditor/Private/SLevelViewportToolBar.cpp
+++ b/Engine/Source/Editor/LevelEditor/Private/SLevelViewportToolBar.cpp
@@ -1288,8 +1288,17 @@ void SLevelViewportToolBar::CreateViewMenuExtensions(FMenuBuilder& MenuBuilder)
 	MenuBuilder.AddSubMenu(LOCTEXT("VisualizeBufferViewModeDisplayName", "Buffer Visualization"),
 		LOCTEXT("BufferVisualizationMenu_ToolTip", "Select a mode for buffer visualization"),
 		FNewMenuDelegate::CreateStatic(&FBufferVisualizationMenuCommands::BuildVisualisationSubMenu),
-		false,
-		FSlateIcon(FEditorStyle::GetStyleSetName(), "EditorViewport.VisualizeBufferMode"));
+		FUIAction(FExecuteAction(), FCanExecuteAction(),
+			FIsActionChecked::CreateLambda([this]()
+			{
+				const TSharedRef<SEditorViewport> ViewportRef = Viewport.Pin().ToSharedRef();
+				const TSharedPtr<FEditorViewportClient> ViewportClient = ViewportRef->GetViewportClient();
+				check(ViewportClient.IsValid());
+				return ViewportClient->IsViewModeEnabled(VMI_VisualizeBuffer);
+			})),
+			//FIsActionChecked::CreateStatic(&FBufferVisualizationMenuCommands::IsVisualisationMenuSelected, ViewportClient)),
+		/* InExtensionHook = */ NAME_None, EUserInterfaceActionType::RadioButton,
+		/* bInOpenSubMenuOnClick = */ false, FSlateIcon(FEditorStyle::GetStyleSetName(), "EditorViewport.VisualizeBufferMode"));
 
 	MenuBuilder.BeginSection("LevelViewportLandscape", LOCTEXT("LandscapeHeader", "Landscape") );
 	{
diff --git a/Engine/Source/Editor/UnrealEd/Private/SEditorViewportViewMenu.cpp b/Engine/Source/Editor/UnrealEd/Private/SEditorViewportViewMenu.cpp
--- a/Engine/Source/Editor/UnrealEd/Private/SEditorViewportViewMenu.cpp
+++ b/Engine/Source/Editor/UnrealEd/Private/SEditorViewportViewMenu.cpp
@@ -364,7 +364,26 @@ TSharedRef<SWidget> SEditorViewportViewMenu::GenerateViewMenuContent() const
 					}
 				};
 
-				ViewMenuBuilder.AddSubMenu( LOCTEXT("OptimizationSubMenu", "Optimization Viewmodes"), LOCTEXT("Optimization_ToolTip", "Select optimization visualizer"), FNewMenuDelegate::CreateStatic( &Local::BuildOptimizationMenu, ParentToolBar ) );
+				ViewMenuBuilder.AddSubMenu(
+					LOCTEXT("OptimizationSubMenu", "Optimization Viewmodes"), LOCTEXT("Optimization_ToolTip", "Select optimization visualizer"),
+					FNewMenuDelegate::CreateStatic(&Local::BuildOptimizationMenu, ParentToolBar),
+					FUIAction(FExecuteAction(), FCanExecuteAction(),
+						FIsActionChecked::CreateLambda([this]()
+						{
+								const TSharedRef<SEditorViewport> ViewportRef = Viewport.Pin().ToSharedRef();
+								const TSharedPtr<FEditorViewportClient> ViewportClient = ViewportRef->GetViewportClient();
+								check(ViewportClient.IsValid());
+								const EViewModeIndex ViewMode = ViewportClient->GetViewMode();
+								return (
+									// Texture Streaming Accuracy
+									ViewMode == VMI_LightComplexity || ViewMode == VMI_LightmapDensity || ViewMode == VMI_StationaryLightOverlap
+									|| ViewMode == VMI_ShaderComplexity || ViewMode == VMI_ShaderComplexityWithQuadOverdraw || ViewMode == VMI_QuadOverdraw
+									// Texture Streaming Accuracy
+									|| ViewMode == VMI_PrimitiveDistanceAccuracy || ViewMode == VMI_MeshUVDensityAccuracy || ViewMode == VMI_MaterialTextureScaleAccuracy || ViewMode == VMI_RequiredTextureResolution
+								);
+						})),
+					/* InExtensionHook = */ NAME_None, EUserInterfaceActionType::RadioButton,
+					/* bInOpenSubMenuOnClick = */ false, FSlateIcon(FEditorStyle::GetStyleSetName(), "EditorViewport.QuadOverdrawMode"));
 			}
 
 #if RHI_RAYTRACING
@@ -397,7 +416,20 @@ TSharedRef<SWidget> SEditorViewportViewMenu::GenerateViewMenuContent() const
 					}
 				};
 
-				ViewMenuBuilder.AddSubMenu(LOCTEXT("VisualizeGroupedLODDisplayName", "Level of Detail Coloration"), LOCTEXT("GroupedLODMenu_ToolTip", "Select a mode for LOD Coloration"), FNewMenuDelegate::CreateStatic(&Local::BuildLODMenu), /*Default*/false, FSlateIcon(FEditorStyle::GetStyleSetName(), "EditorViewport.GroupLODColorationMode"));
+				ViewMenuBuilder.AddSubMenu(
+					LOCTEXT("VisualizeGroupedLODDisplayName", "Level of Detail Coloration"), LOCTEXT("GroupedLODMenu_ToolTip", "Select a mode for LOD Coloration"),
+					FNewMenuDelegate::CreateStatic(&Local::BuildLODMenu),
+					FUIAction(FExecuteAction(), FCanExecuteAction(),
+						FIsActionChecked::CreateLambda([this]()
+						{
+								const TSharedRef<SEditorViewport> ViewportRef = Viewport.Pin().ToSharedRef();
+								const TSharedPtr<FEditorViewportClient> ViewportClient = ViewportRef->GetViewportClient();
+								check(ViewportClient.IsValid());
+								const EViewModeIndex ViewMode = ViewportClient->GetViewMode();
+								return (ViewMode == VMI_LODColoration || ViewMode == VMI_HLODColoration);
+						})),
+					/* InExtensionHook = */ NAME_None, EUserInterfaceActionType::RadioButton,
+					/* bInOpenSubMenuOnClick = */ false, FSlateIcon(FEditorStyle::GetStyleSetName(), "EditorViewport.GroupLODColorationMode"));
 			}
 
 			ViewMenuBuilder.EndSection();
diff --git a/Engine/Source/Runtime/Slate/Private/Framework/MultiBox/MultiBoxBuilder.cpp b/Engine/Source/Runtime/Slate/Private/Framework/MultiBox/MultiBoxBuilder.cpp
--- a/Engine/Source/Runtime/Slate/Private/Framework/MultiBox/MultiBoxBuilder.cpp
+++ b/Engine/Source/Runtime/Slate/Private/Framework/MultiBox/MultiBoxBuilder.cpp
@@ -244,7 +244,7 @@ void FMenuBuilder::AddSubMenu( const TAttribute<FText>& InMenuLabel, const TAttr
 	ApplySectionBeginning();
 
 	const bool bIsSubMenu = true;
-	TSharedRef< FMenuEntryBlock > NewMenuEntryBlock( new FMenuEntryBlock( InExtensionHook, InMenuLabel, InToolTip, InSubMenu, ExtenderStack.Top(), bIsSubMenu, bInOpenSubMenuOnClick, InIcon, InUIAction, InUserInterfaceActionType, bCloseSelfOnly, bInShouldCloseWindowAfterMenuSelection ) );
+	TSharedRef< FMenuEntryBlock > NewMenuEntryBlock( new FMenuEntryBlock( InExtensionHook, InMenuLabel, InToolTip, InSubMenu, ExtenderStack.Top(), bIsSubMenu, bInOpenSubMenuOnClick, InIcon, InUIAction, InUserInterfaceActionType, bCloseSelfOnly, bInShouldCloseWindowAfterMenuSelection, CommandListStack.Last() ) );
 
 	MultiBox->AddMultiBlock( NewMenuEntryBlock );
 }
diff --git a/Engine/Source/Runtime/Slate/Private/Framework/MultiBox/SMenuEntryBlock.cpp b/Engine/Source/Runtime/Slate/Private/Framework/MultiBox/SMenuEntryBlock.cpp
--- a/Engine/Source/Runtime/Slate/Private/Framework/MultiBox/SMenuEntryBlock.cpp
+++ b/Engine/Source/Runtime/Slate/Private/Framework/MultiBox/SMenuEntryBlock.cpp
@@ -28,8 +28,8 @@ FMenuEntryBlock::FMenuEntryBlock( const FName& InExtensionHook, const TSharedPtr
 }
 
 
-FMenuEntryBlock::FMenuEntryBlock( const FName& InExtensionHook, const TAttribute<FText>& InLabel, const TAttribute<FText>& InToolTip, const FNewMenuDelegate& InEntryBuilder, TSharedPtr<FExtender> InExtender, bool bInSubMenu, bool bInSubMenuOnClick, const FSlateIcon& InIcon, const FUIAction& InUIAction, const EUserInterfaceActionType InUserInterfaceActionType, bool bInCloseSelfOnly, bool bInShouldCloseWindowAfterMenuSelection)
-	: FMultiBlock( InUIAction, InExtensionHook, EMultiBlockType::MenuEntry )
+FMenuEntryBlock::FMenuEntryBlock( const FName& InExtensionHook, const TAttribute<FText>& InLabel, const TAttribute<FText>& InToolTip, const FNewMenuDelegate& InEntryBuilder, TSharedPtr<FExtender> InExtender, bool bInSubMenu, bool bInSubMenuOnClick, const FSlateIcon& InIcon, const FUIAction& InUIAction, const EUserInterfaceActionType InUserInterfaceActionType, bool bInCloseSelfOnly, bool bInShouldCloseWindowAfterMenuSelection, TSharedPtr< const FUICommandList > InCommandList)
+	: FMultiBlock( InUIAction, InExtensionHook, EMultiBlockType::MenuEntry, /* bInIsPartOfHeading = */ false, InCommandList )
 	, LabelOverride( InLabel )
 	, ToolTipOverride( InToolTip )
 	, IconOverride( InIcon )
diff --git a/Engine/Source/Runtime/Slate/Private/Framework/MultiBox/SMenuEntryBlock.h b/Engine/Source/Runtime/Slate/Private/Framework/MultiBox/SMenuEntryBlock.h
--- a/Engine/Source/Runtime/Slate/Private/Framework/MultiBox/SMenuEntryBlock.h
+++ b/Engine/Source/Runtime/Slate/Private/Framework/MultiBox/SMenuEntryBlock.h
@@ -100,8 +100,9 @@ public:
 	 * @param	InUserInterfaceActionType	Type of interface action
 	 * @param	bInShouldCloseWindowAfterMenuSelection	In the case of a submenu, whether it should close after an item is selected
 	 * @param	bInInvertLabelOnHover	Whether to invert the label text's color on hover
+	 * @param	InCommandList		The list of commands bound to delegates that should be executed for menu entries
 	 */
-	FMenuEntryBlock( const FName& InExtensionHook, const TAttribute<FText>& InLabel, const TAttribute<FText>& InToolTip, const FNewMenuDelegate& InEntryBuilder, TSharedPtr<FExtender> InExtender, bool bInSubMenu, bool bInSubMenuOnClick, const FSlateIcon& InIcon, const FUIAction& InUIAction, const EUserInterfaceActionType InUserInterfaceActionType, bool bInCloseSelfOnly, bool bInShouldCloseWindowAfterMenuSelection = true);
+	FMenuEntryBlock( const FName& InExtensionHook, const TAttribute<FText>& InLabel, const TAttribute<FText>& InToolTip, const FNewMenuDelegate& InEntryBuilder, TSharedPtr<FExtender> InExtender, bool bInSubMenu, bool bInSubMenuOnClick, const FSlateIcon& InIcon, const FUIAction& InUIAction, const EUserInterfaceActionType InUserInterfaceActionType, bool bInCloseSelfOnly, bool bInShouldCloseWindowAfterMenuSelection = true, TSharedPtr< const FUICommandList > InCommandList = nullptr );
 
 	FMenuEntryBlock( const FName& InExtensionHook, const TAttribute<FText>& InLabel, const TAttribute<FText>& InToolTip, const FSlateIcon& InIcon, const FUIAction& InUIAction, const EUserInterfaceActionType InUserInterfaceActionType, bool bInCloseSelfOnly, bool bInShouldCloseWindowAfterMenuSelection = true);
 
diff --git a/Engine/Source/Runtime/Slate/Public/Framework/MultiBox/MultiBox.h b/Engine/Source/Runtime/Slate/Public/Framework/MultiBox/MultiBox.h
--- a/Engine/Source/Runtime/Slate/Public/Framework/MultiBox/MultiBox.h
+++ b/Engine/Source/Runtime/Slate/Public/Framework/MultiBox/MultiBox.h
@@ -76,8 +76,9 @@ public:
 	 *
 	 * @InAction UI action delegates that should be used in place of UI commands (dynamic menu items)
 	 */
-	FMultiBlock( const FUIAction& InAction,  FName InExtensionHook = NAME_None, EMultiBlockType InType = EMultiBlockType::None, bool bInIsPartOfHeading = false )
+	FMultiBlock( const FUIAction& InAction,  FName InExtensionHook = NAME_None, EMultiBlockType InType = EMultiBlockType::None, bool bInIsPartOfHeading = false, TSharedPtr< const FUICommandList > InCommandList = nullptr )
 		: DirectActions( InAction )
+		, ActionList( InCommandList )
 		, ExtensionHook( InExtensionHook )
 		, Type( InType )
 		, TutorialHighlightName( NAME_None )
