Copy up from Dev-Editor [REDACTED_HANDLE]

diff --git a/Engine/Source/Editor/UnrealEd/Private/EditorEngine.cpp b/Engine/Source/Editor/UnrealEd/Private/EditorEngine.cpp
--- a/Engine/Source/Editor/UnrealEd/Private/EditorEngine.cpp
+++ b/Engine/Source/Editor/UnrealEd/Private/EditorEngine.cpp
@@ -75,6 +75,7 @@
 #include "FileHelpers.h"
 #include "EditorModeInterpolation.h"
 #include "Dialogs/Dialogs.h"
+#include "Dialogs/DialogsPrivate.h"
 #include "UnrealEdGlobals.h"
 #include "Matinee/MatineeActor.h"
 #include "InteractiveFoliageActor.h"
@@ -344,13 +345,10 @@ UEditorEngine::UEditorEngine(const FObjectInitializer& ObjectInitializer)
 	}
 
 	DetailMode = DM_MAX;
-	PlayInEditorViewportIndex = -1;
 	CurrentPlayWorldDestination = -1;
 	bDisableDeltaModification = false;
-	bPlayOnLocalPcSession = false;
 	bAllowMultiplePIEWorlds = true;
 	bIsEndingPlay = false;
-	NumOnlinePIEInstances = 0;
 	DefaultWorldFeatureLevel = GMaxRHIFeatureLevel;
 	PreviewPlatform = FPreviewPlatformInfo(DefaultWorldFeatureLevel);
 
@@ -955,12 +953,6 @@ void UEditorEngine::Init(IEngineLoop* InEngineLoop)
 
 	LoadEditorFeatureLevel();
 
-	PRAGMA_DISABLE_DEPRECATION_WARNINGS
-	// ILayers (GEditor->Layers) has been deprecated, use ULayersSubsystem (GEditor->GetEditorSubsystem<ULayersSubsystem>()) instead.
-	// We temporaily assign Layers = GEditor->GetEditorSubsystem<ULayersSubsystem>(), but we will remove Layers in future releases.
-	Layers = MakeShareable(GetEditorSubsystem<ULayersSubsystem>(), [](ULayersSubsystem*) {});
-	PRAGMA_ENABLE_DEPRECATION_WARNINGS
-
 
 	// Init transactioning.
 	Trans = CreateTrans();
@@ -1595,10 +1587,10 @@ void UEditorEngine::Tick( float DeltaSeconds, bool bIdleMode )
 		}
 	}
 
-	// kick off a "Play From Here" if we got one
-	if (bIsPlayWorldQueued)
+	// Kick off a Play Session request if one was queued up during the last frame.
+	if (PlaySessionRequest.IsSet())
 	{
-		StartQueuedPlayMapRequest();
+		StartQueuedPlaySessionRequest();
 	}
 	else if( bIsToggleBetweenPIEandSIEQueued )
 	{
@@ -1970,7 +1962,7 @@ float UEditorEngine::GetMaxTickRate( float DeltaTime, bool bAllowFrameRateSmooth
 	if( !ShouldThrottleCPUUsage() )
 	{
 		// do not limit fps in VR Preview mode
-		if (bUseVRPreviewForPlayWorld)
+		if (IsVRPreviewActive())
 		{
 			return 0.0f;
 		}
@@ -4697,7 +4689,7 @@ EAppReturnType::Type UEditorEngine::OnModalMessageDialog(EAppMsgType::Type InMes
 {
 	if( IsInGameThread() && FSlateApplication::IsInitialized() && FSlateApplication::Get().CanAddModalWindow() )
 	{
-		return OpenMsgDlgInt(InMessage, InText, InTitle);
+		return OpenMessageDialog_Internal(InMessage, InText, InTitle);
 	}
 	else
 	{
@@ -6190,9 +6182,11 @@ bool UEditorEngine::ShouldThrottleCPUUsage() const
 {
 	bool bShouldThrottle = false;
 
-	bool bIsForeground = FPlatformApplicationMisc::IsThisApplicationForeground();
+	const bool bRunningCommandlet = IsRunningCommandlet();
+
+	const bool bIsForeground = FPlatformApplicationMisc::IsThisApplicationForeground();
 
-	if( !bIsForeground )
+	if( !bIsForeground && !bRunningCommandlet )
 	{
 		const UEditorPerformanceSettings* Settings = GetDefault<UEditorPerformanceSettings>();
 		bShouldThrottle = Settings->bThrottleCPUWhenNotForeground;
@@ -6221,7 +6215,7 @@ bool UEditorEngine::ShouldThrottleCPUUsage() const
 		}
 	}
 
-	return bShouldThrottle && !IsRunningCommandlet();
+	return bShouldThrottle;
 }
 
 bool UEditorEngine::AreAllWindowsHidden() const
@@ -7331,10 +7325,18 @@ void UEditorEngine::AutomationLoadMap(const FString& MapName, FString* OutError)
 		UE_LOG(LogEditor, Log, TEXT("Starting PIE for the automation tests for world, %s"), *GWorld->GetMapName());
 
 		FFailedGameStartHandler FailHandler;
-		FPlayInEditorOverrides Overrides;
-		Overrides.bDedicatedServer = false;
-		Overrides.NumberOfClients = 1;
-		PlayInEditor(GWorld, /*bInSimulateInEditor=*/false, Overrides);
+
+		FRequestPlaySessionParams RequestParams;
+		ULevelEditorPlaySettings* EditorPlaySettings = NewObject<ULevelEditorPlaySettings>();
+		EditorPlaySettings->SetPlayNumberOfClients(1);
+		EditorPlaySettings->bLaunchSeparateServer = false;
+		RequestParams.EditorPlaySettings = EditorPlaySettings;
+
+		RequestPlaySession(RequestParams);
+
+		// Immediately launch the session 
+		StartQueuedPlaySessionRequest();
+
 		if (!FailHandler.CanProceed())
 		{
 			*OutError = TEXT("Error encountered.");
@@ -7351,7 +7353,7 @@ void UEditorEngine::AutomationLoadMap(const FString& MapName, FString* OutError)
 bool UEditorEngine::IsHMDTrackingAllowed() const
 {
 	// @todo vreditor: Added GEnableVREditorHacks check below to allow head movement in non-PIE editor; needs revisit
-	return GEnableVREditorHacks || (PlayWorld && (bUseVRPreviewForPlayWorld || GetDefault<ULevelEditorPlaySettings>()->ViewportGetsHMDControl));
+	return GEnableVREditorHacks || (PlayWorld && (IsVRPreviewActive() || GetDefault<ULevelEditorPlaySettings>()->ViewportGetsHMDControl));
 }
 
 void UEditorEngine::OnModuleCompileStarted(bool bIsAsyncCompile)
@@ -7529,4 +7531,13 @@ void UEditorEngine::SaveEditorFeatureLevel()
 	Settings->PostEditChange();
 }
 
+void UEditorEngine::AddReferencedObjects(FReferenceCollector& Collector)
+{
+	if (PlayInEditorSessionInfo.IsSet())
+	{
+		Collector.AddReferencedObject(PlayInEditorSessionInfo->OriginalRequestParams.EditorPlaySettings, this);
+	}
+}
+
+
 #undef LOCTEXT_NAMESPACE 
diff --git a/Engine/Source/Editor/UnrealEd/Private/Factories/EditorFactories.cpp b/Engine/Source/Editor/UnrealEd/Private/Factories/EditorFactories.cpp
--- a/Engine/Source/Editor/UnrealEd/Private/Factories/EditorFactories.cpp
+++ b/Engine/Source/Editor/UnrealEd/Private/Factories/EditorFactories.cpp
@@ -6220,8 +6220,10 @@ EReimportResult::Type UReimportFbxSkeletalMeshFactory::Reimport( UObject* Obj, i
 	ReimportUI->SkeletalMeshImportData = ImportData;
 	const FSkeletalMeshModel* SkeletalMeshModel = SkeletalMesh->GetImportedModel();
 
+	bool bIsBuildAvailable = SkeletalMesh->IsLODImportedDataBuildAvailable(0);
+	
 	//Manage the content type from the source file index
-	ReimportUI->bAllowContentTypeImport = SkeletalMeshModel && SkeletalMeshModel->LODModels.Num() > 0 && !SkeletalMeshModel->LODModels[0].RawSkeletalMeshBulkData.IsEmpty();
+	ReimportUI->bAllowContentTypeImport = SkeletalMeshModel && SkeletalMeshModel->LODModels.Num() > 0 && !SkeletalMesh->IsLODImportedDataEmpty(0);
 	if (!ReimportUI->bAllowContentTypeImport)
 	{
 		//No content type allow reimport All (legacy)
@@ -6305,13 +6307,18 @@ EReimportResult::Type UReimportFbxSkeletalMeshFactory::Reimport( UObject* Obj, i
 		if (LODInfo && SkeletalMesh->GetImportedModel() && SkeletalMesh->GetImportedModel()->LODModels.IsValidIndex(0))
 		{
 			const FSkeletalMeshLODModel& LODModel = SkeletalMesh->GetImportedModel()->LODModels[0];
-			if (LODModel.RawSkeletalMeshBulkData.IsBuildDataAvailable())
+			
+			if (bIsBuildAvailable)
 			{
 				//Set the build settings
 				LODInfo->BuildSettings.bComputeWeightedNormals = SKImportData->bComputeWeightedNormals;
 				LODInfo->BuildSettings.bRecomputeNormals = SKImportData->NormalImportMethod == EFBXNormalImportMethod::FBXNIM_ComputeNormals;
 				LODInfo->BuildSettings.bRecomputeTangents = SKImportData->NormalImportMethod != EFBXNormalImportMethod::FBXNIM_ImportNormalsAndTangents;
 				LODInfo->BuildSettings.bUseMikkTSpace = SKImportData->NormalGenerationMethod == EFBXNormalGenerationMethod::MikkTSpace;
+				LODInfo->BuildSettings.ThresholdPosition = SKImportData->ThresholdPosition;
+				LODInfo->BuildSettings.ThresholdTangentNormal = SKImportData->ThresholdTangentNormal;
+				LODInfo->BuildSettings.ThresholdUV = SKImportData->ThresholdUV;
+				LODInfo->BuildSettings.MorphThresholdPosition = SKImportData->MorphThresholdPosition;
 			}
 		}
 	}
diff --git a/Engine/Source/Editor/UnrealEd/Private/Factories/SkeletalMeshImport.cpp b/Engine/Source/Editor/UnrealEd/Private/Factories/SkeletalMeshImport.cpp
--- a/Engine/Source/Editor/UnrealEd/Private/Factories/SkeletalMeshImport.cpp
+++ b/Engine/Source/Editor/UnrealEd/Private/Factories/SkeletalMeshImport.cpp
@@ -31,6 +31,7 @@
 #include "IMeshReductionManagerModule.h"
 #include "Rendering/SkeletalMeshModel.h"
 #include "Engine/AssetUserData.h"
+#include "UObject/MetaData.h"
 
 DEFINE_LOG_CATEGORY_STATIC(LogSkeletalMeshImport, Log, All);
 
@@ -239,7 +240,10 @@ ExistingSkelMeshData* SaveExistingSkelMeshData(USkeletalMesh* ExistingSkelMesh,
 	int32 SafeReimportLODIndex = ReimportLODIndex < 0 ? 0 : ReimportLODIndex;
 
 	ExistingMeshDataPtr = new ExistingSkelMeshData();
-		
+
+	//Save the package UMetaData
+	ExistingMeshDataPtr->ExistingUMetaDataTagValues = UMetaData::GetMapForObject(ExistingSkelMesh);
+
 	ExistingMeshDataPtr->UseMaterialNameSlotWorkflow = SkeletalMeshIsUsingMaterialSlotNameWorkflow(ExistingSkelMesh->AssetImportData);
 	ExistingMeshDataPtr->MinLOD = ExistingSkelMesh->MinLod;
 	ExistingMeshDataPtr->DisableBelowMinLodStripping = ExistingSkelMesh->DisableBelowMinLodStripping;
@@ -355,7 +359,7 @@ ExistingSkelMeshData* SaveExistingSkelMeshData(USkeletalMesh* ExistingSkelMesh,
 			FSkeletalMeshLODModel& LODModel = ImportedResource->LODModels[LODModelIndex];
 			LODModel.RawPointIndices.Lock( LOCK_READ_ONLY );
 			LODModel.LegacyRawPointIndices.Lock( LOCK_READ_ONLY );
-			LODModel.RawSkeletalMeshBulkData.GetBulkData().Lock( LOCK_READ_ONLY );
+			LODModel.RawSkeletalMeshBulkData_DEPRECATED.GetBulkData().Lock( LOCK_READ_ONLY );
 			int32 ReductionLODIndex = LODModelIndex + OffsetReductionLODIndex;
 			if (ImportedResource->OriginalReductionSourceMeshData.IsValidIndex(ReductionLODIndex) && !ImportedResource->OriginalReductionSourceMeshData[ReductionLODIndex]->IsEmpty())
 			{
@@ -379,7 +383,7 @@ ExistingSkelMeshData* SaveExistingSkelMeshData(USkeletalMesh* ExistingSkelMesh,
 			FSkeletalMeshLODModel& LODModel = ImportedResource->LODModels[LODModelIndex];
 			LODModel.RawPointIndices.Unlock();
 			LODModel.LegacyRawPointIndices.Unlock();
-			LODModel.RawSkeletalMeshBulkData.GetBulkData().Unlock();
+			LODModel.RawSkeletalMeshBulkData_DEPRECATED.GetBulkData().Unlock();
 		}
 
 		ExistingMeshDataPtr->ExistingLODInfo = ExistingSkelMesh->GetLODInfoArray();
@@ -678,6 +682,15 @@ void RestoreExistingSkelMeshData(ExistingSkelMeshData* MeshData, USkeletalMesh*
 	{
 		return;
 	}
+
+	//Restore the package metadata
+	if (MeshData->ExistingUMetaDataTagValues)
+	{
+		UMetaData* PackageMetaData = SkeletalMesh->GetOutermost()->GetMetaData();
+		checkSlow(PackageMetaData);
+		PackageMetaData->SetObjectValues(SkeletalMesh, *MeshData->ExistingUMetaDataTagValues);
+	}
+
 	int32 SafeReimportLODIndex = ReimportLODIndex < 0 ? 0 : ReimportLODIndex;
 	SkeletalMesh->MinLod = MeshData->MinLOD;
 	SkeletalMesh->DisableBelowMinLodStripping = MeshData->DisableBelowMinLodStripping;
diff --git a/Engine/Source/Editor/UnrealEd/Private/Fbx/FbxMainExport.cpp b/Engine/Source/Editor/UnrealEd/Private/Fbx/FbxMainExport.cpp
--- a/Engine/Source/Editor/UnrealEd/Private/Fbx/FbxMainExport.cpp
+++ b/Engine/Source/Editor/UnrealEd/Private/Fbx/FbxMainExport.cpp
@@ -1599,7 +1599,7 @@ int32 FFbxExporter::FLevelSequenceAnimTrackAdapter::GetStartFrame() const
 {
 	FFrameRate TickResolution = MovieScene->GetTickResolution();
 	FFrameRate DisplayRate = MovieScene->GetDisplayRate();
-	return FFrameRate::TransformTime(FFrameTime(MovieScene::DiscreteInclusiveLower(MovieScene->GetPlaybackRange()) * RootToLocalTransform.Inverse()), TickResolution, DisplayRate).RoundToFrame().Value;
+	return FFrameRate::TransformTime(FFrameTime(MovieScene::DiscreteInclusiveLower(MovieScene->GetPlaybackRange()) * RootToLocalTransform.InverseLinearOnly()), TickResolution, DisplayRate).RoundToFrame().Value;
 }
 
 int32 FFbxExporter::FLevelSequenceAnimTrackAdapter::GetLength() const
@@ -1616,21 +1616,19 @@ void FFbxExporter::FLevelSequenceAnimTrackAdapter::UpdateAnimation( int32 LocalF
 	FFrameRate DisplayRate = MovieScene->GetDisplayRate();
 
 	FFrameTime LocalTime = FFrameRate::TransformTime(FFrameTime(LocalFrame), DisplayRate, TickResolution);
-	FFrameTime GlobalTime = LocalTime * RootToLocalTransform.Inverse();
+	FFrameTime GlobalTime = LocalTime * RootToLocalTransform.InverseLinearOnly();
 
 	FMovieSceneContext Context = FMovieSceneContext(FMovieSceneEvaluationRange(GlobalTime, TickResolution), MovieScenePlayer->GetPlaybackStatus()).SetHasJumped(true);
 
 	MovieScenePlayer->GetEvaluationTemplate().Evaluate( Context, *MovieScenePlayer );
 }
 
-
 float FFbxExporter::FLevelSequenceAnimTrackAdapter::GetFrameRate() const
 {
 	return MovieScene->GetDisplayRate().AsDecimal();
 }
 
-
-bool FFbxExporter::ExportLevelSequenceTracks(UMovieScene* MovieScene, IMovieScenePlayer* MovieScenePlayer, FbxNode* FbxActor, AActor* Actor, const TArray<UMovieSceneTrack*>& Tracks, const FMovieSceneSequenceTransform& RootToLocalTransform)
+bool FFbxExporter::ExportLevelSequenceTracks(UMovieScene* MovieScene, IMovieScenePlayer* MovieScenePlayer, FMovieSceneSequenceIDRef InSequenceID, FbxNode* FbxActor, AActor* Actor, const TArray<UMovieSceneTrack*>& Tracks, const FMovieSceneSequenceTransform& RootToLocalTransform)
 {
 	USkeletalMeshComponent* SkeletalMeshComp = Actor ? Cast<USkeletalMeshComponent>(Actor->GetComponentByClass(USkeletalMeshComponent::StaticClass())) : nullptr;
 
@@ -1645,7 +1643,7 @@ bool FFbxExporter::ExportLevelSequenceTracks(UMovieScene* MovieScene, IMovieScen
 		if (Track->IsA(UMovieScene3DTransformTrack::StaticClass()) && !bSkip3DTransformTrack)
 		{
 			UMovieScene3DTransformTrack* TransformTrack = (UMovieScene3DTransformTrack*)Track;
-			ExportLevelSequence3DTransformTrack(FbxActor, *TransformTrack, Actor, MovieScene->GetPlaybackRange(), RootToLocalTransform);
+			ExportLevelSequence3DTransformTrack(FbxActor, MovieScenePlayer, InSequenceID, *TransformTrack, Actor, MovieScene->GetPlaybackRange(), RootToLocalTransform);
 		}
 		else if (Track->IsA(UMovieScenePropertyTrack::StaticClass()))
 		{
@@ -1662,19 +1660,17 @@ bool FFbxExporter::ExportLevelSequenceTracks(UMovieScene* MovieScene, IMovieScen
 			}
 		}
 	}
-
 	return true;
 }
 
-
-bool FFbxExporter::ExportLevelSequence( UMovieScene* MovieScene, const TArray<FGuid>& Bindings, IMovieScenePlayer* MovieScenePlayer, INodeNameAdapter& NodeNameAdapter, FMovieSceneSequenceIDRef SequenceID, const FMovieSceneSequenceTransform& RootToLocalTransform)
+bool FFbxExporter::ExportLevelSequence(UMovieScene* MovieScene, const TArray<FGuid>& Bindings, IMovieScenePlayer* MovieScenePlayer, INodeNameAdapter& NodeNameAdapter, FMovieSceneSequenceIDRef SequenceID, const FMovieSceneSequenceTransform& RootToLocalTransform)
 {
-	if ( MovieScene == nullptr || MovieScenePlayer == nullptr )
+	if (MovieScene == nullptr || MovieScenePlayer == nullptr)
 	{
 		return false;
 	}
 
-	for ( const FMovieSceneBinding& MovieSceneBinding : MovieScene->GetBindings() )
+	for (const FMovieSceneBinding& MovieSceneBinding : MovieScene->GetBindings())
 	{
 		// If there are specific bindings to export, export those only
 		if (Bindings.Num() != 0 && !Bindings.Contains(MovieSceneBinding.GetObjectGuid()))
@@ -1684,18 +1680,18 @@ bool FFbxExporter::ExportLevelSequence( UMovieScene* MovieScene, const TArray<FG
 
 		bool bAnyBindingsExported = false;
 
-		for ( TWeakObjectPtr<UObject> RuntimeObject : MovieScenePlayer->FindBoundObjects(MovieSceneBinding.GetObjectGuid(), SequenceID) )
+		for (TWeakObjectPtr<UObject> RuntimeObject : MovieScenePlayer->FindBoundObjects(MovieSceneBinding.GetObjectGuid(), SequenceID))
 		{
-			if ( RuntimeObject.IsValid() )
+			if (RuntimeObject.IsValid())
 			{
-				AActor* Actor = Cast<AActor>( RuntimeObject.Get() );
-				UActorComponent* Component = Cast<UActorComponent>( RuntimeObject.Get() );
-				if ( Actor == nullptr && Component != nullptr )
+				AActor* Actor = Cast<AActor>(RuntimeObject.Get());
+				UActorComponent* Component = Cast<UActorComponent>(RuntimeObject.Get());
+				if (Actor == nullptr && Component != nullptr)
 				{
 					Actor = Component->GetOwner();
 				}
 
-				if ( Actor == nullptr )
+				if (Actor == nullptr)
 				{
 					continue;
 				}
@@ -1703,9 +1699,9 @@ bool FFbxExporter::ExportLevelSequence( UMovieScene* MovieScene, const TArray<FG
 				FbxNode* FbxActor = FindActor(Actor, &NodeNameAdapter);
 
 				// now it should export everybody
-				if ( FbxActor )
+				if (FbxActor)
 				{
-					ExportLevelSequenceTracks(MovieScene, MovieScenePlayer, FbxActor, Actor, MovieSceneBinding.GetTracks(), RootToLocalTransform);
+					ExportLevelSequenceTracks(MovieScene, MovieScenePlayer, SequenceID, FbxActor, Actor, MovieSceneBinding.GetTracks(), RootToLocalTransform);
 					bAnyBindingsExported = true;
 				}
 			}
@@ -1714,7 +1710,7 @@ bool FFbxExporter::ExportLevelSequence( UMovieScene* MovieScene, const TArray<FG
 		// If no bindings exported, create a dummy actor to export tracks onto
 		if (!bAnyBindingsExported)
 		{
-			ExportLevelSequenceTracks(MovieScene, MovieScenePlayer, nullptr, nullptr, MovieSceneBinding.GetTracks(), RootToLocalTransform);
+			ExportLevelSequenceTracks(MovieScene, MovieScenePlayer, SequenceID, nullptr, nullptr, MovieSceneBinding.GetTracks(), RootToLocalTransform);
 		}
 	}
 
@@ -2516,7 +2512,7 @@ void FFbxExporter::ExportChannelToFbxCurve(FbxAnimCurve& InFbxCurve, const FMovi
 
 		FbxTime FbxTime;
 		FbxAnimCurveKey FbxKey;
-		const double KeyTimeSeconds = GetExportOptions()->bExportLocalTime ? KeyTime / TickResolution : (KeyTime * RootToLocalTransform.Inverse()) / TickResolution;
+		const double KeyTimeSeconds = GetExportOptions()->bExportLocalTime ? KeyTime / TickResolution : (KeyTime * RootToLocalTransform.InverseLinearOnly()) / TickResolution;
 
 		FbxTime.SetSecondDouble(KeyTimeSeconds);
 
@@ -2573,13 +2569,14 @@ void FFbxExporter::ExportChannelToFbxCurve(FbxAnimCurve& InFbxCurve, const FMovi
 	InFbxCurve.KeyModifyEnd();
 }
 
-void FFbxExporter::ExportLevelSequence3DTransformTrack( FbxNode* FbxNode, UMovieScene3DTransformTrack& TransformTrack, AActor* Actor, const TRange<FFrameNumber>& InPlaybackRange, const FMovieSceneSequenceTransform& RootToLocalTransform )
+void FFbxExporter::ExportLevelSequence3DTransformTrack(FbxNode* FbxNode, IMovieScenePlayer* MovieScenePlayer, FMovieSceneSequenceIDRef InSequenceID, UMovieScene3DTransformTrack& TransformTrack, AActor* Actor, const TRange<FFrameNumber>& InPlaybackRange, const FMovieSceneSequenceTransform& RootToLocalTransform)
 {
-	FbxAnimLayer* BaseLayer = AnimStack->GetMember<FbxAnimLayer>( 0 );
-
-	const bool bIsCameraActor = Actor ? Actor->IsA(ACameraActor::StaticClass()) : false;
-	const bool bIsLightActor = Actor ? Actor->IsA(ALight::StaticClass()) : false;
-	const bool bBakeRotations = bIsCameraActor || bIsLightActor;
+	//if more than one section, we use baked version of all sections.
+	if (TransformTrack.GetAllSections().Num() > 1)
+	{
+		ExportLevelSequenceInterrogated3DTransformTrack(FbxNode, MovieScenePlayer, InSequenceID, TransformTrack, Actor, InPlaybackRange, RootToLocalTransform);
+		return;
+	}
 
 	// TODO: Support more than one section?
 	UMovieScene3DTransformSection* TransformSection = TransformTrack.GetAllSections().Num() > 0
@@ -2591,6 +2588,12 @@ void FFbxExporter::ExportLevelSequence3DTransformTrack( FbxNode* FbxNode, UMovie
 		return;
 	}
 
+	FbxAnimLayer* BaseLayer = AnimStack->GetMember<FbxAnimLayer>(0);
+
+	const bool bIsCameraActor = Actor ? Actor->IsA(ACameraActor::StaticClass()) : false;
+	const bool bIsLightActor = Actor ? Actor->IsA(ALight::StaticClass()) : false;
+	const bool bBakeRotations = bIsCameraActor || bIsLightActor;
+
 	if (!FbxNode)
 	{
 		FbxNode = CreateNode(TransformTrack.GetDisplayName().ToString());
@@ -2599,21 +2602,21 @@ void FFbxExporter::ExportLevelSequence3DTransformTrack( FbxNode* FbxNode, UMovie
 	FFrameRate TickResolution = TransformTrack.GetTypedOuter<UMovieScene>()->GetTickResolution();
 	FFrameRate DisplayRate = TransformTrack.GetTypedOuter<UMovieScene>()->GetDisplayRate();
 
-	FbxAnimCurveNode* TranslationNode = FbxNode->LclTranslation.GetCurveNode( BaseLayer, true );
-	FbxAnimCurveNode* RotationNode = FbxNode->LclRotation.GetCurveNode( BaseLayer, true );
-	FbxAnimCurveNode* ScaleNode = FbxNode->LclScaling.GetCurveNode( BaseLayer, true );
+	FbxAnimCurveNode* TranslationNode = FbxNode->LclTranslation.GetCurveNode(BaseLayer, true);
+	FbxAnimCurveNode* RotationNode = FbxNode->LclRotation.GetCurveNode(BaseLayer, true);
+	FbxAnimCurveNode* ScaleNode = FbxNode->LclScaling.GetCurveNode(BaseLayer, true);
 
-	FbxAnimCurve* FbxCurveTransX = FbxNode->LclTranslation.GetCurve( BaseLayer, FBXSDK_CURVENODE_COMPONENT_X, true );
-	FbxAnimCurve* FbxCurveTransY = FbxNode->LclTranslation.GetCurve( BaseLayer, FBXSDK_CURVENODE_COMPONENT_Y, true );
-	FbxAnimCurve* FbxCurveTransZ = FbxNode->LclTranslation.GetCurve( BaseLayer, FBXSDK_CURVENODE_COMPONENT_Z, true );
+	FbxAnimCurve* FbxCurveTransX = FbxNode->LclTranslation.GetCurve(BaseLayer, FBXSDK_CURVENODE_COMPONENT_X, true);
+	FbxAnimCurve* FbxCurveTransY = FbxNode->LclTranslation.GetCurve(BaseLayer, FBXSDK_CURVENODE_COMPONENT_Y, true);
+	FbxAnimCurve* FbxCurveTransZ = FbxNode->LclTranslation.GetCurve(BaseLayer, FBXSDK_CURVENODE_COMPONENT_Z, true);
 
-	FbxAnimCurve* FbxCurveRotX = FbxNode->LclRotation.GetCurve( BaseLayer, FBXSDK_CURVENODE_COMPONENT_X, true );
-	FbxAnimCurve* FbxCurveRotY = FbxNode->LclRotation.GetCurve( BaseLayer, FBXSDK_CURVENODE_COMPONENT_Y, true );
-	FbxAnimCurve* FbxCurveRotZ = FbxNode->LclRotation.GetCurve( BaseLayer, FBXSDK_CURVENODE_COMPONENT_Z, true );
+	FbxAnimCurve* FbxCurveRotX = FbxNode->LclRotation.GetCurve(BaseLayer, FBXSDK_CURVENODE_COMPONENT_X, true);
+	FbxAnimCurve* FbxCurveRotY = FbxNode->LclRotation.GetCurve(BaseLayer, FBXSDK_CURVENODE_COMPONENT_Y, true);
+	FbxAnimCurve* FbxCurveRotZ = FbxNode->LclRotation.GetCurve(BaseLayer, FBXSDK_CURVENODE_COMPONENT_Z, true);
 
-	FbxAnimCurve* FbxCurveScaleX = FbxNode->LclScaling.GetCurve( BaseLayer, FBXSDK_CURVENODE_COMPONENT_X, true );
-	FbxAnimCurve* FbxCurveScaleY = FbxNode->LclScaling.GetCurve( BaseLayer, FBXSDK_CURVENODE_COMPONENT_Y, true );
-	FbxAnimCurve* FbxCurveScaleZ = FbxNode->LclScaling.GetCurve( BaseLayer, FBXSDK_CURVENODE_COMPONENT_Z, true );
+	FbxAnimCurve* FbxCurveScaleX = FbxNode->LclScaling.GetCurve(BaseLayer, FBXSDK_CURVENODE_COMPONENT_X, true);
+	FbxAnimCurve* FbxCurveScaleY = FbxNode->LclScaling.GetCurve(BaseLayer, FBXSDK_CURVENODE_COMPONENT_Y, true);
+	FbxAnimCurve* FbxCurveScaleZ = FbxNode->LclScaling.GetCurve(BaseLayer, FBXSDK_CURVENODE_COMPONENT_Z, true);
 
 	TArrayView<FMovieSceneFloatChannel*> FloatChannels = TransformSection->GetChannelProxy().GetChannels<FMovieSceneFloatChannel>();
 
@@ -2656,7 +2659,7 @@ void FFbxExporter::ExportLevelSequence3DTransformTrack( FbxNode* FbxNode, UMovie
 		FbxCurveRotZ->KeyModifyBegin();
 
 		int32 LocalStartFrame = FFrameRate::TransformTime(FFrameTime(MovieScene::DiscreteInclusiveLower(InPlaybackRange)), TickResolution, DisplayRate).RoundToFrame().Value;
-		int32 StartFrame = FFrameRate::TransformTime(FFrameTime(MovieScene::DiscreteInclusiveLower(InPlaybackRange) * RootToLocalTransform.Inverse()), TickResolution, DisplayRate).RoundToFrame().Value;
+		int32 StartFrame = FFrameRate::TransformTime(FFrameTime(MovieScene::DiscreteInclusiveLower(InPlaybackRange) * RootToLocalTransform.InverseLinearOnly()), TickResolution, DisplayRate).RoundToFrame().Value;
 		int32 AnimationLength = FFrameRate::TransformTime(FFrameTime(FFrameNumber(MovieScene::DiscreteSize(InPlaybackRange))), TickResolution, DisplayRate).RoundToFrame().Value;
 
 		for (int32 FrameCount = 0; FrameCount <= AnimationLength; ++FrameCount)
@@ -2705,6 +2708,172 @@ void FFbxExporter::ExportLevelSequence3DTransformTrack( FbxNode* FbxNode, UMovie
 	}
 }
 
+
+static void GetLocationAtTime(IMovieScenePlayer* MovieScenePlayer, FMovieSceneEvaluationTrack* Track, UObject* Object, FFrameTime KeyTime, FVector& KeyPos, FRotator& KeyRot, FVector& KeyScale, FFrameRate TickResolution)
+{
+	FMovieSceneInterrogationData InterrogationData;
+	MovieScenePlayer->GetEvaluationTemplate().CopyActuators(InterrogationData.GetAccumulator());
+
+	FMovieSceneContext Context(FMovieSceneEvaluationRange(KeyTime, TickResolution));
+	Track->Interrogate(Context, InterrogationData, Object);
+
+	for (const FTransformData& Transform : InterrogationData.Iterate<FTransformData>(UMovieScene3DTransformSection::GetInterrogationKey()))
+	{
+		KeyPos = Transform.Translation;
+		KeyRot = Transform.Rotation;
+		KeyScale = Transform.Scale;
+		break;
+	}
+}
+void FFbxExporter::ExportLevelSequenceInterrogated3DTransformTrack(FbxNode* FbxNode, IMovieScenePlayer* MovieScenePlayer, FMovieSceneSequenceIDRef InSequenceID, UMovieScene3DTransformTrack& TransformTrack, AActor* Actor, const TRange<FFrameNumber>& InPlaybackRange, const FMovieSceneSequenceTransform& RootToLocalTransform)
+{
+	if (TransformTrack.GetAllSections().Num() <= 0)
+	{
+		return;
+	}
+
+	FbxAnimLayer* BaseLayer = AnimStack->GetMember<FbxAnimLayer>(0);
+
+	const bool bIsCameraActor = Actor ? Actor->IsA(ACameraActor::StaticClass()) : false;
+	const bool bIsLightActor = Actor ? Actor->IsA(ALight::StaticClass()) : false;
+	const bool bBakeRotations = bIsCameraActor || bIsLightActor;
+
+
+	if (!FbxNode)
+	{
+		FbxNode = CreateNode(TransformTrack.GetDisplayName().ToString());
+	}
+
+	FFrameRate TickResolution = TransformTrack.GetTypedOuter<UMovieScene>()->GetTickResolution();
+	FFrameRate DisplayRate = TransformTrack.GetTypedOuter<UMovieScene>()->GetDisplayRate();
+
+	FbxAnimCurveNode* TranslationNode = FbxNode->LclTranslation.GetCurveNode(BaseLayer, true);
+	FbxAnimCurveNode* RotationNode = FbxNode->LclRotation.GetCurveNode(BaseLayer, true);
+	FbxAnimCurveNode* ScaleNode = FbxNode->LclScaling.GetCurveNode(BaseLayer, true);
+
+	FbxAnimCurve* FbxCurveTransX = FbxNode->LclTranslation.GetCurve(BaseLayer, FBXSDK_CURVENODE_COMPONENT_X, true);
+	FbxAnimCurve* FbxCurveTransY = FbxNode->LclTranslation.GetCurve(BaseLayer, FBXSDK_CURVENODE_COMPONENT_Y, true);
+	FbxAnimCurve* FbxCurveTransZ = FbxNode->LclTranslation.GetCurve(BaseLayer, FBXSDK_CURVENODE_COMPONENT_Z, true);
+
+	FbxAnimCurve* FbxCurveRotX = FbxNode->LclRotation.GetCurve(BaseLayer, FBXSDK_CURVENODE_COMPONENT_X, true);
+	FbxAnimCurve* FbxCurveRotY = FbxNode->LclRotation.GetCurve(BaseLayer, FBXSDK_CURVENODE_COMPONENT_Y, true);
+	FbxAnimCurve* FbxCurveRotZ = FbxNode->LclRotation.GetCurve(BaseLayer, FBXSDK_CURVENODE_COMPONENT_Z, true);
+
+	FbxAnimCurve* FbxCurveScaleX = FbxNode->LclScaling.GetCurve(BaseLayer, FBXSDK_CURVENODE_COMPONENT_X, true);
+	FbxAnimCurve* FbxCurveScaleY = FbxNode->LclScaling.GetCurve(BaseLayer, FBXSDK_CURVENODE_COMPONENT_Y, true);
+	FbxAnimCurve* FbxCurveScaleZ = FbxNode->LclScaling.GetCurve(BaseLayer, FBXSDK_CURVENODE_COMPONENT_Z, true);
+
+	FTransform RotationDirectionConvert;
+	if (bIsCameraActor)
+	{
+		FRotator Rotator = FFbxDataConverter::GetCameraRotation().GetInverse();
+		RotationDirectionConvert = FTransform(Rotator);
+	}
+	else if (bIsLightActor)
+	{
+		FRotator Rotator = FFbxDataConverter::GetLightRotation().GetInverse();
+		RotationDirectionConvert = FTransform(Rotator);
+	}
+
+	FbxCurveTransX->KeyModifyBegin();
+	FbxCurveTransY->KeyModifyBegin();
+	FbxCurveTransZ->KeyModifyBegin();
+
+	FbxCurveRotX->KeyModifyBegin();
+	FbxCurveRotY->KeyModifyBegin();
+	FbxCurveRotZ->KeyModifyBegin();
+
+	if (!bIsCameraActor)
+	{
+		FbxCurveScaleX->KeyModifyBegin();
+		FbxCurveScaleY->KeyModifyBegin();
+		FbxCurveScaleZ->KeyModifyBegin();
+	}
+
+	int32 LocalStartFrame = FFrameRate::TransformTime(FFrameTime(MovieScene::DiscreteInclusiveLower(InPlaybackRange)), TickResolution, DisplayRate).RoundToFrame().Value;
+	int32 StartFrame = FFrameRate::TransformTime(FFrameTime(MovieScene::DiscreteInclusiveLower(InPlaybackRange) * RootToLocalTransform.InverseLinearOnly()), TickResolution, DisplayRate).RoundToFrame().Value;
+	int32 AnimationLength = FFrameRate::TransformTime(FFrameTime(FFrameNumber(MovieScene::DiscreteSize(InPlaybackRange))), TickResolution, DisplayRate).RoundToFrame().Value;
+
+	FMovieSceneEvaluationTemplate* Template = MovieScenePlayer->GetEvaluationTemplate().FindTemplate(InSequenceID);
+	FMovieSceneEvaluationTrack* EvalTrack = nullptr;
+	if (Template)
+	{
+		EvalTrack = Template->FindTrack(TransformTrack.GetSignature());
+	}
+	if (!EvalTrack)
+	{
+		if (Actor)
+		{
+			FString String = Actor->GetActorLabel();
+			UE_LOG(LogFbx, Warning, TEXT("Exporting 3D TransformTrack on %s failed, can not create eval track.\n\n"), *String);
+		}
+		else
+		{
+			UE_LOG(LogFbx, Warning, TEXT("Exporting 3D TransformTrack on an Unbound Actor failed, can not create eval track.\n\n"));
+		}
+		return;
+
+	}
+
+	for (int32 FrameCount = 0; FrameCount <= AnimationLength; ++FrameCount)
+	{
+		int32 LocalFrame = LocalStartFrame + FrameCount;
+
+		FFrameTime LocalTime = FFrameRate::TransformTime(FFrameTime(LocalFrame), DisplayRate, TickResolution);
+
+		FVector Trans = FVector::ZeroVector;
+		FRotator Rotator;
+		FVector Scale;
+
+		GetLocationAtTime(MovieScenePlayer, EvalTrack, Actor, LocalTime, Trans, Rotator, Scale, TickResolution);
+
+		FTransform RelativeTransform;
+		RelativeTransform.SetTranslation(Trans);
+		RelativeTransform.SetRotation(Rotator.Quaternion());
+		RelativeTransform.SetScale3D(Scale);
+
+		RelativeTransform = RotationDirectionConvert * RelativeTransform;
+
+		FbxVector4 KeyTrans = Converter.ConvertToFbxPos(RelativeTransform.GetTranslation());
+		FbxVector4 KeyRot = Converter.ConvertToFbxRot(RelativeTransform.GetRotation().Euler());
+		FbxVector4 KeyScale = Converter.ConvertToFbxScale(RelativeTransform.GetScale3D());
+
+		FbxTime FbxTime;
+		FbxTime.SetSecondDouble(GetExportOptions()->bExportLocalTime ? DisplayRate.AsSeconds(LocalFrame) : DisplayRate.AsSeconds(StartFrame + FrameCount));
+
+		FbxCurveTransX->KeySet(FbxCurveTransX->KeyAdd(FbxTime), FbxTime, KeyTrans[0]);
+		FbxCurveTransY->KeySet(FbxCurveTransY->KeyAdd(FbxTime), FbxTime, KeyTrans[1]);
+		FbxCurveTransZ->KeySet(FbxCurveTransZ->KeyAdd(FbxTime), FbxTime, KeyTrans[2]);
+
+		FbxCurveRotX->KeySet(FbxCurveRotX->KeyAdd(FbxTime), FbxTime, KeyRot[0]);
+		FbxCurveRotY->KeySet(FbxCurveRotY->KeyAdd(FbxTime), FbxTime, KeyRot[1]);
+		FbxCurveRotZ->KeySet(FbxCurveRotZ->KeyAdd(FbxTime), FbxTime, KeyRot[2]);
+
+		if (!bIsCameraActor)
+		{
+			FbxCurveScaleX->KeySet(FbxCurveScaleX->KeyAdd(FbxTime), FbxTime, KeyScale[0]);
+			FbxCurveScaleY->KeySet(FbxCurveScaleY->KeyAdd(FbxTime), FbxTime, KeyScale[1]);
+			FbxCurveScaleZ->KeySet(FbxCurveScaleZ->KeyAdd(FbxTime), FbxTime, KeyScale[2]);
+		}
+	}
+
+	FbxCurveTransX->KeyModifyEnd();
+	FbxCurveTransY->KeyModifyEnd();
+	FbxCurveTransZ->KeyModifyEnd();
+
+	FbxCurveRotX->KeyModifyEnd();
+	FbxCurveRotY->KeyModifyEnd();
+	FbxCurveRotZ->KeyModifyEnd();
+
+	if (!bIsCameraActor)
+	{
+		FbxCurveScaleX->KeyModifyEnd();
+		FbxCurveScaleY->KeyModifyEnd();
+		FbxCurveScaleZ->KeyModifyEnd();
+	}
+}
+
+
 void FFbxExporter::ExportLevelSequencePropertyTrack( FbxNode* FbxNode, UMovieScenePropertyTrack& PropertyTrack, const TRange<FFrameNumber>& InPlaybackRange, const FMovieSceneSequenceTransform& RootToLocalTransform)
 {
 	// TODO: Support more than one section?
diff --git a/Engine/Source/Editor/UnrealEd/Private/Fbx/FbxSkeletalMeshImport.cpp b/Engine/Source/Editor/UnrealEd/Private/Fbx/FbxSkeletalMeshImport.cpp
--- a/Engine/Source/Editor/UnrealEd/Private/Fbx/FbxSkeletalMeshImport.cpp
+++ b/Engine/Source/Editor/UnrealEd/Private/Fbx/FbxSkeletalMeshImport.cpp
@@ -1463,7 +1463,7 @@ bool UnFbx::FFbxImporter::ReplaceSkeletalMeshGeometryImportData(const USkeletalM
 
 	//Load the original skeletal mesh import data
 	FSkeletalMeshImportData OriginalSkeletalMeshImportData;
-	SkeletalMeshLODModel.RawSkeletalMeshBulkData.LoadRawMesh(OriginalSkeletalMeshImportData);
+	SkeletalMesh->LoadLODImportedData(LodIndex, OriginalSkeletalMeshImportData);
 
 	//Backup the new geometry and skinning to be able to apply the skinning to the old geometry
 	FSkeletalMeshImportData NewGeometryAndRigData = *ImportData;
@@ -1512,7 +1512,7 @@ bool UnFbx::FFbxImporter::ReplaceSkeletalMeshSkinningImportData(const USkeletalM
 
 	//Load the original skeletal mesh import data
 	FSkeletalMeshImportData OriginalSkeletalMeshImportData;
-	SkeletalMeshLODModel.RawSkeletalMeshBulkData.LoadRawMesh(OriginalSkeletalMeshImportData);
+	SkeletalMesh->LoadLODImportedData(LodIndex, OriginalSkeletalMeshImportData);
 
 	ImportData->bDiffPose = OriginalSkeletalMeshImportData.bDiffPose;
 	ImportData->bUseT0AsRefPose = OriginalSkeletalMeshImportData.bUseT0AsRefPose;
@@ -1609,7 +1609,7 @@ bool UnFbx::FFbxImporter::GatherPointsForMorphTarget(FSkeletalMeshImportData* Ou
 	{
 		int32 OriginalPointIdx = OutData->PointToRawMap[ PointIdx ];
 
-		if ( ( NewImportData.Points[ OriginalPointIdx ] - OutData->Points[ PointIdx ] ).SizeSquared() > FMath::Square( THRESH_VECTORS_ARE_NEAR ) ) // THRESH_POINTS_ARE_NEAR is too big, we might not be recomputing some normals that can have changed significantly
+		if ( ( NewImportData.Points[ OriginalPointIdx ] - OutData->Points[ PointIdx ] ).SizeSquared() > FMath::Square(THRESH_POINTS_ARE_SAME) )
 		{
 			ModifiedPoints.Add( PointIdx );
 		}
@@ -1681,6 +1681,8 @@ USkeletalMesh* UnFbx::FFbxImporter::ImportSkeletalMesh(FImportSkeletalMeshArgs &
 		return nullptr;
 	}
 
+	int32 SafeLODIndex = ImportSkeletalMeshArgs.LodIndex < 0 ? 0 : ImportSkeletalMeshArgs.LodIndex;
+
 	int32 SkelType = 0; // 0 for skeletal mesh, 1 for rigid mesh
 	
 	FbxNode* Node = ImportSkeletalMeshArgs.NodeArray[0];
@@ -1766,12 +1768,9 @@ USkeletalMesh* UnFbx::FFbxImporter::ImportSkeletalMesh(FImportSkeletalMeshArgs &
 	//Adjust the import data from the import options
 	if (ExistingSkelMesh != nullptr)
 	{
-		int32 SafeLODIndex = ImportSkeletalMeshArgs.LodIndex < 0 ? 0 : ImportSkeletalMeshArgs.LodIndex;
 		if (ExistingSkelMesh->GetImportedModel() && ExistingSkelMesh->GetImportedModel()->LODModels.IsValidIndex(SafeLODIndex))
 		{
-			const FRawSkeletalMeshBulkData& RawSkeletalMeshBulkData = ExistingSkelMesh->GetImportedModel()->LODModels[SafeLODIndex].RawSkeletalMeshBulkData;
-			GeoImportVersion = RawSkeletalMeshBulkData.GeoImportVersion;
-			SkinningImportVersion = RawSkeletalMeshBulkData.SkinningImportVersion;
+			ExistingSkelMesh->GetLODImportedDataVersions(SafeLODIndex, GeoImportVersion, SkinningImportVersion);
 		}
 
 		if ((ImportOptions->bImportAsSkeletalSkinning || ImportOptions->bImportAsSkeletalGeometry))
@@ -1812,9 +1811,32 @@ USkeletalMesh* UnFbx::FFbxImporter::ImportSkeletalMesh(FImportSkeletalMeshArgs &
 		return nullptr;
 	}
 
+	FSkeletalMeshBuildSettings BuildOptions;
+	//Make sure the build option change in the re-import ui is reconduct
+	BuildOptions.bBuildAdjacencyBuffer = true;
+	BuildOptions.bRecomputeNormals = !ImportOptions->ShouldImportNormals() || !SkelMeshImportDataPtr->bHasNormals;
+	BuildOptions.bRecomputeTangents = !ImportOptions->ShouldImportTangents() || !SkelMeshImportDataPtr->bHasTangents;
+	BuildOptions.bUseMikkTSpace = (ImportOptions->NormalGenerationMethod == EFBXNormalGenerationMethod::MikkTSpace) && (!ImportOptions->ShouldImportNormals() || !ImportOptions->ShouldImportTangents());
+	BuildOptions.bComputeWeightedNormals = ImportOptions->bComputeWeightedNormals;
+	BuildOptions.bRemoveDegenerates = ImportOptions->bRemoveDegenerates;
+	BuildOptions.ThresholdPosition = ImportOptions->OverlappingThresholds.ThresholdPosition;
+	BuildOptions.ThresholdTangentNormal = ImportOptions->OverlappingThresholds.ThresholdTangentNormal;
+	BuildOptions.ThresholdUV = ImportOptions->OverlappingThresholds.ThresholdUV;
+	BuildOptions.MorphThresholdPosition = ImportOptions->OverlappingThresholds.MorphThresholdPosition;
+	
 	if (ExistingSkelMesh)
 	{
 		ExistingSkelMesh->PreEditChange(NULL);
+		//We update the build settings before saving the asset data we want to restore after the re-import
+		int32 SafeReimportLODIndex = ImportSkeletalMeshArgs.LodIndex < 0 ? 0 : ImportSkeletalMeshArgs.LodIndex;
+		FSkeletalMeshLODInfo* LODInfoPtr = ExistingSkelMesh->GetLODInfo(SafeReimportLODIndex);
+		if (LODInfoPtr)
+		{
+			//Adjacency buffer cannot be change in the re-import options, it must be backup before and restored after the copy.
+			bool bBuildAdjacencyBuffer = LODInfoPtr->BuildSettings.bBuildAdjacencyBuffer;
+			LODInfoPtr->BuildSettings = BuildOptions;
+			LODInfoPtr->BuildSettings.bBuildAdjacencyBuffer = bBuildAdjacencyBuffer;
+		}
 		//The backup of the skeletal mesh data empty the LOD array in the ImportedResource of the skeletal mesh
 		//If the import fail after this step the editor can crash when updating the bone later since the LODModel will not exist anymore
 		ExistSkelMeshDataPtr = SaveExistingSkelMeshData(ExistingSkelMesh, !ImportOptions->bImportMaterials, ImportSkeletalMeshArgs.LodIndex);
@@ -1839,8 +1861,8 @@ USkeletalMesh* UnFbx::FFbxImporter::ImportSkeletalMesh(FImportSkeletalMeshArgs &
 	check(ImportedResource->LODModels.Num() == 0);
 	ImportedResource->LODModels.Empty();
 	ImportedResource->LODModels.Add(new FSkeletalMeshLODModel());
-
-	FSkeletalMeshLODModel& LODModel = ImportedResource->LODModels[0];
+	const int32 ImportLODModelIndex = 0;
+	FSkeletalMeshLODModel& LODModel = ImportedResource->LODModels[ImportLODModelIndex];
 
 	// process materials from import data
 	ProcessImportMeshMaterials(SkeletalMesh->Materials, *SkelMeshImportDataPtr);
@@ -1862,22 +1884,19 @@ USkeletalMesh* UnFbx::FFbxImporter::ImportSkeletalMesh(FImportSkeletalMeshArgs &
 	ProcessImportMeshInfluences(*SkelMeshImportDataPtr);
 
 	//Store the original fbx import data the SkelMeshImportDataPtr should not be modified after this
-	LODModel.RawSkeletalMeshBulkData.SaveRawMesh(*SkelMeshImportDataPtr);
+	SkeletalMesh->SaveLODImportedData(ImportLODModelIndex, *SkelMeshImportDataPtr);
 	if (ImportOptions->bImportAsSkeletalSkinning)
 	{
-		LODModel.RawSkeletalMeshBulkData.SkinningImportVersion = ESkeletalMeshSkinningImportVersions::LatestVersion;
-		LODModel.RawSkeletalMeshBulkData.GeoImportVersion = GeoImportVersion;
+		SkeletalMesh->SetLODImportedDataVersions(ImportLODModelIndex, GeoImportVersion, ESkeletalMeshSkinningImportVersions::LatestVersion);
 	}
 	else if(ImportOptions->bImportAsSkeletalGeometry)
 	{
-		LODModel.RawSkeletalMeshBulkData.GeoImportVersion = ESkeletalMeshGeoImportVersions::LatestVersion;
-		LODModel.RawSkeletalMeshBulkData.SkinningImportVersion = SkinningImportVersion;
+		SkeletalMesh->SetLODImportedDataVersions(ImportLODModelIndex, ESkeletalMeshGeoImportVersions::LatestVersion, SkinningImportVersion);
 	}
 	else
 	{
 		//We reimport both
-		LODModel.RawSkeletalMeshBulkData.SkinningImportVersion = ESkeletalMeshSkinningImportVersions::LatestVersion;
-		LODModel.RawSkeletalMeshBulkData.GeoImportVersion = ESkeletalMeshGeoImportVersions::LatestVersion;
+		SkeletalMesh->SetLODImportedDataVersions(ImportLODModelIndex, ESkeletalMeshGeoImportVersions::LatestVersion, ESkeletalMeshSkinningImportVersions::LatestVersion);
 	}
 	SkeletalMesh->ResetLODInfo();
 	FSkeletalMeshLODInfo& NewLODInfo = SkeletalMesh->AddLODInfo();
@@ -1904,17 +1923,6 @@ USkeletalMesh* UnFbx::FFbxImporter::ImportSkeletalMesh(FImportSkeletalMeshArgs &
 		TArray<int32> LODPointToRawMap;
 		SkelMeshImportDataPtr->CopyLODImportData(LODPoints,LODWedges,LODFaces,LODInfluences,LODPointToRawMap);
 
-		FSkeletalMeshBuildSettings BuildOptions;
-		BuildOptions.bBuildAdjacencyBuffer = true;
-		BuildOptions.bRecomputeNormals = !ImportOptions->ShouldImportNormals() || !SkelMeshImportDataPtr->bHasNormals;
-		BuildOptions.bRecomputeTangents = !ImportOptions->ShouldImportTangents() || !SkelMeshImportDataPtr->bHasTangents;
-		BuildOptions.bUseMikkTSpace = (ImportOptions->NormalGenerationMethod == EFBXNormalGenerationMethod::MikkTSpace) && (!ImportOptions->ShouldImportNormals() || !ImportOptions->ShouldImportTangents());
-		BuildOptions.bComputeWeightedNormals = ImportOptions->bComputeWeightedNormals;
-		BuildOptions.bRemoveDegenerates = ImportOptions->bRemoveDegenerates;
-		BuildOptions.ThresholdPosition = ImportOptions->OverlappingThresholds.ThresholdPosition;
-		BuildOptions.ThresholdTangentNormal = ImportOptions->OverlappingThresholds.ThresholdTangentNormal;
-		BuildOptions.ThresholdUV = ImportOptions->OverlappingThresholds.ThresholdUV;
-
 		bool bUseLegacyMeshUtilities = false;
 		bool bBuildSuccess = false;
 		if (bUseLegacyMeshUtilities)
@@ -1927,7 +1935,7 @@ USkeletalMesh* UnFbx::FFbxImporter::ImportSkeletalMesh(FImportSkeletalMeshArgs &
 			TArray<FText> WarningMessages;
 			TArray<FName> WarningNames;
 			// Create actual rendering data.
-			bBuildSuccess = MeshUtilities.BuildSkeletalMesh(ImportedResource->LODModels[0], SkeletalMesh->RefSkeleton, LODInfluences, LODWedges, LODFaces, LODPoints, LODPointToRawMap, LegacyBuildOptions, &WarningMessages, &WarningNames);
+			bBuildSuccess = MeshUtilities.BuildSkeletalMesh(ImportedResource->LODModels[ImportLODModelIndex], SkeletalMesh->RefSkeleton, LODInfluences, LODWedges, LODFaces, LODPoints, LODPointToRawMap, LegacyBuildOptions, &WarningMessages, &WarningNames);
 
 			// temporary hack of message/names, should be one token or a struct
 			if (WarningMessages.Num() > 0 && WarningNames.Num() == WarningMessages.Num())
@@ -1943,14 +1951,13 @@ USkeletalMesh* UnFbx::FFbxImporter::ImportSkeletalMesh(FImportSkeletalMeshArgs &
 		else
 		{
 			//The imported LOD is always 0 here, the LOD custom import will import the LOD alone(in a temporary skeletalmesh) and add it to the base skeletal mesh later
-			const int32 LODIndex = 0;
-			check(SkeletalMesh->GetLODInfo(LODIndex) != nullptr);
+			check(SkeletalMesh->GetLODInfo(ImportLODModelIndex) != nullptr);
 			//Set the build options
-			SkeletalMesh->GetLODInfo(LODIndex)->BuildSettings = BuildOptions;
+			SkeletalMesh->GetLODInfo(ImportLODModelIndex)->BuildSettings = BuildOptions;
 			//New MeshDescription build process
 			IMeshBuilderModule& MeshBuilderModule = FModuleManager::Get().LoadModuleChecked<IMeshBuilderModule>("MeshBuilder");
 			//We must build the LODModel so we can restore properly the mesh, but we do not have to regenerate LODs
-			bBuildSuccess = MeshBuilderModule.BuildSkeletalMesh(SkeletalMesh, LODIndex, false);
+			bBuildSuccess = MeshBuilderModule.BuildSkeletalMesh(SkeletalMesh, ImportLODModelIndex, false);
 		}
 
 		if( !bBuildSuccess )
@@ -1972,9 +1979,9 @@ USkeletalMesh* UnFbx::FFbxImporter::ImportSkeletalMesh(FImportSkeletalMeshArgs &
 		}
 		if (ImportSkeletalMeshArgs.ImportMeshSectionsData)
 		{
-			for (int32 SectionIndex = 0; SectionIndex < ImportedResource->LODModels[0].Sections.Num(); ++SectionIndex)
+			for (int32 SectionIndex = 0; SectionIndex < ImportedResource->LODModels[ImportLODModelIndex].Sections.Num(); ++SectionIndex)
 			{
-				const FSkelMeshSection &SkelMeshSection = ImportedResource->LODModels[0].Sections[SectionIndex];
+				const FSkelMeshSection &SkelMeshSection = ImportedResource->LODModels[ImportLODModelIndex].Sections[SectionIndex];
 				int32 MaterialIndex = SkelMeshSection.MaterialIndex;
 				if (SkeletalMesh->GetLODInfo(0)->LODMaterialMap.IsValidIndex(SectionIndex) && SkeletalMesh->GetLODInfo(0)->LODMaterialMap[SectionIndex] != INDEX_NONE)
 				{
@@ -2401,8 +2408,7 @@ USkeletalMesh* UnFbx::FFbxImporter::ReimportSkeletalMesh(USkeletalMesh* Mesh, UF
 		FSkeletalMeshLODInfo* LODInfo = Mesh->GetLODInfo(0);
 		if(LODInfo && Mesh->GetImportedModel() && Mesh->GetImportedModel()->LODModels.IsValidIndex(0))
 		{
-			const FSkeletalMeshLODModel& LODModel = Mesh->GetImportedModel()->LODModels[0];
-			if (!LODModel.RawSkeletalMeshBulkData.IsBuildDataAvailable())
+			if (!Mesh->IsLODImportedDataBuildAvailable(0))
 			{
 				//Set the build settings
 				LODInfo->BuildSettings.bBuildAdjacencyBuffer = true;
@@ -2411,6 +2417,10 @@ USkeletalMesh* UnFbx::FFbxImporter::ReimportSkeletalMesh(USkeletalMesh* Mesh, UF
 				LODInfo->BuildSettings.bRecomputeTangents = SKImportData->NormalImportMethod != EFBXNormalImportMethod::FBXNIM_ImportNormalsAndTangents;
 				LODInfo->BuildSettings.bRemoveDegenerates = true;
 				LODInfo->BuildSettings.bUseMikkTSpace = SKImportData->NormalGenerationMethod == EFBXNormalGenerationMethod::MikkTSpace;
+				LODInfo->BuildSettings.ThresholdPosition = SKImportData->ThresholdPosition;
+				LODInfo->BuildSettings.ThresholdTangentNormal = SKImportData->ThresholdTangentNormal;
+				LODInfo->BuildSettings.ThresholdUV = SKImportData->ThresholdUV;
+				LODInfo->BuildSettings.MorphThresholdPosition = SKImportData->MorphThresholdPosition;
 			}
 		}
 	}
@@ -3911,6 +3921,9 @@ void UnFbx::FFbxImporter::InsertNewLODToBaseSkeletalMesh(USkeletalMesh* InSkelet
 
 	// Assign new FSkeletalMeshLODModel to desired slot in selected skeletal mesh.
 	DestImportedResource->LODModels[DesiredLOD] = NewLODModel;
+	//Copy the import data into the base skeletalmesh for the imported LOD
+	USkeletalMesh::CopyImportedData(0, InSkeletalMesh, DesiredLOD, BaseSkeletalMesh);
+	
 
 	// If this LOD had been generated previously by automatic mesh reduction, clear that flag.
 	LODInfo.bHasBeenSimplified = false;
@@ -4315,6 +4328,8 @@ void UnFbx::FFbxImporter::ImportMorphTargetsInternal( TArray<FbxNode*>& SkelMesh
 		BaseImportData.MorphTargetNames.Add(ShapeName);
 		TSet<uint32>& ModifiedPoints = BaseImportData.MorphTargetModifiedPoints.AddDefaulted_GetRef();
 		GatherPointsForMorphTarget(&ShapeImportData, SkelMeshNodeArray, &ShapeArray, ModifiedPoints);
+		//We do not need this data anymore empty it so we reduce the size of what we save into memory
+		ShapeImportData.PointToRawMap.Empty();
 		BaseImportData.MorphTargets.Add(ShapeImportData);
 		check(BaseImportData.MorphTargetNames.Num() == BaseImportData.MorphTargets.Num() && BaseImportData.MorphTargetNames.Num() == BaseImportData.MorphTargetModifiedPoints.Num());
 	}
@@ -4322,7 +4337,7 @@ void UnFbx::FFbxImporter::ImportMorphTargetsInternal( TArray<FbxNode*>& SkelMesh
 	if (BaseSkelMesh->GetImportedModel() && BaseSkelMesh->GetImportedModel()->LODModels.IsValidIndex(LODIndex))
 	{
 		//If we can build the skeletal mesh there is no need to build the morph target now, all the necessary build morph target data was copied before.
-		if (!BaseSkelMesh->GetImportedModel()->LODModels[LODIndex].RawSkeletalMeshBulkData.IsBuildDataAvailable())
+		if (!BaseSkelMesh->IsLODImportedDataBuildAvailable(LODIndex))
 		{
 			//Build MorphTargets
 			FLODUtilities::BuildMorphTargets(
@@ -4331,7 +4346,9 @@ void UnFbx::FFbxImporter::ImportMorphTargetsInternal( TArray<FbxNode*>& SkelMesh
 				LODIndex,
 				ImportOptions->ShouldImportNormals(),
 				ImportOptions->ShouldImportTangents(),
-				(ImportOptions->NormalGenerationMethod == EFBXNormalGenerationMethod::MikkTSpace));
+				(ImportOptions->NormalGenerationMethod == EFBXNormalGenerationMethod::MikkTSpace),
+				ImportOptions->OverlappingThresholds
+				);
 		}
 	}
 
@@ -4364,10 +4381,7 @@ void UnFbx::FFbxImporter::ImportFbxMorphTarget(TArray<FbxNode*> &SkelMeshNodeArr
 	{
 		ImportMorphTargetsInternal( SkelMeshNodeArray, BaseSkelMesh, LODIndex, BaseSkeletalMeshImportData);
 		//Save the rawMesh
-		if (BaseSkelMesh->GetImportedModel() && BaseSkelMesh->GetImportedModel()->LODModels.IsValidIndex(LODIndex))
-		{
-			BaseSkelMesh->GetImportedModel()->LODModels[LODIndex].RawSkeletalMeshBulkData.SaveRawMesh(BaseSkeletalMeshImportData);
-		}
+		BaseSkelMesh->SaveLODImportedData(LODIndex, BaseSkeletalMeshImportData);
 	}
 }
 
diff --git a/Engine/Source/Editor/UnrealEd/Private/FbxExporter.h b/Engine/Source/Editor/UnrealEd/Private/FbxExporter.h
--- a/Engine/Source/Editor/UnrealEd/Private/FbxExporter.h
+++ b/Engine/Source/Editor/UnrealEd/Private/FbxExporter.h
@@ -141,7 +141,7 @@ public:
 	 *
 	 * @return	true, if successful
 	 */
-	bool ExportLevelSequenceTracks(UMovieScene* MovieScene, IMovieScenePlayer* MovieScenePlayer, FbxNode* FbxActor, AActor* Actor, const TArray<UMovieSceneTrack*>& Tracks, const FMovieSceneSequenceTransform& RootToLocalTransform);
+	bool ExportLevelSequenceTracks(UMovieScene* MovieScene, IMovieScenePlayer* MovieScenePlayer, FMovieSceneSequenceIDRef InSequenceID, FbxNode* FbxActor, AActor* Actor, const TArray<UMovieSceneTrack*>& Tracks, const FMovieSceneSequenceTransform& RootToLocalTransform);
 
 	/**
 	 * Exports all the animation sequences part of a single Group in a Matinee sequence
@@ -270,6 +270,7 @@ private:
 	class IAnimTrackAdapter
 	{
 	public:
+		virtual ~IAnimTrackAdapter() {}
 		virtual int32 GetLocalStartFrame() const = 0;
 		virtual int32 GetStartFrame() const { return GetLocalStartFrame(); }
 		virtual int32 GetLength() const = 0;
@@ -456,7 +457,12 @@ private:
 	/**
 	 * Exports a level sequence 3D transform track into the FBX animation stack.
 	 */
-	void ExportLevelSequence3DTransformTrack( FbxNode* FbxActor, UMovieScene3DTransformTrack& TransformTrack, AActor* Actor, const TRange<FFrameNumber>& InPlaybackRange, const FMovieSceneSequenceTransform& RootToLocalTransform );
+	void ExportLevelSequence3DTransformTrack(FbxNode* FbxActor, IMovieScenePlayer* MovieScenePlayer, FMovieSceneSequenceIDRef InSequenceID, UMovieScene3DTransformTrack& TransformTrack, AActor* Actor, const TRange<FFrameNumber>& InPlaybackRange, const FMovieSceneSequenceTransform& RootToLocalTransform);
+
+	/**
+	 * Exports a level sequence 3D transform track that's getting interrogated(sample all sections) nto the FBX animation stack.
+	 */
+	void ExportLevelSequenceInterrogated3DTransformTrack(FbxNode* FbxActor, IMovieScenePlayer* MovieScenePlayer, FMovieSceneSequenceIDRef InSequenceID, UMovieScene3DTransformTrack& TransformTrack, AActor* Actor, const TRange<FFrameNumber>& InPlaybackRange, const FMovieSceneSequenceTransform& RootToLocalTransform);
 
 	/** 
 	 * Exports a level sequence property track into the FBX animation stack. 
diff --git a/Engine/Source/Editor/UnrealEd/Private/FbxMeshUtils.cpp b/Engine/Source/Editor/UnrealEd/Private/FbxMeshUtils.cpp
--- a/Engine/Source/Editor/UnrealEd/Private/FbxMeshUtils.cpp
+++ b/Engine/Source/Editor/UnrealEd/Private/FbxMeshUtils.cpp
@@ -366,10 +366,10 @@ namespace FbxMeshUtils
 			bool bUseLODs = true;
 			int32 MaxLODLevel = 0;
 			TArray<FString> LODStrings;
-			TArray<FbxNode*>* MeshObject = NULL;;
+			TArray<FbxNode*>* MeshObject = NULL;
 
 			//Set the build options if the BuildDat is not available so it is the same option we use to import the LOD
-			if (ImportedResource && ImportedResource->LODModels.IsValidIndex(LODLevel) && !ImportedResource->LODModels[LODLevel].RawSkeletalMeshBulkData.IsBuildDataAvailable())
+			if (ImportedResource && ImportedResource->LODModels.IsValidIndex(LODLevel) && !SelectedSkelMesh->IsLODImportedDataBuildAvailable(LODLevel))
 			{
 				FSkeletalMeshLODInfo* LODInfo = SelectedSkelMesh->GetLODInfo(LODLevel);
 				if (LODInfo)
@@ -383,6 +383,7 @@ namespace FbxMeshUtils
 					LODInfo->BuildSettings.ThresholdPosition = ImportOptions->OverlappingThresholds.ThresholdPosition;
 					LODInfo->BuildSettings.ThresholdTangentNormal = ImportOptions->OverlappingThresholds.ThresholdTangentNormal;
 					LODInfo->BuildSettings.ThresholdUV = ImportOptions->OverlappingThresholds.ThresholdUV;
+					LODInfo->BuildSettings.MorphThresholdPosition = ImportOptions->OverlappingThresholds.MorphThresholdPosition;
 				}
 			}
 
