Merging CL#10802309

#UE4 Added option in resavepackagescommandlet to keep the guid of the packages being resaved to minimize differences on disk after the resave process

diff --git a/Engine/Source/Editor/UnrealEd/Private/Commandlets/ContentCommandlets.cpp b/Engine/Source/Editor/UnrealEd/Private/Commandlets/ContentCommandlets.cpp
--- a/Engine/Source/Editor/UnrealEd/Private/Commandlets/ContentCommandlets.cpp
+++ b/Engine/Source/Editor/UnrealEd/Private/Commandlets/ContentCommandlets.cpp
@@ -766,7 +766,9 @@ void UResavePackagesCommandlet::LoadAndSaveOnePackage(const FString& Filename)
 							*Filename,Linker->Summary.GetFileVersionUE4(), Linker->Summary.GetFileVersionLicenseeUE4(), GPackageFileUE4Version, VER_LATEST_ENGINE_LICENSEEUE4 );
 					}
 
-					if( SavePackageHelper(Package, Filename) )
+					const static bool bKeepPackageGUIDOnSave = FParse::Param(FCommandLine::Get(), TEXT("KeepPackageGUIDOnSave"));
+					ESaveFlags SaveFlags = bKeepPackageGUIDOnSave ? SAVE_KeepGUID : SAVE_None;
+					if( SavePackageHelper(Package, Filename, RF_Standalone, GWarn, nullptr, SaveFlags) )
 					{
 						if (Verbosity == VERY_VERBOSE)
 						{
