Avoid trying to reload an asset after force delete failure

diff --git a/Engine/Source/Editor/UnrealEd/Private/ObjectTools.cpp b/Engine/Source/Editor/UnrealEd/Private/ObjectTools.cpp
--- a/Engine/Source/Editor/UnrealEd/Private/ObjectTools.cpp
+++ b/Engine/Source/Editor/UnrealEd/Private/ObjectTools.cpp
@@ -2699,7 +2699,7 @@ namespace ObjectTools
 			ReloadEditorWorldForReferenceReplacementIfNecessary(ObjectsToDelete);
 		}
 
-		TArray<UPackage*> PackagesToReload;
+		TArray<UPackage*> PackagesFailedToDelete;
 		{
 			int32 ReplaceableObjectsNum = 0;
 			{
@@ -2820,8 +2820,8 @@ namespace ObjectTools
 				// if the delete fails at this point, it means the object won't be able to be purged and might be left in a weird state, as a last resort queue its package for reload
 				else
 				{
-					UE_LOG(LogObjectTools, Warning, TEXT("ForceDeleteObject failed to delete %s, queuing its package for reloading"), *CurObject->GetName());
-					PackagesToReload.AddUnique(CurObject->GetOutermost());
+					UE_LOG(LogObjectTools, Warning, TEXT("ForceDeleteObject failed to delete %s, this package is now potentially corrupt"), *CurObject->GetName());
+					PackagesFailedToDelete.AddUnique(CurObject->GetOutermost());
 					It.RemoveCurrent();
 				}
 
@@ -2846,13 +2846,23 @@ namespace ObjectTools
 		}
 		ObjectsToDelete.Empty();
 
-		// Reload packages of objects we failed to clean as a last resort since they might be left in an unstable state due to the force replace references
-		ensureMsgf(PackagesToReload.Num() == 0, TEXT("Failed to unload all packages during ForceDeleteObjects"));
-		if (PackagesToReload.Num() > 0)
+		// Final report of packages we failed to delete. This is mostly for crash reporter. To fix this ensure you need
+		// to fix the reference replacement/isreferenced mismatch. The former does not find subobjects, the latter does:
+		if(PackagesFailedToDelete.Num())
 		{
-			FText ErrorMessage;
-			bool bSuccess = UPackageTools::ReloadPackages(PackagesToReload, ErrorMessage, UPackageTools::EReloadPackagesInteractionMode::AssumePositive);
-			ensureMsgf(bSuccess, TEXT("Failed to reload package as a last resort in ForceDeleteObjects"));
+			FString FailedPackageNames;
+			for(UPackage* Package : PackagesFailedToDelete)
+			{
+				FailedPackageNames.Append(FString::Printf(TEXT("\r\n%s"), *Package->GetName()));
+			}
+
+			ensureMsgf( false, 
+				TEXT(R"(
+					Failed to unload all packages during ForceDeleteObjects - 
+					these packages are likely corrupt. Consider restarting the 
+					editor, noting which assets remain and then deleting them 
+					from the file system manually: %s)"), *FailedPackageNames
+			);
 		}
 
 		GWarn->EndSlowTask();
