Fixing morphtarget reimport

diff --git a/Engine/Source/Editor/UnrealEd/Private/Fbx/FbxSkeletalMeshImport.cpp b/Engine/Source/Editor/UnrealEd/Private/Fbx/FbxSkeletalMeshImport.cpp
--- a/Engine/Source/Editor/UnrealEd/Private/Fbx/FbxSkeletalMeshImport.cpp
+++ b/Engine/Source/Editor/UnrealEd/Private/Fbx/FbxSkeletalMeshImport.cpp
@@ -1454,75 +1454,12 @@ bool UnFbx::FFbxImporter::FillSkeletalMeshImportData(TArray<FbxNode*>& NodeArray
 
 bool UnFbx::FFbxImporter::ReplaceSkeletalMeshGeometryImportData(const USkeletalMesh* SkeletalMesh, FSkeletalMeshImportData* ImportData, int32 LodIndex)
 {
-	FSkeletalMeshModel *ImportedResource = SkeletalMesh->GetImportedModel();
-	check(ImportedResource && ImportedResource->LODModels.IsValidIndex(LodIndex));
-	FSkeletalMeshLODModel& SkeletalMeshLODModel = ImportedResource->LODModels[LodIndex];
-
-	const FSkeletalMeshLODInfo* LodInfo = SkeletalMesh->GetLODInfo(LodIndex);
-	check(LodInfo);
-
-	//Load the original skeletal mesh import data
-	FSkeletalMeshImportData OriginalSkeletalMeshImportData;
-	SkeletalMeshLODModel.RawSkeletalMeshBulkData.LoadRawMesh(OriginalSkeletalMeshImportData);
-
-	//Backup the new geometry and skinning to be able to apply the skinning to the old geometry
-	FSkeletalMeshImportData NewGeometryAndRigData = *ImportData;
-
-	ImportData->bHasNormals = OriginalSkeletalMeshImportData.bHasNormals;
-	ImportData->bHasTangents = OriginalSkeletalMeshImportData.bHasTangents;
-	ImportData->bHasVertexColors = OriginalSkeletalMeshImportData.bHasVertexColors;
-	ImportData->NumTexCoords = OriginalSkeletalMeshImportData.NumTexCoords;
-	
-	ImportData->Materials.Reset();
-	ImportData->Points.Reset();
-	ImportData->Faces.Reset();
-	ImportData->Wedges.Reset();
-	ImportData->PointToRawMap.Reset();
-
-	//Material is a special case since we cannot serialize the UMaterialInstance when saving the RawSkeletalMeshBulkData
-	//So it has to be reconstructed.
-	ImportData->MaxMaterialIndex = 0;
-	for (int32 MaterialIndex = 0; MaterialIndex < SkeletalMesh->Materials.Num(); ++MaterialIndex)
-	{
-		SkeletalMeshImportData::FMaterial NewMaterial;
-
-		NewMaterial.MaterialImportName = SkeletalMesh->Materials[MaterialIndex].ImportedMaterialSlotName.ToString();
-		NewMaterial.Material = SkeletalMesh->Materials[MaterialIndex].MaterialInterface;
-		// Add an entry for each unique material
-		ImportData->MaxMaterialIndex = FMath::Max(ImportData->MaxMaterialIndex, (uint32)(ImportData->Materials.Add(NewMaterial)));
-	}
-
-	ImportData->NumTexCoords = OriginalSkeletalMeshImportData.NumTexCoords;
-	ImportData->Points += OriginalSkeletalMeshImportData.Points;
-	ImportData->Faces += OriginalSkeletalMeshImportData.Faces;
-	ImportData->Wedges += OriginalSkeletalMeshImportData.Wedges;
-	ImportData->PointToRawMap += OriginalSkeletalMeshImportData.PointToRawMap;
-
-	return ImportData->ApplyRigToGeo(NewGeometryAndRigData);
+	return FSkeletalMeshImportData::ReplaceSkeletalMeshGeometryImportData(SkeletalMesh, ImportData, LodIndex);
 }
 
 bool UnFbx::FFbxImporter::ReplaceSkeletalMeshSkinningImportData(const USkeletalMesh* SkeletalMesh, FSkeletalMeshImportData* ImportData, int32 LodIndex)
 {
-	FSkeletalMeshModel *ImportedResource = SkeletalMesh->GetImportedModel();
-	check(ImportedResource && ImportedResource->LODModels.IsValidIndex(LodIndex));
-	FSkeletalMeshLODModel& SkeletalMeshLODModel = ImportedResource->LODModels[LodIndex];
-
-	const FSkeletalMeshLODInfo* LodInfo = SkeletalMesh->GetLODInfo(LodIndex);
-	check(LodInfo);
-
-	//Load the original skeletal mesh import data
-	FSkeletalMeshImportData OriginalSkeletalMeshImportData;
-	SkeletalMeshLODModel.RawSkeletalMeshBulkData.LoadRawMesh(OriginalSkeletalMeshImportData);
-
-	ImportData->bDiffPose = OriginalSkeletalMeshImportData.bDiffPose;
-	ImportData->bUseT0AsRefPose = OriginalSkeletalMeshImportData.bUseT0AsRefPose;
-
-	ImportData->RefBonesBinary.Reset();
-
-	ImportData->RefBonesBinary += OriginalSkeletalMeshImportData.RefBonesBinary;
-
-	//Fix the old skinning to match the new geometry
-	return ImportData->ApplyRigToGeo(OriginalSkeletalMeshImportData);
+	return FSkeletalMeshImportData::ReplaceSkeletalMeshRigImportData(SkeletalMesh, ImportData, LodIndex);
 }
 
 bool UnFbx::FFbxImporter::FillSkeletalMeshImportPoints(FSkeletalMeshImportData* OutData, FbxNode* RootNode, FbxNode* Node, FbxShape* FbxShape)
@@ -1591,7 +1528,7 @@ bool UnFbx::FFbxImporter::GatherPointsForMorphTarget(FSkeletalMeshImportData* Ou
 		FbxMesh* FbxMesh = Node->GetMesh();
 
 		FbxShape* FbxShape = nullptr;
-		if (FbxShapeArray)
+		if (FbxShapeArray && FbxShapeArray->IsValidIndex(NodeIndex))
 		{
 			FbxShape = (*FbxShapeArray)[NodeIndex];
 		}
@@ -1774,18 +1711,15 @@ USkeletalMesh* UnFbx::FFbxImporter::ImportSkeletalMesh(FImportSkeletalMeshArgs &
 			SkinningImportVersion = RawSkeletalMeshBulkData.SkinningImportVersion;
 		}
 
-		if ((ImportOptions->bImportAsSkeletalSkinning || ImportOptions->bImportAsSkeletalGeometry))
+		if (ImportOptions->bImportAsSkeletalSkinning)
 		{
-			if (ImportOptions->bImportAsSkeletalSkinning)
-			{
-				//Replace geometry import data by original existing skel mesh geometry data
-				ReplaceSkeletalMeshGeometryImportData(ExistingSkelMesh, SkelMeshImportDataPtr, SafeLODIndex);
-			}
-			else if (ImportOptions->bImportAsSkeletalGeometry)
-			{
-				//Replace skinning import data by original existing skel mesh skinning data
-				ReplaceSkeletalMeshSkinningImportData(ExistingSkelMesh, SkelMeshImportDataPtr, SafeLODIndex);
-			}
+			//Replace geometry import data by original existing skel mesh geometry data
+			ReplaceSkeletalMeshGeometryImportData(ExistingSkelMesh, SkelMeshImportDataPtr, SafeLODIndex);
+		}
+		else if (ImportOptions->bImportAsSkeletalGeometry)
+		{
+			//Replace skinning import data by original existing skel mesh skinning data
+			ReplaceSkeletalMeshSkinningImportData(ExistingSkelMesh, SkelMeshImportDataPtr, SafeLODIndex);
 		}
 
 		//Reuse the vertex color in case we are not re-importing weights only and there is some vertexcolor in the existing skeleMesh
diff --git a/Engine/Source/Runtime/Engine/Private/Rendering/SkeletalMeshLODImporterData.cpp b/Engine/Source/Runtime/Engine/Private/Rendering/SkeletalMeshLODImporterData.cpp
--- a/Engine/Source/Runtime/Engine/Private/Rendering/SkeletalMeshLODImporterData.cpp
+++ b/Engine/Source/Runtime/Engine/Private/Rendering/SkeletalMeshLODImporterData.cpp
@@ -150,6 +150,9 @@ bool FSkeletalMeshImportData::ReplaceSkeletalMeshGeometryImportData(const USkele
 	ImportData->Faces.Reset();
 	ImportData->Wedges.Reset();
 	ImportData->PointToRawMap.Reset();
+	ImportData->MorphTargetNames.Reset();
+	ImportData->MorphTargets.Reset();
+	ImportData->MorphTargetModifiedPoints.Reset();
 
 	//Material is a special case since we cannot serialize the UMaterialInstance when saving the RawSkeletalMeshBulkData
 	//So it has to be reconstructed.
@@ -169,6 +172,9 @@ bool FSkeletalMeshImportData::ReplaceSkeletalMeshGeometryImportData(const USkele
 	ImportData->Faces += OriginalSkeletalMeshImportData.Faces;
 	ImportData->Wedges += OriginalSkeletalMeshImportData.Wedges;
 	ImportData->PointToRawMap += OriginalSkeletalMeshImportData.PointToRawMap;
+	ImportData->MorphTargetNames += OriginalSkeletalMeshImportData.MorphTargetNames;
+	ImportData->MorphTargets += OriginalSkeletalMeshImportData.MorphTargets;
+	ImportData->MorphTargetModifiedPoints += OriginalSkeletalMeshImportData.MorphTargetModifiedPoints;
 
 	return ImportData->ApplyRigToGeo(NewGeometryAndRigData);
 }
@@ -190,7 +196,6 @@ bool FSkeletalMeshImportData::ReplaceSkeletalMeshRigImportData(const USkeletalMe
 	ImportData->bUseT0AsRefPose = OriginalSkeletalMeshImportData.bUseT0AsRefPose;
 
 	ImportData->RefBonesBinary.Reset();
-
 	ImportData->RefBonesBinary += OriginalSkeletalMeshImportData.RefBonesBinary;
 
 	//Fix the old rig to match the new geometry
@@ -206,9 +211,6 @@ bool FSkeletalMeshImportData::ApplyRigToGeo(FSkeletalMeshImportData& Other)
 	FWedgePosition::FillWedgePosition(OldGeoOverlappingPosition, Other.Points, Other.Wedges, THRESH_POINTS_ARE_SAME);
 	FOctreeQueryHelper OctreeQueryHelper(OldGeoOverlappingPosition.GetOctree());
 
-	int32 NewWedgesNum = Wedges.Num();
-	int32 OldWedgesNum = Other.Wedges.Num();
-
 	//////////////////////////////////////////////////////////////////////////
 	// Found the Remapping between old vertex index and new vertex index
 	// The old vertex index are the key, the index of the first array
@@ -219,7 +221,8 @@ bool FSkeletalMeshImportData::ApplyRigToGeo(FSkeletalMeshImportData& Other)
 	// new vertex will have correct bone weight apply to them.
 	TArray<TArray<int32>> OldToNewRemap;
 	OldToNewRemap.AddDefaulted(Other.Points.Num());
-	for (int32 WedgeIndex = 0; WedgeIndex < NewWedgesNum; ++WedgeIndex)
+
+	for (int32 WedgeIndex = 0, NewWedgesNum = Wedges.Num(); WedgeIndex < NewWedgesNum; ++WedgeIndex)
 	{
 		const FVector2D& CurWedgeUV = Wedges[WedgeIndex].UVs[0];
 		int32 NewVertexIndex = (int32)(Wedges[WedgeIndex].VertexIndex);
