Fixed potential denormalized quaternion when maninpulating atmospheric lights.

https://jira.it.epicgames.net/browse/UE-85859

diff --git a/Engine/Source/Editor/UnrealEd/Private/LevelEditorViewport.cpp b/Engine/Source/Editor/UnrealEd/Private/LevelEditorViewport.cpp
--- a/Engine/Source/Editor/UnrealEd/Private/LevelEditorViewport.cpp
+++ b/Engine/Source/Editor/UnrealEd/Private/LevelEditorViewport.cpp
@@ -4164,6 +4164,7 @@ void FLevelEditorViewportClient::MouseMove(FViewport* InViewport, int32 x, int32
 			SelectedAtmosphericLight->SetWorldTransform(ComponentTransform);
 
 			UserControlledAtmosphericLightMatrix = ComponentTransform;
+			UserControlledAtmosphericLightMatrix.NormalizeRotation();
 		}
 	}
 
