Updating copyright for Engine Editor.

#rnx

diff --git a/Engine/Source/Editor/DetailCustomizations/Private/AnimationAssetDetail.cpp b/Engine/Source/Editor/DetailCustomizations/Private/AnimationAssetDetail.cpp
--- a/Engine/Source/Editor/DetailCustomizations/Private/AnimationAssetDetail.cpp
+++ b/Engine/Source/Editor/DetailCustomizations/Private/AnimationAssetDetail.cpp
@@ -1,4 +1,4 @@
-// Copyright 1998-2019 Epic Games, Inc. All Rights Reservekd.
+// Copyright Epic Games, Inc. All Rights Reservekd.
 
 #include "CoreMinimal.h"
 #include "Widgets/DeclarativeSyntaxSupport.h"
diff --git a/Engine/Source/Editor/DetailCustomizations/Private/PoseAssetDetails.cpp b/Engine/Source/Editor/DetailCustomizations/Private/PoseAssetDetails.cpp
--- a/Engine/Source/Editor/DetailCustomizations/Private/PoseAssetDetails.cpp
+++ b/Engine/Source/Editor/DetailCustomizations/Private/PoseAssetDetails.cpp
@@ -1,4 +1,4 @@
-// Copyright 1998-2019 Epic Games, Inc. All Rights Reservekd.
+// Copyright Epic Games, Inc. All Rights Reservekd.
 
 #include "PoseAssetDetails.h"
 #include "Misc/MessageDialog.h"
diff --git a/Engine/Source/Editor/UnrealEd/Private/AboutScreen.cpp b/Engine/Source/Editor/UnrealEd/Private/AboutScreen.cpp
--- a/Engine/Source/Editor/UnrealEd/Private/AboutScreen.cpp
+++ b/Engine/Source/Editor/UnrealEd/Private/AboutScreen.cpp
@@ -1,4 +1,4 @@
-// Copyright 1998-2019 Epic Games, Inc. All Rights Reserved.
+// Copyright Epic Games, Inc. All Rights Reserved.
 
 #include "AboutScreen.h"
 #include "Fonts/SlateFontInfo.h"
@@ -29,7 +29,7 @@ void SAboutScreen::Construct(const FArguments& InArgs)
 #pragma warning(push)
 #pragma warning(disable : 4428)	// universal-character-name encountered in source
 #endif
-	AboutLines.Add(MakeShareable(new FLineDefinition(LOCTEXT("Copyright1", "Copyright 1998-2019 Epic Games, Inc. All rights reserved"), 11, FLinearColor(1.f, 1.f, 1.f), FMargin(0.f) )));
+	AboutLines.Add(MakeShareable(new FLineDefinition(LOCTEXT("Copyright1", "Copyright Epic Games, Inc. All rights reserved"), 11, FLinearColor(1.f, 1.f, 1.f), FMargin(0.f) )));
 	AboutLines.Add(MakeShareable(new FLineDefinition(LOCTEXT("Copyright2", "Epic, Epic Games, Unreal, and their respective logos are trademarks or registered trademarks of Epic Games, Inc.\nin the United States of America and elsewhere."), 8, FLinearColor(1.f, 1.f, 1.f), FMargin(0.0f,2.0f) )));
 
 #ifdef _MSC_VER
diff --git a/Engine/Source/Editor/VirtualTexturingEditor/Private/RuntimeVirtualTextureDetailsCustomization.cpp b/Engine/Source/Editor/VirtualTexturingEditor/Private/RuntimeVirtualTextureDetailsCustomization.cpp
--- a/Engine/Source/Editor/VirtualTexturingEditor/Private/RuntimeVirtualTextureDetailsCustomization.cpp
+++ b/Engine/Source/Editor/VirtualTexturingEditor/Private/RuntimeVirtualTextureDetailsCustomization.cpp
@@ -1,5 +1,5 @@
-// Copyright 1998-2019 Epic Games, Inc. All Rights Reserved.
-// Copyright 1998-2019 Epic Games, Inc. All Rights Reserved.
+// Copyright Epic Games, Inc. All Rights Reserved.
+// Copyright Epic Games, Inc. All Rights Reserved.
 
 #include "RuntimeVirtualTextureDetailsCustomization.h"
 
