Support slow task cancellation checks with nested tasks and throughout the slow task lifescope (and after its lifescope; before a new root slow task is created).

diff --git a/Engine/Source/Editor/UnrealEd/Private/FeedbackContextEditor.cpp b/Engine/Source/Editor/UnrealEd/Private/FeedbackContextEditor.cpp
--- a/Engine/Source/Editor/UnrealEd/Private/FeedbackContextEditor.cpp
+++ b/Engine/Source/Editor/UnrealEd/Private/FeedbackContextEditor.cpp
@@ -383,10 +383,13 @@ void FFeedbackContextEditor::StartSlowTask( const FText& Task, bool bShowCancelB
 {
 	FFeedbackContext::StartSlowTask( Task, bShowCancelButton );
 
-	TSharedPtr<SWindow> ParentWindow;
 	if (GEditor)
-	{		
+	{
+		// reset the cancellation flag
+		HasTaskBeenCancelled = false;
+
 		// If there is a pie window and it is active attempt to parent any slow task dialogs to it to prevent the game window from falling behind due to a slowtask window opening.
+		TSharedPtr<SWindow> ParentWindow;
 		if (FWorldContext* PieWorldContext = GEditor->GetPIEWorldContext())
 		{
 			FSlatePlayInEditorInfo* SlatePlayInEditorSession = GEditor->SlatePlayInEditorMap.Find(PieWorldContext->ContextHandle);
@@ -400,69 +403,66 @@ void FFeedbackContextEditor::StartSlowTask( const FText& Task, bool bShowCancelB
 				}
 			}
 		}
-	}
-
-	// Attempt to parent the slow task window to the slate main frame
-	if(!ParentWindow.IsValid() && FModuleManager::Get().IsModuleLoaded( "MainFrame" ) )
-	{
-		IMainFrameModule& MainFrame = FModuleManager::LoadModuleChecked<IMainFrameModule>( "MainFrame" );
-		ParentWindow = MainFrame.GetParentWindow();
-	}
-
-	if (GIsEditor && ParentWindow.IsValid())
-	{
-		GSlowTaskOccurred = GIsSlowTask;
-
-		// Don't show the progress dialog if the Build Progress dialog is already visible
-		bool bProgressWindowShown = BuildProgressWidget.IsValid();
-
-		// Don't show the progress dialog if a Slate menu is currently open
-		const bool bHaveOpenMenu = FSlateApplication::Get().AnyMenusVisible();
 
-		if( bHaveOpenMenu )
+		// Attempt to parent the slow task window to the slate main frame
+		if (!ParentWindow.IsValid() && FModuleManager::Get().IsModuleLoaded("MainFrame"))
 		{
-			UE_LOG(LogSlate, Log, TEXT("Prevented a slow task dialog from being summoned while a context menu was open") );
+			IMainFrameModule& MainFrame = FModuleManager::LoadModuleChecked<IMainFrameModule>("MainFrame");
+			ParentWindow = MainFrame.GetParentWindow();
 		}
 
-		// reset the cancellation flag
-		HasTaskBeenCancelled = false;
-
-		if (!bProgressWindowShown && !bHaveOpenMenu && FSlateApplication::Get().CanDisplayWindows())
+		if (ParentWindow.IsValid())
 		{
-			FOnCancelClickedDelegate OnCancelClicked;
-			if ( bShowCancelButton )
+			GSlowTaskOccurred = GIsSlowTask;
+
+			// Don't show the progress dialog if the Build Progress dialog is already visible
+			bool bProgressWindowShown = BuildProgressWidget.IsValid();
+
+			// Don't show the progress dialog if a Slate menu is currently open
+			const bool bHaveOpenMenu = FSlateApplication::Get().AnyMenusVisible();
+
+			if (bHaveOpenMenu)
 			{
-				// The cancel button is only displayed if a delegate is bound to it.
-				OnCancelClicked = FOnCancelClickedDelegate::CreateRaw(this, &FFeedbackContextEditor::OnUserCancel);
+				UE_LOG(LogSlate, Log, TEXT("Prevented a slow task dialog from being summoned while a context menu was open"));
 			}
 
-			const bool bFocusAndActivate = FPlatformApplicationMisc::IsThisApplicationForeground();
+			if (!bProgressWindowShown && !bHaveOpenMenu && FSlateApplication::Get().CanDisplayWindows())
+			{
+				FOnCancelClickedDelegate OnCancelClicked;
+				if (bShowCancelButton)
+				{
+					// The cancel button is only displayed if a delegate is bound to it.
+					OnCancelClicked = FOnCancelClickedDelegate::CreateRaw(this, &FFeedbackContextEditor::OnUserCancel);
+				}
 
-			TSharedRef<SWindow> SlowTaskWindowRef = SNew(SWindow)
-				.SizingRule(ESizingRule::Autosized)
-				.AutoCenter(EAutoCenter::PreferredWorkArea)
-				.IsPopupWindow(true)
-				.CreateTitleBar(true)
-				.ActivationPolicy(bFocusAndActivate ? EWindowActivationPolicy::Always : EWindowActivationPolicy::Never)
-				.FocusWhenFirstShown(bFocusAndActivate);
+				const bool bFocusAndActivate = FPlatformApplicationMisc::IsThisApplicationForeground();
 
-			SlowTaskWindowRef->SetContent(
-				SNew(SSlowTaskWidget)
-				.ScopeStack(GetScopeStackSharedPtr())
-				.OnCancelClickedDelegate( OnCancelClicked )
-			);
+				TSharedRef<SWindow> SlowTaskWindowRef = SNew(SWindow)
+					.SizingRule(ESizingRule::Autosized)
+					.AutoCenter(EAutoCenter::PreferredWorkArea)
+					.IsPopupWindow(true)
+					.CreateTitleBar(true)
+					.ActivationPolicy(bFocusAndActivate ? EWindowActivationPolicy::Always : EWindowActivationPolicy::Never)
+					.FocusWhenFirstShown(bFocusAndActivate);
 
-			SlowTaskWindow = SlowTaskWindowRef;
+				SlowTaskWindowRef->SetContent(
+					SNew(SSlowTaskWidget)
+					.ScopeStack(GetScopeStackSharedPtr())
+					.OnCancelClickedDelegate(OnCancelClicked)
+				);
 
-			const bool bSlowTask = true;
-			FSlateApplication::Get().AddModalWindow( SlowTaskWindowRef, ParentWindow, bSlowTask );
+				SlowTaskWindow = SlowTaskWindowRef;
 
-			SlowTaskWindowRef->ShowWindow();
+				const bool bSlowTask = true;
+				FSlateApplication::Get().AddModalWindow(SlowTaskWindowRef, ParentWindow, bSlowTask);
 
-			TickSlate(SlowTaskWindow.Pin());
-		}
+				SlowTaskWindowRef->ShowWindow();
 
-		FPlatformSplash::SetSplashText( SplashTextType::StartupProgress, *Task.ToString() );
+				TickSlate(SlowTaskWindow.Pin());
+			}
+
+			FPlatformSplash::SetSplashText(SplashTextType::StartupProgress, *Task.ToString());
+		}
 	}
 }
 
@@ -560,7 +560,7 @@ void FFeedbackContextEditor::ProgressReported( const float TotalProgressInterp,
 				else
 				{
 					NewDisplayMessage.AppendChar( TCHAR( ' ' ) );
-				}				
+				}
 			}
 			NewDisplayMessage.Append( FString::Printf( TEXT( " %i%%" ), int(TotalProgressInterp * 100.f) ) );
 			DisplayMessage = FText::FromString( NewDisplayMessage );
@@ -570,12 +570,9 @@ void FFeedbackContextEditor::ProgressReported( const float TotalProgressInterp,
 	}
 }
 
-/** Whether or not the user has canceled out of this dialog */
 bool FFeedbackContextEditor::ReceivedUserCancel( void )
 {
-	const bool res = HasTaskBeenCancelled;
-	HasTaskBeenCancelled = false;
-	return res;
+	return HasTaskBeenCancelled;
 }
 
 void FFeedbackContextEditor::OnUserCancel()
diff --git a/Engine/Source/Editor/UnrealEd/Public/FeedbackContextEditor.h b/Engine/Source/Editor/UnrealEd/Public/FeedbackContextEditor.h
--- a/Engine/Source/Editor/UnrealEd/Public/FeedbackContextEditor.h
+++ b/Engine/Source/Editor/UnrealEd/Public/FeedbackContextEditor.h
@@ -42,7 +42,11 @@ public:
 
 	void SetContext( FContextSupplier* InSupplier ) override {}
 
-	/** Whether or not the user has canceled out of this dialog */
+	/**
+	 * Whether or not the user has canceled out of the progress dialog 
+	 * (i.e. the ongoing slow task or the last one that ran).
+	 * The user cancel flag is reset when starting a new root slow task.
+	 */
 	virtual bool ReceivedUserCancel() override;
 
 	void OnUserCancel();
diff --git a/Engine/Source/Runtime/Core/Public/Misc/FeedbackContext.h b/Engine/Source/Runtime/Core/Public/Misc/FeedbackContext.h
--- a/Engine/Source/Runtime/Core/Public/Misc/FeedbackContext.h
+++ b/Engine/Source/Runtime/Core/Public/Misc/FeedbackContext.h
@@ -24,7 +24,11 @@ public:
 	/** Ask the user a binary question, returning their answer */
 	virtual bool YesNof( const FText& Question ) { return false; }
 	
-	/** Return whether the user has requested to cancel the current slow task */
+	/**
+	 * Whether or not the user has canceled out of the progress dialog
+	 * (i.e. the ongoing slow task or the last one that ran).
+	 * The user cancel flag is reset when starting a new root slow task.
+	 */
 	virtual bool ReceivedUserCancel() { return false; };
 	
 	/** Public const access to the current state of the scope stack */
