#UE4 Move AssetReferenceFilter to EditorEngine so it can be used in Commandlets

diff --git a/Engine/Source/Editor/ClassViewer/Private/ClassViewerFilter.cpp b/Engine/Source/Editor/ClassViewer/Private/ClassViewerFilter.cpp
--- a/Engine/Source/Editor/ClassViewer/Private/ClassViewerFilter.cpp
+++ b/Engine/Source/Editor/ClassViewer/Private/ClassViewerFilter.cpp
@@ -398,7 +398,7 @@ FClassViewerFilter::FClassViewerFilter(const FClassViewerInitializationOptions&
 	AssetRegistry(FModuleManager::GetModuleChecked<FAssetRegistryModule>(TEXT("AssetRegistry")).Get())
 {
 	// Create a game-specific filter, if the referencing property/assets were supplied
-	if (GUnrealEd)
+	if (GEditor)
 	{
 		FAssetReferenceFilterContext AssetReferenceFilterContext;
 		AssetReferenceFilterContext.ReferencingAssets = InInitOptions.AdditionalReferencingAssets;
@@ -411,7 +411,7 @@ FClassViewerFilter::FClassViewerFilter(const FClassViewerInitializationOptions&
 				AssetReferenceFilterContext.ReferencingAssets.Add(FAssetData(ReferencingObject));
 			}
 		}
-		AssetReferenceFilter = GUnrealEd->MakeAssetReferenceFilter(AssetReferenceFilterContext);
+		AssetReferenceFilter = GEditor->MakeAssetReferenceFilter(AssetReferenceFilterContext);
 	}
 }
 
diff --git a/Engine/Source/Editor/ContentBrowser/Private/SAssetPicker.cpp b/Engine/Source/Editor/ContentBrowser/Private/SAssetPicker.cpp
--- a/Engine/Source/Editor/ContentBrowser/Private/SAssetPicker.cpp
+++ b/Engine/Source/Editor/ContentBrowser/Private/SAssetPicker.cpp
@@ -244,7 +244,7 @@ void SAssetPicker::Construct( const FArguments& InArgs )
 				AssetReferenceFilterContext.ReferencingAssets.Add(FAssetData(ReferencingObject));
 			}
 		}
-		TSharedPtr<IAssetReferenceFilter> AssetReferenceFilter = GUnrealEd ? GUnrealEd->MakeAssetReferenceFilter(AssetReferenceFilterContext) : nullptr;
+		TSharedPtr<IAssetReferenceFilter> AssetReferenceFilter = GEditor ? GEditor->MakeAssetReferenceFilter(AssetReferenceFilterContext) : nullptr;
 		if (AssetReferenceFilter.IsValid())
 		{
 			FOnShouldFilterAsset ConfigFilter = InArgs._AssetPickerConfig.OnShouldFilterAsset;
diff --git a/Engine/Source/Editor/GraphEditor/Private/SGraphPanel.cpp b/Engine/Source/Editor/GraphEditor/Private/SGraphPanel.cpp
--- a/Engine/Source/Editor/GraphEditor/Private/SGraphPanel.cpp
+++ b/Engine/Source/Editor/GraphEditor/Private/SGraphPanel.cpp
@@ -1005,7 +1005,7 @@ FReply SGraphPanel::OnDrop( const FGeometry& MyGeometry, const FDragDropEvent& D
 
 bool SGraphPanel::PassesAssetReferenceFilter(const TArray<FAssetData>& ReferencedAssets, FText* OutFailureReason) const
 {
-	if (GUnrealEd)
+	if (GEditor)
 	{
 		FAssetReferenceFilterContext AssetReferenceFilterContext;
 		UObject* GraphOuter = GraphObj ? GraphObj->GetOuter() : nullptr;
@@ -1013,7 +1013,7 @@ bool SGraphPanel::PassesAssetReferenceFilter(const TArray<FAssetData>& Reference
 		{
 			AssetReferenceFilterContext.ReferencingAssets.Add(FAssetData(GraphOuter));
 		}
-		TSharedPtr<IAssetReferenceFilter> AssetReferenceFilter = GUnrealEd->MakeAssetReferenceFilter(AssetReferenceFilterContext);
+		TSharedPtr<IAssetReferenceFilter> AssetReferenceFilter = GEditor->MakeAssetReferenceFilter(AssetReferenceFilterContext);
 		if (AssetReferenceFilter.IsValid())
 		{
 			for (const FAssetData& Asset : ReferencedAssets)
diff --git a/Engine/Source/Editor/PropertyEditor/Private/UserInterface/PropertyEditor/SPropertyEditorAsset.cpp b/Engine/Source/Editor/PropertyEditor/Private/UserInterface/PropertyEditor/SPropertyEditorAsset.cpp
--- a/Engine/Source/Editor/PropertyEditor/Private/UserInterface/PropertyEditor/SPropertyEditorAsset.cpp
+++ b/Engine/Source/Editor/PropertyEditor/Private/UserInterface/PropertyEditor/SPropertyEditorAsset.cpp
@@ -1161,7 +1161,7 @@ void SPropertyEditorAsset::OnUse()
 		&& !OnShouldFilterAsset.IsBound()
 		&& AllowedClassFilters.Num() == 0
 		&& DisallowedClassFilters.Num() == 0
-		&& (GUnrealEd ? !GUnrealEd->MakeAssetReferenceFilter(FAssetReferenceFilterContext()).IsValid() : true))
+		&& (GEditor ? !GEditor->MakeAssetReferenceFilter(FAssetReferenceFilterContext()).IsValid() : true))
 	{
 		PropertyEditor->GetPropertyHandle()->SetObjectValueFromSelection();
 	}
@@ -1364,7 +1364,7 @@ bool SPropertyEditorAsset::IsClassAllowed(const UClass* InClass) const
 
 bool SPropertyEditorAsset::CanSetBasedOnAssetReferenceFilter( const FAssetData& InAssetData, FText* OutOptionalFailureReason) const
 {
-	if (GUnrealEd && InAssetData.IsValid())
+	if (GEditor && InAssetData.IsValid())
 	{
 		TSharedPtr<IPropertyHandle> PropertyHandleToUse = GetMostSpecificPropertyHandle();
 		FAssetReferenceFilterContext AssetReferenceFilterContext;
@@ -1390,7 +1390,7 @@ bool SPropertyEditorAsset::CanSetBasedOnAssetReferenceFilter( const FAssetData&
 			}
 		}
 
-		TSharedPtr<IAssetReferenceFilter> AssetReferenceFilter = GUnrealEd->MakeAssetReferenceFilter(AssetReferenceFilterContext);
+		TSharedPtr<IAssetReferenceFilter> AssetReferenceFilter = GEditor->MakeAssetReferenceFilter(AssetReferenceFilterContext);
 		if (AssetReferenceFilter.IsValid() && !AssetReferenceFilter->PassesFilter(InAssetData, OutOptionalFailureReason))
 		{
 			return false;
diff --git a/Engine/Source/Editor/UnrealEd/Classes/Editor/EditorEngine.h b/Engine/Source/Editor/UnrealEd/Classes/Editor/EditorEngine.h
--- a/Engine/Source/Editor/UnrealEd/Classes/Editor/EditorEngine.h
+++ b/Engine/Source/Editor/UnrealEd/Classes/Editor/EditorEngine.h
@@ -411,6 +411,21 @@ struct FPreviewPlatformInfo
 
 };
 
+/** Struct used in filtering allowed references between assets. Passes context about the referencers to game-level filters */
+struct FAssetReferenceFilterContext
+{
+	TArray<FAssetData> ReferencingAssets;
+};
+
+/** Used in filtering allowed references between assets. Implement a subclass of this and return it in OnMakeAssetReferenceFilter */
+class IAssetReferenceFilter
+{
+public:
+	virtual ~IAssetReferenceFilter() { }
+	/** Filter function to pass/fail an asset. Called in some situations that are performance-sensitive so is expected to be fast. */
+	virtual bool PassesFilter(const FAssetData& AssetData, FText* OutOptionalFailureReason = nullptr) const = 0;
+};
+
 /**
  * Engine that drives the Editor.
  * Separate from UGameEngine because it may have much different functionality than desired for an instance of a game itself.
@@ -2518,6 +2533,18 @@ public:
 	*/
 	void GetReferencedAssetsForEditorSelection(TArray<UObject*>& Objects, const bool bIgnoreOtherAssetsIfBPReferenced = false);
 
+	/** Returns a filter to restruct what assets show up in asset pickers based on what the selection is used for (i.e. what will reference the assets) */
+	DECLARE_DELEGATE_RetVal_OneParam(TSharedPtr<IAssetReferenceFilter>, FOnMakeAssetReferenceFilter, const FAssetReferenceFilterContext& /*Context*/);
+	FOnMakeAssetReferenceFilter& OnMakeAssetReferenceFilter() { return OnMakeAssetReferenceFilterDelegate; }
+	TSharedPtr<IAssetReferenceFilter> MakeAssetReferenceFilter(const FAssetReferenceFilterContext& Context) { return OnMakeAssetReferenceFilterDelegate.IsBound() ? OnMakeAssetReferenceFilterDelegate.Execute(Context) : nullptr; }
+
+protected:
+
+	/** Returns a filter to restruct what assets show up in asset pickers based on what the selection is used for (i.e. what will reference the assets) */
+	FOnMakeAssetReferenceFilter OnMakeAssetReferenceFilterDelegate;
+
+public:
+
 	/** Returns the WorldContext for the editor world. For now, there will always be exactly 1 of these in the editor. 
 	 *
 	 * @param	bEnsureIsGWorld		Temporarily allow callers to validate that this maps to GWorld to help track down any conversion bugs
diff --git a/Engine/Source/Editor/UnrealEd/Classes/Editor/UnrealEdEngine.h b/Engine/Source/Editor/UnrealEd/Classes/Editor/UnrealEdEngine.h
--- a/Engine/Source/Editor/UnrealEd/Classes/Editor/UnrealEdEngine.h
+++ b/Engine/Source/Editor/UnrealEd/Classes/Editor/UnrealEdEngine.h
@@ -111,21 +111,6 @@ struct FTemplateMapInfo
 class FPerformanceMonitor;
 
 
-/** Struct used in filtering allowed references between assets. Passes context about the referencers to game-level filters */
-struct FAssetReferenceFilterContext
-{
-	TArray<FAssetData> ReferencingAssets;
-};
-
-/** Used in filtering allowed references between assets. Implement a subclass of this and return it in OnMakeAssetReferenceFilter */
-class IAssetReferenceFilter
-{
-public:
-	virtual ~IAssetReferenceFilter() { }
-	/** Filter function to pass/fail an asset. Called in some situations that are performance-sensitive so is expected to be fast. */
-	virtual bool PassesFilter(const FAssetData& AssetData, FText* OutOptionalFailureReason = nullptr) const = 0;
-};
-
 UCLASS(config=Engine, transient)
 class UNREALED_API UUnrealEdEngine : public UEditorEngine, public FNotifyHook
 {
@@ -836,16 +821,8 @@ public:
 
 	bool IsComponentSelected(const UPrimitiveComponent* PrimComponent);
 
-	/** Returns a filter to restruct what assets show up in asset pickers based on what the selection is used for (i.e. what will reference the assets) */
-	DECLARE_DELEGATE_RetVal_OneParam(TSharedPtr<IAssetReferenceFilter>, FOnMakeAssetReferenceFilter, const FAssetReferenceFilterContext& /*Context*/);
-	FOnMakeAssetReferenceFilter& OnMakeAssetReferenceFilter() { return OnMakeAssetReferenceFilterDelegate; }
-	TSharedPtr<IAssetReferenceFilter> MakeAssetReferenceFilter(const FAssetReferenceFilterContext& Context) { return OnMakeAssetReferenceFilterDelegate.IsBound() ? OnMakeAssetReferenceFilterDelegate.Execute(Context) : nullptr; }
-	
 protected:
 
-	/** Returns a filter to restruct what assets show up in asset pickers based on what the selection is used for (i.e. what will reference the assets) */
-	FOnMakeAssetReferenceFilter OnMakeAssetReferenceFilterDelegate;
-
 	/** Called when global editor selection changes */
 	void OnEditorSelectionChanged(UObject* SelectionThatChanged);
 
