Fixing crash when quickly clicking an actor while the level is transitioning

diff --git a/Engine/Source/Editor/UnrealEd/Private/LevelEditorViewport.cpp b/Engine/Source/Editor/UnrealEd/Private/LevelEditorViewport.cpp
--- a/Engine/Source/Editor/UnrealEd/Private/LevelEditorViewport.cpp
+++ b/Engine/Source/Editor/UnrealEd/Private/LevelEditorViewport.cpp
@@ -2178,35 +2178,38 @@ void FLevelEditorViewportClient::ProcessClick(FSceneView& View, HHitProxy* HitPr
 		{
 			HActor* ActorHitProxy = (HActor*)HitProxy;
 			AActor* ConsideredActor = ActorHitProxy->Actor;
-			while (ConsideredActor->IsChildActor())
+			if (ConsideredActor) // It is possible to be clicking something during level transition if you spam click, and it might not be valid by this point
 			{
-				ConsideredActor = ConsideredActor->GetParentActor();
-			}
+				while (ConsideredActor->IsChildActor())
+				{
+					ConsideredActor = ConsideredActor->GetParentActor();
+				}
 
-			// We want to process the click on the component only if:
-			// 1. The actor clicked is already selected
-			// 2. The actor selected is the only actor selected
-			// 3. The actor selected is blueprintable
-			// 4. No components are already selected and the click was a double click
-			// 5. OR, a component is already selected and the click was NOT a double click
-			const bool bActorAlreadySelectedExclusively = GEditor->GetSelectedActors()->IsSelected(ConsideredActor) && ( GEditor->GetSelectedActorCount() == 1 );
-			const bool bActorIsBlueprintable = FKismetEditorUtilities::CanCreateBlueprintOfClass(ConsideredActor->GetClass());
-			const bool bComponentAlreadySelected = GEditor->GetSelectedComponentCount() > 0;
-			const bool bWasDoubleClick = ( Click.GetEvent() == IE_DoubleClick );
+				// We want to process the click on the component only if:
+				// 1. The actor clicked is already selected
+				// 2. The actor selected is the only actor selected
+				// 3. The actor selected is blueprintable
+				// 4. No components are already selected and the click was a double click
+				// 5. OR, a component is already selected and the click was NOT a double click
+				const bool bActorAlreadySelectedExclusively = GEditor->GetSelectedActors()->IsSelected(ConsideredActor) && (GEditor->GetSelectedActorCount() == 1);
+				const bool bActorIsBlueprintable = FKismetEditorUtilities::CanCreateBlueprintOfClass(ConsideredActor->GetClass());
+				const bool bComponentAlreadySelected = GEditor->GetSelectedComponentCount() > 0;
+				const bool bWasDoubleClick = (Click.GetEvent() == IE_DoubleClick);
 
-			const bool bSelectComponent = bActorAlreadySelectedExclusively && bActorIsBlueprintable && (bComponentAlreadySelected != bWasDoubleClick);
+				const bool bSelectComponent = bActorAlreadySelectedExclusively && bActorIsBlueprintable && (bComponentAlreadySelected != bWasDoubleClick);
 
-			if (bSelectComponent)
-			{
-				LevelViewportClickHandlers::ClickComponent(this, ActorHitProxy, Click);
-			}
-			else
-			{
-				LevelViewportClickHandlers::ClickActor(this, ConsideredActor, Click, true);
-			}
+				if (bSelectComponent)
+				{
+					LevelViewportClickHandlers::ClickComponent(this, ActorHitProxy, Click);
+				}
+				else
+				{
+					LevelViewportClickHandlers::ClickActor(this, ConsideredActor, Click, true);
+				}
 
-			// We clicked an actor, allow the pivot to reposition itself.
-			// GUnrealEd->SetPivotMovedIndependently(false);
+				// We clicked an actor, allow the pivot to reposition itself.
+				// GUnrealEd->SetPivotMovedIndependently(false);
+			}
 		}
 		else if (HitProxy->IsA(HInstancedStaticMeshInstance::StaticGetType()))
 		{
