Fix for undoing a parameter edit condition not working after undoing a parameter edit

diff --git a/Engine/Source/Editor/UnrealEd/Private/PreviewMaterial.cpp b/Engine/Source/Editor/UnrealEd/Private/PreviewMaterial.cpp
--- a/Engine/Source/Editor/UnrealEd/Private/PreviewMaterial.cpp
+++ b/Engine/Source/Editor/UnrealEd/Private/PreviewMaterial.cpp
@@ -1693,9 +1693,6 @@ void UMaterialEditorInstanceConstant::PostEditUndo()
 		UpdateSourceInstanceParent();
 
 		Context.AddMaterialInstance(SourceInstance);
-
-		// Fully update static parameters before recreating render state for all components
-		SetSourceInstance(SourceInstance);
 	}
 }
 #endif
