[UE-86647] AsyncIODeleteTest - add debug output to remote debug a test failing on Mac.
#rnx

diff --git a/Engine/Source/Editor/UnrealEd/Private/Cooker/AsyncIODelete.cpp b/Engine/Source/Editor/UnrealEd/Private/Cooker/AsyncIODelete.cpp
--- a/Engine/Source/Editor/UnrealEd/Private/Cooker/AsyncIODelete.cpp
+++ b/Engine/Source/Editor/UnrealEd/Private/Cooker/AsyncIODelete.cpp
@@ -143,15 +143,26 @@ void FAsyncIODelete::Teardown()
 	bInitialized = false;
 }
 
-void FAsyncIODelete::WaitForAllTasks()
+bool FAsyncIODelete::WaitForAllTasks(float TimeLimitSeconds)
 {
 	if (!bInitialized)
 	{
-		return;
+		return true;
 	}
 
-	TasksComplete->Wait();
+	if (TimeLimitSeconds <= 0.f)
+	{
+		TasksComplete->Wait();
+	}
+	else
+	{
+		if (!TasksComplete->Wait(FTimespan::FromSeconds(TimeLimitSeconds)))
+		{
+			return false;
+		}
+	}
 	check(ActiveTaskCount == 0);
+	return true;
 }
 
 bool FAsyncIODelete::Delete(const FStringView& PathToDelete, EPathType ExpectedType)
diff --git a/Engine/Source/Editor/UnrealEd/Private/Cooker/AsyncIODelete.h b/Engine/Source/Editor/UnrealEd/Private/Cooker/AsyncIODelete.h
--- a/Engine/Source/Editor/UnrealEd/Private/Cooker/AsyncIODelete.h
+++ b/Engine/Source/Editor/UnrealEd/Private/Cooker/AsyncIODelete.h
@@ -90,7 +90,7 @@ public:
 	}
 
 	/** Synchronously wait for all deletes to complete.  */
-	void WaitForAllTasks();
+	bool WaitForAllTasks(float TimeLimitSeconds = 0.0f);
 
 private:
 	/** Update the ActiveTaskCount and events for a task completing */
diff --git a/Engine/Source/Editor/UnrealEd/Private/Cooker/AsyncIODeleteTest.cpp b/Engine/Source/Editor/UnrealEd/Private/Cooker/AsyncIODeleteTest.cpp
--- a/Engine/Source/Editor/UnrealEd/Private/Cooker/AsyncIODeleteTest.cpp
+++ b/Engine/Source/Editor/UnrealEd/Private/Cooker/AsyncIODeleteTest.cpp
@@ -55,53 +55,67 @@ bool FAsyncIODeleteTest::RunTest(const FString& Parameters)
 
 	FileManager.MakeDirectory(*TestRoot);
 
+	const bool bVerbose = true;
+	auto StartSection = [bVerbose](const TCHAR* Section)
+	{
+		if (bVerbose)
+		{
+			UE_LOG(LogCore, Display, TEXT("%s"), Section);
+		}
+	};
+	const float MaxWaitTime = 5.0f;
+	auto WaitForAllTasksAndVerify = [this, MaxWaitTime](FAsyncIODelete& InAsyncIODelete)
+	{
+		bool WaitResult = InAsyncIODelete.WaitForAllTasks(MaxWaitTime);
+		TestTrue(TEXT("WaitForAllTasks timed out"), WaitResult);
+	};
 	{
 		FAsyncIODelete AsyncIODelete(TempRoot);
 
-		// Waiting for tasks to complete when none have been launched should succeed
-		AsyncIODelete.WaitForAllTasks();
+		StartSection(TEXT("Waiting for tasks to complete when none have been launched should succeed"));
+		WaitForAllTasksAndVerify(AsyncIODelete);
 
-		// Moving file and directory from the source location should be finished by the time DeleteFile/DeleteDirectory returns
+		StartSection(TEXT("Moving file and directory from the source location should be finished by the time DeleteFile/DeleteDirectory returns"));
 		CreateTestPathsToDelete();
 		AsyncIODelete.DeleteFile(TestFile1);
 		AsyncIODelete.DeleteDirectory(TestDir1);
 		TestRequestedPathsDeleted(TEXT("AsyncIODelete::Delete should have moved the deleted paths before returning."));
 
-		// Deleting the temporary files/directories should be finished before WaitForAllTasks returns
-		AsyncIODelete.WaitForAllTasks();
+		StartSection(TEXT("Deleting the temporary files/directories should be finished before WaitForAllTasks returns"));
+		WaitForAllTasksAndVerify(AsyncIODelete);
 		TestTempRootCountsEqual(TempRoot, 0, 0, TEXT("AsyncIODelete should have deleted the moved paths before WaitForAllTasks returned."));
 
-		// Two FAsyncIODelete constructed at once should be legal, as long as they have different TempRoots
+		StartSection(TEXT("Two FAsyncIODelete constructed at once should be legal, as long as they have different TempRoots"));
 		FAsyncIODelete AsyncIODelete2(TempRoot2);
 
-		// Use the pause feature to verify that the paths are indeed moved into the TempRoot
+		StartSection(TEXT("Use the pause feature to verify that the paths are indeed moved into the TempRoot"));
 		AsyncIODelete2.SetDeletesPaused(true);
 		CreateTestPathsToDelete();
 		AsyncIODelete2.DeleteFile(TestFile1);
 		AsyncIODelete2.DeleteDirectory(TestDir1);
 		TestRequestedPathsDeleted(TEXT("AsyncIODelete::Delete should have moved the deleted paths before returning even when paused."));
-		AsyncIODelete2.WaitForAllTasks();
+		WaitForAllTasksAndVerify(AsyncIODelete2);
 		TestTempRootCountsEqual(TempRoot2, 1, 1, TEXT("AsyncIODelete should not have deleted the moved paths because it is paused."));
 		AsyncIODelete2.SetDeletesPaused(false);
-		AsyncIODelete2.WaitForAllTasks();
+		WaitForAllTasksAndVerify(AsyncIODelete2);
 		TestTempRootCountsEqual(TempRoot2, 0, 0, TEXT("AsyncIODelete should have deleted the moved paths after unpausing."));
 
-		// Verify Teardown() deletes the TempRoot and Setup() creates it
+		StartSection(TEXT("Verify Teardown() deletes the TempRoot and Setup() creates it"));
 		AsyncIODelete2.Teardown();
 		TestTrue(TEXT("AsyncIODelete::Teardown should have deleted its TempRoot."), !FileManager.DirectoryExists(*TempRoot2));
 		AsyncIODelete2.Setup();
 		TestTrue(TEXT("AsyncIODelete::Setup should have created its TempRoot."), FileManager.DirectoryExists(*TempRoot2));
 
-		// Manual setup works as long as you call SetTempRoot before Setup
+		StartSection(TEXT("Manual setup works as long as you call SetTempRoot before Setup"));
 		FAsyncIODelete AsyncIODelete3;
 		AsyncIODelete3.SetTempRoot(TempRoot3);
 		AsyncIODelete3.Setup();
 		TestTrue(TEXT("Setup should have created the TempRoot."), FileManager.DirectoryExists(*TempRoot3));
 
-		// Check that even after Setup, waiting for tasks to complete when none have been launched should succeed
-		AsyncIODelete.WaitForAllTasks();
+		StartSection(TEXT("Check that even after Setup, waiting for tasks to complete when none have been launched should succeed"));
+		WaitForAllTasksAndVerify(AsyncIODelete);
 
-		// Changing TempRoot and then deleting a file works
+		StartSection(TEXT("Changing TempRoot and then deleting a file works"));
 		CreateTestPathsToDelete();
 		AsyncIODelete3.DeleteFile(TestFile1);
 		AsyncIODelete3.DeleteDirectory(TestDir1);
@@ -112,10 +126,10 @@ bool FAsyncIODeleteTest::RunTest(const FString& Parameters)
 		AsyncIODelete3.DeleteDirectory(TestDir1);
 		TestRequestedPathsDeleted(TEXT("AsyncIODelete::Delete should have worked after changing the TempRoot."));
 		TestTrue(TEXT("Delete should have created the new TempRoot after SetTempRoot."), FileManager.DirectoryExists(*TempRoot4));
-		AsyncIODelete3.WaitForAllTasks();
+		WaitForAllTasksAndVerify(AsyncIODelete3);
 		TestTempRootCountsEqual(TempRoot4, 0, 0, TEXT("AsyncIODelete::Delete should have created the tasks to delete moved files after changing the TempRoot."));
 
-		// Attempting to delete a parent directory of the temproot, the temproot itself, or a child inside of it fails
+		StartSection(TEXT("Attempting to delete a parent directory of the temproot, the temproot itself, or a child inside of it fails"));
 		FString SubDirInTempRoot4 = FPaths::Combine(TempRoot4, TEXT("SubDir"));
 		FileManager.MakeDirectory(*SubDirInTempRoot4, bApplyToTreeTrue); // Note it's illegal to add files into TempRoot, but we're not currently checking for it and we're not colliding with the DeleteN paths AsyncIODelete uses, so this breaking of the rule will not cause problems
 		TestFalse(TEXT("AsyncIODelete should refuse to delete a parent of its TempRoot."), AsyncIODelete3.DeleteDirectory(TestRoot));
