Re-add OnBeginPlay after both calls removed via two seperate fixes and merge.

diff --git a/Engine/Source/Editor/UnrealEd/Private/PlayLevel.cpp b/Engine/Source/Editor/UnrealEd/Private/PlayLevel.cpp
--- a/Engine/Source/Editor/UnrealEd/Private/PlayLevel.cpp
+++ b/Engine/Source/Editor/UnrealEd/Private/PlayLevel.cpp
@@ -3269,6 +3269,11 @@ UGameInstance* UEditorEngine::CreatePIEGameInstance(int32 InPIEInstance, bool bI
 	UGameViewportClient* ViewportClient = NULL;
 	ULocalPlayer *NewLocalPlayer = NULL;
 
+	if (GEngine->XRSystem.IsValid() && !bInSimulateInEditor)
+	{
+		GEngine->XRSystem->OnBeginPlay(*PieWorldContext);
+	}
+
 	if (!PieWorldContext->RunAsDedicated)
 	{
 		bool bCreateNewAudioDevice = PlayInSettings->IsCreateAudioDeviceForEveryPlayer();
