HLOD rebuild commandlet - add levels packages to be saved

diff --git a/Engine/Source/Editor/UnrealEd/Private/Commandlets/ContentCommandlets.cpp b/Engine/Source/Editor/UnrealEd/Private/Commandlets/ContentCommandlets.cpp
--- a/Engine/Source/Editor/UnrealEd/Private/Commandlets/ContentCommandlets.cpp
+++ b/Engine/Source/Editor/UnrealEd/Private/Commandlets/ContentCommandlets.cpp
@@ -1675,6 +1675,8 @@ void UResavePackagesCommandlet::PerformAdditionalOperations(class UWorld* World,
 				TSet<UPackage*> PackagesToSave;
 				for(ULevel* Level : World->GetLevels())
 				{
+					PackagesToSave.Add(Level->GetOutermost());
+
 					if(Level->bIsVisible)
 					{
 						Builder.GetMeshesPackagesToSave(Level, PackagesToSave);
