Fix cooker to filter out plugin content for unsupported target platforms

diff --git a/Engine/Source/Editor/UnrealEd/Classes/CookOnTheSide/CookOnTheFlyServer.h b/Engine/Source/Editor/UnrealEd/Classes/CookOnTheSide/CookOnTheFlyServer.h
--- a/Engine/Source/Editor/UnrealEd/Classes/CookOnTheSide/CookOnTheFlyServer.h
+++ b/Engine/Source/Editor/UnrealEd/Classes/CookOnTheSide/CookOnTheFlyServer.h
@@ -684,6 +684,13 @@ private:
 	*/
 	void InitializeTargetPlatforms();
 
+	/**
+	* Some content plugins does not support all target platforms.
+	* Build up a map of unsupported packages per platform that can be checked before saving.
+	*/
+	void DiscoverPlatformSpecificNeverCookPackages(
+		const TArray<FName>& TargetPlatformNames, const TArray<FString>& TargetPlatformStrings);
+
 	/**
 	* Clean up the sandbox
 	*/
diff --git a/Engine/Source/Editor/UnrealEd/Private/CookOnTheFlyServer.cpp b/Engine/Source/Editor/UnrealEd/Private/CookOnTheFlyServer.cpp
--- a/Engine/Source/Editor/UnrealEd/Private/CookOnTheFlyServer.cpp
+++ b/Engine/Source/Editor/UnrealEd/Private/CookOnTheFlyServer.cpp
@@ -1271,6 +1271,7 @@ struct FPackageTracker : public FUObjectArray::FUObjectCreateListener, public FU
 
 	FThreadSafeSet<FName>						NeverCookPackageList;
 	FThreadSafeSet<FName>						UncookedEditorOnlyPackages; // set of packages that have been rejected due to being referenced by editor-only properties
+	TMap<FName, TSet<FName>> 					PlatformSpecificNeverCookPackages;
 
 	// Currently targeted platforms
 	TArray<FName>			AllTargetPlatformNames;
@@ -3910,14 +3911,24 @@ void UCookOnTheFlyServer::SaveCookedPackage(UPackage* Package, uint32 SaveFlags,
 				Result = ESavePackageResult::ContainsEditorOnlyData;
 				bCookPackage = false;
 			}
-
 			// Check whether or not game-specific behaviour should prevent this package from being cooked for the target platform
-			if (UAssetManager::IsValid() && !UAssetManager::Get().ShouldCookForPlatform(Package, Target))
+			else if (UAssetManager::IsValid() && !UAssetManager::Get().ShouldCookForPlatform(Package, Target))
 			{
 				Result = ESavePackageResult::ContainsEditorOnlyData;
 				bCookPackage = false;
 				UE_LOG(LogCook, Display, TEXT("Excluding %s -> %s"), *Package->GetName(), *PlatFilename);
 			}
+			// check if this package is unsupported for the target platform (typically plugin content)
+			else 
+			{
+				TSet<FName>* NeverCookPackages = PackageTracker->PlatformSpecificNeverCookPackages.Find(FName(*Target->PlatformName()));
+				if (NeverCookPackages && NeverCookPackages->Find(FName(*PackagePathName)))
+				{
+					Result = ESavePackageResult::ContainsEditorOnlyData;
+					bCookPackage = false;
+					UE_LOG(LogCook, Display, TEXT("Excluding %s -> %s"), *Package->GetName(), *PlatFilename);
+				}
+			}
 
 			if (bCookPackage == true)
 			{
@@ -6665,6 +6676,62 @@ void UCookOnTheFlyServer::InitializeTargetPlatforms()
 	}
 }
 
+void UCookOnTheFlyServer::DiscoverPlatformSpecificNeverCookPackages(
+	const TArray<FName>& TargetPlatformNames, const TArray<FString>& TargetPlatformStrings)
+{
+	TArray<FName> PluginUnsupportedTargetPlatforms;
+	TArray<FAssetData> PluginAssets;
+	FARFilter PluginARFilter;
+	FString PluginPackagePath;
+
+	TArray<TSharedRef<IPlugin>> AllContentPlugins = IPluginManager::Get().GetEnabledPluginsWithContent();
+	for (TSharedRef<IPlugin> Plugin : AllContentPlugins)
+	{
+		const FPluginDescriptor& Descriptor = Plugin->GetDescriptor();
+
+		// we are only interested in plugins that does not support all platforms
+		if (Descriptor.SupportedTargetPlatforms.Num() == 0)
+		{
+			continue;
+		}
+
+		// find any unsupported target platforms for this plugin
+		PluginUnsupportedTargetPlatforms.Reset();
+		for (int32 I = 0, Count = TargetPlatformNames.Num(); I < Count; ++I)
+		{
+			if (!Descriptor.SupportedTargetPlatforms.Contains(TargetPlatformStrings[I]))
+			{
+				PluginUnsupportedTargetPlatforms.Add(TargetPlatformNames[I]);
+			}
+		}
+
+		// if there are unsupported target platforms,
+		// then add all packages for this plugin for these platforms to the PlatformSpecificNeverCookPackages map
+		if (PluginUnsupportedTargetPlatforms.Num() > 0)
+		{
+			PluginPackagePath.Reset(127);
+			PluginPackagePath.AppendChar(TEXT('/'));
+			PluginPackagePath.Append(Plugin->GetName());
+
+			PluginARFilter.bRecursivePaths = true;
+			PluginARFilter.bIncludeOnlyOnDiskAssets = true;
+			PluginARFilter.PackagePaths.Reset(1);
+			PluginARFilter.PackagePaths.Emplace(*PluginPackagePath);
+
+			PluginAssets.Reset();
+			AssetRegistry->GetAssets(PluginARFilter, PluginAssets);
+
+			for (FName& PlatformName : PluginUnsupportedTargetPlatforms)
+			{
+				TSet<FName>& NeverCookPackages = PackageTracker->PlatformSpecificNeverCookPackages.FindOrAdd(PlatformName);
+				for (const FAssetData& Asset : PluginAssets)
+				{
+					NeverCookPackages.Add(Asset.PackageName);
+				}
+			}
+		}
+	}
+}
 
 void UCookOnTheFlyServer::TermSandbox()
 {
@@ -6815,12 +6882,20 @@ void UCookOnTheFlyServer::StartCookByTheBook( const FCookByTheBookStartupOptions
 
 	}
 	
-
-	CookByTheBookOptions->TargetPlatformNames.Empty();
-	for (const ITargetPlatform* Platform : TargetPlatforms)
+	// build persistent list of all target platform names in CookByTheBookOptions->TargetPlatformNames,
+	// also use temp list of platform strings to discover PlatformSpecificNeverCookPackages
 	{
-		FName PlatformName = FName(*Platform->PlatformName());
-		CookByTheBookOptions->TargetPlatformNames.Add(PlatformName); // build list of all target platform names
+		TArray<FString> TargetPlatformStrings;
+		TargetPlatformStrings.Reserve(TargetPlatforms.Num());
+
+		CookByTheBookOptions->TargetPlatformNames.Empty(TargetPlatforms.Num());
+		for (const ITargetPlatform* Platform : TargetPlatforms)
+		{
+			FString& PlatformString = TargetPlatformStrings.Emplace_GetRef(Platform->PlatformName());
+			CookByTheBookOptions->TargetPlatformNames.Emplace(*PlatformString);
+		}
+
+		DiscoverPlatformSpecificNeverCookPackages(CookByTheBookOptions->TargetPlatformNames, TargetPlatformStrings);
 	}
 	const TArray<FName>& TargetPlatformNames = CookByTheBookOptions->TargetPlatformNames;
 
