TBA: Improvements and fixes for textasset commandlet

diff --git a/Engine/Source/Editor/UnrealEd/Private/Commandlets/TextAssetCommandlet.cpp b/Engine/Source/Editor/UnrealEd/Private/Commandlets/TextAssetCommandlet.cpp
--- a/Engine/Source/Editor/UnrealEd/Private/Commandlets/TextAssetCommandlet.cpp
+++ b/Engine/Source/Editor/UnrealEd/Private/Commandlets/TextAssetCommandlet.cpp
@@ -20,6 +20,7 @@
 #include "Serialization/Formatters/JsonArchiveOutputFormatter.h"
 #include "Serialization/MemoryWriter.h"
 #include "Serialization/ArchiveUObjectFromStructuredArchive.h"
+#include "ProfilingDebugging/CpuProfilerTrace.h"
 
 DEFINE_LOG_CATEGORY(LogTextAsset);
 
@@ -51,12 +52,48 @@ void FindMismatchedSerializers()
 
 int32 UTextAssetCommandlet::Main(const FString& CmdLineParams)
 {
+	TRACE_CPUPROFILER_EVENT_SCOPE_TEXT(TEXT("UTextAssetCommandlet::Main"));
+
 	TArray<FString> Blacklist;
 
+	static const FString BackupExtension = TEXT("textassetbackup");
+	static const FString BackupRoundtripExtension = TEXT("textassetbackup_roundtrip");
+	static const FString BackupExtension_WithDot = TEXT(".") + BackupExtension;
+	static const FString BackupRoundtripExtension_WithDot = TEXT(".") + BackupRoundtripExtension;
+
+	// Repair any damage caused by a failed run of this commandlet
+	TArray<FString> IntermediateFilenames;
+	struct FVisitor : public IPlatformFile::FDirectoryVisitor
+	{
+		virtual bool Visit(const TCHAR* FilenameOrDirectory, bool bIsDirectory) override
+		{
+			FString Extension = FPaths::GetExtension(FilenameOrDirectory);
+			if (!bIsDirectory && (Extension == BackupExtension))
+			{
+				UE_LOG(LogTextAsset, Display, TEXT("Cleaning up old intermediate file %s"), FilenameOrDirectory);
+
+				FString BinaryFilename = FPaths::GetPath(FilenameOrDirectory) / FPaths::GetBaseFilename(FilenameOrDirectory);
+				FString TextFilename = FPaths::ChangeExtension(BinaryFilename, FPackageName::GetTextAssetPackageExtension());
+				FString RoundtripBackup = BinaryFilename + BackupRoundtripExtension_WithDot;
+
+				IFileManager::Get().Delete(*BinaryFilename);
+				IFileManager::Get().Delete(*TextFilename);
+				IFileManager::Get().Delete(*RoundtripBackup);
+
+				IFileManager::Get().Move(*BinaryFilename, FilenameOrDirectory);
+			}
+			return true;
+		}
+	} RepairVisitor;
+
+	IFileManager::Get().IterateDirectoryRecursively(*FPaths::ProjectContentDir(), RepairVisitor);
+	IFileManager::Get().IterateDirectoryRecursively(*FPaths::EngineContentDir(), RepairVisitor);
+
 	FString ModeString = TEXT("ResaveText");
 	FString IterationsString = TEXT("1");
-	FString FilenameFilterString, OutputPathString;
+	FString FilenameFilterString, OutputPathString, CSVFilename;
 
+	FParse::Value(*CmdLineParams, TEXT("csv="), CSVFilename);
 	FParse::Value(*CmdLineParams, TEXT("mode="), ModeString);
 	FParse::Value(*CmdLineParams, TEXT("filter="), FilenameFilterString);
 	FParse::Value(*CmdLineParams, TEXT("outputpath="), OutputPathString);
@@ -67,8 +104,9 @@ int32 UTextAssetCommandlet::Main(const FString& CmdLineParams)
 		ResaveText,
 		ResaveBinary,
 		RoundTrip,
+		LoadBinary,
 		LoadText,
-		FindMismatchedSerializers
+		FindMismatchedSerializers,
 	};
 
 	TMap<FString, EMode> Modes;
@@ -76,6 +114,7 @@ int32 UTextAssetCommandlet::Main(const FString& CmdLineParams)
 	Modes.Add(TEXT("ResaveBinary"), EMode::ResaveBinary);
 	Modes.Add(TEXT("RoundTrip"), EMode::RoundTrip);
 	Modes.Add(TEXT("LoadText"), EMode::LoadText);
+	Modes.Add(TEXT("LoadBinary"), EMode::LoadBinary);
 	Modes.Add(TEXT("FindMismatchedSerializers"), EMode::FindMismatchedSerializers);
 
 	check(Modes.Contains(ModeString));
@@ -124,6 +163,13 @@ int32 UTextAssetCommandlet::Main(const FString& CmdLineParams)
 		//IFileManager::Get().FindFilesRecursive(InputAssetFilenames, *BasePath, *(Wildcard + FPackageName::GetTextMapPackageExtension()), true, false, false);
 		break;
 	}
+
+	case EMode::LoadBinary:
+	{
+		IFileManager::Get().FindFilesRecursive(InputAssetFilenames, *ProjectContentDir, *(Wildcard + FPackageName::GetAssetPackageExtension()), true, false, true);
+		//IFileManager::Get().FindFilesRecursive(InputAssetFilenames, *BasePath, *(Wildcard + FPackageName::GetTextMapPackageExtension()), true, false, false);
+		break;
+	}
 	}
 
 	TArray<TTuple<FString, FString>> FilesToProcess;
@@ -173,6 +219,7 @@ int32 UTextAssetCommandlet::Main(const FString& CmdLineParams)
 		}
 
 		case EMode::LoadText:
+		case EMode::LoadBinary:
 		{
 			break;
 		}
@@ -184,23 +231,6 @@ int32 UTextAssetCommandlet::Main(const FString& CmdLineParams)
 		}
 	}
 
-	TArray<FString> IntermediateFilenames;
-	struct FVisitor : public IPlatformFile::FDirectoryVisitor
-	{
-		virtual bool Visit(const TCHAR* FilenameOrDirectory, bool bIsDirectory) override
-		{
-			if (!bIsDirectory && FString(FilenameOrDirectory).Contains(TEXT(".txtassettemp")))
-			{
-				UE_LOG(LogTextAsset, Display, TEXT("Cleaning up old intermediate file %s"), FilenameOrDirectory);
-				IFileManager::Get().Delete(FilenameOrDirectory);
-			}
-			return true;
-		}
-	} Visitor;
-
-	IFileManager::Get().IterateDirectoryRecursively(*FPaths::ProjectContentDir(), Visitor);
-	IFileManager::Get().IterateDirectoryRecursively(*FPaths::EngineContentDir(), Visitor);
-
 	const FString FailedDiffsPath = FPaths::ProjectSavedDir() / TEXT("FailedDiffs");
 	static const bool bKeepFailedDiffs = FParse::Param(FCommandLine::Get(), TEXT("keepfaileddiffs"));
 	if (bKeepFailedDiffs)
@@ -211,6 +241,17 @@ int32 UTextAssetCommandlet::Main(const FString& CmdLineParams)
 	float TotalPackageLoadTime = 0.0;
 	float TotalPackageSaveTime = 0.0;
 
+	FArchive* CSVWriter = nullptr;
+	if (CSVFilename.Len() > 0)
+	{
+		CSVWriter = IFileManager::Get().CreateFileWriter(*CSVFilename);
+		if (CSVWriter != nullptr)
+		{
+			FString CSVLine = FString::Printf(TEXT("Total Time,Num Files,AvgFileTime,MinFileTime,MaxFileTime,TotalLoadTime\n"));
+			CSVWriter->Serialize(TCHAR_TO_ANSI(*CSVLine), CSVLine.Len());
+		}
+	}
+
 	for (int32 Iteration = 0; Iteration < NumSaveIterations; ++Iteration)
 	{
 		if (NumSaveIterations > 1)
@@ -227,6 +268,7 @@ int32 UTextAssetCommandlet::Main(const FString& CmdLineParams)
 		FString MinTimePackage;
 		float IterationPackageLoadTime = 0.0;
 		float IterationPackageSaveTime = 0.0;
+		double ThisPackageLoadTime = 0.0;
 
 		TArray<FString> PhaseSuccess;
 		TArray<TArray<FString>> PhaseFails;
@@ -238,6 +280,8 @@ int32 UTextAssetCommandlet::Main(const FString& CmdLineParams)
 			FString SourceLongPackageName = FPackageName::FilenameToLongPackageName(SourceFilename);
 			FString DestinationFilename = FileToProcess.Get<1>();
 
+			TRACE_CPUPROFILER_EVENT_SCOPE_TEXT(*SourceFilename);
+
 			IntermediateFilenames.Empty();
 
 			double StartTime = FPlatformTime::Seconds();
@@ -251,7 +295,7 @@ int32 UTextAssetCommandlet::Main(const FString& CmdLineParams)
 
 				IFileManager::Get().Delete(*WorkingFilenames[1], false, false, true);
 
-				FString SourceBackupFilename = SourceFilename + TEXT(".bak");
+				FString SourceBackupFilename = SourceFilename + BackupExtension_WithDot;
 				if (IFileManager::Get().FileExists(*SourceBackupFilename))
 				{
 					IFileManager::Get().Delete(*SourceFilename, false, false, true);
@@ -266,7 +310,7 @@ int32 UTextAssetCommandlet::Main(const FString& CmdLineParams)
 				CollectGarbage(RF_NoFlags, true);
 
 				// Make a copy of the resaved source package which we can use as the base revision for each test
-				FString BaseBinaryPackageBackup = SourceFilename + TEXT(".bak2");
+				FString BaseBinaryPackageBackup = SourceFilename + BackupRoundtripExtension_WithDot;
 				IFileManager::Get().Copy(*BaseBinaryPackageBackup, *SourceFilename, true);
 
 				FSHAHash SourceHash;
@@ -310,12 +354,25 @@ int32 UTextAssetCommandlet::Main(const FString& CmdLineParams)
 						case 1: // text only
 						{
 							Bucket = 1;
+
+							if (i > 0)
+							{
+								IFileManager::Get().Delete(*WorkingFilenames[0]);
+							}
+
 							break;
 						}
 
 						case 2: // alternate
 						{
 							Bucket = i % 2;
+
+							if ((i > 0) && Bucket == 0)
+							{
+								// We're doing alternating text/binary saves, so we need to delete the text version as we have no way of forcing the load to choose between text and binary
+								IFileManager::Get().Delete(*WorkingFilenames[0]);
+							}
+
 							break;
 						}
 
@@ -326,12 +383,6 @@ int32 UTextAssetCommandlet::Main(const FString& CmdLineParams)
 						}
 						};
 
-						if (Phase == 2 && Bucket == 1)
-						{
-							// We're doing alternating text/binary saves, so we need to delete the text version as we have no way of forcing the load to choose between text and binary
-							IFileManager::Get().Delete(*WorkingFilenames[1]);
-						}
-
 						UPackage* Package = LoadPackage(nullptr, *SourceLongPackageName, LOAD_None);
 						SavePackageHelper(Package, *WorkingFilenames[Bucket], RF_Standalone, GWarn, nullptr, SAVE_KeepGUID);
 						ResetLoaders(Package);
@@ -368,7 +419,8 @@ int32 UTextAssetCommandlet::Main(const FString& CmdLineParams)
 							bPhasesMatched[Phase] = bPhasesMatched[Phase] && Hash == Refs[0];
 						}
 
-						UE_LOG(LogTextAsset, Display, TEXT("\tPass %i [%s] %s"), Pass++, *Hash.ToString(), bPhasesMatched[Phase] ? TEXT("OK") : TEXT("FAILED"));
+						UE_LOG(LogTextAsset, Display, TEXT("\tPass %i [%s] %s"), Pass, *Hash.ToString(), bPhasesMatched[Phase] ? TEXT("OK") : TEXT("FAILED"));
+						Pass++;
 					}
 
 					if (!bPhasesMatched[Phase])
@@ -454,35 +506,40 @@ int32 UTextAssetCommandlet::Main(const FString& CmdLineParams)
 				IterationPackageLoadTime += Timer;
 				TotalPackageLoadTime += Timer;
 
+				bool bSaveSuccessful = false;
+
 				if (Package)
 				{
 					{
 						SCOPE_SECONDS_COUNTER(Timer);
 
 						IFileManager::Get().Delete(*DestinationFilename, false, true, true);
-						SavePackageHelper(Package, *DestinationFilename, RF_Standalone, GWarn, nullptr, SAVE_KeepGUID);
+						bSaveSuccessful = SavePackageHelper(Package, *DestinationFilename, RF_Standalone, GWarn, nullptr, SAVE_KeepGUID);
 					}
 					TotalPackageSaveTime += Timer;
 					IterationPackageSaveTime += Timer;
 				}
 
-				if (bVerifyJson)
+				if (bSaveSuccessful)
 				{
-					FArchive* File = IFileManager::Get().CreateFileReader(*DestinationFilename);
-					TSharedPtr< FJsonObject > RootObject;
-					TSharedRef< TJsonReader<char> > Reader = TJsonReaderFactory<char>::Create(File);
-					ensure(FJsonSerializer::Deserialize(Reader, RootObject));
-					delete File;
-				}
+					if (bVerifyJson)
+					{
+						FArchive* File = IFileManager::Get().CreateFileReader(*DestinationFilename);
+						TSharedPtr< FJsonObject > RootObject;
+						TSharedRef< TJsonReader<char> > Reader = TJsonReaderFactory<char>::Create(File);
+						ensure(FJsonSerializer::Deserialize(Reader, RootObject));
+						delete File;
+					}
 
-				if (OutputPathString.Len() > 0)
-				{
-					FString CopyFilename = DestinationFilename;
-					FPaths::MakePathRelativeTo(CopyFilename, *FPaths::RootDir());
-					CopyFilename = OutputPathString / CopyFilename;
-					CopyFilename.RemoveFromEnd(TEXT(".tmp"));
-					IFileManager::Get().MakeDirectory(*FPaths::GetPath(CopyFilename));
-					IFileManager::Get().Move(*CopyFilename, *DestinationFilename);
+					if (OutputPathString.Len() > 0)
+					{
+						FString CopyFilename = DestinationFilename;
+						FPaths::MakePathRelativeTo(CopyFilename, *FPaths::RootDir());
+						CopyFilename = OutputPathString / CopyFilename;
+						CopyFilename.RemoveFromEnd(TEXT(".tmp"));
+						IFileManager::Get().MakeDirectory(*FPaths::GetPath(CopyFilename));
+						IFileManager::Get().Move(*CopyFilename, *DestinationFilename);
+					}
 				}
 
 				break;
@@ -491,17 +548,37 @@ int32 UTextAssetCommandlet::Main(const FString& CmdLineParams)
 			case EMode::LoadText:
 			{
 				UPackage* Package = nullptr;
+				UE_LOG(LogTextAsset, Display, TEXT("Loading Text Asset '%s'"), *SourceFilename);
 
 				CollectGarbage(RF_NoFlags, true);
-				double Timer = 0.0;
+				ThisPackageLoadTime = 0.0;
 				{
-					SCOPE_SECONDS_COUNTER(Timer);
-					UE_LOG(LogTextAsset, Display, TEXT("Loading Text Asset '%s'"), *SourceFilename);
+					SCOPE_SECONDS_COUNTER(ThisPackageLoadTime);
 					Package = LoadPackage(nullptr, *SourceFilename, 0);
 				}
 				CollectGarbage(RF_NoFlags, true);
-				IterationPackageLoadTime += Timer;
-				TotalPackageLoadTime += Timer;
+				IterationPackageLoadTime += ThisPackageLoadTime;
+				TotalPackageLoadTime += ThisPackageLoadTime;
+
+				Package = nullptr;
+
+				break;
+			}
+
+			case EMode::LoadBinary:
+			{
+				UPackage* Package = nullptr;
+				UE_LOG(LogTextAsset, Display, TEXT("Loading Binary Asset '%s'"), *SourceFilename);
+				CollectGarbage(RF_NoFlags, true);
+
+				ThisPackageLoadTime = 0.0;
+				{
+					SCOPE_SECONDS_COUNTER(ThisPackageLoadTime);
+					Package = LoadPackage(nullptr, *SourceFilename, 0);
+				}
+				CollectGarbage(RF_NoFlags, true);
+				IterationPackageLoadTime += ThisPackageLoadTime;
+				TotalPackageLoadTime += ThisPackageLoadTime;
 
 				Package = nullptr;
 
@@ -512,16 +589,33 @@ int32 UTextAssetCommandlet::Main(const FString& CmdLineParams)
 			double EndTime = FPlatformTime::Seconds();
 			double Time = EndTime - StartTime;
 
-			if (Time > MaxTime)
+			if (Mode == EMode::LoadBinary || Mode == EMode::LoadText)
 			{
-				MaxTime = Time;
-				MaxTimePackage = SourceFilename;
-			}
+				if (ThisPackageLoadTime > MaxTime)
+				{
+					MaxTime = ThisPackageLoadTime;
+					MaxTimePackage = SourceFilename;
+				}
 
-			if (Time < MinTime)
+				if (ThisPackageLoadTime < MinTime)
+				{
+					MinTime = ThisPackageLoadTime;
+					MinTimePackage = SourceFilename;
+				}
+			}
+			else
 			{
-				MinTime = Time;
-				MinTimePackage = SourceFilename;
+				if (Time > MaxTime)
+				{
+					MaxTime = Time;
+					MaxTimePackage = SourceFilename;
+				}
+
+				if (Time < MinTime)
+				{
+					MinTime = Time;
+					MinTimePackage = SourceFilename;
+				}
 			}
 
 			TotalTime += Time;
@@ -553,13 +647,35 @@ int32 UTextAssetCommandlet::Main(const FString& CmdLineParams)
 			}
 		}
 
+		double AvgFileTime, MinFileTime, MaxFileTime;
+
+		if (Mode == EMode::LoadBinary || Mode == EMode::LoadText)
+		{
+			AvgFileTime = IterationPackageLoadTime;
+		}
+		else
+		{
+			AvgFileTime = TotalTime;
+		}
+
+		AvgFileTime /= (double)NumFiles;
+		MinFileTime = MinTime;
+		MaxFileTime = MaxTime;
+
 		UE_LOG(LogTextAsset, Display, TEXT("\tTotal Time:\t%.2fs"), TotalTime);
-		UE_LOG(LogTextAsset, Display, TEXT("\tAvg File Time:  \t%.2fms"), (TotalTime * 1000.0) / (double)NumFiles);
-		UE_LOG(LogTextAsset, Display, TEXT("\tMin File Time:  \t%.2fms (%s)"), MinTime * 1000.0, *MinTimePackage);
-		UE_LOG(LogTextAsset, Display, TEXT("\tMax File Time:  \t%.2fms (%s)"), MaxTime * 1000.0, *MaxTimePackage);
+		UE_LOG(LogTextAsset, Display, TEXT("\tTotal Files:\t%i"), NumFiles);
+		UE_LOG(LogTextAsset, Display, TEXT("\tAvg File Time:  \t%.2fms"), AvgFileTime * 1000.0);
+		UE_LOG(LogTextAsset, Display, TEXT("\tMin File Time:  \t%.2fms (%s)"), MinFileTime * 1000.0, *MinTimePackage);
+		UE_LOG(LogTextAsset, Display, TEXT("\tMax File Time:  \t%.2fms (%s)"), MaxFileTime * 1000.0, *MaxTimePackage);
 		UE_LOG(LogTextAsset, Display, TEXT("\tTotal Package Load Time:  \t%.2fs"), IterationPackageLoadTime);
 
-		if (Mode != EMode::LoadText)
+		if (CSVWriter != nullptr)
+		{
+			FString CSVLine = FString::Printf(TEXT("%f,%i,%f,%f,%f,%f\n"), TotalTime, NumFiles, AvgFileTime, MinFileTime, MaxFileTime, IterationPackageLoadTime);
+			CSVWriter->Serialize(TCHAR_TO_ANSI(*CSVLine), CSVLine.Len());
+		}
+
+		if (Mode != EMode::LoadText && Mode != EMode::ResaveText)
 		{
 			UE_LOG(LogTextAsset, Display, TEXT("\tTotal Package Save Time:  \t%.2fs"), IterationPackageSaveTime);
 		}
@@ -572,6 +688,11 @@ int32 UTextAssetCommandlet::Main(const FString& CmdLineParams)
 	UE_LOG(LogTextAsset, Display, TEXT("\tTotal Files Processed:  \t%i"), FilesToProcess.Num());
 	UE_LOG(LogTextAsset, Display, TEXT("\tAvg Iteration Package Load Time:  \t%.2fs"), TotalPackageLoadTime / (float)NumSaveIterations);
 
+	if (CSVWriter != nullptr)
+	{
+		delete CSVWriter;
+	}
+
 	if (Mode != EMode::LoadText)
 	{
 		UE_LOG(LogTextAsset, Display, TEXT("\tAvg Iteration Save Time:  \t%.2fs"), TotalPackageSaveTime / (float)NumSaveIterations);
