TBA: Don't attempt to do a JSON validation of the output file when running the text asset commandlet in "resave binary" mode

diff --git a/Engine/Source/Editor/UnrealEd/Private/Commandlets/TextAssetCommandlet.cpp b/Engine/Source/Editor/UnrealEd/Private/Commandlets/TextAssetCommandlet.cpp
--- a/Engine/Source/Editor/UnrealEd/Private/Commandlets/TextAssetCommandlet.cpp
+++ b/Engine/Source/Editor/UnrealEd/Private/Commandlets/TextAssetCommandlet.cpp
@@ -522,7 +522,7 @@ int32 UTextAssetCommandlet::Main(const FString& CmdLineParams)
 
 				if (bSaveSuccessful)
 				{
-					if (bVerifyJson)
+					if (bVerifyJson && Mode == EMode::ResaveText)
 					{
 						FArchive* File = IFileManager::Get().CreateFileReader(*DestinationFilename);
 						TSharedPtr< FJsonObject > RootObject;
