Deprecating ARRAY_COUNT and changing it to UE_ARRAY_COUNT.

diff --git a/Engine/Plugins/NetcodeUnitTest/NetcodeUnitTest/Source/NetcodeUnitTest/Private/UnitTests/VMReflection.cpp b/Engine/Plugins/NetcodeUnitTest/NetcodeUnitTest/Source/NetcodeUnitTest/Private/UnitTests/VMReflection.cpp
--- a/Engine/Plugins/NetcodeUnitTest/NetcodeUnitTest/Source/NetcodeUnitTest/Private/UnitTests/VMReflection.cpp
+++ b/Engine/Plugins/NetcodeUnitTest/NetcodeUnitTest/Source/NetcodeUnitTest/Private/UnitTests/VMReflection.cpp
@@ -869,7 +869,7 @@ bool UVMReflection::ExecuteUnitTest()
 		Results[2] = (uint8*)((FVMReflection(TestObj)->*"BytePropArray")["uint8"][2], &bError[2]);
 		Results[3] = (uint8*)((FVMReflection(TestObj)->*"BytePropArray")["uint8"][3], &bError[3]);
 
-		for (int32 i=0; i<ARRAY_COUNT(Results); i++)
+		for (int32 i=0; i< UE_ARRAY_COUNT(Results); i++)
 		{
 			if (!bError[i] && Results[i] != NULL)
 			{
@@ -983,7 +983,7 @@ bool UVMReflection::ExecuteUnitTest()
 		Results[2] = (uint8*)((FVMReflection(TestObj)->*"DynBytePropArray")["uint8"][2], &bError[2]);
 		Results[3] = (uint8*)((FVMReflection(TestObj)->*"DynBytePropArray")["uint8"][3], &bError[3]);
 
-		for (int32 i=0; i<ARRAY_COUNT(Results); i++)
+		for (int32 i=0; i< UE_ARRAY_COUNT(Results); i++)
 		{
 			if (!bError[i] && Results[i] != NULL)
 			{
diff --git a/Engine/Source/Editor/UnrealEd/Private/Lightmass/ImportVolumetricLightmap.cpp b/Engine/Source/Editor/UnrealEd/Private/Lightmass/ImportVolumetricLightmap.cpp
--- a/Engine/Source/Editor/UnrealEd/Private/Lightmass/ImportVolumetricLightmap.cpp
+++ b/Engine/Source/Editor/UnrealEd/Private/Lightmass/ImportVolumetricLightmap.cpp
@@ -109,7 +109,7 @@ bool CopyFromBrickmapTexel(
 
 		*(FFloat3Packed*)&BrickData.AmbientVector.Data[LinearDestCellIndex * sizeof(FFloat3Packed)] = FilteredVolumeLookupReconverted<FFloat3Packed>(BrickTextureCoordinate, CurrentLevelData.BrickDataDimensions, (const FFloat3Packed*)BrickData.AmbientVector.Data.GetData());
 		
-		for (int32 i = 0; i < ARRAY_COUNT(BrickData.SHCoefficients); i++)
+		for (int32 i = 0; i < UE_ARRAY_COUNT(BrickData.SHCoefficients); i++)
 		{
 			*(FColor*)&BrickData.SHCoefficients[i].Data[LinearDestCellIndex * sizeof(FColor)] = FilteredVolumeLookupReconverted<FColor>(BrickTextureCoordinate, CurrentLevelData.BrickDataDimensions, (const FColor*)BrickData.SHCoefficients[i].Data.GetData());
 		}
@@ -161,7 +161,7 @@ void FLightmassProcessor::ImportIrradianceTasks(bool& bGenerateSkyShadowing, TAr
 				Swarm.ReadChannel(Channel, &NewBrick.AverageClosestGeometryDistance, sizeof(NewBrick.AverageClosestGeometryDistance));
 				ReadArray(Channel, NewBrick.AmbientVector);
 
-				for (int32 i = 0; i < ARRAY_COUNT(NewBrick.SHCoefficients); i++)
+				for (int32 i = 0; i < UE_ARRAY_COUNT(NewBrick.SHCoefficients); i++)
 				{
 					ReadArray(Channel, NewBrick.SHCoefficients[i]);
 				}
@@ -197,7 +197,7 @@ int32 NumDilateOverEmbeddedVoxelsPasses = 2;
 struct FFilteredBrickData
 {
 	FFloat3Packed AmbientVector;
-	FColor SHCoefficients[ARRAY_COUNT(FImportedVolumetricLightmapBrick::SHCoefficients)];
+	FColor SHCoefficients[UE_ARRAY_COUNT(FImportedVolumetricLightmapBrick::SHCoefficients)];
 	uint8 DirectionalLightShadowing;
 };
 
@@ -270,10 +270,10 @@ void FilterWithNeighbors(
 					{
 						//@todo - filter SkyBentNormal from neighbors too
 						FLinearColor AmbientVector = FLinearColor(0, 0, 0, 0);
-						FLinearColor SHCoefficients[ARRAY_COUNT(Brick.SHCoefficients)];
+						FLinearColor SHCoefficients[UE_ARRAY_COUNT(Brick.SHCoefficients)];
 						FLinearColor DirectionalLightShadowing = FLinearColor(0, 0, 0, 0);
 
-						for (int32 i = 0; i < ARRAY_COUNT(SHCoefficients); i++)
+						for (int32 i = 0; i < UE_ARRAY_COUNT(SHCoefficients); i++)
 						{
 							SHCoefficients[i] = FLinearColor(0, 0, 0, 0);
 						}
@@ -316,9 +316,9 @@ void FilterWithNeighbors(
 												FLinearColor NeighborAmbientVector = FilteredVolumeLookup<FFloat3Packed>(BrickTextureCoordinate, CurrentLevelData.BrickDataDimensions, (const FFloat3Packed*)CurrentLevelData.BrickData.AmbientVector.Data.GetData());
 												AmbientVector += NeighborAmbientVector * Weight;
 
-												for (int32 i = 0; i < ARRAY_COUNT(SHCoefficients); i++)
+												for (int32 i = 0; i < UE_ARRAY_COUNT(SHCoefficients); i++)
 												{
-													static_assert(ARRAY_COUNT(SHCoefficients) == 6, "Assuming 2 SHCoefficient vectors per color channel");
+													static_assert(UE_ARRAY_COUNT(SHCoefficients) == 6, "Assuming 2 SHCoefficient vectors per color channel");
 
 													// Weight by ambient before filtering, normalized SH coefficients don't filter properly
 													float AmbientCoefficient = NeighborAmbientVector.Component(i / 2);
@@ -349,9 +349,9 @@ void FilterWithNeighbors(
 							const FLinearColor FilteredAmbientColor = AmbientVector * InvTotalWeight;
 							FilteredBrickData[LinearVoxelIndex].AmbientVector = ConvertFromLinearColor<FFloat3Packed>(FilteredAmbientColor);
 
-							for (int32 i = 0; i < ARRAY_COUNT(SHCoefficients); i++)
+							for (int32 i = 0; i < UE_ARRAY_COUNT(SHCoefficients); i++)
 							{
-								static_assert(ARRAY_COUNT(SHCoefficients) == 6, "Assuming 2 SHCoefficient vectors per color channel");
+								static_assert(UE_ARRAY_COUNT(SHCoefficients) == 6, "Assuming 2 SHCoefficient vectors per color channel");
 
 								float AmbientCoefficient = FMath::Max(FilteredAmbientColor.Component(i / 2), KINDA_SMALL_NUMBER);
 
@@ -385,7 +385,7 @@ void FilterWithNeighbors(
 						FFloat3Packed* DestAmbientVector = (FFloat3Packed*)&CurrentLevelData.BrickData.AmbientVector.Data[LinearBrickLayoutCellIndex * sizeof(FFloat3Packed)];
 						*DestAmbientVector = FilteredBrickData[LinearVoxelIndex].AmbientVector;
 
-						for (int32 i = 0; i < ARRAY_COUNT(FilteredBrickData[LinearVoxelIndex].SHCoefficients); i++)
+						for (int32 i = 0; i < UE_ARRAY_COUNT(FilteredBrickData[LinearVoxelIndex].SHCoefficients); i++)
 						{
 							FColor* DestCoefficients = (FColor*)&CurrentLevelData.BrickData.SHCoefficients[i].Data[LinearBrickLayoutCellIndex * sizeof(FColor)];
 							*DestCoefficients = FilteredBrickData[LinearVoxelIndex].SHCoefficients[i];
@@ -981,7 +981,7 @@ void FLightmassProcessor::ImportVolumetricLightmap()
 		CurrentLevelData.BrickData.SkyBentNormal.Format = PF_B8G8R8A8;
 		CurrentLevelData.BrickData.DirectionalLightShadowing.Format = PF_G8;
 
-		for (int32 i = 0; i < ARRAY_COUNT(CurrentLevelData.BrickData.SHCoefficients); i++)
+		for (int32 i = 0; i < UE_ARRAY_COUNT(CurrentLevelData.BrickData.SHCoefficients); i++)
 		{
 			CurrentLevelData.BrickData.SHCoefficients[i].Format = PF_B8G8R8A8;
 		}
@@ -1041,7 +1041,7 @@ void FLightmassProcessor::ImportVolumetricLightmap()
 
 		CurrentLevelData.BrickData.DirectionalLightShadowing.Resize(TotalBrickDataSize * GPixelFormats[CurrentLevelData.BrickData.DirectionalLightShadowing.Format].BlockBytes);
 
-		for (int32 i = 0; i < ARRAY_COUNT(CurrentLevelData.BrickData.SHCoefficients); i++)
+		for (int32 i = 0; i < UE_ARRAY_COUNT(CurrentLevelData.BrickData.SHCoefficients); i++)
 		{
 			const int32 Stride = GPixelFormats[CurrentLevelData.BrickData.SHCoefficients[i].Format].BlockBytes;
 			CurrentLevelData.BrickData.SHCoefficients[i].Resize(TotalBrickDataSize * Stride);
@@ -1076,7 +1076,7 @@ void FLightmassProcessor::ImportVolumetricLightmap()
 				(const uint8*)Brick.AmbientVector.GetData(),
 				CurrentLevelData.BrickData.AmbientVector.Data.GetData());
 
-			for (int32 i = 0; i < ARRAY_COUNT(Brick.SHCoefficients); i++)
+			for (int32 i = 0; i < UE_ARRAY_COUNT(Brick.SHCoefficients); i++)
 			{
 				CopyBrickToAtlasVolumeTexture(
 					GPixelFormats[CurrentLevelData.BrickData.SHCoefficients[i].Format].BlockBytes,
@@ -1229,7 +1229,7 @@ void FLightmassProcessor::ImportVolumetricLightmap()
 	int32 BrickDataSize = CurrentLevelData.BrickData.AmbientVector.Data.Num() + CurrentLevelData.BrickData.SkyBentNormal.Data.Num() + CurrentLevelData.BrickData.DirectionalLightShadowing.Data.Num()
 			+ CurrentLevelData.BrickData.LQLightColor.Data.Num() + CurrentLevelData.BrickData.LQLightColor.Data.Num();
 
-	for (int32 i = 0; i < ARRAY_COUNT(CurrentLevelData.BrickData.SHCoefficients); i++)
+	for (int32 i = 0; i < UE_ARRAY_COUNT(CurrentLevelData.BrickData.SHCoefficients); i++)
 	{
 		BrickDataSize += CurrentLevelData.BrickData.SHCoefficients[i].Data.Num();
 	}
diff --git a/Engine/Source/Programs/UnrealLightmass/Private/Lighting/AdaptiveVolumetricLightmap.cpp b/Engine/Source/Programs/UnrealLightmass/Private/Lighting/AdaptiveVolumetricLightmap.cpp
--- a/Engine/Source/Programs/UnrealLightmass/Private/Lighting/AdaptiveVolumetricLightmap.cpp
+++ b/Engine/Source/Programs/UnrealLightmass/Private/Lighting/AdaptiveVolumetricLightmap.cpp
@@ -52,7 +52,7 @@ FOverlapInterval GetInterval(const FBox& Box, const FVector& Vector)
 	FOverlapInterval Result;
 	Result.Min = Result.Max = FVector::DotProduct(Vector, BoxVertices[0]);
 
-	for (int32 i = 1; i < ARRAY_COUNT(BoxVertices); ++i) 
+	for (int32 i = 1; i < UE_ARRAY_COUNT(BoxVertices); ++i) 
 	{
 		float Projection = FVector::DotProduct(Vector, BoxVertices[i]);
 		Result.Min = FMath::Min(Result.Min, Projection);
@@ -99,7 +99,7 @@ bool IntersectTriangleAndAABB(const FTriangle& Triangle, const FBox& Box)
 		FVector::CrossProduct(BoxNormal2, TriangleEdge2)
 	};
 
-	for (int i = 0; i < ARRAY_COUNT(TestDirections); ++i) 
+	for (int i = 0; i < UE_ARRAY_COUNT(TestDirections); ++i) 
 	{
 		if (!OverlapOnAxis(Box, Triangle, TestDirections[i])) 
 		{
@@ -155,7 +155,7 @@ bool FStaticLightingSystem::DoesVoxelIntersectSceneGeometry(const FBox& CellBoun
 					Triangle.Vertices[1] = Vertices[1].WorldPosition;
 					Triangle.Vertices[2] = Vertices[2].WorldPosition;
 
-					FBox TriangleAABB(Triangle.Vertices, ARRAY_COUNT(Triangle.Vertices));
+					FBox TriangleAABB(Triangle.Vertices, UE_ARRAY_COUNT(Triangle.Vertices));
 
 					if (ExpandedCellBounds.Intersect(TriangleAABB))
 					{
@@ -516,7 +516,7 @@ void FStaticLightingSystem::ProcessVolumetricLightmapBrickTask(FVolumetricLightm
 	BrickData.DirectionalLightShadowing.Empty(TotalBrickSize);
 	BrickData.DirectionalLightShadowing.AddDefaulted(TotalBrickSize);
 
-	for (int32 i = 0; i < ARRAY_COUNT(BrickData.SHCoefficients); i++)
+	for (int32 i = 0; i < UE_ARRAY_COUNT(BrickData.SHCoefficients); i++)
 	{
 		BrickData.SHCoefficients[i].Empty(TotalBrickSize);
 		BrickData.SHCoefficients[i].AddDefaulted(TotalBrickSize);
@@ -804,7 +804,7 @@ void FStaticLightingSystem::CalculateAdaptiveVolumetricLightmap(int32 TaskIndex)
 
 void FIrradianceBrickData::SetFromVolumeLightingSample(int32 Index, const FVolumeLightingSample& Sample, bool bInsideGeometry, float MinDistanceToSurface, bool bBorderVoxel)
 {
-	static_assert(ARRAY_COUNT(Sample.HighQualityCoefficients) >= ARRAY_COUNT(SHCoefficients) + 1, "Coefficient mismatch");
+	static_assert(UE_ARRAY_COUNT(Sample.HighQualityCoefficients) >= UE_ARRAY_COUNT(SHCoefficients) + 1, "Coefficient mismatch");
 
 	AmbientVector[Index] = FFloat3Packed(FLinearColor(Sample.HighQualityCoefficients[0][0], Sample.HighQualityCoefficients[0][1], Sample.HighQualityCoefficients[0][2], 0.0f));
 
diff --git a/Engine/Source/Runtime/Core/Private/HAL/LowLevelMemTracker.cpp b/Engine/Source/Runtime/Core/Private/HAL/LowLevelMemTracker.cpp
--- a/Engine/Source/Runtime/Core/Private/HAL/LowLevelMemTracker.cpp
+++ b/Engine/Source/Runtime/Core/Private/HAL/LowLevelMemTracker.cpp
@@ -128,7 +128,7 @@ extern const TCHAR* LLMGetTagName(ELLMTag Tag)
 #undef LLM_TAG_NAME_ARRAY
 
 	int32 Index = (int32)Tag;
-	if (Index >= 0 && Index < ARRAY_COUNT(Names))
+	if (Index >= 0 && Index < UE_ARRAY_COUNT(Names))
 	{
 		return Names[Index];
 	}
@@ -145,7 +145,7 @@ extern const ANSICHAR* LLMGetTagNameANSI(ELLMTag Tag)
 #undef LLM_TAG_NAME_ARRAY
 
 	int32 Index = (int32)Tag;
-	if (Index >= 0 && Index < ARRAY_COUNT(Names))
+	if (Index >= 0 && Index < UE_ARRAY_COUNT(Names))
 	{
 		return Names[Index];
 	}
@@ -162,7 +162,7 @@ extern FName LLMGetTagStat(ELLMTag Tag)
 #undef LLM_TAG_STAT_ARRAY
 
 	int32 Index = (int32)Tag;
-	if (Index >= 0 && Index < ARRAY_COUNT(Names))
+	if (Index >= 0 && Index < UE_ARRAY_COUNT(Names))
 	{
 		return Names[Index];
 	}
@@ -179,7 +179,7 @@ extern FName LLMGetTagStatGroup(ELLMTag Tag)
 #undef LLM_TAG_STAT_ARRAY
 
 	int32 Index = (int32)Tag;
-	if (Index >= 0 && Index < ARRAY_COUNT(Names))
+	if (Index >= 0 && Index < UE_ARRAY_COUNT(Names))
 	{
 		return Names[Index];
 	}
@@ -196,7 +196,7 @@ extern int32 LLMGetTagParent(ELLMTag Tag)
 #undef LLM_TAG_NAME_ARRAY
 
 	int32 Index = (int32)Tag;
-	if( Index >= 0 && Index < ARRAY_COUNT(ParentTags))
+	if( Index >= 0 && Index < UE_ARRAY_COUNT(ParentTags))
 	{
 		return ParentTags[Index];
 	}
diff --git a/Engine/Source/Runtime/Core/Private/Misc/ConfigCacheIni.cpp b/Engine/Source/Runtime/Core/Private/Misc/ConfigCacheIni.cpp
--- a/Engine/Source/Runtime/Core/Private/Misc/ConfigCacheIni.cpp
+++ b/Engine/Source/Runtime/Core/Private/Misc/ConfigCacheIni.cpp
@@ -3246,7 +3246,7 @@ static void GetSourceIniHierarchyFilenames(const TCHAR* InBaseIniName, const TCH
 	bool bHasPlatformExtensionProjectConfigDir = FPaths::DirectoryExists(*PlatformExtensionProjectConfigDir);
 
 	// go over all the config layers
-	for (int32 LayerIndex = 0; LayerIndex < ARRAY_COUNT(GConfigLayers); LayerIndex++)
+	for (int32 LayerIndex = 0; LayerIndex < UE_ARRAY_COUNT(GConfigLayers); LayerIndex++)
 	{
 		const FConfigLayer& Layer = GConfigLayers[LayerIndex];
 		bool bHasPlatformTag = Layer.Path.Contains(TEXT("{PLATFORM}"));
@@ -3298,10 +3298,10 @@ static void GetSourceIniHierarchyFilenames(const TCHAR* InBaseIniName, const TCH
 			// we assume none of the more special tags in expanded ones
 			checkfSlow(FCString::Strstr(*Layer.Path, TEXT("{USERSETTINGS}")) == nullptr && FCString::Strstr(*Layer.Path, TEXT("{USER}")) == nullptr, TEXT("Expanded config %s shouldn't have a {USER*} tags in it"), *Layer.Path);
 			// last layer is expected to not be expanded
-			checkfSlow(LayerIndex < ARRAY_COUNT(GConfigLayers) - 1, TEXT("Final layer %s shouldn't be an expansion layer, as it needs to generate the  hierarchy cache key"), *Layer.Path);
+			checkfSlow(LayerIndex < UE_ARRAY_COUNT(GConfigLayers) - 1, TEXT("Final layer %s shouldn't be an expansion layer, as it needs to generate the  hierarchy cache key"), *Layer.Path);
 
 			// loop over all the possible expansions
-			for (int32 ExpansionIndex = 0; ExpansionIndex < ARRAY_COUNT(GConfigLayerExpansions); ExpansionIndex++)
+			for (int32 ExpansionIndex = 0; ExpansionIndex < UE_ARRAY_COUNT(GConfigLayerExpansions); ExpansionIndex++)
 			{
 				const FConfigLayerExpansion& Expansion = GConfigLayerExpansions[ExpansionIndex];
 
@@ -3342,7 +3342,7 @@ static void GetSourceIniHierarchyFilenames(const TCHAR* InBaseIniName, const TCH
 				}
 
 				// check if we should be generating the cache key - only at the end of all expansions
-				bool bGenerateCacheKey = (Layer.Flag & Flag_GenerateCacheKey) && ExpansionIndex == ARRAY_COUNT(GConfigLayerExpansions);
+				bool bGenerateCacheKey = (Layer.Flag & Flag_GenerateCacheKey) && ExpansionIndex == UE_ARRAY_COUNT(GConfigLayerExpansions);
 				checkfSlow(!(bGenerateCacheKey && bHasPlatformTag), TEXT("Flag_GenerateCacheKey shouldn't have a platform tag"));
 
 				const FDataDrivenPlatformInfoRegistry::FPlatformInfo& Info = FDataDrivenPlatformInfoRegistry::GetPlatformInfo(PlatformName);
@@ -4237,7 +4237,7 @@ private:
 			TEXT("SectionName")
 		};
 		using UnderType = __underlying_type(HistoryType);
-		static_assert(static_cast<UnderType>(HistoryType::Count) == ARRAY_COUNT(Strings), "");
+		static_assert(static_cast<UnderType>(HistoryType::Count) == UE_ARRAY_COUNT(Strings), "");
 		return Strings[static_cast<UnderType>(Type)];
 	}
 
diff --git a/Engine/Source/Runtime/Core/Private/Windows/WindowsPlatformProcess.cpp b/Engine/Source/Runtime/Core/Private/Windows/WindowsPlatformProcess.cpp
--- a/Engine/Source/Runtime/Core/Private/Windows/WindowsPlatformProcess.cpp
+++ b/Engine/Source/Runtime/Core/Private/Windows/WindowsPlatformProcess.cpp
@@ -869,7 +869,7 @@ const TCHAR* FWindowsPlatformProcess::BaseDir()
 			{
 				hCurrentModule = hInstance;
 			}
-			GetModuleFileName(hCurrentModule, Result, ARRAY_COUNT(Result));
+			GetModuleFileName(hCurrentModule, Result, UE_ARRAY_COUNT(Result));
 			FString TempResult(Result);
 			TempResult = TempResult.Replace(TEXT("\\"), TEXT("/"));
 			FCString::Strcpy(Result, *TempResult);
@@ -984,7 +984,7 @@ const TCHAR* FWindowsPlatformProcess::ComputerName()
 	static TCHAR Result[256]=TEXT("");
 	if( !Result[0] )
 	{
-		uint32 Size=ARRAY_COUNT(Result);
+		uint32 Size=UE_ARRAY_COUNT(Result);
 		GetComputerName( Result, (::DWORD*)&Size );
 	}
 	return Result;
@@ -998,7 +998,7 @@ const TCHAR* FWindowsPlatformProcess::UserName(bool bOnlyAlphaNumeric/* = true*/
 	{
 		if( !ResultAlpha[0] )
 		{
-			uint32 Size=ARRAY_COUNT(ResultAlpha);
+			uint32 Size=UE_ARRAY_COUNT(ResultAlpha);
 			GetUserName( ResultAlpha, (::DWORD*)&Size );
 			TCHAR *c, *d;
 			for( c=ResultAlpha, d=ResultAlpha; *c!=0; c++ )
@@ -1012,7 +1012,7 @@ const TCHAR* FWindowsPlatformProcess::UserName(bool bOnlyAlphaNumeric/* = true*/
 	{
 		if( !Result[0] )
 		{
-			uint32 Size=ARRAY_COUNT(Result);
+			uint32 Size=UE_ARRAY_COUNT(Result);
 			GetUserName( Result, (::DWORD*)&Size );
 		}
 		return Result;
@@ -1061,7 +1061,7 @@ const TCHAR* FWindowsPlatformProcess::ExecutablePath()
 	static TCHAR Result[512]=TEXT("");
 	if( !Result[0] )
 	{
-		if ( !GetModuleFileName( hInstance, Result, ARRAY_COUNT(Result) ) )
+		if ( !GetModuleFileName( hInstance, Result, UE_ARRAY_COUNT(Result) ) )
 		{
 			Result[0] = 0;
 		}
@@ -1076,13 +1076,13 @@ const TCHAR* FWindowsPlatformProcess::ExecutableName(bool bRemoveExtension)
 	if( !Result[0] )
 	{
 		// Get complete path for the executable
-		if ( GetModuleFileName( hInstance, Result, ARRAY_COUNT(Result) ) != 0 )
+		if ( GetModuleFileName( hInstance, Result, UE_ARRAY_COUNT(Result) ) != 0 )
 		{
 			// Remove all of the path information by finding the base filename
 			FString FileName = Result;
 			FString FileNameWithExt = Result;
-			FCString::Strncpy( Result, *( FPaths::GetBaseFilename(FileName) ), ARRAY_COUNT(Result) );
-			FCString::Strncpy( ResultWithExt, *( FPaths::GetCleanFilename(FileNameWithExt) ), ARRAY_COUNT(ResultWithExt) );
+			FCString::Strncpy( Result, *( FPaths::GetBaseFilename(FileName) ), UE_ARRAY_COUNT(Result) );
+			FCString::Strncpy( ResultWithExt, *( FPaths::GetCleanFilename(FileNameWithExt) ), UE_ARRAY_COUNT(ResultWithExt) );
 		}
 		// If the call failed, zero out the memory to be safe
 		else
@@ -1123,7 +1123,7 @@ const FString FWindowsPlatformProcess::GetModulesDirectory()
 
 		// Get the directory for it
 		TCHAR Buffer[MAX_PATH] = TEXT("");
-		GetModuleFileName(hCurrentModule, Buffer, ARRAY_COUNT(Buffer));
+		GetModuleFileName(hCurrentModule, Buffer, UE_ARRAY_COUNT(Buffer));
 		*FCString::Strrchr(Buffer, '\\') = 0;
 
 		// Normalize the resulting path
diff --git a/Engine/Source/Runtime/Core/Public/HAL/LowLevelMemTracker.h b/Engine/Source/Runtime/Core/Public/HAL/LowLevelMemTracker.h
--- a/Engine/Source/Runtime/Core/Public/HAL/LowLevelMemTracker.h
+++ b/Engine/Source/Runtime/Core/Public/HAL/LowLevelMemTracker.h
@@ -63,7 +63,7 @@
 			if (bLog)
 			{
 				TCHAR DescriptionString[4096];
-				GET_VARARGS(DescriptionString, ARRAY_COUNT(DescriptionString), ARRAY_COUNT(DescriptionString) - 1, Format, Format);
+				GET_VARARGS(DescriptionString, UE_ARRAY_COUNT(DescriptionString), UE_ARRAY_COUNT(DescriptionString) - 1, Format, Format);
 
 				FPlatformMisc::LowLevelOutputDebugString(DescriptionString);
 
diff --git a/Engine/Source/Runtime/Core/Public/Templates/UnrealTemplate.h b/Engine/Source/Runtime/Core/Public/Templates/UnrealTemplate.h
--- a/Engine/Source/Runtime/Core/Public/Templates/UnrealTemplate.h
+++ b/Engine/Source/Runtime/Core/Public/Templates/UnrealTemplate.h
@@ -130,14 +130,15 @@ FORCEINLINE const T& AsConst(T& Ref)
 
 #ifdef __clang__
 	template <typename T>
-	auto ArrayCountHelper(T& t) -> typename TEnableIf<__is_array(T), char(&)[sizeof(t) / sizeof(t[0]) + 1]>::Type;
+	auto UE4ArrayCountHelper(T& t) -> typename TEnableIf<__is_array(T), char(&)[sizeof(t) / sizeof(t[0]) + 1]>::Type;
 #else
 	template <typename T, uint32 N>
-	char (&ArrayCountHelper(const T (&)[N]))[N + 1];
+	char (&UE4ArrayCountHelper(const T (&)[N]))[N + 1];
 #endif
 
 // Number of elements in an array.
-#define ARRAY_COUNT( array ) (sizeof(ArrayCountHelper(array)) - 1)
+#define UE_ARRAY_COUNT( array ) (sizeof(UE4ArrayCountHelper(array)) - 1)
+#define ARRAY_COUNT( array ) DEPRECATED_MACRO(4.24, "The ARRAY_COUNT macro has been deprecated in favor of UE_ARRAY_COUNT.") UE_ARRAY_COUNT( array )
 
 // Offset of a struct member.
 #ifndef UNREAL_CODE_ANALYZER
diff --git a/Engine/Source/Runtime/Engine/Private/PrecomputedVolumetricLightmap.cpp b/Engine/Source/Runtime/Engine/Private/PrecomputedVolumetricLightmap.cpp
--- a/Engine/Source/Runtime/Engine/Private/PrecomputedVolumetricLightmap.cpp
+++ b/Engine/Source/Runtime/Engine/Private/PrecomputedVolumetricLightmap.cpp
@@ -70,7 +70,7 @@ FArchive& operator<<(FArchive& Ar,FPrecomputedVolumetricLightmapData& Volume)
 
 	Ar << Volume.BrickData.AmbientVector;
 
-	for (int32 i = 0; i < ARRAY_COUNT(Volume.BrickData.SHCoefficients); i++)
+	for (int32 i = 0; i < UE_ARRAY_COUNT(Volume.BrickData.SHCoefficients); i++)
 	{
 		Ar << Volume.BrickData.SHCoefficients[i];
 	}
@@ -133,7 +133,7 @@ int32 FVolumetricLightmapBrickData::GetMinimumVoxelSize() const
 	check(AmbientVector.Format != PF_Unknown);
 	int32 VoxelSize = GPixelFormats[AmbientVector.Format].BlockBytes;
 
-	for (int32 i = 0; i < ARRAY_COUNT(SHCoefficients); i++)
+	for (int32 i = 0; i < UE_ARRAY_COUNT(SHCoefficients); i++)
 	{
 		VoxelSize += GPixelFormats[SHCoefficients[i].Format].BlockBytes;
 	}
@@ -172,7 +172,7 @@ ENGINE_API void FPrecomputedVolumetricLightmapData::InitRHI()
 	IndirectionTexture.CreateTexture(IndirectionTextureDimensions);
 	BrickData.AmbientVector.CreateTexture(BrickDataDimensions);
 
-	for (int32 i = 0; i < ARRAY_COUNT(BrickData.SHCoefficients); i++)
+	for (int32 i = 0; i < UE_ARRAY_COUNT(BrickData.SHCoefficients); i++)
 	{
 		BrickData.SHCoefficients[i].CreateTexture(BrickDataDimensions);
 	}
@@ -190,7 +190,7 @@ ENGINE_API void FPrecomputedVolumetricLightmapData::ReleaseRHI()
 	IndirectionTexture.Texture.SafeRelease();
 	BrickData.AmbientVector.Texture.SafeRelease();
 
-	for (int32 i = 0; i < ARRAY_COUNT(BrickData.SHCoefficients); i++)
+	for (int32 i = 0; i < UE_ARRAY_COUNT(BrickData.SHCoefficients); i++)
 	{
 		BrickData.SHCoefficients[i].Texture.SafeRelease();
 	}
diff --git a/Engine/Source/Runtime/RenderCore/Private/RenderGraphEvent.cpp b/Engine/Source/Runtime/RenderCore/Private/RenderGraphEvent.cpp
--- a/Engine/Source/Runtime/RenderCore/Private/RenderGraphEvent.cpp
+++ b/Engine/Source/Runtime/RenderCore/Private/RenderGraphEvent.cpp
@@ -15,7 +15,7 @@ FRDGEventName::FRDGEventName(const TCHAR* EventFormat, ...)
 		va_start(VAList, EventFormat);
 		TCHAR TempStr[256];
 		// Build the string in the temp buffer
-		FCString::GetVarArgs(TempStr, ARRAY_COUNT(TempStr), EventFormat, VAList);
+		FCString::GetVarArgs(TempStr, UE_ARRAY_COUNT(TempStr), EventFormat, VAList);
 		va_end(VAList);
 
 		EventNameStorage = TempStr;
diff --git a/Engine/Source/Runtime/Renderer/Private/PostProcess/PostProcessing.cpp b/Engine/Source/Runtime/Renderer/Private/PostProcess/PostProcessing.cpp
--- a/Engine/Source/Runtime/Renderer/Private/PostProcess/PostProcessing.cpp
+++ b/Engine/Source/Runtime/Renderer/Private/PostProcess/PostProcessing.cpp
@@ -233,7 +233,7 @@ public:
 
 		static const TCHAR* PassLabels[] =
 		{ NULL, TEXT("BloomDownsample1"), TEXT("BloomDownsample2"), TEXT("BloomDownsample3"), TEXT("BloomDownsample4"), TEXT("BloomDownsample5") };
-		static_assert(ARRAY_COUNT(PassLabels) == DownSampleStages, "PassLabel count must be equal to DownSampleStages.");
+		static_assert(UE_ARRAY_COUNT(PassLabels) == DownSampleStages, "PassLabel count must be equal to DownSampleStages.");
 
 		// The first down sample is the input
 		PostProcessDownsamples[0] = SourceDownsample;
@@ -587,7 +587,7 @@ static FRenderingCompositeOutputRef AddBloom(FBloomDownSampleArray& BloomDownSam
 	{
 		// console variable override
 		static const auto CVar = IConsoleManager::Get().FindTConsoleVariableDataInt(TEXT("r.BloomQuality"));
-		BloomQuality = FMath::Clamp(CVar->GetValueOnRenderThread(), 0, (int32)ARRAY_COUNT(BloomQualityStages));
+		BloomQuality = FMath::Clamp(CVar->GetValueOnRenderThread(), 0, (int32)UE_ARRAY_COUNT(BloomQualityStages));
 	}
 
 	// Extract the Context
@@ -659,7 +659,7 @@ static FRenderingCompositeOutputRef AddBloom(FBloomDownSampleArray& BloomDownSam
 			{ Settings.Bloom2Size, &Settings.Bloom2Tint },
 			{ Settings.Bloom1Size, &Settings.Bloom1Tint },
 		};
-		static const uint32 NumBloomStages = ARRAY_COUNT(BloomStages);
+		static const uint32 NumBloomStages = UE_ARRAY_COUNT(BloomStages);
 
 		const uint32 BloomStageCount = BloomQualityStages[BloomQuality - 1];
 		check(BloomStageCount <= NumBloomStages);
