When running cook on the fly asset registry was not generated (since no chunk manifests are generated). This caused certain runtime lookups (like UObjectLibrary::LoadAssetDataFromPaths) to return no results in cook on the fly scenarios. With this fix, if there are no chunk manifests present the entire asset registry is saved instead.

diff --git a/Engine/Source/Editor/UnrealEd/Private/Commandlets/AssetRegistryGenerator.cpp b/Engine/Source/Editor/UnrealEd/Private/Commandlets/AssetRegistryGenerator.cpp
--- a/Engine/Source/Editor/UnrealEd/Private/Commandlets/AssetRegistryGenerator.cpp
+++ b/Engine/Source/Editor/UnrealEd/Private/Commandlets/AssetRegistryGenerator.cpp
@@ -1074,68 +1074,90 @@ bool FAssetRegistryGenerator::SaveAssetRegistry(const FString& SandboxPath, bool
 		const int32 GenericChunkBucket = -1;
 		ChunkBucketNames.Add(GenericChunkBucket, FString());
 
-		// Pass over all chunks and build a mapping of chunk index to asset registry name. All chunks that don't have a unique registry are assigned to the "generic bucket"
-		// which will be written to the master asset registry in chunk 0
-		for (int32 PakchunkIndex = 0; PakchunkIndex < FinalChunkManifests.Num(); ++PakchunkIndex)
+		// When chunk manifests have been generated (e.g. cook by the book) serialize 
+		// an asset registry for each chunk.
+		if (FinalChunkManifests.Num() > 0)
 		{
-			FChunkPackageSet* Manifest = FinalChunkManifests[PakchunkIndex];
-			if (Manifest == nullptr)
+			// Pass over all chunks and build a mapping of chunk index to asset registry name. All chunks that don't have a unique registry are assigned to the "generic bucket"
+			// which will be written to the master asset registry in chunk 0
+			for (int32 PakchunkIndex = 0; PakchunkIndex < FinalChunkManifests.Num(); ++PakchunkIndex)
 			{
-				continue;
-			}
+				FChunkPackageSet* Manifest = FinalChunkManifests[PakchunkIndex];
+				if (Manifest == nullptr)
+				{
+					continue;
+				}
 
-			bool bAddToGenericBucket = true;
+				bool bAddToGenericBucket = true;
 
-			if (bUseAssetManager)
-			{
-				// For chunks with unique asset registry name, pakchunkIndex should equal chunkid
-				FName RegistryName = UAssetManager::Get().GetUniqueAssetRegistryName(PakchunkIndex);
-				if (RegistryName != NAME_None)
+				if (bUseAssetManager)
 				{
-					ChunkBuckets.FindOrAdd(PakchunkIndex).Add(PakchunkIndex);
-					ChunkBucketNames.FindOrAdd(PakchunkIndex) = RegistryName.ToString();
-					bAddToGenericBucket = false;
+					// For chunks with unique asset registry name, pakchunkIndex should equal chunkid
+					FName RegistryName = UAssetManager::Get().GetUniqueAssetRegistryName(PakchunkIndex);
+					if (RegistryName != NAME_None)
+					{
+						ChunkBuckets.FindOrAdd(PakchunkIndex).Add(PakchunkIndex);
+						ChunkBucketNames.FindOrAdd(PakchunkIndex) = RegistryName.ToString();
+						bAddToGenericBucket = false;
+					}
+				}
+
+				if (bAddToGenericBucket)
+				{
+					ChunkBuckets.FindOrAdd(GenericChunkBucket).Add(PakchunkIndex);
 				}
 			}
 
-			if (bAddToGenericBucket)
+			FString SandboxPathWithoutExtension = FPaths::ChangeExtension(SandboxPath, TEXT(""));
+			FString SandboxPathExtension = FPaths::GetExtension(SandboxPath);
+
+			for (TMap<int32, TSet<int32>>::ElementType& ChunkBucketElement : ChunkBuckets)
 			{
-				ChunkBuckets.FindOrAdd(GenericChunkBucket).Add(PakchunkIndex);
-			}
-		}
+				// Prune out the development only packages, and any assets that belong in a different chunk asset registry
+				FAssetRegistryState NewState;
+				NewState.InitializeFromExistingAndPrune(State, CookedPackages, TSet<FName>(), ChunkBucketElement.Value, SaveOptions);
 
-		FString SandboxPathWithoutExtension = FPaths::ChangeExtension(SandboxPath, TEXT(""));
-		FString SandboxPathExtension = FPaths::GetExtension(SandboxPath);
+				InjectEncryptionData(NewState);
 
-		for (TMap<int32, TSet<int32>>::ElementType& ChunkBucketElement : ChunkBuckets)
-		{
-			// Prune out the development only packages, and any assets that belong in a different chunk asset registry
-			FAssetRegistryState NewState;
-			NewState.InitializeFromExistingAndPrune(State, CookedPackages, TSet<FName>(), ChunkBucketElement.Value, SaveOptions);
+				// Create runtime registry data
+				FArrayWriter SerializedAssetRegistry;
+				SerializedAssetRegistry.SetFilterEditorOnly(true);
+
+				NewState.Serialize(SerializedAssetRegistry, SaveOptions);
+
+				// Save the generated registry
+				FString PlatformSandboxPath = SandboxPathWithoutExtension.Replace(TEXT("[Platform]"), *TargetPlatform->PlatformName());
+				PlatformSandboxPath += ChunkBucketNames[ChunkBucketElement.Key] + TEXT(".") + SandboxPathExtension;
 
-			InjectEncryptionData(NewState);
+				FFileHelper::SaveArrayToFile(SerializedAssetRegistry, *PlatformSandboxPath);
+
+				FString FilenameForLog;
+				if (ChunkBucketElement.Key != GenericChunkBucket)
+				{
+					check(ChunkBucketElement.Key < FinalChunkManifests.Num());
+					FChunkPackageSet* ChunkPackageSet = FinalChunkManifests[ChunkBucketElement.Key];
+					check(ChunkPackageSet);
+					FilenameForLog = FString::Printf(TEXT("[chunkbucket %i] "), ChunkBucketElement.Key);
+				}
+				UE_LOG(LogAssetRegistryGenerator, Display, TEXT("Generated asset registry %snum assets %d, size is %5.2fkb"), *FilenameForLog, NewState.GetNumAssets(), (float)SerializedAssetRegistry.Num() / 1024.f);
+			}
+		}
+		// If no chunk manifests have been generated (e.g. cook on the fly)
+		else
+		{
+			// Prune out the development only packages
+			State.PruneAssetData(CookedPackages, TSet<FName>(), SaveOptions);
 
 			// Create runtime registry data
 			FArrayWriter SerializedAssetRegistry;
 			SerializedAssetRegistry.SetFilterEditorOnly(true);
 
-			NewState.Serialize(SerializedAssetRegistry, SaveOptions);
+			State.Serialize(SerializedAssetRegistry, SaveOptions);
 
 			// Save the generated registry
-			FString PlatformSandboxPath = SandboxPathWithoutExtension.Replace(TEXT("[Platform]"), *TargetPlatform->PlatformName());
-			PlatformSandboxPath += ChunkBucketNames[ChunkBucketElement.Key] + TEXT(".") + SandboxPathExtension;
-
+			FString PlatformSandboxPath = SandboxPath.Replace(TEXT("[Platform]"), *TargetPlatform->PlatformName());
 			FFileHelper::SaveArrayToFile(SerializedAssetRegistry, *PlatformSandboxPath);
-
-			FString FilenameForLog;
-			if (ChunkBucketElement.Key != GenericChunkBucket)
-			{
-				check(ChunkBucketElement.Key < FinalChunkManifests.Num());
-				FChunkPackageSet* ChunkPackageSet = FinalChunkManifests[ChunkBucketElement.Key];
-				check(ChunkPackageSet);
-				FilenameForLog = FString::Printf(TEXT("[chunkbucket %i] "), ChunkBucketElement.Key);
-			}
-			UE_LOG(LogAssetRegistryGenerator, Display, TEXT("Generated asset registry %snum assets %d, size is %5.2fkb"), *FilenameForLog, NewState.GetNumAssets(), (float)SerializedAssetRegistry.Num() / 1024.f);
+			UE_LOG(LogAssetRegistryGenerator, Display, TEXT("Generated asset registry num assets %d, size is %5.2fkb"), ObjectToDataMap.Num(), (float)SerializedAssetRegistry.Num() / 1024.f);
 		}
 	}
 
