UResavePackagesCommandlet: do not attempt to check out and save transient packages when building navigation data

diff --git a/Engine/Source/Editor/UnrealEd/Private/Commandlets/ContentCommandlets.cpp b/Engine/Source/Editor/UnrealEd/Private/Commandlets/ContentCommandlets.cpp
--- a/Engine/Source/Editor/UnrealEd/Private/Commandlets/ContentCommandlets.cpp
+++ b/Engine/Source/Editor/UnrealEd/Private/Commandlets/ContentCommandlets.cpp
@@ -1706,7 +1706,7 @@ void UResavePackagesCommandlet::PerformAdditionalOperations(class UWorld* World,
 				for (TActorIterator<ANavigationData> It(World); It; ++It)
 				{
 					UPackage* Package = It->GetOutermost();
-					if (Package != nullptr && Package->IsDirty())
+					if (Package != nullptr && Package->IsDirty() && !Package->HasAnyFlags(RF_Transient))
 					{
 						CheckoutAndSavePackage(Package, CheckedOutPackagesFilenames, bSkipCheckedOutFiles);
 					}
