Add null check when copy/pasting material nodes

diff --git a/Engine/Source/Editor/UnrealEd/Private/MaterialGraph.cpp b/Engine/Source/Editor/UnrealEd/Private/MaterialGraph.cpp
--- a/Engine/Source/Editor/UnrealEd/Private/MaterialGraph.cpp
+++ b/Engine/Source/Editor/UnrealEd/Private/MaterialGraph.cpp
@@ -180,7 +180,9 @@ void UMaterialGraph::LinkGraphNodesFromMaterial()
 				UEdGraphPin* InputPin = CastChecked<UMaterialGraphNode>(Expression->GraphNode)->GetInputPin(InputIndex);
 
 				// InputPin can be null during a PostEditChange when there is a circular dependency between nodes, and nodes have pins that are dynamically created
-				if (InputPin != nullptr && ExpressionInputs[InputIndex]->Expression)
+				if (InputPin != nullptr && ExpressionInputs[InputIndex]->Expression
+					// Unclear why this is null sometimes, but this is safer than crashing
+					&& ExpressionInputs[InputIndex]->Expression->GraphNode)
 				{
 					UMaterialGraphNode* GraphNode = CastChecked<UMaterialGraphNode>(ExpressionInputs[InputIndex]->Expression->GraphNode);
 					InputPin->MakeLinkTo(GraphNode->GetOutputPin(GetValidOutputIndex(ExpressionInputs[InputIndex])));
