Checking that SourceInstance or SourceFunction are valid when undoing and calling functions on them

diff --git a/Engine/Source/Editor/UnrealEd/Private/PreviewMaterial.cpp b/Engine/Source/Editor/UnrealEd/Private/PreviewMaterial.cpp
--- a/Engine/Source/Editor/UnrealEd/Private/PreviewMaterial.cpp
+++ b/Engine/Source/Editor/UnrealEd/Private/PreviewMaterial.cpp
@@ -1681,12 +1681,12 @@ void UMaterialEditorInstanceConstant::PostEditUndo()
 {
 	Super::PostEditUndo();
 
-	if (bIsFunctionPreviewMaterial)
+	if (bIsFunctionPreviewMaterial && SourceFunction)
 	{
 		bIsFunctionInstanceDirty = true;
 		ApplySourceFunctionChanges();
 	}
-	else
+	else if (SourceInstance)
 	{
 		FMaterialUpdateContext Context;
 
