When trying to find the most relevant blueprint class for a given compiled class, always choose the compiled class when it has been cooked

diff --git a/Engine/Source/Editor/UnrealEd/Private/Kismet2/BlueprintEditorUtils.cpp b/Engine/Source/Editor/UnrealEd/Private/Kismet2/BlueprintEditorUtils.cpp
--- a/Engine/Source/Editor/UnrealEd/Private/Kismet2/BlueprintEditorUtils.cpp
+++ b/Engine/Source/Editor/UnrealEd/Private/Kismet2/BlueprintEditorUtils.cpp
@@ -2129,7 +2129,7 @@ UClass* FBlueprintEditorUtils::GetSkeletonClass(UClass* FromClass)
 
 UClass* FBlueprintEditorUtils::GetMostUpToDateClass(UClass* FromClass)
 {
-	if(!FromClass || FromClass->HasAnyClassFlags(CLASS_Native))
+	if(!FromClass || FromClass->HasAnyClassFlags(CLASS_Native) || FromClass->bCooked)
 	{
 		return FromClass;
 	}
