You cannot create a blueprint of a null class

diff --git a/Engine/Source/Editor/UnrealEd/Private/Kismet2/Kismet2.cpp b/Engine/Source/Editor/UnrealEd/Private/Kismet2/Kismet2.cpp
--- a/Engine/Source/Editor/UnrealEd/Private/Kismet2/Kismet2.cpp
+++ b/Engine/Source/Editor/UnrealEd/Private/Kismet2/Kismet2.cpp
@@ -990,23 +990,28 @@ void FKismetEditorUtilities::ConformBlueprintFlagsAndComponents(UBlueprint* Blue
 /** @return		true is it's possible to create a blueprint from the specified class */
 bool FKismetEditorUtilities::CanCreateBlueprintOfClass(const UClass* Class)
 {
-	bool bAllowDerivedBlueprints = false;
-	GConfig->GetBool(TEXT("Kismet"), TEXT("AllowDerivedBlueprints"), /*out*/ bAllowDerivedBlueprints, GEngineIni);
-
-	const bool bCanCreateBlueprint =
-		!Class->HasAnyClassFlags(CLASS_Deprecated)
-		&& !Class->HasAnyClassFlags(CLASS_NewerVersionExists)
-		&& (!Class->ClassGeneratedBy || (bAllowDerivedBlueprints && !IsClassABlueprintSkeleton(Class)));
+	bool bCanCreateBlueprint = false;
+	
+	if (Class)
+	{
+		bool bAllowDerivedBlueprints = false;
+		GConfig->GetBool(TEXT("Kismet"), TEXT("AllowDerivedBlueprints"), /*out*/ bAllowDerivedBlueprints, GEngineIni);
 
-	const bool bIsBPGC = (Cast<UBlueprintGeneratedClass>(Class) != nullptr);
+		bCanCreateBlueprint = !Class->HasAnyClassFlags(CLASS_Deprecated)
+			&& !Class->HasAnyClassFlags(CLASS_NewerVersionExists)
+			&& (!Class->ClassGeneratedBy || (bAllowDerivedBlueprints && !IsClassABlueprintSkeleton(Class)));
 
-	const UEdGraphSchema_K2* K2Schema = GetDefault<UEdGraphSchema_K2>();
-	const bool bIsValidClass = Class->GetBoolMetaDataHierarchical(FBlueprintMetadata::MD_IsBlueprintBase)
-		|| (Class == UObject::StaticClass())
-		|| (Class->HasAnyClassFlags(CLASS_CompiledFromBlueprint) || Class == USceneComponent::StaticClass() || Class == UActorComponent::StaticClass())
-		|| bIsBPGC;  // BPs are always considered inheritable
+		const bool bIsBPGC = (Cast<UBlueprintGeneratedClass>(Class) != nullptr);
 
-	return bCanCreateBlueprint && bIsValidClass;
+		const bool bIsValidClass = Class->GetBoolMetaDataHierarchical(FBlueprintMetadata::MD_IsBlueprintBase)
+			|| (Class == UObject::StaticClass())
+			|| (Class->HasAnyClassFlags(CLASS_CompiledFromBlueprint) || Class == USceneComponent::StaticClass() || Class == UActorComponent::StaticClass())
+			|| bIsBPGC;  // BPs are always considered inheritable
+			
+		bCanCreateBlueprint &= bIsValidClass;
+	}
+	
+	return bCanCreateBlueprint;
 }
 
 UBlueprint* FKismetEditorUtilities::CreateBlueprintFromActor(const FString& Path, AActor* Actor, const bool bReplaceActor, bool bKeepMobility /*= false*/)
