Dev-Editor-OpenWorld optim integration (9660099): Optimize FPackageName::DoesPackageExist usage in FileHelpers by going through the AssetRegistry when possible.

diff --git a/Engine/Source/Editor/UnrealEd/Private/FileHelpers.cpp b/Engine/Source/Editor/UnrealEd/Private/FileHelpers.cpp
--- a/Engine/Source/Editor/UnrealEd/Private/FileHelpers.cpp
+++ b/Engine/Source/Editor/UnrealEd/Private/FileHelpers.cpp
@@ -202,7 +202,31 @@ namespace FileDialogHelpers
 	}
 }
 
+namespace FileHelperPackageUtil
+{
+	/**
+	 * DoesPackageExist helper that rely on the AssetRegistry to validate if a package exists instead of hitting the FS
+	 * Fallback to the FS if the asset registry initial scan isn't done or we aren't in Editor
+	 */
+	bool DoesPackageExist(UPackage* Package, FString* OutFilename = nullptr)
+	{
+		// Test using asset registry to figure out existence
+		IAssetRegistry& AssetRegistry = FAssetRegistryModule::GetRegistry();
+		if (!AssetRegistry.IsLoadingAssets() || !GIsEditor)
+		{
+			TArray<FAssetData> Data;
+			FAssetRegistryModule::GetRegistry().GetAssetsByPackageName(Package->GetFName(), Data, true);
 
+			if (Data.Num() > 0 && OutFilename)
+			{
+				*OutFilename = FPackageName::LongPackageNameToFilename(Package->GetName(), Package->ContainsMap() ? FPackageName::GetMapPackageExtension() : FPackageName::GetAssetPackageExtension());
+			}
+
+			return Data.Num() > 0;
+		}
+		return FPackageName::DoesPackageExist(Package->GetName(), nullptr, OutFilename);
+	}
+}
 
 /**
  * Queries the user if they want to quit out of interpolation editing before save.
@@ -544,7 +568,7 @@ static bool SaveWorld(UWorld* World,
 	FString	CleanFilename;
 
 	// Does a filename already exist for this package?
-	const bool bPackageExists = FPackageName::DoesPackageExist( PackageName, NULL, &ExistingFilename );
+	const bool bPackageExists = FileHelperPackageUtil::DoesPackageExist(Package, &ExistingFilename );
 
 	if ( ForceFilename )
 	{
@@ -1364,7 +1388,7 @@ bool FEditorFileUtils::AddCheckoutPackageItems(bool bCheckDirty, TArray<UPackage
 				}
 				
 				FString Filename;
-				if (FPackageName::DoesPackageExist(Package->GetName(), NULL, &Filename))
+				if (FileHelperPackageUtil::DoesPackageExist(Package, &Filename))
 				{
 					if (IFileManager::Get().IsReadOnly(*Filename))
 					{
@@ -1418,7 +1442,7 @@ bool FEditorFileUtils::AddCheckoutPackageItems(bool bCheckDirty, TArray<UPackage
 		bool bPkgReadOnly = true;
 		bool bCareAboutReadOnly = SourceControlProvider.UsesLocalReadOnlyState();
 		// Find the filename for this package
-		bool bFoundFile = FPackageName::DoesPackageExist(CurPackage->GetName(), NULL, &Filename);
+		bool bFoundFile = FileHelperPackageUtil::DoesPackageExist(CurPackage, &Filename);
 		if (bFoundFile)
 		{
 			// determine if the package file is read only
@@ -1655,7 +1679,7 @@ bool FEditorFileUtils::PromptToCheckoutPackages(bool bCheckDirty, const TArray<U
 					UPackage* PackageToMakeWritable = *PkgsToMakeWritableIter;
 					FString Filename;
 
-					bool bFoundFile = FPackageName::DoesPackageExist( PackageToMakeWritable->GetName(), NULL, &Filename );
+					bool bFoundFile = FileHelperPackageUtil::DoesPackageExist( PackageToMakeWritable, &Filename );
 					if( bFoundFile )
 					{
 						// If we're ignoring the package due to the user ignoring it for saving, remove it from the ignore list
@@ -2808,7 +2832,7 @@ static InternalSavePackageResult InternalSavePackage(UPackage* PackageToSave, bo
 	if( bIsValidPath )
 	{
 		FString ExistingFilename;
-		const bool bPackageAlreadyExists = FPackageName::DoesPackageExist(PackageName, NULL, &ExistingFilename);
+		const bool bPackageAlreadyExists = FileHelperPackageUtil::DoesPackageExist(PackageToSave, &ExistingFilename);
 		if (!bPackageAlreadyExists)
 		{
 			// Construct a filename from long package name.
@@ -3140,7 +3164,7 @@ static bool InternalSavePackagesFast(const TArray<UPackage*>& PackagesToSave, bo
 
 		// Check if a file exists for this package
 		FString Filename;
-		bool bFoundFile = FPackageName::DoesPackageExist(CurPackage->GetName(), NULL, &Filename);
+		bool bFoundFile = FileHelperPackageUtil::DoesPackageExist(CurPackage, &Filename);
 		if (bFoundFile)
 		{
 			// determine if the package file is read only
