Level editor transform gizmo now uses component's transform instead of actor when components are selected

diff --git a/Engine/Source/Editor/UnrealEd/Private/EditorModeManager.cpp b/Engine/Source/Editor/UnrealEd/Private/EditorModeManager.cpp
--- a/Engine/Source/Editor/UnrealEd/Private/EditorModeManager.cpp
+++ b/Engine/Source/Editor/UnrealEd/Private/EditorModeManager.cpp
@@ -1180,12 +1180,19 @@ FMatrix FEditorModeTools::GetLocalCoordinateSystem()
 
 	if (!CustomCoordinateSystemProvided)
 	{
-		const int32 Num = GetSelectedActors()->CountSelections<AActor>();
-
-		// Coordinate system needs to come from the last actor selected
-		if (Num > 0)
+		if (USceneComponent* SceneComponent = GetSelectedComponents()->GetBottom<USceneComponent>())
+		{
+			Matrix = FQuatRotationMatrix(SceneComponent->GetComponentQuat());
+		}
+		else
 		{
-			Matrix = FQuatRotationMatrix(GetSelectedActors()->GetBottom<AActor>()->GetActorQuat());
+			const int32 Num = GetSelectedActors()->CountSelections<AActor>();
+
+			// Coordinate system needs to come from the last actor selected
+			if (Num > 0)
+			{
+				Matrix = FQuatRotationMatrix(GetSelectedActors()->GetBottom<AActor>()->GetActorQuat());
+			}
 		}
 	}
 
