Removed "Save Checked Out" button from file saving dialog when exiting the editor.

#rnx

diff --git a/Engine/Source/Editor/UnrealEd/Private/FileHelpers.cpp b/Engine/Source/Editor/UnrealEd/Private/FileHelpers.cpp
--- a/Engine/Source/Editor/UnrealEd/Private/FileHelpers.cpp
+++ b/Engine/Source/Editor/UnrealEd/Private/FileHelpers.cpp
@@ -1624,11 +1624,6 @@ bool FEditorFileUtils::PromptToCheckoutPackages(bool bCheckDirty, const TArray<U
 		
 		// Prepare the buttons for the checkout dialog
 
-		if (OutPackagesNotNeedingCheckout != nullptr && OutPackagesNotNeedingCheckout->Num() > 0)
-		{
-			CheckoutPackagesDialogModule.AddButton(DRT_Save, NSLOCTEXT("PackagesDialogModule", "Dlg_SaveCheckedOutButton", "Save Checked Out"), NSLOCTEXT("PackagesDialogModule", "Dlg_SaveCheckedOutTooltip", "Save already checked out packages, but not the packages in this dialog."));
-		}
-
 		// The checkout button should be disabled if no packages can be checked out.
 		CheckoutPackagesDialogModule.AddButton(DRT_CheckOut, NSLOCTEXT("PackagesDialogModule", "Dlg_CheckOutButtonp", "Check Out Selected"), NSLOCTEXT("PackagesDialogModule", "Dlg_CheckOutTooltip", "Attempt to Check Out Checked Assets"), CheckOutSelectedDisabledAttrib );
 		
