Disable AsyncDelete on Mac, since it is failing in EngineTest System_MacEditor
#rnx

diff --git a/Engine/Source/Editor/UnrealEd/Private/Cooker/AsyncIODelete.cpp b/Engine/Source/Editor/UnrealEd/Private/Cooker/AsyncIODelete.cpp
--- a/Engine/Source/Editor/UnrealEd/Private/Cooker/AsyncIODelete.cpp
+++ b/Engine/Source/Editor/UnrealEd/Private/Cooker/AsyncIODelete.cpp
@@ -52,6 +52,7 @@ void FAsyncIODelete::SetTempRoot(const FStringView& InOwnedTempRoot)
 void FAsyncIODelete::SetDeletesPaused(bool bInPaused)
 {
 	bPaused = bInPaused;
+#if ASYNCIODELETE_ASYNC_ENABLED
 	if (!bPaused)
 	{
 		IFileManager& FileManager = IFileManager::Get();
@@ -67,6 +68,7 @@ void FAsyncIODelete::SetDeletesPaused(bool bInPaused)
 		}
 		PausedDeletes.Empty();
 	}
+#endif
 }
 
 bool FAsyncIODelete::Setup()
@@ -82,6 +84,7 @@ bool FAsyncIODelete::Setup()
 		return false;
 	}
 
+#if ASYNCIODELETE_ASYNC_ENABLED
 	// Delete the TempRoot directory to clear the results from any previous process using the same TempRoot that did not shut down cleanly
 	uint32 ErrorCode;
 	if (!DeleteTempRootDirectory(ErrorCode))
@@ -107,6 +110,7 @@ bool FAsyncIODelete::Setup()
 	// TempRoot and bPaused are preserved across setup/teardown and may have any value
 	check(PausedDeletes.Num() == 0);
 	check(DeleteCounter == 0);
+#endif
 
 	// We are now setup and ready to create DeleteTasks
 	bInitialized = true;
@@ -121,6 +125,7 @@ void FAsyncIODelete::Teardown()
 		return;
 	}
 
+#if ASYNCIODELETE_ASYNC_ENABLED
 	// Clear task variables
 	WaitForAllTasks();
 	check(ActiveTaskCount == 0 && TasksComplete != nullptr && TasksComplete->Wait(0));
@@ -138,6 +143,7 @@ void FAsyncIODelete::Teardown()
 	// Clear delete variables; we don't need to run the tasks for the remaining pauseddeletes because synchronously deleting the temp directory above did the work they were going to do
 	PausedDeletes.Empty();
 	DeleteCounter = 0;
+#endif
 
 	// We are now torn down and ready for a new setup
 	bInitialized = false;
@@ -145,6 +151,7 @@ void FAsyncIODelete::Teardown()
 
 bool FAsyncIODelete::WaitForAllTasks(float TimeLimitSeconds)
 {
+#if ASYNCIODELETE_ASYNC_ENABLED
 	if (!bInitialized)
 	{
 		return true;
@@ -152,7 +159,7 @@ bool FAsyncIODelete::WaitForAllTasks(float TimeLimitSeconds)
 
 	if (TimeLimitSeconds <= 0.f)
 	{
-		TasksComplete->Wait();
+	TasksComplete->Wait();
 	}
 	else
 	{
@@ -162,6 +169,7 @@ bool FAsyncIODelete::WaitForAllTasks(float TimeLimitSeconds)
 		}
 	}
 	check(ActiveTaskCount == 0);
+#endif
 	return true;
 }
 
@@ -195,16 +203,19 @@ bool FAsyncIODelete::Delete(const FStringView& PathToDelete, EPathType ExpectedT
 		return false;
 	}
 
+#if ASYNCIODELETE_ASYNC_ENABLED
 	if (DeleteCounter == UINT32_MAX)
 	{
 		Teardown();
 	}
+#endif
 	if (!Setup())
 	{
 		// Setup failed; we are not able to provide asynchronous deletes; fall back to synchronous
 		return SynchronousDelete(PathToDeleteSZ, ExpectedType);
 	}
 
+#if ASYNCIODELETE_ASYNC_ENABLED
 	const FString TempPath = FPaths::Combine(TempRoot, FString::Printf(TEXT("Delete%u"), DeleteCounter));
 	DeleteCounter++;
 
@@ -228,8 +239,12 @@ bool FAsyncIODelete::Delete(const FStringView& PathToDelete, EPathType ExpectedT
 		CreateDeleteTask(TempPath, ExpectedType);
 	}
 	return true;
+#else
+	return SynchronousDelete(PathToDeleteSZ, ExpectedType);
+#endif
 }
 
+#if ASYNCIODELETE_ASYNC_ENABLED
 void FAsyncIODelete::CreateDeleteTask(const FStringView& InDeletePath, EPathType PathType)
 {
 	{
@@ -254,6 +269,7 @@ void FAsyncIODelete::OnTaskComplete()
 		TasksComplete->Trigger();
 	}
 }
+#endif
 
 bool FAsyncIODelete::SynchronousDelete(const TCHAR* InDeletePath, EPathType PathType)
 {
@@ -277,6 +293,7 @@ bool FAsyncIODelete::SynchronousDelete(const TCHAR* InDeletePath, EPathType Path
 	return Result;
 }
 
+#if ASYNCIODELETE_ASYNC_ENABLED
 bool FAsyncIODelete::DeleteTempRootDirectory(uint32& OutErrorCode)
 {
 	OutErrorCode = 0;
@@ -315,6 +332,7 @@ bool FAsyncIODelete::DeleteTempRootDirectory(uint32& OutErrorCode)
 	}
 	return bDeleteSucceeded;
 }
+#endif
 
 #if WITH_ASYNCIODELETE_DEBUG
 void FAsyncIODelete::AddTempRoot(const FStringView& InTempRoot)
diff --git a/Engine/Source/Editor/UnrealEd/Private/Cooker/AsyncIODelete.h b/Engine/Source/Editor/UnrealEd/Private/Cooker/AsyncIODelete.h
--- a/Engine/Source/Editor/UnrealEd/Private/Cooker/AsyncIODelete.h
+++ b/Engine/Source/Editor/UnrealEd/Private/Cooker/AsyncIODelete.h
@@ -7,6 +7,7 @@
 #include "Async/Future.h"
 #include "Containers/Array.h"
 #include "Containers/StringView.h"
+#include "HAL/Platform.h"
 #include "Misc/AssertionMacros.h"
 
 #ifndef WITH_ASYNCIODELETE_DEBUG
@@ -17,6 +18,13 @@
 	#endif
 #endif
 
+// Temporarily disable AsyncIODelete on Mac until we can diagnose why it's failing
+#if PLATFORM_MAC
+#define ASYNCIODELETE_ASYNC_ENABLED 0
+#else
+#define ASYNCIODELETE_ASYNC_ENABLED 1
+#endif
+
 class FEvent;
 
 /**
@@ -93,8 +101,10 @@ public:
 	bool WaitForAllTasks(float TimeLimitSeconds = 0.0f);
 
 private:
+#if ASYNCIODELETE_ASYNC_ENABLED
 	/** Update the ActiveTaskCount and events for a task completing */
 	void OnTaskComplete();
+#endif
 
 	enum class EPathType
 	{
@@ -104,20 +114,26 @@ private:
 	/** Asynchronously delete a file or directory.  We handle both in the same function, but we want the interface to be explicit. */
 	bool Delete(const FStringView& PathToDelete, EPathType ExpectedType);
 
+#if ASYNCIODELETE_ASYNC_ENABLED
 	/** Create and store the task to delete the given Directory or File */
 	void CreateDeleteTask(const FStringView& InDeletePath, EPathType PathType);
+#endif
 
 	/** Delete the given path synchronously; called from a task or in error fallback cases from the public thread */
 	bool SynchronousDelete(const TCHAR* InDeletePath, EPathType PathType);
 
+#if ASYNCIODELETE_ASYNC_ENABLED
 	bool DeleteTempRootDirectory(uint32& OutErrorCode);
+#endif
 
 	FString	TempRoot;
+#if ASYNCIODELETE_ASYNC_ENABLED
 	TArray<FString> PausedDeletes;
 	FCriticalSection CriticalSection; // We use a CriticalSection instead of a TAtomic ActiveTaskCount so that we can atomically { trigger TasksComplete if ActiveTaskCount == 0 }
 	FEvent* TasksComplete = nullptr;
 	uint32 ActiveTaskCount = 0;
 	uint32 DeleteCounter = 0;
+#endif
 	bool bInitialized = false;
 	bool bPaused = false;
 
diff --git a/Engine/Source/Editor/UnrealEd/Private/Cooker/AsyncIODeleteTest.cpp b/Engine/Source/Editor/UnrealEd/Private/Cooker/AsyncIODeleteTest.cpp
--- a/Engine/Source/Editor/UnrealEd/Private/Cooker/AsyncIODeleteTest.cpp
+++ b/Engine/Source/Editor/UnrealEd/Private/Cooker/AsyncIODeleteTest.cpp
@@ -12,6 +12,12 @@
 IMPLEMENT_SIMPLE_AUTOMATION_TEST(FAsyncIODeleteTest, "System.Core.Misc.AsyncIODelete", EAutomationTestFlags::ApplicationContextMask | EAutomationTestFlags::EngineFilter);
 bool FAsyncIODeleteTest::RunTest(const FString& Parameters)
 {
+#if ASYNCIODELETE_ASYNC_ENABLED
+	UE_LOG(LogCore, Display, TEXT("FAsyncIODeleteTest::RunTest, ASYNCIODELETE is ENABLED"));
+#else
+	UE_LOG(LogCore, Display, TEXT("FAsyncIODeleteTest::RunTest, ASYNCIODELETE is DISABLED"));
+#endif
+
 	FString TestRoot = FPaths::CreateTempFilename(FPlatformProcess::UserTempDir(), TEXT("AsyncIODelete"), TEXT(""));
 	IFileManager& FileManager = IFileManager::Get();
 	const bool bApplyToTreeTrue = true;
@@ -70,6 +76,7 @@ bool FAsyncIODeleteTest::RunTest(const FString& Parameters)
 		TestTrue(TEXT("WaitForAllTasks timed out"), WaitResult);
 	};
 	{
+		StartSection(TEXT("Constructing first FAsyncIODelete"));
 		FAsyncIODelete AsyncIODelete(TempRoot);
 
 		StartSection(TEXT("Waiting for tasks to complete when none have been launched should succeed"));
@@ -95,22 +102,32 @@ bool FAsyncIODeleteTest::RunTest(const FString& Parameters)
 		AsyncIODelete2.DeleteDirectory(TestDir1);
 		TestRequestedPathsDeleted(TEXT("AsyncIODelete::Delete should have moved the deleted paths before returning even when paused."));
 		WaitForAllTasksAndVerify(AsyncIODelete2);
+#if ASYNCIODELETE_ASYNC_ENABLED
 		TestTempRootCountsEqual(TempRoot2, 1, 1, TEXT("AsyncIODelete should not have deleted the moved paths because it is paused."));
+#endif
 		AsyncIODelete2.SetDeletesPaused(false);
 		WaitForAllTasksAndVerify(AsyncIODelete2);
+#if ASYNCIODELETE_ASYNC_ENABLED
 		TestTempRootCountsEqual(TempRoot2, 0, 0, TEXT("AsyncIODelete should have deleted the moved paths after unpausing."));
+#endif
 
 		StartSection(TEXT("Verify Teardown() deletes the TempRoot and Setup() creates it"));
 		AsyncIODelete2.Teardown();
+#if ASYNCIODELETE_ASYNC_ENABLED
 		TestTrue(TEXT("AsyncIODelete::Teardown should have deleted its TempRoot."), !FileManager.DirectoryExists(*TempRoot2));
+#endif
 		AsyncIODelete2.Setup();
+#if ASYNCIODELETE_ASYNC_ENABLED
 		TestTrue(TEXT("AsyncIODelete::Setup should have created its TempRoot."), FileManager.DirectoryExists(*TempRoot2));
+#endif
 
 		StartSection(TEXT("Manual setup works as long as you call SetTempRoot before Setup"));
 		FAsyncIODelete AsyncIODelete3;
 		AsyncIODelete3.SetTempRoot(TempRoot3);
 		AsyncIODelete3.Setup();
+#if ASYNCIODELETE_ASYNC_ENABLED
 		TestTrue(TEXT("Setup should have created the TempRoot."), FileManager.DirectoryExists(*TempRoot3));
+#endif
 
 		StartSection(TEXT("Check that even after Setup, waiting for tasks to complete when none have been launched should succeed"));
 		WaitForAllTasksAndVerify(AsyncIODelete);
@@ -120,14 +137,20 @@ bool FAsyncIODeleteTest::RunTest(const FString& Parameters)
 		AsyncIODelete3.DeleteFile(TestFile1);
 		AsyncIODelete3.DeleteDirectory(TestDir1);
 		AsyncIODelete3.SetTempRoot(TempRoot4);
+#if ASYNCIODELETE_ASYNC_ENABLED
 		TestTrue(TEXT("SetTempRoot should have deleted the old TempRoot."), !FileManager.DirectoryExists(*TempRoot3));
+#endif
 		CreateTestPathsToDelete();
 		AsyncIODelete3.DeleteFile(TestFile1);
 		AsyncIODelete3.DeleteDirectory(TestDir1);
 		TestRequestedPathsDeleted(TEXT("AsyncIODelete::Delete should have worked after changing the TempRoot."));
+#if ASYNCIODELETE_ASYNC_ENABLED
 		TestTrue(TEXT("Delete should have created the new TempRoot after SetTempRoot."), FileManager.DirectoryExists(*TempRoot4));
+#endif
 		WaitForAllTasksAndVerify(AsyncIODelete3);
+#if ASYNCIODELETE_ASYNC_ENABLED
 		TestTempRootCountsEqual(TempRoot4, 0, 0, TEXT("AsyncIODelete::Delete should have created the tasks to delete moved files after changing the TempRoot."));
+#endif
 
 		StartSection(TEXT("Attempting to delete a parent directory of the temproot, the temproot itself, or a child inside of it fails"));
 		FString SubDirInTempRoot4 = FPaths::Combine(TempRoot4, TEXT("SubDir"));
@@ -137,8 +160,10 @@ bool FAsyncIODeleteTest::RunTest(const FString& Parameters)
 		TestFalse(TEXT("AsyncIODelete should refuse to delete a child of its TempRoot."), AsyncIODelete3.DeleteDirectory(SubDirInTempRoot4));
 	}
 
+#if ASYNCIODELETE_ASYNC_ENABLED
 	TestTrue(TEXT("AsyncIODelete destructor should have deleted its TempRoot."),
 		!FileManager.DirectoryExists(*TempRoot) && !FileManager.DirectoryExists(*TempRoot2) && !FileManager.DirectoryExists(*TempRoot3) && !FileManager.DirectoryExists(*TempRoot4));
+#endif
 
 	return true;
 }
\ No newline at end of file
