Change Ed mode default grid drawing to true.  #rb Matt.Kuhlenschmidt #jira UE-86749 #fyi Rex.Hill

diff --git a/Engine/Source/Editor/UnrealEd/Private/EditorModes.cpp b/Engine/Source/Editor/UnrealEd/Private/EditorModes.cpp
--- a/Engine/Source/Editor/UnrealEd/Private/EditorModes.cpp
+++ b/Engine/Source/Editor/UnrealEd/Private/EditorModes.cpp
@@ -27,7 +27,7 @@ const FEditorModeID FBuiltinEditorModes::EM_SceneDepthPicker(TEXT("EM_SceneDepth
 
 FEdModeDefault::FEdModeDefault()
 {
-	bDrawGrid = false;
+	bDrawGrid = true;
 	bDrawPivot = false;
 	bDrawBaseInfo = false;
 	bDrawWorldBox = false;
