Fix handling of user pressing 'x' in prompt to checkout packages dialog
#rnx

diff --git a/Engine/Source/Editor/UnrealEd/Private/FileHelpers.cpp b/Engine/Source/Editor/UnrealEd/Private/FileHelpers.cpp
--- a/Engine/Source/Editor/UnrealEd/Private/FileHelpers.cpp
+++ b/Engine/Source/Editor/UnrealEd/Private/FileHelpers.cpp
@@ -1723,7 +1723,7 @@ bool FEditorFileUtils::PromptToCheckoutPackages(bool bCheckDirty, const TArray<U
 				bResult = true;
 				bPerformedOperation = true;
 			}
-			else if (UserResponse == DRT_Cancel)
+			else if (UserResponse == DRT_Cancel || UserResponse == DRT_None)
 			{
 				// Handle the case of the user canceling out of the dialog
 				bResult = false;
