Fix duplicate root component in SCS editor when creating a blueprint via the harvest components feature

diff --git a/Engine/Source/Editor/UnrealEd/Private/Kismet2/Kismet2.cpp b/Engine/Source/Editor/UnrealEd/Private/Kismet2/Kismet2.cpp
--- a/Engine/Source/Editor/UnrealEd/Private/Kismet2/Kismet2.cpp
+++ b/Engine/Source/Editor/UnrealEd/Private/Kismet2/Kismet2.cpp
@@ -1503,6 +1503,8 @@ public:
 				Pair.Key->SetRelativeTransform(Pair.Value);
 			}
 
+			SCS->RemoveNode(SCS->GetDefaultSceneRootNode());
+
 			FTransform NewActorTransform = FTransform::Identity;
 			if (RootActors.Num() == 1)
 			{
