PatchDiff: Updating AssetRegistryDiff tool to include class-by-class breakdowns of the nondeterminism detected

diff --git a/Engine/Source/Editor/UnrealEd/Classes/Commandlets/DiffAssetRegistriesCommandlet.h b/Engine/Source/Editor/UnrealEd/Classes/Commandlets/DiffAssetRegistriesCommandlet.h
--- a/Engine/Source/Editor/UnrealEd/Classes/Commandlets/DiffAssetRegistriesCommandlet.h
+++ b/Engine/Source/Editor/UnrealEd/Classes/Commandlets/DiffAssetRegistriesCommandlet.h
@@ -113,10 +113,34 @@ class UDiffAssetRegistriesCommandlet : public UCommandlet
 		}
 	};
 
+	struct FDeterminismInfo
+	{
+		int64 DirectSize;
+		int64 DirectCount;
+		int64 IndirectSize;
+		int64 IndirectCount;
+
+		TArray<int64> Direct;
+		TArray<int64> Indirect;
+		
+		void AddDirect(const FChangeInfo& Rhs)
+		{
+			DirectSize += Rhs.GetTotalChangeSize();
+			DirectCount += Rhs.GetTotalChangeCount();
+		}
+
+		void AddIndirect(const FChangeInfo& Rhs)
+		{
+			IndirectSize += Rhs.GetTotalChangeSize();
+			IndirectCount += Rhs.GetTotalChangeCount();
+		}
+	};
+
 	struct FChunkChangeInfo
 	{
 		TSet<FName> IncludedAssets;
 		TMap<FName, FChangeInfo> ChangesByClass;
+		TMap<FName, FDeterminismInfo> Determinism;
 	};
 
 	GENERATED_UCLASS_BODY()
@@ -149,7 +173,7 @@ private:
 	bool	IsInRelevantChunk(FAssetRegistryState& InRegistryState, FName InAssetPath);
 
 	void	LogChangedFiles(FArchive *CSVFile, const FString &OldPath, const FString &NewPath);
-	void	LogClassSummary(FArchive *CSVFile, const FString& HeaderPrefix, const TMap<FName, FChangeInfo>& InChangeInfoByAsset, bool bDoWarnings);
+	void	LogClassSummary(FArchive *CSVFile, const FString& HeaderPrefix, const TMap<FName, FChangeInfo>& InChangeInfoByAsset, bool bDoWarnings, TMap<FName, FDeterminismInfo> DeterminismInfo = TMap<FName, FDeterminismInfo>());
 	void	SummarizeDeterminism();
 
 	void	PopulateChangelistMap(const FString &Branch, const FString &CL, bool bEnginePackages);
@@ -194,6 +218,7 @@ private:
 	FChangeInfo						NondeterministicSummary;
 	FChangeInfo						IndirectNondeterministicSummary;
 	TMap<FName, FChangeInfo>		ChangeSummaryByClass;
+	TMap<FName, FDeterminismInfo>   DeterminismByClass;
 	TMap<FName, FChangeInfo>		ChangeInfoByAsset;
 	TMap<int32, FChangeInfo>		ChangeSummaryByChangelist;
 	TMap<FName, FName>				AssetPathToClassName;
diff --git a/Engine/Source/Editor/UnrealEd/Private/Commandlets/DiffAssetRegistriesCommandlet.cpp b/Engine/Source/Editor/UnrealEd/Private/Commandlets/DiffAssetRegistriesCommandlet.cpp
--- a/Engine/Source/Editor/UnrealEd/Private/Commandlets/DiffAssetRegistriesCommandlet.cpp
+++ b/Engine/Source/Editor/UnrealEd/Private/Commandlets/DiffAssetRegistriesCommandlet.cpp
@@ -711,13 +711,27 @@ void UDiffAssetRegistriesCommandlet::SummarizeDeterminism()
 			}
 		}
 
+		TArray<int32> ChunkIds = GetAssetChunks(NewState, AssetPath);
+		FName ClassName = GetClassName(NewState, AssetPath);
 		if (classification == 'c')
 		{
 			NondeterministicSummary += ChangeInfo;
+			
+			for (int32 ChunkId : ChunkIds)
+			{
+				ChangesByChunk.FindOrAdd(ChunkId).Determinism.FindOrAdd(ClassName).AddDirect(ChangeInfo);
+			}
+			DeterminismByClass.FindOrAdd(ClassName).AddDirect(ChangeInfo);
 		}
 		else if (classification == 'n')
 		{
 			IndirectNondeterministicSummary += ChangeInfo;
+
+			for (int32 ChunkId : ChunkIds)
+			{
+				ChangesByChunk.FindOrAdd(ChunkId).Determinism.FindOrAdd(ClassName).AddIndirect(ChangeInfo);
+			}
+			DeterminismByClass.FindOrAdd(ClassName).AddIndirect(ChangeInfo);
 		}
 	}
 }
@@ -909,7 +923,7 @@ void UDiffAssetRegistriesCommandlet::LogChangedFiles(FArchive *CSVFile, FString
 	}
 }
 
-void UDiffAssetRegistriesCommandlet::LogClassSummary(FArchive *CSVFile, const FString& HeaderPrefix, const TMap<FName, FChangeInfo>& InChangeInfoByAsset, bool bDoWarnings)
+void UDiffAssetRegistriesCommandlet::LogClassSummary(FArchive *CSVFile, const FString& HeaderPrefix, const TMap<FName, FChangeInfo>& InChangeInfoByAsset, bool bDoWarnings, TMap<FName, FDeterminismInfo> DeterminismInfo)
 {
 	const float InvToMB = 1.0 / (1024 * 1024);
 	FString HeaderSpacer = (HeaderPrefix.Len() ? TEXT(" ") : TEXT(""));
@@ -956,9 +970,15 @@ void UDiffAssetRegistriesCommandlet::LogClassSummary(FArchive *CSVFile, const FS
 			//We'll need to display the header, since this is the first row that has sufficient changes
 			if (CSVFile)
 			{
+				FString ExtraNondeterminismHeader = TEXT("");
+				if (DeterminismInfo.Num())
+				{
+					ExtraNondeterminismHeader = TEXT(",DirectNondeterministicCount,DirectNondeterministicSize,IndirectNondeterministicCount,IndirectNondeterministicSize");
+				}
+
 				CSVFile->Logf(TEXT(""));
 				CSVFile->Logf(TEXT("%s%sClass Summary"), *HeaderPrefix, *HeaderSpacer);
-				CSVFile->Logf(TEXT("Name,Percentage,TotalSize,Adds,AddedSize,Changes,ChangesSize,Deletes,DeletedSize,Unchanged,UnchangedSize"));
+				CSVFile->Logf(TEXT("Name,Percentage,TotalCount,TotalSize,Adds,AddedSize,Changes,ChangesSize,Deletes,DeletedSize,Unchanged,UnchangedSize%s"), *ExtraNondeterminismHeader);
 			}
 
 			bChangesPastThreshold = true;
@@ -966,14 +986,40 @@ void UDiffAssetRegistriesCommandlet::LogClassSummary(FArchive *CSVFile, const FS
 
 		if (CSVFile)
 		{
-			CSVFile->Logf(TEXT("%s,%0.02f,%lld,%lld,%lld,%lld,%lld,%lld,%lld,%lld,%lld"),
+			FString ExtraNondeterminismData = TEXT("");
+			if (DeterminismInfo.Num())
+			{
+				int64 DirectCount = 0;
+				int64 IndirectCount = 0;
+
+				int64 DirectSize = 0;
+				int64 IndirectSize = 0;
+				if (DeterminismInfo.Contains(ClassName))
+				{
+					int64 TotalCount = Changes.GetTotalChangeCount();
+					DirectCount = DeterminismInfo[ClassName].DirectCount;
+					IndirectCount = DeterminismInfo[ClassName].IndirectCount;
+
+					DirectSize = DeterminismInfo[ClassName].DirectSize;
+					IndirectSize = DeterminismInfo[ClassName].IndirectSize;
+				}
+
+
+				ExtraNondeterminismData = FString::Printf(TEXT(",%lld,%lld,%lld,%lld"),
+					DirectCount, DirectSize, IndirectCount, IndirectSize);
+			}
+
+
+			CSVFile->Logf(TEXT("%s,%0.02f,%lld,%lld,%lld,%lld,%lld,%lld,%lld,%lld,%lld,%lld%s"),
 				*ClassName.ToString(),
 				Changes.GetChangePercentage() * 100.0,
+				Changes.GetTotalChangeCount(),
 				Changes.GetTotalChangeSize(),
 				Changes.Adds, Changes.AddedBytes,
 				Changes.Changes, Changes.ChangedBytes,
 				Changes.Deletes, Changes.DeletedBytes,
-				Changes.Unchanged, Changes.UnchangedBytes);
+				Changes.Unchanged, Changes.UnchangedBytes,
+				*ExtraNondeterminismData);
 		}
 
 		// log summary & change
@@ -1311,7 +1357,7 @@ void UDiffAssetRegistriesCommandlet::DiffAssetRegistries(const FString& OldPath,
 
 	//Overall Class Summary
 	//Actually do the warnings in this run, since it's the overall one
-	LogClassSummary(CSVFile, TEXT("Overall"), ChangeSummaryByClass, true);
+	LogClassSummary(CSVFile, TEXT("Overall"), ChangeSummaryByClass, true, DeterminismByClass);
 
 	//Chunk-by-chunk class summaries
 	if (bGroupByChunk)
@@ -1324,7 +1370,7 @@ void UDiffAssetRegistriesCommandlet::DiffAssetRegistries(const FString& OldPath,
 			FString ChunkHeader = (ChunkID == -1) ? TEXT("Untagged") : FString::Printf(TEXT("Chunk %d"), ChunkID);
 			if (ChangesByChunk[ChunkID].ChangesByClass.Num())
 			{
-				LogClassSummary(CSVFile, *ChunkHeader, ChangesByChunk[ChunkID].ChangesByClass, false);
+				LogClassSummary(CSVFile, *ChunkHeader, ChangesByChunk[ChunkID].ChangesByClass, false, ChangesByChunk[ChunkID].Determinism);
 			}
 		}
 	}
