Exported UActorFactoryBlueprint to allow extension.

diff --git a/Engine/Source/Editor/UnrealEd/Classes/ActorFactories/ActorFactoryBlueprint.h b/Engine/Source/Editor/UnrealEd/Classes/ActorFactories/ActorFactoryBlueprint.h
--- a/Engine/Source/Editor/UnrealEd/Classes/ActorFactories/ActorFactoryBlueprint.h
+++ b/Engine/Source/Editor/UnrealEd/Classes/ActorFactories/ActorFactoryBlueprint.h
@@ -11,8 +11,8 @@
 class AActor;
 struct FAssetData;
 
-UCLASS(MinimalAPI, config=Editor, collapsecategories, hidecategories=Object)
-class UActorFactoryBlueprint : public UActorFactory
+UCLASS(config=Editor, collapsecategories, hidecategories=Object)
+class UNREALED_API UActorFactoryBlueprint : public UActorFactory
 {
 	GENERATED_UCLASS_BODY()
 
