Fixed duplicated localization keys

#rnx

diff --git a/Engine/Plugins/Enterprise/VariantManager/Source/VariantManager/Private/DisplayNodes/VariantManagerActorNode.cpp b/Engine/Plugins/Enterprise/VariantManager/Source/VariantManager/Private/DisplayNodes/VariantManagerActorNode.cpp
--- a/Engine/Plugins/Enterprise/VariantManager/Source/VariantManager/Private/DisplayNodes/VariantManagerActorNode.cpp
+++ b/Engine/Plugins/Enterprise/VariantManager/Source/VariantManager/Private/DisplayNodes/VariantManagerActorNode.cpp
@@ -69,7 +69,7 @@ void FVariantManagerActorNode::BuildContextMenu(FMenuBuilder& MenuBuilder)
 	{
 		MenuBuilder.AddSubMenu(
 			LOCTEXT("RebindActorName", "Rebind to other Actor"),
-			LOCTEXT("RebindActorName", "Connect this binding to a different actor and try reusing the captured properties"),
+			LOCTEXT("RebindActorName_Tooltip", "Connect this binding to a different actor and try reusing the captured properties"),
 			FNewMenuDelegate::CreateSP(this, &FVariantManagerActorNode::AddAssignActorSubMenu));
 	}
 
diff --git a/Engine/Plugins/Experimental/MeshModelingToolset/Source/MeshModelingToolsEditorOnly/Private/MergeMeshesTool.cpp b/Engine/Plugins/Experimental/MeshModelingToolset/Source/MeshModelingToolsEditorOnly/Private/MergeMeshesTool.cpp
--- a/Engine/Plugins/Experimental/MeshModelingToolset/Source/MeshModelingToolsEditorOnly/Private/MergeMeshesTool.cpp
+++ b/Engine/Plugins/Experimental/MeshModelingToolset/Source/MeshModelingToolsEditorOnly/Private/MergeMeshesTool.cpp
@@ -130,7 +130,7 @@ void UMergeMeshesTool::Shutdown(EToolShutdownType ShutdownType)
 		// Hide or destroy the sources
 		{
 			if (MergeProps->bDeleteInputActors) 
-				GetToolManager()->BeginUndoTransaction(LOCTEXT("MergeMeshes", "Remove Sources"));
+				GetToolManager()->BeginUndoTransaction(LOCTEXT("RemoveSources", "Remove Sources"));
 			
 			for (auto& ComponentTarget : ComponentTargets)
 			{
diff --git a/Engine/Plugins/Experimental/MeshModelingToolset/Source/MeshModelingToolsEditorOnly/Private/VoxelCSGMeshesTool.cpp b/Engine/Plugins/Experimental/MeshModelingToolset/Source/MeshModelingToolsEditorOnly/Private/VoxelCSGMeshesTool.cpp
--- a/Engine/Plugins/Experimental/MeshModelingToolset/Source/MeshModelingToolsEditorOnly/Private/VoxelCSGMeshesTool.cpp
+++ b/Engine/Plugins/Experimental/MeshModelingToolset/Source/MeshModelingToolsEditorOnly/Private/VoxelCSGMeshesTool.cpp
@@ -117,7 +117,7 @@ void UVoxelCSGMeshesTool::Shutdown(EToolShutdownType ShutdownType)
 	{
 		// Generate the result
 		{
-			GetToolManager()->BeginUndoTransaction(LOCTEXT("VoxelCSGMeshes", "Boolean Meshes"));
+			GetToolManager()->BeginUndoTransaction(LOCTEXT("BooleanMeshes", "Boolean Meshes"));
 
 			GenerateAsset(Result);
 
@@ -128,7 +128,7 @@ void UVoxelCSGMeshesTool::Shutdown(EToolShutdownType ShutdownType)
 		{
 			
 			if (CSGProps->bDeleteInputActors) 
-				GetToolManager()->BeginUndoTransaction(LOCTEXT("VoxelCSGMeshes", "Remove Sources"));
+				GetToolManager()->BeginUndoTransaction(LOCTEXT("RemoveSources", "Remove Sources"));
 			
 			for (auto& ComponentTarget : ComponentTargets)
 			{
diff --git a/Engine/Plugins/MovieScene/MovieRenderPipeline/Source/MovieRenderPipelineEditor/Private/Widgets/SMoviePipelineQueueEditor.cpp b/Engine/Plugins/MovieScene/MovieRenderPipeline/Source/MovieRenderPipelineEditor/Private/Widgets/SMoviePipelineQueueEditor.cpp
--- a/Engine/Plugins/MovieScene/MovieRenderPipeline/Source/MovieRenderPipelineEditor/Private/Widgets/SMoviePipelineQueueEditor.cpp
+++ b/Engine/Plugins/MovieScene/MovieRenderPipeline/Source/MovieRenderPipelineEditor/Private/Widgets/SMoviePipelineQueueEditor.cpp
@@ -347,7 +347,7 @@ TSharedRef<SWidget> SQueueJobListRow::GenerateWidgetForColumn(const FName& Colum
 			[
 				SNew(STextBlock)
 				// .Style(FTakeRecorderStyle::Get(), "TakeRecorder.Source.Switch")
-				.Text(LOCTEXT("PendingJobStatus_Label", "Ready"))
+				.Text(LOCTEXT("PendingJobStatusReady_Label", "Ready"))
 			]
 
 			// Progress Bar
@@ -365,7 +365,7 @@ TSharedRef<SWidget> SQueueJobListRow::GenerateWidgetForColumn(const FName& Colum
 			.HAlign(HAlign_Center)
 			[
 				SNew(STextBlock)
-				.Text(LOCTEXT("PendingJobStatus_Label", "Completed!"))
+				.Text(LOCTEXT("PendingJobStatusCompleted_Label", "Completed!"))
 			];
 	}
 
diff --git a/Engine/Source/Editor/FoliageEdit/Private/SFoliageEdit.cpp b/Engine/Source/Editor/FoliageEdit/Private/SFoliageEdit.cpp
--- a/Engine/Source/Editor/FoliageEdit/Private/SFoliageEdit.cpp
+++ b/Engine/Source/Editor/FoliageEdit/Private/SFoliageEdit.cpp
@@ -141,7 +141,7 @@ void SFoliageEdit::Construct(const FArguments& InArgs)
 						.VAlign(VAlign_Center)
 						[
 							SNew(STextBlock)
-							.Text(LOCTEXT("BrushSize", "Brush Size"))
+							.Text(LOCTEXT("BrushSize_Legacy", "Brush Size"))
 							.Font(StandardFont)
 						]
 						+ SHorizontalBox::Slot()
@@ -214,7 +214,7 @@ void SFoliageEdit::Construct(const FArguments& InArgs)
 						.VAlign(VAlign_Center)
 						[
 							SNew(STextBlock)
-							.Text(LOCTEXT("EraseDensity", "Erase Density"))
+							.Text(LOCTEXT("EraseDensity_Legacy", "Erase Density"))
 							.Font(StandardFont)
 						]
 						+ SHorizontalBox::Slot()
diff --git a/Engine/Source/Editor/UMGEditor/Private/WidgetBlueprintEditorUtils.cpp b/Engine/Source/Editor/UMGEditor/Private/WidgetBlueprintEditorUtils.cpp
--- a/Engine/Source/Editor/UMGEditor/Private/WidgetBlueprintEditorUtils.cpp
+++ b/Engine/Source/Editor/UMGEditor/Private/WidgetBlueprintEditorUtils.cpp
@@ -712,7 +712,7 @@ void FWidgetBlueprintEditorUtils::BuildReplaceWithMenu(FMenuBuilder& Menu, TShar
 					{
 						Menu.AddMenuEntry(
 							FText::Format(LOCTEXT("ReplaceWithNamedSlot", "Replace With '{0}'"), SlotNameTxt),
-							FText::Format(LOCTEXT("ReplaceWithChildTooltip", "Remove this widget and insert '{0}' content into the parent."), SlotNameTxt),
+							FText::Format(LOCTEXT("ReplaceWithNamedSlotTooltip", "Remove this widget and insert '{0}' content into the parent."), SlotNameTxt),
 							FSlateIcon(),
 							FUIAction(
 								FExecuteAction::CreateStatic(&FWidgetBlueprintEditorUtils::ReplaceWidgetWithNamedSlot, BlueprintEditor, BP, Widget, SlotName),
diff --git a/Engine/Source/Editor/UnrealEd/Private/Settings/LevelEditorPlaySettingsCustomization.h b/Engine/Source/Editor/UnrealEd/Private/Settings/LevelEditorPlaySettingsCustomization.h
--- a/Engine/Source/Editor/UnrealEd/Private/Settings/LevelEditorPlaySettingsCustomization.h
+++ b/Engine/Source/Editor/UnrealEd/Private/Settings/LevelEditorPlaySettingsCustomization.h
@@ -731,7 +731,7 @@ private:
 			Desc += LOCTEXT("MultiplayerDescription_OneProcess", "The following will all run under one UE4 instance:\n").ToString();
 			if (PlayNetMode == EPlayNetMode::PIE_Client)
 			{
-				Desc += LOCTEXT("MultiplayerDescription_DedicatedServer", "A hidden dedicated server instance will run in editor. ").ToString();
+				Desc += LOCTEXT("MultiplayerDescription_DedicatedServerHidden", "A hidden dedicated server instance will run in editor. ").ToString();
 				if (PlayNumberOfClients == 1)
 				{
 					Desc += LOCTEXT("MultiplayerDescription_EditorClient", "The editor will connect as a client. ").ToString();
@@ -802,7 +802,7 @@ private:
 			else
 			{
 				// Client requires additional dedicated server instance
-				Desc += LOCTEXT("MultiplayerDescription_DedicatedServer", "A dedicated server will open in a new window. ").ToString();
+				Desc += LOCTEXT("MultiplayerDescription_DedicatedServerNewWindow", "A dedicated server will open in a new window. ").ToString();
 				if (PlayNumberOfClients == 1)
 				{
 					Desc += LOCTEXT("MultiplayerDescription_EditorClient", "The editor will connect as a client. ").ToString();
