Fixing crash when reimporting static mesh with more than 8 LODs

Alexis.Matte

diff --git a/Engine/Source/Editor/UnrealEd/Private/Fbx/FbxStaticMeshImport.cpp b/Engine/Source/Editor/UnrealEd/Private/Fbx/FbxStaticMeshImport.cpp
--- a/Engine/Source/Editor/UnrealEd/Private/Fbx/FbxStaticMeshImport.cpp
+++ b/Engine/Source/Editor/UnrealEd/Private/Fbx/FbxStaticMeshImport.cpp
@@ -1259,6 +1259,14 @@ UStaticMesh* UnFbx::FFbxImporter::ReimportStaticMesh(UStaticMesh* Mesh, UFbxStat
 		//Import all LODs
 		for (int32 LODIndex = 1; LODIndex < FbxMeshesLod.Num(); ++LODIndex)
 		{
+			if (LODIndex >= MAX_STATIC_MESH_LODS)
+			{
+				AddTokenizedErrorMessage(FTokenizedMessage::Create(EMessageSeverity::Warning, FText::Format(
+					LOCTEXT("ImporterLimits_MaximumStaticMeshLODReach", "Reached the maximum number of LODs for a Static Mesh({0}) - discarding {1} LOD meshes."), FText::AsNumber(MAX_STATIC_MESH_LODS), FText::AsNumber(FbxMeshesLod.Num() - MAX_STATIC_MESH_LODS))
+				), FFbxErrors::Generic_Mesh_TooManyLODs);
+				break;
+			}
+
 			TArray<FbxNode*> &LODMeshesArray = FbxMeshesLod[LODIndex];
 
 			if (LODMeshesArray[0]->GetMesh() == nullptr)
@@ -1303,6 +1311,14 @@ UStaticMesh* UnFbx::FFbxImporter::ReimportStaticMesh(UStaticMesh* Mesh, UFbxStat
 				// import LOD meshes
 				for (int32 LODIndex = 1; LODIndex < NodeParent->GetChildCount(); LODIndex++)
 				{
+					if (LODIndex >= MAX_STATIC_MESH_LODS)
+					{
+						AddTokenizedErrorMessage(FTokenizedMessage::Create(EMessageSeverity::Warning, FText::Format(
+							LOCTEXT("ImporterLimits_MaximumStaticMeshLODReach", "Reached the maximum number of LODs for a Static Mesh({0}) - discarding {1} LOD meshes."), FText::AsNumber(MAX_STATIC_MESH_LODS), FText::AsNumber(NodeParent->GetChildCount() - MAX_STATIC_MESH_LODS))
+						), FFbxErrors::Generic_Mesh_TooManyLODs);
+						break;
+					}
+
 					AllNodeInLod.Empty();
 					FindAllLODGroupNode(AllNodeInLod, NodeParent, LODIndex);
 					if (AllNodeInLod.Num() > 0)
