- Integrate 11065049 : Fix crash in Niagara after static mesh reimport

#rnx

diff --git a/Engine/Source/Editor/UnrealEd/Private/FbxMeshUtils.cpp b/Engine/Source/Editor/UnrealEd/Private/FbxMeshUtils.cpp
--- a/Engine/Source/Editor/UnrealEd/Private/FbxMeshUtils.cpp
+++ b/Engine/Source/Editor/UnrealEd/Private/FbxMeshUtils.cpp
@@ -204,6 +204,7 @@ namespace FbxMeshUtils
 					}
 
 					// Update mesh component
+					BaseStaticMesh->PostEditChange();
 					BaseStaticMesh->MarkPackageDirty();
 
 					// Import worked
@@ -532,6 +533,7 @@ namespace FbxMeshUtils
 					}
 
 					bSuccess = true;
+
 					// Set LOD source filename
 					SelectedSkelMesh->GetLODInfo(SelectedLOD)->SourceImportFilename = UAssetImportData::SanitizeImportFilename(Filename, nullptr);
 					SelectedSkelMesh->GetLODInfo(SelectedLOD)->bImportWithBaseMesh = false;
