Error in multiple Cook steps: Failed import for StaticMesh .. Referenced by export StaticMeshComponent

diff --git a/Engine/Source/Editor/UnrealEd/Private/Commandlets/ContentCommandlets.cpp b/Engine/Source/Editor/UnrealEd/Private/Commandlets/ContentCommandlets.cpp
--- a/Engine/Source/Editor/UnrealEd/Private/Commandlets/ContentCommandlets.cpp
+++ b/Engine/Source/Editor/UnrealEd/Private/Commandlets/ContentCommandlets.cpp
@@ -1649,6 +1649,16 @@ void UResavePackagesCommandlet::PerformAdditionalOperations(class UWorld* World,
 					Builder.PreviewBuild();
 				}
 
+				// Get the list of packages that needs to be saved after cluster rebuilding.
+				TSet<UPackage*> PackagesToSave;
+				for (ULevel* Level : World->GetLevels())
+				{
+					if (Level->bIsVisible)
+					{
+						PackagesToSave.Add(Level->GetOutermost());
+					}
+				}
+
 				if (bGenerateMeshProxies || bForceProxyGeneration)
 				{
 					Builder.BuildMeshesForLODActors(bForceProxyGeneration);
@@ -1671,12 +1681,9 @@ void UResavePackagesCommandlet::PerformAdditionalOperations(class UWorld* World,
 					GShaderCompilingManager->ProcessAsyncResults(false, false);
 				}
 
-				// Get the list of packages needs to be saved.
-				TSet<UPackage*> PackagesToSave;
+				// Get the list of packages needs to be saved after proxy mesh generation.
 				for(ULevel* Level : World->GetLevels())
 				{
-					PackagesToSave.Add(Level->GetOutermost());
-
 					if(Level->bIsVisible)
 					{
 						Builder.GetMeshesPackagesToSave(Level, PackagesToSave);
diff --git a/Engine/Source/Editor/UnrealEd/Private/HierarchicalLOD.cpp b/Engine/Source/Editor/UnrealEd/Private/HierarchicalLOD.cpp
--- a/Engine/Source/Editor/UnrealEd/Private/HierarchicalLOD.cpp
+++ b/Engine/Source/Editor/UnrealEd/Private/HierarchicalLOD.cpp
@@ -302,6 +302,7 @@ int32 HashLODActorForClusterComparison(ALODActor* LODActor)
 
 	HashValue = HashCombine(HashValue, LODActor->LODLevel);
 	HashValue = HashCombine(HashValue, LODActor->SubActors.Num());
+	HashValue = HashCombine(HashValue, GetTypeHash(LODActor->GetLODActorTag()));
 
 	TArray<AActor*> Actors = LODActor->SubActors;
 	Actors.Sort();
@@ -812,26 +813,32 @@ void FHierarchicalLODBuilder::BuildMeshesForLODActors(bool bForceAll)
 					UHLODProxy* Proxy = Utilities->CreateOrRetrieveLevelHLODProxy(LevelIter, LODIndex);
 					UPackage* AssetsOuter = Proxy->GetOutermost();
 					checkf(AssetsOuter != nullptr, TEXT("Failed to created outer for generated HLOD assets"));
-					AssetsOuter->Modify();
 
-					int32 LODActorIndex = 0;
-					TArray<ALODActor*>& LODLevel = LODLevelActors[LODIndex];
-					for (ALODActor* Actor : LODLevel)
+					if (AssetsOuter)
 					{
-						SlowTask.EnterProgressFrame(100.0f / (float)NumLODActors, FText::Format(LOCTEXT("HierarchicalLOD_BuildLODActorMeshesProgress", "Building LODActor Mesh {0} of {1} (LOD Level {2})"), FText::AsNumber(LODActorIndex), FText::AsNumber(LODLevelActors[LODIndex].Num()), FText::AsNumber(LODIndex + 1)));
+						int32 LODActorIndex = 0;
+						TArray<ALODActor*>& LODLevel = LODLevelActors[LODIndex];
+						for (ALODActor* Actor : LODLevel)
+						{
+							SlowTask.EnterProgressFrame(100.0f / (float)NumLODActors, FText::Format(LOCTEXT("HierarchicalLOD_BuildLODActorMeshesProgress", "Building LODActor Mesh {0} of {1} (LOD Level {2})"), FText::AsNumber(LODActorIndex), FText::AsNumber(LODLevelActors[LODIndex].Num()), FText::AsNumber(LODIndex + 1)));
 
-						bool bBuildSuccessful = Utilities->BuildStaticMeshForLODActor(Actor, AssetsOuter, BuildLODLevelSettings[LODIndex], BaseMaterial);
+							bool bBuildSuccessful = Utilities->BuildStaticMeshForLODActor(Actor, AssetsOuter, BuildLODLevelSettings[LODIndex], BaseMaterial);
 
-						// Report an error if the build failed
-						if (!bBuildSuccessful)
-						{
-							FMessageLog("HLODResults").Error()
-								->AddToken(FTextToken::Create(LOCTEXT("HLODError_MeshNotBuildOne", "Cannot create proxy mesh for ")))
-								->AddToken(FUObjectToken::Create(Actor))
-								->AddToken(FTextToken::Create(LOCTEXT("HLODError_MeshNotBuildTwo", " this could be caused by incorrect mesh components in the sub actors")));
-						}
+							// Report an error if the build failed
+							if (!bBuildSuccessful)
+							{
+								FMessageLog("HLODResults").Error()
+									->AddToken(FTextToken::Create(LOCTEXT("HLODError_MeshNotBuildOne", "Cannot create proxy mesh for ")))
+									->AddToken(FUObjectToken::Create(Actor))
+									->AddToken(FTextToken::Create(LOCTEXT("HLODError_MeshNotBuildTwo", " this could be caused by incorrect mesh components in the sub actors")));
+							}
+							else
+							{
+								AssetsOuter->Modify();
+							}
 
-						++LODActorIndex;
+							++LODActorIndex;
+						}
 					}
 				}
 			}
