If nothing is selected in the level editor, don't attempt to set the blueprint editors' 'object being debugged' to the first instance of the blueprint-being-editedt's type

diff --git a/Engine/Source/Editor/UnrealEd/Private/Kismet2/BlueprintEditorUtils.cpp b/Engine/Source/Editor/UnrealEd/Private/Kismet2/BlueprintEditorUtils.cpp
--- a/Engine/Source/Editor/UnrealEd/Private/Kismet2/BlueprintEditorUtils.cpp
+++ b/Engine/Source/Editor/UnrealEd/Private/Kismet2/BlueprintEditorUtils.cpp
@@ -8366,43 +8366,6 @@ void FBlueprintEditorUtils::FindAndSetDebuggableBlueprintInstances()
 					}
 				}
 			}
-			// Nothing of this type selected, just find any instance of one of these objects.
-			if (!bFoundItemToDebug)
-			{
-				for (TObjectIterator<UObject> It; It; ++It)
-				{
-					AActor* ObjectAsActor = Cast<AActor>( *It );
-					UWorld* ActorWorld = ObjectAsActor ? ObjectAsActor->GetWorld() : nullptr;
-					if( ActorWorld && ( ActorWorld->WorldType != EWorldType::EditorPreview) && ActorWorld->WorldType != EWorldType::Inactive )
-					{
-						if( IsObjectADebugCandidate(ObjectAsActor, EachBlueprint, true/*bInDisallowDerivedBlueprints*/ ) == true )
-						{
-							EachBlueprint->SetObjectBeingDebugged( ObjectAsActor );
-							bFoundItemToDebug = true;
-							BlueprintsToRefresh.Add( EachBlueprint );
-							break;
-						}
-						else if( SimilarInstanceUnselected == nullptr)
-						{
-							// If we haven't found a similar unselected instance already check for one now
-							if( IsObjectADebugCandidate(ObjectAsActor, EachBlueprint, false/*bInDisallowDerivedBlueprints*/ ) == true )
-							{
-								SimilarInstanceUnselected = ObjectAsActor;
-							}						
-						}
-					}
-				}
-			}
-
-			// If we didn't find and exact type match, but we did find a related type use that.
-			if( bFoundItemToDebug == false )
-			{
-				if( ( SimilarInstanceSelected != nullptr ) || ( SimilarInstanceUnselected != nullptr ) )
-				{
-					EachBlueprint->SetObjectBeingDebugged( SimilarInstanceSelected != nullptr ? SimilarInstanceSelected : SimilarInstanceUnselected );
-					BlueprintsToRefresh.Add( EachBlueprint );
-				}
-			}
 		}
 
 		// Refresh all blueprint windows that we have made a change to the debugging selection of
