Localization PAK chunking fixes

 - Fixed the localization data for some chunked assets being left in the catch-all chunk.
   - The loc gather doesn't always use valid object paths, but it does always have a valid package path at the root.
   - This changes the logic to first try and convert to a package path, then see if it's a valid package path.
 - Fixed the loaded chunk IDs list in FLocalizationResourceTextSource containing duplicate items.
   - FCoreDelegates::OnPakFileMounted is currently called for each sub-chunk PAK file.
 - Fixed IChunkDataGenerator::GenerateChunkDataFiles incorrectly being called for each sub-chunk.
 - Fixed FLocalizationChunkDataGenerator potentially having the incorrect list of localization targets and cultures when cooking for multiple platforms at once.
 - Added a packaging setting to control which chunk ID should be the catch-all chunk for localization.
   - This defaults to chunk 0, but can be overridden on a per-project basis.

#rnx

diff --git a/Engine/Source/Editor/UnrealEd/Classes/Settings/ProjectPackagingSettings.h b/Engine/Source/Editor/UnrealEd/Classes/Settings/ProjectPackagingSettings.h
--- a/Engine/Source/Editor/UnrealEd/Classes/Settings/ProjectPackagingSettings.h
+++ b/Engine/Source/Editor/UnrealEd/Classes/Settings/ProjectPackagingSettings.h
@@ -285,6 +285,10 @@ public:
 	UPROPERTY(config, EditAnywhere, Category=Packaging, AdvancedDisplay)
 	TArray<FString> LocalizationTargetsToChunk;
 
+	/** The chunk ID that should be used as the catch-all chunk for any non-asset localized strings */
+	UPROPERTY(config, EditAnywhere, Category=Packaging, AdvancedDisplay)
+	int32 LocalizationTargetCatchAllChunkId = 0;
+
 	/**
 	 * Cook all things in the project content directory
 	 */
diff --git a/Engine/Source/Editor/UnrealEd/Private/Commandlets/AssetRegistryGenerator.cpp b/Engine/Source/Editor/UnrealEd/Private/Commandlets/AssetRegistryGenerator.cpp
--- a/Engine/Source/Editor/UnrealEd/Private/Commandlets/AssetRegistryGenerator.cpp
+++ b/Engine/Source/Editor/UnrealEd/Private/Commandlets/AssetRegistryGenerator.cpp
@@ -434,7 +434,7 @@ bool FAssetRegistryGenerator::GenerateStreamingInstallManifest(int64 InExtraFlav
 			}
 
 			// Allow the extra data generation steps to run and add their output to the manifest
-			if (ChunkDataGenerators.Num() > 0)
+			if (ChunkDataGenerators.Num() > 0 && SubChunkIndex == 0)
 			{
 				TSet<FName> PackagesInChunk;
 				PackagesInChunk.Reserve(Manifest->Num());
diff --git a/Engine/Source/Editor/UnrealEd/Private/CookOnTheFlyServer.cpp b/Engine/Source/Editor/UnrealEd/Private/CookOnTheFlyServer.cpp
--- a/Engine/Source/Editor/UnrealEd/Private/CookOnTheFlyServer.cpp
+++ b/Engine/Source/Editor/UnrealEd/Private/CookOnTheFlyServer.cpp
@@ -6973,8 +6973,6 @@ void UCookOnTheFlyServer::StartCookByTheBook( const FCookByTheBookStartupOptions
 
 	// Find all the localized packages and map them back to their source package
 	{
-		UE_LOG(LogCook, Display, TEXT("Discovering localized assets"));
-
 		TArray<FString> AllCulturesToCook = CookByTheBookStartupOptions.CookCultures;
 		for (const FString& CultureName : CookByTheBookStartupOptions.CookCultures)
 		{
@@ -6986,6 +6984,8 @@ void UCookOnTheFlyServer::StartCookByTheBook( const FCookByTheBookStartupOptions
 		}
 		AllCulturesToCook.Sort();
 
+		UE_LOG(LogCook, Display, TEXT("Discovering localized assets for cultures: %s"), *FString::Join(AllCulturesToCook, TEXT(", ")));
+
 		TArray<FString> RootPaths;
 		FPackageName::QueryRootContentPaths(RootPaths);
 
@@ -7032,7 +7032,7 @@ void UCookOnTheFlyServer::StartCookByTheBook( const FCookByTheBookStartupOptions
 		{
 			for (TPair<FName, FAssetRegistryGenerator*>& Pair : RegistryGenerators)
 			{
-				Pair.Value->RegisterChunkDataGenerator(MakeShared<FLocalizationChunkDataGenerator>(MoveTemp(LocalizationTargetsToChunk), MoveTemp(AllCulturesToCook)));
+				Pair.Value->RegisterChunkDataGenerator(MakeShared<FLocalizationChunkDataGenerator>(PackagingSettings->LocalizationTargetCatchAllChunkId, LocalizationTargetsToChunk, AllCulturesToCook));
 			}
 		}
 	}
diff --git a/Engine/Source/Editor/UnrealEd/Private/LocalizationChunkDataGenerator.cpp b/Engine/Source/Editor/UnrealEd/Private/LocalizationChunkDataGenerator.cpp
--- a/Engine/Source/Editor/UnrealEd/Private/LocalizationChunkDataGenerator.cpp
+++ b/Engine/Source/Editor/UnrealEd/Private/LocalizationChunkDataGenerator.cpp
@@ -11,16 +11,23 @@
 
 DEFINE_LOG_CATEGORY_STATIC(LogLocalizationChunkDataGenerator, Log, All);
 
-FLocalizationChunkDataGenerator::FLocalizationChunkDataGenerator(TArray<FString> InLocalizationTargetsToChunk, TArray<FString> InAllCulturesToCook)
-	: LocalizationTargetsToChunk(MoveTemp(InLocalizationTargetsToChunk))
+FLocalizationChunkDataGenerator::FLocalizationChunkDataGenerator(const int32 InCatchAllChunkId, TArray<FString> InLocalizationTargetsToChunk, TArray<FString> InAllCulturesToCook)
+	: CatchAllChunkId(InCatchAllChunkId)
+	, LocalizationTargetsToChunk(MoveTemp(InLocalizationTargetsToChunk))
 	, AllCulturesToCook(MoveTemp(InAllCulturesToCook))
 {
 }
 
 void FLocalizationChunkDataGenerator::GenerateChunkDataFiles(const int32 InChunkId, const TSet<FName>& InPackagesInChunk, const FString& InPlatformName, FSandboxPlatformFile* InSandboxFile, TArray<FString>& OutChunkFilenames)
 {
-	// We can skip this chunk if it's empty
-	if (InPackagesInChunk.Num() == 0)
+	// The primary chunk doesn't gain a suffix to make it unique, as it is replacing the offline localization data that is usually staged verbatim
+	const bool bIsPrimaryChunk = InChunkId == 0;
+	
+	// The catch-all chunk is the only chunk that can contain non-asset localization data (this is chunk 0 by default)
+	const bool bIsCatchAllChunk = InChunkId == CatchAllChunkId;
+
+	// We can skip empty non-primary chunks
+	if (!bIsPrimaryChunk && InPackagesInChunk.Num() == 0)
 	{
 		return;
 	}
@@ -41,10 +48,6 @@ void FLocalizationChunkDataGenerator::GenerateChunkDataFiles(const int32 InChunk
 			continue;
 		}
 
-		// Chunk 0 is the only chunk that can contain non-asset localization data, as it acts as the "catch-all" since it's always available
-		// It is also the only chunk that doesn't gain a suffix to make it unique (as it is replacing the offline localization data that is usually staged verbatim)
-		const bool bIsPrimaryChunk = InChunkId == 0;
-
 		// Prepare to produce the localization target for this chunk
 		const TArray<FString> AvailableCulturesToCook = SourceLocTextHelper->GetAllCultures();
 		const FString ChunkTargetName = TextLocalizationResourceUtil::GetLocalizationTargetNameForChunkId(SourceLocTextHelper->GetTargetName(), InChunkId);
@@ -54,19 +57,17 @@ void FLocalizationChunkDataGenerator::GenerateChunkDataFiles(const int32 InChunk
 		bool bChunkHasText = false;
 		FLocTextHelper ChunkLocTextHelper(ChunkTargetRoot, FString::Printf(TEXT("%s.manifest"), *ChunkTargetName), FString::Printf(TEXT("%s.archive"), *ChunkTargetName), SourceLocTextHelper->GetNativeCulture(), SourceLocTextHelper->GetForeignCultures(), nullptr);
 		ChunkLocTextHelper.LoadAll(ELocTextHelperLoadFlags::Create); // Create the in-memory manifest and archives
-		SourceLocTextHelper->EnumerateSourceTexts([bIsPrimaryChunk, SourceLocTextHelper, &bChunkHasText, &ChunkLocTextHelper, &InPackagesInChunk, &AvailableCulturesToCook](TSharedRef<FManifestEntry> InManifestEntry)
+		SourceLocTextHelper->EnumerateSourceTexts([bIsCatchAllChunk, SourceLocTextHelper, &bChunkHasText, &ChunkLocTextHelper, &InPackagesInChunk, &AvailableCulturesToCook](TSharedRef<FManifestEntry> InManifestEntry)
 		{
 			for (const FManifestContext& ManifestContext : InManifestEntry->Contexts)
 			{
-				bool bIncludeInChunk = false;
-				if (FPackageName::IsValidObjectPath(ManifestContext.SourceLocation))
-				{
-					const FName SourcePackageName = *FPackageName::ObjectPathToPackageName(ManifestContext.SourceLocation);
-					bIncludeInChunk = InPackagesInChunk.Contains(SourcePackageName);
-				}
-				else
+				bool bIncludeInChunk = bIsCatchAllChunk;
 				{
-					bIncludeInChunk = bIsPrimaryChunk;
+					const FString SourcePackageName = FPackageName::ObjectPathToPackageName(ManifestContext.SourceLocation);
+					if (FPackageName::IsValidLongPackageName(SourcePackageName))
+					{
+						bIncludeInChunk = InPackagesInChunk.Contains(*SourcePackageName);
+					}
 				}
 
 				if (bIncludeInChunk)
@@ -87,7 +88,7 @@ void FLocalizationChunkDataGenerator::GenerateChunkDataFiles(const int32 InChunk
 			return true; // continue enumeration
 		}, false);
 
-		// If this chunk has no localization data then we can skip generating the LocRes (unless it's the primary chunk)
+		// We can skip empty non-primary chunks
 		if (!bIsPrimaryChunk && !bChunkHasText)
 		{
 			continue;
diff --git a/Engine/Source/Editor/UnrealEd/Private/LocalizationChunkDataGenerator.h b/Engine/Source/Editor/UnrealEd/Private/LocalizationChunkDataGenerator.h
--- a/Engine/Source/Editor/UnrealEd/Private/LocalizationChunkDataGenerator.h
+++ b/Engine/Source/Editor/UnrealEd/Private/LocalizationChunkDataGenerator.h
@@ -13,7 +13,7 @@ class FLocTextHelper;
 class FLocalizationChunkDataGenerator : public IChunkDataGenerator
 {
 public:
-	FLocalizationChunkDataGenerator(TArray<FString> InLocalizationTargetsToChunk, TArray<FString> InAllCulturesToCook);
+	FLocalizationChunkDataGenerator(const int32 InCatchAllChunkId, TArray<FString> InLocalizationTargetsToChunk, TArray<FString> InAllCulturesToCook);
 	virtual ~FLocalizationChunkDataGenerator() = default;
 
 	//~ IChunkDataGenerator
@@ -23,6 +23,9 @@ private:
 	/** Update CachedLocalizationTargetData if needed */
 	void ConditionalCacheLocalizationTargetData();
 
+	/** The chunk ID that should be used as the catch-all chunk for any non-asset localized strings */
+	int32 CatchAllChunkId;
+
 	/** List of localization targets that should be chunked */
 	TArray<FString> LocalizationTargetsToChunk;
 
diff --git a/Engine/Source/Runtime/Core/Private/Internationalization/LocalizationResourceTextSource.h b/Engine/Source/Runtime/Core/Private/Internationalization/LocalizationResourceTextSource.h
--- a/Engine/Source/Runtime/Core/Private/Internationalization/LocalizationResourceTextSource.h
+++ b/Engine/Source/Runtime/Core/Private/Internationalization/LocalizationResourceTextSource.h
@@ -26,7 +26,7 @@ public:
 	 */
 	void RegisterChunkId(const int32 InChunkId)
 	{
-		ChunkIds.Add(InChunkId);
+		ChunkIds.AddUnique(InChunkId);
 	}
 
 	/**
