Merging CL#11083537

Fix drag and drop file import
#rnx

diff --git a/Engine/Source/Editor/UnrealEd/Private/Subsystems/ImportSubsystem.cpp b/Engine/Source/Editor/UnrealEd/Private/Subsystems/ImportSubsystem.cpp
--- a/Engine/Source/Editor/UnrealEd/Private/Subsystems/ImportSubsystem.cpp
+++ b/Engine/Source/Editor/UnrealEd/Private/Subsystems/ImportSubsystem.cpp
@@ -121,7 +121,7 @@ void UImportSubsystem::Deinitialize()
 void UImportSubsystem::ImportNextTick(const TArray<FString>& Files, const FString& DestinationPath)
 {
 	FAssetToolsModule& AssetToolsModule = FModuleManager::LoadModuleChecked<FAssetToolsModule>("AssetTools");
-	if (AssetToolsModule.Get().GetWritableFolderBlacklist()->PassesStartsWithFilter(DestinationPath))
+	if (!AssetToolsModule.Get().GetWritableFolderBlacklist()->PassesStartsWithFilter(DestinationPath))
 	{
 		AssetToolsModule.Get().NotifyBlockedByWritableFolderFilter();
 		return;
