Don't skip generation of the localization data for the catch-all chunk, even if it's empty

#rnx

diff --git a/Engine/Source/Editor/UnrealEd/Private/LocalizationChunkDataGenerator.cpp b/Engine/Source/Editor/UnrealEd/Private/LocalizationChunkDataGenerator.cpp
--- a/Engine/Source/Editor/UnrealEd/Private/LocalizationChunkDataGenerator.cpp
+++ b/Engine/Source/Editor/UnrealEd/Private/LocalizationChunkDataGenerator.cpp
@@ -26,8 +26,9 @@ void FLocalizationChunkDataGenerator::GenerateChunkDataFiles(const int32 InChunk
 	// The catch-all chunk is the only chunk that can contain non-asset localization data (this is chunk 0 by default)
 	const bool bIsCatchAllChunk = InChunkId == CatchAllChunkId;
 
-	// We can skip empty non-primary chunks
-	if (!bIsPrimaryChunk && InPackagesInChunk.Num() == 0)
+	// We can skip empty non-primary and non-catch-all chunks
+	const bool bCanSkipChunk = !bIsPrimaryChunk && !bIsCatchAllChunk;
+	if (bCanSkipChunk && InPackagesInChunk.Num() == 0)
 	{
 		return;
 	}
@@ -88,8 +89,8 @@ void FLocalizationChunkDataGenerator::GenerateChunkDataFiles(const int32 InChunk
 			return true; // continue enumeration
 		}, false);
 
-		// We can skip empty non-primary chunks
-		if (!bIsPrimaryChunk && !bChunkHasText)
+		// We can skip empty non-primary and non-catch-all chunks
+		if (bCanSkipChunk && !bChunkHasText)
 		{
 			continue;
 		}
