Fix PIE crashing on start when a camera is showing a picture-in-picture preview

diff --git a/Engine/Source/Editor/UnrealEd/Private/EditorEngine.cpp b/Engine/Source/Editor/UnrealEd/Private/EditorEngine.cpp
--- a/Engine/Source/Editor/UnrealEd/Private/EditorEngine.cpp
+++ b/Engine/Source/Editor/UnrealEd/Private/EditorEngine.cpp
@@ -1823,7 +1823,7 @@ void UEditorEngine::Tick( float DeltaSeconds, bool bIdleMode )
 					continue;
 				}
 
-				if ( ViewportClient->IsVisible() && (!bAllWindowsHidden || ViewportClient->WantsDrawWhenAppIsHidden()) )
+				if (ViewportClient->IsVisible() && (!bAllWindowsHidden || ViewportClient->WantsDrawWhenAppIsHidden()) && ViewportClient->GetScene())
 				{
 					// Only update ortho viewports if that mode is turned on, the viewport client we are about to update is orthographic and the current editing viewport is orthographic and tracking mouse movement.
 					bUpdateLinkedOrthoViewports = GetDefault<ULevelEditorViewportSettings>()->bUseLinkedOrthographicViewports && ViewportClient->IsOrtho() && GCurrentLevelEditingViewportClient && GCurrentLevelEditingViewportClient->IsOrtho() && GCurrentLevelEditingViewportClient->IsTracking();
@@ -2060,7 +2060,10 @@ bool UEditorEngine::UpdateSingleViewportClient(FEditorViewportClient* InViewport
 		// Add view information for perspective viewports.
 		if( InViewportClient->IsPerspective() )
 		{
-			InViewportClient->GetWorld()->ViewLocationsRenderedLastFrame.Add(InViewportClient->GetViewLocation());
+			if (UWorld* ViewportClientWorld = InViewportClient->GetWorld())
+			{
+				ViewportClientWorld->ViewLocationsRenderedLastFrame.Add(InViewportClient->GetViewLocation());
+			}
 	
 			// If we're currently simulating in editor, then we'll need to make sure that sub-levels are streamed in.
 			// When using PIE, this normally happens by UGameViewportClient::Draw().  But for SIE, we need to do
