Added support for data sets bigger than 2GB to the ShaderPipelineCache tools commandlet

diff --git a/Engine/Source/Editor/UnrealEd/Private/Commandlets/ShaderPipelineCacheToolsCommandlet.cpp b/Engine/Source/Editor/UnrealEd/Private/Commandlets/ShaderPipelineCacheToolsCommandlet.cpp
--- a/Engine/Source/Editor/UnrealEd/Private/Commandlets/ShaderPipelineCacheToolsCommandlet.cpp
+++ b/Engine/Source/Editor/UnrealEd/Private/Commandlets/ShaderPipelineCacheToolsCommandlet.cpp
@@ -16,8 +16,16 @@ const TCHAR* STABLE_CSV_EXT = TEXT("stablepc.csv");
 const TCHAR* STABLE_CSV_COMPRESSED_EXT = TEXT("stablepc.csv.compressed");
 const TCHAR* STABLE_COMPRESSED_EXT = TEXT(".compressed");
 const int32  STABLE_COMPRESSED_EXT_LEN = 11; // len of ".compressed";
-const int32  STABLE_COMPRESSED_VER = 1;
+const int32  STABLE_COMPRESSED_VER = 2;
+const int64  STABLE_MAX_CHUNK_SIZE = MAX_int32 - 100 * 1024 * 1024;
 
+struct FSCDataChunk
+{
+	FSCDataChunk() : UncomressedOutputLines(), OutputLinesAr(UncomressedOutputLines) {}
+
+	TArray<uint8> UncomressedOutputLines;
+	FMemoryWriter OutputLinesAr;
+};
 
 void ExpandWildcards(TArray<FString>& Parts)
 {
@@ -81,7 +89,7 @@ void LoadStableSCL(TMultiMap<FStableShaderKeyAndValue, FSHAHash>& StableMap, con
 	UE_LOG(LogShaderPipelineCacheTools, Display, TEXT("Loaded %d shader info lines"), SourceFileContents.Num() - 1);
 }
 
-static bool LoadAndDecompressStableCSV(const FString& Filename, TArray<uint8>& UncompressedData)
+static bool LoadAndDecompressStableCSV(const FString& Filename, TArray<FString>& OutputLines)
 {
 	bool bResult = false;
 	FArchive* Ar = IFileManager::Get().CreateFileReader(*Filename);
@@ -90,22 +98,41 @@ static bool LoadAndDecompressStableCSV(const FString& Filename, TArray<uint8>& U
 		if (Ar->TotalSize() > 8)
 		{
 			int32 CompressedVersion = 0;
-			int32 UncompressedSize = 0;
-			int32 CompressedSize = 0;
+			int32 NumChunks = 1;
 			
 			Ar->Serialize(&CompressedVersion, sizeof(int32));
-			Ar->Serialize(&UncompressedSize, sizeof(int32));
-			Ar->Serialize(&CompressedSize, sizeof(int32));
-			
-			TArray<uint8> CompressedData;
-			CompressedData.SetNumUninitialized(CompressedSize);
-			Ar->Serialize(CompressedData.GetData(), CompressedSize);
+			if (CompressedVersion > 1)
+			{
+				Ar->Serialize(&NumChunks, sizeof(int32));
+			}
 
-			UncompressedData.SetNumUninitialized(UncompressedSize);
-			bResult = FCompression::UncompressMemory(NAME_Zlib, UncompressedData.GetData(), UncompressedSize, CompressedData.GetData(), CompressedSize);
-			if (!bResult)
+			for (int32 Index = 0; Index < NumChunks; ++Index)
 			{
-				UE_LOG(LogShaderPipelineCacheTools, Display, TEXT("Failed to decompress file %s"), *Filename);
+				int32 UncompressedSize = 0;
+				int32 CompressedSize = 0;
+
+				Ar->Serialize(&UncompressedSize, sizeof(int32));
+				Ar->Serialize(&CompressedSize, sizeof(int32));
+
+				TArray<uint8> CompressedData;
+				CompressedData.SetNumUninitialized(CompressedSize);
+				Ar->Serialize(CompressedData.GetData(), CompressedSize);
+
+				TArray<uint8> UncompressedData;
+				UncompressedData.SetNumUninitialized(UncompressedSize);
+				bResult = FCompression::UncompressMemory(NAME_Zlib, UncompressedData.GetData(), UncompressedSize, CompressedData.GetData(), CompressedSize);
+				if (!bResult)
+				{
+					UE_LOG(LogShaderPipelineCacheTools, Display, TEXT("Failed to decompress file %s"), *Filename);
+				}
+
+				FMemoryReader MemArchive(UncompressedData);
+				FString LineCSV;
+				while (!MemArchive.AtEnd())
+				{
+					MemArchive << LineCSV;
+					OutputLines.Add(LineCSV);
+				}
 			}
 		}
 		else
@@ -128,16 +155,8 @@ static bool LoadStableCSV(const FString& Filename, TArray<FString>& OutputLines)
 	bool bResult = false;
 	if (Filename.EndsWith(STABLE_CSV_COMPRESSED_EXT))
 	{
-		TArray<uint8> DecompressedData;
-		if (LoadAndDecompressStableCSV(Filename, DecompressedData))
+		if (LoadAndDecompressStableCSV(Filename, OutputLines))
 		{
-			FMemoryReader MemArchive(DecompressedData);
-			FString LineCSV;
-			while (!MemArchive.AtEnd())
-			{
-				MemArchive << LineCSV;
-				OutputLines.Add(LineCSV);
-			}
 			bResult = true;
 		}
 	}
@@ -149,42 +168,69 @@ static bool LoadStableCSV(const FString& Filename, TArray<FString>& OutputLines)
 	return bResult;
 }
 
-static int64 SaveStableCSV(const FString& Filename, const TArray<uint8>& UncompressedData)
+static int64 SaveStableCSV(const FString& Filename, const FSCDataChunk* DataChunks, int32 NumChunks)
 {
 	if (Filename.EndsWith(STABLE_CSV_COMPRESSED_EXT))
 	{
-		int32 UncompressedSize = UncompressedData.Num();
-		UE_LOG(LogShaderPipelineCacheTools, Display, TEXT("Compressing output, size = %.1fKB"), UncompressedSize/1024.f);
-		int32 CompressedSize = FCompression::CompressMemoryBound(NAME_Zlib, UncompressedSize);
-		TArray<uint8> CompressedData;
-		CompressedData.SetNumZeroed(CompressedSize);
+		UE_LOG(LogShaderPipelineCacheTools, Display, TEXT("Compressing output, %d chunks"), NumChunks);
+
+		struct FSCCompressedChunk
+		{
+			FSCCompressedChunk(int32 UncompressedSize)
+			{
+				CompressedSize = FCompression::CompressMemoryBound(NAME_Zlib, UncompressedSize);
+				CompressedData.SetNumZeroed(CompressedSize);
+			}
 
-		if (FCompression::CompressMemory(NAME_Zlib, CompressedData.GetData(), CompressedSize, UncompressedData.GetData(), UncompressedSize))
+			TArray<uint8> CompressedData;
+			int32 CompressedSize;
+		};
+
+		TArray<FSCCompressedChunk> CompressedChunks;
+
+		for (int32 Index = 0; Index < NumChunks; ++Index)
 		{
-			FArchive* Ar = IFileManager::Get().CreateFileWriter(*Filename);
-			if (!Ar)
+			const FSCDataChunk& Chunk = DataChunks[Index];
+			CompressedChunks.Add(FSCCompressedChunk(Chunk.UncomressedOutputLines.Num()));
+
+			UE_LOG(LogShaderPipelineCacheTools, Display, TEXT("Compressing chunk %d, size = %.1fKB"), Index, Chunk.UncomressedOutputLines.Num() / 1024.f);
+			if (FCompression::CompressMemory(NAME_Zlib, CompressedChunks[Index].CompressedData.GetData(), CompressedChunks[Index].CompressedSize, Chunk.UncomressedOutputLines.GetData(), Chunk.UncomressedOutputLines.Num()) == false)
 			{
-				UE_LOG(LogShaderPipelineCacheTools, Fatal, TEXT("Failed to open %s"), *Filename);
-				return -1;
+				UE_LOG(LogShaderPipelineCacheTools, Fatal, TEXT("Failed to compress chunk %d (%.1f KB)"), Index, Chunk.UncomressedOutputLines.Num() / 1024.f);
 			}
+		}
+
+		FArchive* Ar = IFileManager::Get().CreateFileWriter(*Filename);
+		if (!Ar)
+		{
+			UE_LOG(LogShaderPipelineCacheTools, Fatal, TEXT("Failed to open %s"), *Filename);
+			return -1;
+		}
 			
-			int32 CompressedVersion = STABLE_COMPRESSED_VER;
+		int32 CompressedVersion = STABLE_COMPRESSED_VER;
 			
-			Ar->Serialize(&CompressedVersion, sizeof(int32));
+		Ar->Serialize(&CompressedVersion, sizeof(int32));
+		Ar->Serialize(&NumChunks, sizeof(int32));
+
+		for (int32 Index = 0; Index < NumChunks; ++Index)
+		{
+			int32 UncompressedSize = DataChunks[Index].UncomressedOutputLines.Num();
+			int32 CompressedSize = CompressedChunks[Index].CompressedSize;
 			Ar->Serialize(&UncompressedSize, sizeof(int32));
 			Ar->Serialize(&CompressedSize, sizeof(int32));
-			Ar->Serialize(CompressedData.GetData(), CompressedSize);
-			delete Ar;
-		}
-		else
-		{
-			UE_LOG(LogShaderPipelineCacheTools, Fatal, TEXT("Failed to compress (%.1f KB)"), UncompressedSize/1024.f);
-			return -1;
+			Ar->Serialize(CompressedChunks[Index].CompressedData.GetData(), CompressedSize);
 		}
+
+		delete Ar;
 	}
 	else
 	{
-		FMemoryReader MemArchive(UncompressedData);
+		if (NumChunks > 1)
+		{
+			UE_LOG(LogShaderPipelineCacheTools, Fatal, TEXT("SaveStableCSV does not support saving uncompressed files larger than 2GB."));
+		}
+
+		FMemoryReader MemArchive(DataChunks[0].UncomressedOutputLines);
 		FString CombinedCSV;
 		FString LineCSV;
 		while (!MemArchive.AtEnd())
@@ -770,10 +816,11 @@ int32 ExpandPSOSC(const TArray<FString>& Tokens)
 		return 1;
 	}
 
+
 	int32 NumLines = 0;
-	TArray<uint8> UncomressedOutputLines;
-	FMemoryWriter OutputLinesAr(UncomressedOutputLines);
-	TSet<FString> DeDup;
+	FSCDataChunk DataChunks[16];
+	int32 CurrentChunk = 0;
+	TSet<uint32> DeDup;
 
 	{
 		FString PSOLine = FString::Printf(TEXT("\"%s\""), *FPipelineCacheFileFormatPSO::CommonHeaderLine());
@@ -783,7 +830,7 @@ int32 ExpandPSOSC(const TArray<FString>& Tokens)
 			PSOLine += FString::Printf(TEXT(",\"shaderslot%d: %s\""), SlotIndex, *FStableShaderKeyAndValue::HeaderLine());
 		}
 
-		OutputLinesAr << PSOLine;
+		DataChunks[CurrentChunk].OutputLinesAr << PSOLine;
 		NumLines++;
 	}
 
@@ -853,11 +900,15 @@ int32 ExpandPSOSC(const TArray<FString>& Tokens)
 				}
 			}
 
-
-			if (!DeDup.Contains(PSOLine))
+			const uint32 PSOLineHash = FCrc::MemCrc32(PSOLine.GetCharArray().GetData(), sizeof(TCHAR) * PSOLine.Len());
+			if (!DeDup.Contains(PSOLineHash))
 			{
-				DeDup.Add(PSOLine);
-				OutputLinesAr << PSOLine;
+				DeDup.Add(PSOLineHash);
+				if (DataChunks[CurrentChunk].OutputLinesAr.TotalSize() + (int64)((PSOLine.Len() + 1) * sizeof(TCHAR)) >= STABLE_MAX_CHUNK_SIZE)
+				{
+					++CurrentChunk;
+				}
+				DataChunks[CurrentChunk].OutputLinesAr << PSOLine;
 				NumLines++;
 			}
 		}
@@ -897,7 +948,7 @@ int32 ExpandPSOSC(const TArray<FString>& Tokens)
 		}
 	}
 
-	int64 FileSize = SaveStableCSV(OutputFilename, UncomressedOutputLines);
+	int64 FileSize = SaveStableCSV(OutputFilename, DataChunks, CurrentChunk + 1);
 	if (FileSize < 1)
 	{
 		return 1;
@@ -1404,7 +1455,7 @@ int32 DiffStable(const TArray<FString>& Tokens)
 
 int32 DecompressCSV(const TArray<FString>& Tokens)
 {
-	TArray<uint8> DecompressedData;
+	TArray<FString> DecompressedData;
 	for (int32 TokenIndex = 0; TokenIndex < Tokens.Num(); TokenIndex++)
 	{
 		const FString& CompressedFilename = Tokens[TokenIndex];
@@ -1417,17 +1468,22 @@ int32 DecompressCSV(const TArray<FString>& Tokens)
 		DecompressedData.Reset();
 		if (LoadAndDecompressStableCSV(CompressedFilename, DecompressedData))
 		{
-			FMemoryReader MemArchive(DecompressedData);
-			FString LineCSV;
-			while (!MemArchive.AtEnd())
+			FString FilenameCSV = CompressedFilename.LeftChop(STABLE_COMPRESSED_EXT_LEN);
+			FPlatformFileManager::Get().GetPlatformFile().DeleteFile(*FilenameCSV);
+
+			for (const FString& LineCSV : DecompressedData)
 			{
-				MemArchive << LineCSV;
 				CombinedCSV.Append(LineCSV);
 				CombinedCSV.Append(LINE_TERMINATOR);
+
+				if ((int64)(CombinedCSV.Len() * sizeof(TCHAR)) >= (int64)(MAX_int32 - 1024 * 1024))
+				{
+					FFileHelper::SaveStringToFile(CombinedCSV, *FilenameCSV, FFileHelper::EEncodingOptions::AutoDetect, &IFileManager::Get(), FILEWRITE_Append);
+					CombinedCSV.Empty();
+				}
 			}
 
-			FString FilenameCSV = CompressedFilename.LeftChop(STABLE_COMPRESSED_EXT_LEN);
-			FFileHelper::SaveStringToFile(CombinedCSV, *FilenameCSV);
+			FFileHelper::SaveStringToFile(CombinedCSV, *FilenameCSV, FFileHelper::EEncodingOptions::AutoDetect, &IFileManager::Get(), FILEWRITE_Append);
 		}
 	}
 
