Optimizations for the ShaderPipelineCache tools commandlet

diff --git a/Engine/Source/Editor/UnrealEd/Private/Commandlets/ShaderPipelineCacheToolsCommandlet.cpp b/Engine/Source/Editor/UnrealEd/Private/Commandlets/ShaderPipelineCacheToolsCommandlet.cpp
--- a/Engine/Source/Editor/UnrealEd/Private/Commandlets/ShaderPipelineCacheToolsCommandlet.cpp
+++ b/Engine/Source/Editor/UnrealEd/Private/Commandlets/ShaderPipelineCacheToolsCommandlet.cpp
@@ -64,7 +64,7 @@ void ExpandWildcards(TArray<FString>& Parts)
 }
 
 
-void LoadStableSCL(TMultiMap<FStableShaderKeyAndValue, FSHAHash>& StableMap, const FString& Filename)
+void LoadStableSCL(TMultiMap<int32, FSHAHash>& StableMap, TArray<FStableShaderKeyAndValue>& StableArray, const FString& Filename)
 {
 	UE_LOG(LogShaderPipelineCacheTools, Display, TEXT("Loading %s...."), *Filename);
 	TArray<FString> SourceFileContents;
@@ -78,7 +78,8 @@ void LoadStableSCL(TMultiMap<FStableShaderKeyAndValue, FSHAHash>& StableMap, con
 			FMemory::Memzero(Item);
 			Item.ParseFromString(SourceFileContents[Index]);
 			check(!(Item.OutputHash == FSHAHash()));
-			StableMap.AddUnique(Item, Item.OutputHash);
+			int32 ItemIndex = StableArray.Add(Item);
+			StableMap.AddUnique(ItemIndex, Item.OutputHash);
 		}
 	}
 	if (SourceFileContents.Num() < 1)
@@ -339,7 +340,7 @@ int32 DumpPSOSC(FString& Token)
 	return 0;
 }
 
-static void PrintShaders(const TMap<FSHAHash, TArray<FStableShaderKeyAndValue>>& InverseMap, const FSHAHash& Shader, const TCHAR *Label)
+static void PrintShaders(const TMap<FSHAHash, TArray<int32>>& InverseMap, TArray<FStableShaderKeyAndValue>& StableArray, const FSHAHash& Shader, const TCHAR *Label)
 {
 	UE_LOG(LogShaderPipelineCacheTools, Display, TEXT(" -- %s"), Label);
 
@@ -348,25 +349,25 @@ static void PrintShaders(const TMap<FSHAHash, TArray<FStableShaderKeyAndValue>>&
 		UE_LOG(LogShaderPipelineCacheTools, Display, TEXT("    null"));
 		return;
 	}
-	const TArray<FStableShaderKeyAndValue>* Out = InverseMap.Find(Shader);
+	const TArray<int32>* Out = InverseMap.Find(Shader);
 	if (!Out)
 	{
 		UE_LOG(LogShaderPipelineCacheTools, Display, TEXT("    No shaders found with hash %s"), *Shader.ToString());
 		return;
 	}
-	for (const FStableShaderKeyAndValue& Item : *Out)
+	for (const int32& Item : *Out)
 	{
-		UE_LOG(LogShaderPipelineCacheTools, Display, TEXT("    %s"), *Item.ToString());
+		UE_LOG(LogShaderPipelineCacheTools, Display, TEXT("    %s"), *StableArray[Item].ToString());
 	}
 }
 
-static bool GetStableShadersAndZeroHash(const TMap<FSHAHash, TArray<FStableShaderKeyAndValue>>& InverseMap, const FSHAHash& Shader, TArray<FStableShaderKeyAndValue>& StableShaders, bool& bOutAnyActiveButMissing)
+static bool GetStableShadersAndZeroHash(const TMap<FSHAHash, TArray<int32>>& InverseMap, TArray<FStableShaderKeyAndValue>& StableArray, const FSHAHash& Shader, TArray<int32>& StableShaders, bool& bOutAnyActiveButMissing)
 {
 	if (Shader == FSHAHash())
 	{
 		return false;
 	}
-	const TArray<FStableShaderKeyAndValue>* Out = InverseMap.Find(Shader);
+	const TArray<int32>* Out = InverseMap.Find(Shader);
 	if (!Out)
 	{
 		UE_LOG(LogShaderPipelineCacheTools, Warning, TEXT("No shaders found with hash %s"), *Shader.ToString());
@@ -375,17 +376,17 @@ static bool GetStableShadersAndZeroHash(const TMap<FSHAHash, TArray<FStableShade
 		return false;
 	}
 	StableShaders.Reserve(Out->Num());
-	for (const FStableShaderKeyAndValue& Item : *Out)
+	for (const int32& Item : *Out)
 	{
-		FStableShaderKeyAndValue Temp = Item;
+		FStableShaderKeyAndValue Temp = StableArray[Item];
 		Temp.OutputHash = FSHAHash();
-		if (StableShaders.Contains(Temp))
+		if (StableShaders.Contains(Item))
 		{
 			UE_LOG(LogShaderPipelineCacheTools, Error, TEXT("Duplicate stable shader. This is bad because it means our stable key is not exhaustive."));
-			UE_LOG(LogShaderPipelineCacheTools, Error, TEXT(" %s"), *Item.ToString());
+			UE_LOG(LogShaderPipelineCacheTools, Error, TEXT(" %s"), *Temp.ToString());
 			continue;
 		}
-		StableShaders.Add(Temp);
+		StableShaders.Add(Item);
 	}
 	return true;
 }
@@ -435,13 +436,13 @@ bool CouldBeUsedTogether(const FStableShaderKeyAndValue& A, const FStableShaderK
 
 int32 DumpSCLCSV(const FString& Token)
 {
-
-	TMultiMap<FStableShaderKeyAndValue, FSHAHash> StableMap;
-	LoadStableSCL(StableMap, Token);
+	TArray<FStableShaderKeyAndValue> StableArray;
+	TMultiMap<int32, FSHAHash> StableMap;
+	LoadStableSCL(StableMap, StableArray, Token);
 	UE_LOG(LogShaderPipelineCacheTools, Display, TEXT("    %s"), *FStableShaderKeyAndValue::HeaderLine());
 	for (const auto& Pair : StableMap)
 	{
-		FStableShaderKeyAndValue Temp(Pair.Key);
+		FStableShaderKeyAndValue Temp(StableArray[Pair.Key]);
 		Temp.OutputHash = Pair.Value;
 		UE_LOG(LogShaderPipelineCacheTools, Display, TEXT("    %s"), *Temp.ToString());
 	}
@@ -465,7 +466,7 @@ struct FPermuation
 	FStableShaderKeyAndValue Slots[SF_NumFrequencies];
 };
 
-void GeneratePermuations(TArray<FPermuation>& Permutations, FPermuation& WorkingPerm, int32 SlotIndex , const TArray<FStableShaderKeyAndValue> StableShadersPerSlot[SF_NumFrequencies], const bool ActivePerSlot[SF_NumFrequencies])
+void GeneratePermuations(TArray<FPermuation>& Permutations, FPermuation& WorkingPerm, int32 SlotIndex , const TArray<int32> StableShadersPerSlot[SF_NumFrequencies], const TArray<FStableShaderKeyAndValue>& StableArray, const bool ActivePerSlot[SF_NumFrequencies])
 {
 	check(SlotIndex >= 0 && SlotIndex <= SF_NumFrequencies);
 	while (SlotIndex < SF_NumFrequencies && !ActivePerSlot[SlotIndex])
@@ -488,7 +489,7 @@ void GeneratePermuations(TArray<FPermuation>& Permutations, FPermuation& Working
 				continue;
 			}
 			check(SlotIndex != SF_Compute && SlotIndexInner != SF_Compute); // there is never any matching with compute shaders
-			if (!CouldBeUsedTogether(StableShadersPerSlot[SlotIndex][StableIndex], WorkingPerm.Slots[SlotIndexInner]))
+			if (!CouldBeUsedTogether(StableArray[StableShadersPerSlot[SlotIndex][StableIndex]], WorkingPerm.Slots[SlotIndexInner]))
 			{
 				bKeep = false;
 				break;
@@ -498,20 +499,21 @@ void GeneratePermuations(TArray<FPermuation>& Permutations, FPermuation& Working
 		{
 			continue;
 		}
-		WorkingPerm.Slots[SlotIndex] = StableShadersPerSlot[SlotIndex][StableIndex];
-		GeneratePermuations(Permutations, WorkingPerm, SlotIndex + 1, StableShadersPerSlot, ActivePerSlot);
+		WorkingPerm.Slots[SlotIndex] = StableArray[StableShadersPerSlot[SlotIndex][StableIndex]];
+		GeneratePermuations(Permutations, WorkingPerm, SlotIndex + 1, StableShadersPerSlot, StableArray, ActivePerSlot);
 	}
 }
 
 int32 ExpandPSOSC(const TArray<FString>& Tokens)
 {
-	TMultiMap<FStableShaderKeyAndValue, FSHAHash> StableMap;
+	TArray<FStableShaderKeyAndValue> StableArray;
+	TMultiMap<int32, FSHAHash> StableMap;
 	check(Tokens.Last().EndsWith(STABLE_CSV_EXT) || Tokens.Last().EndsWith(STABLE_CSV_COMPRESSED_EXT));
 	for (int32 Index = 0; Index < Tokens.Num() - 1; Index++)
 	{
 		if (Tokens[Index].EndsWith(TEXT(".scl.csv")))
 		{
-			LoadStableSCL(StableMap, Tokens[Index]);
+			LoadStableSCL(StableMap, StableArray, Tokens[Index]);
 		}
 	}
 	if (!StableMap.Num())
@@ -524,7 +526,7 @@ int32 ExpandPSOSC(const TArray<FString>& Tokens)
 		UE_LOG(LogShaderPipelineCacheTools, Verbose, TEXT("    %s"), *FStableShaderKeyAndValue::HeaderLine());
 		for (const auto& Pair : StableMap)
 		{
-			FStableShaderKeyAndValue Temp(Pair.Key);
+			FStableShaderKeyAndValue Temp(StableArray[Pair.Key]);
 			Temp.OutputHash = Pair.Value;
 			UE_LOG(LogShaderPipelineCacheTools, Verbose, TEXT("    %s"), *Temp.ToString());
 		}
@@ -532,7 +534,7 @@ int32 ExpandPSOSC(const TArray<FString>& Tokens)
 	//self test
 	for (const auto& Pair : StableMap)
 	{
-		FStableShaderKeyAndValue Item(Pair.Key);
+		FStableShaderKeyAndValue Item(StableArray[Pair.Key]);
 		Item.OutputHash = Pair.Value;
 		check(Pair.Value != FSHAHash());
 		FString TestString = Item.ToString();
@@ -614,7 +616,7 @@ int32 ExpandPSOSC(const TArray<FString>& Tokens)
 
 		for (const auto& Pair : StableMap)
 		{
-			FStableShaderKeyAndValue Temp(Pair.Key);
+			FStableShaderKeyAndValue Temp(StableArray[Pair.Key]);
 			Temp.OutputHash = Pair.Value;
 			InverseMap.FindOrAdd(Pair.Value).Add(Temp.ToString());
 		}
@@ -641,13 +643,12 @@ int32 ExpandPSOSC(const TArray<FString>& Tokens)
 			}
 		}
 	}
-	TMap<FSHAHash, TArray<FStableShaderKeyAndValue>> InverseMap;
+
+	TMap<FSHAHash, TArray<int32>> InverseMap;
 
 	for (const auto& Pair : StableMap)
 	{
-		FStableShaderKeyAndValue Item(Pair.Key);
-		Item.OutputHash = Pair.Value;
-		InverseMap.FindOrAdd(Item.OutputHash).AddUnique(Item);
+		InverseMap.FindOrAdd(Pair.Value).AddUnique(Pair.Key);
 	}
 
 	int32 TotalStablePSOs = 0;
@@ -677,22 +678,22 @@ int32 ExpandPSOSC(const TArray<FString>& Tokens)
 	{ 
 		NumExamined++;
 		check(SF_Vertex == 0 && SF_Compute == 5);
-		TArray<FStableShaderKeyAndValue> StableShadersPerSlot[SF_NumFrequencies];
+		TArray<int32> StableShadersPerSlot[SF_NumFrequencies];
 		bool ActivePerSlot[SF_NumFrequencies] = { false };
 
 		bool OutAnyActiveButMissing = false;
 
 		if (Item.Type == FPipelineCacheFileFormatPSO::DescriptorType::Compute)
 		{
-			ActivePerSlot[SF_Compute] = GetStableShadersAndZeroHash(InverseMap, Item.ComputeDesc.ComputeShader, StableShadersPerSlot[SF_Compute], OutAnyActiveButMissing);
+			ActivePerSlot[SF_Compute] = GetStableShadersAndZeroHash(InverseMap, StableArray, Item.ComputeDesc.ComputeShader, StableShadersPerSlot[SF_Compute], OutAnyActiveButMissing);
 		}
 		else
 		{
-			ActivePerSlot[SF_Vertex] = GetStableShadersAndZeroHash(InverseMap, Item.GraphicsDesc.VertexShader, StableShadersPerSlot[SF_Vertex], OutAnyActiveButMissing);
-			ActivePerSlot[SF_Pixel] = GetStableShadersAndZeroHash(InverseMap, Item.GraphicsDesc.FragmentShader, StableShadersPerSlot[SF_Pixel], OutAnyActiveButMissing);
-			ActivePerSlot[SF_Geometry] = GetStableShadersAndZeroHash(InverseMap, Item.GraphicsDesc.GeometryShader, StableShadersPerSlot[SF_Geometry], OutAnyActiveButMissing);
-			ActivePerSlot[SF_Hull] = GetStableShadersAndZeroHash(InverseMap, Item.GraphicsDesc.HullShader, StableShadersPerSlot[SF_Hull], OutAnyActiveButMissing);
-			ActivePerSlot[SF_Domain] = GetStableShadersAndZeroHash(InverseMap, Item.GraphicsDesc.DomainShader, StableShadersPerSlot[SF_Domain], OutAnyActiveButMissing);
+			ActivePerSlot[SF_Vertex] = GetStableShadersAndZeroHash(InverseMap, StableArray, Item.GraphicsDesc.VertexShader, StableShadersPerSlot[SF_Vertex], OutAnyActiveButMissing);
+			ActivePerSlot[SF_Pixel] = GetStableShadersAndZeroHash(InverseMap, StableArray, Item.GraphicsDesc.FragmentShader, StableShadersPerSlot[SF_Pixel], OutAnyActiveButMissing);
+			ActivePerSlot[SF_Geometry] = GetStableShadersAndZeroHash(InverseMap, StableArray, Item.GraphicsDesc.GeometryShader, StableShadersPerSlot[SF_Geometry], OutAnyActiveButMissing);
+			ActivePerSlot[SF_Hull] = GetStableShadersAndZeroHash(InverseMap, StableArray, Item.GraphicsDesc.HullShader, StableShadersPerSlot[SF_Hull], OutAnyActiveButMissing);
+			ActivePerSlot[SF_Domain] = GetStableShadersAndZeroHash(InverseMap, StableArray, Item.GraphicsDesc.DomainShader, StableShadersPerSlot[SF_Domain], OutAnyActiveButMissing);
 		}
 
 
@@ -701,16 +702,16 @@ int32 ExpandPSOSC(const TArray<FString>& Tokens)
 			UE_LOG(LogShaderPipelineCacheTools, Display, TEXT("PSO had an active shader slot that did not match any current shaders, ignored."));
 			if (Item.Type == FPipelineCacheFileFormatPSO::DescriptorType::Compute)
 			{
-				PrintShaders(InverseMap, Item.ComputeDesc.ComputeShader, TEXT("ComputeShader"));
+				PrintShaders(InverseMap, StableArray, Item.ComputeDesc.ComputeShader, TEXT("ComputeShader"));
 			}
 			else
 			{
 				UE_LOG(LogShaderPipelineCacheTools, Display, TEXT("   %s"), *Item.GraphicsDesc.StateToString());
-				PrintShaders(InverseMap, Item.GraphicsDesc.VertexShader, TEXT("VertexShader"));
-				PrintShaders(InverseMap, Item.GraphicsDesc.FragmentShader, TEXT("FragmentShader"));
-				PrintShaders(InverseMap, Item.GraphicsDesc.GeometryShader, TEXT("GeometryShader"));
-				PrintShaders(InverseMap, Item.GraphicsDesc.HullShader, TEXT("HullShader"));
-				PrintShaders(InverseMap, Item.GraphicsDesc.DomainShader, TEXT("DomainShader"));
+				PrintShaders(InverseMap, StableArray, Item.GraphicsDesc.VertexShader, TEXT("VertexShader"));
+				PrintShaders(InverseMap, StableArray, Item.GraphicsDesc.FragmentShader, TEXT("FragmentShader"));
+				PrintShaders(InverseMap, StableArray, Item.GraphicsDesc.GeometryShader, TEXT("GeometryShader"));
+				PrintShaders(InverseMap, StableArray, Item.GraphicsDesc.HullShader, TEXT("HullShader"));
+				PrintShaders(InverseMap, StableArray, Item.GraphicsDesc.DomainShader, TEXT("DomainShader"));
 			}
 			continue;
 		}
@@ -740,7 +741,7 @@ int32 ExpandPSOSC(const TArray<FString>& Tokens)
 						bool bFoundCompat = false;
 						for (int32 StableIndexInner = 0; StableIndexInner < StableShadersPerSlot[SlotIndexInner].Num(); StableIndexInner++)
 						{
-							if (CouldBeUsedTogether(StableShadersPerSlot[SlotIndex][StableIndex], StableShadersPerSlot[SlotIndexInner][StableIndexInner]))
+							if (CouldBeUsedTogether(StableArray[StableShadersPerSlot[SlotIndex][StableIndex]], StableArray[StableShadersPerSlot[SlotIndexInner][StableIndexInner]]))
 							{
 								bFoundCompat = true;
 								break;
@@ -774,11 +775,11 @@ int32 ExpandPSOSC(const TArray<FString>& Tokens)
 				UE_LOG(LogShaderPipelineCacheTools, Warning, TEXT("PSO did not create any stable PSOs! (no cross shader slot compatibility)"));
 				UE_LOG(LogShaderPipelineCacheTools, Warning, TEXT("   %s"), *Item.GraphicsDesc.StateToString());
 
-				PrintShaders(InverseMap, Item.GraphicsDesc.VertexShader, TEXT("VertexShader"));
-				PrintShaders(InverseMap, Item.GraphicsDesc.FragmentShader, TEXT("FragmentShader"));
-				PrintShaders(InverseMap, Item.GraphicsDesc.GeometryShader, TEXT("GeometryShader"));
-				PrintShaders(InverseMap, Item.GraphicsDesc.HullShader, TEXT("HullShader"));
-				PrintShaders(InverseMap, Item.GraphicsDesc.DomainShader, TEXT("DomainShader"));
+				PrintShaders(InverseMap, StableArray, Item.GraphicsDesc.VertexShader, TEXT("VertexShader"));
+				PrintShaders(InverseMap, StableArray, Item.GraphicsDesc.FragmentShader, TEXT("FragmentShader"));
+				PrintShaders(InverseMap, StableArray, Item.GraphicsDesc.GeometryShader, TEXT("GeometryShader"));
+				PrintShaders(InverseMap, StableArray, Item.GraphicsDesc.HullShader, TEXT("HullShader"));
+				PrintShaders(InverseMap, StableArray, Item.GraphicsDesc.DomainShader, TEXT("DomainShader"));
 
 				continue;
 			}
@@ -796,7 +797,7 @@ int32 ExpandPSOSC(const TArray<FString>& Tokens)
 
 		TArray<FPermuation>& Permutations(Current.Permutations);
 		FPermuation WorkingPerm;
-		GeneratePermuations(Permutations, WorkingPerm, 0, StableShadersPerSlot, ActivePerSlot);
+		GeneratePermuations(Permutations, WorkingPerm, 0, StableShadersPerSlot, StableArray, ActivePerSlot);
 		if (!Permutations.Num())
 		{
 			UE_LOG(LogShaderPipelineCacheTools, Error, TEXT("PSO did not create any stable PSOs! (somehow)"));
@@ -1017,14 +1018,15 @@ void BuildDateSortedListOfFiles(const TArray<FString>& TokenList, FilenameFilter
 
 int32 BuildPSOSC(const TArray<FString>& Tokens)
 {
-	TMultiMap<FStableShaderKeyAndValue, FSHAHash> StableMap;
+	TArray<FStableShaderKeyAndValue> StableArray;
+	TMultiMap<int32, FSHAHash> StableMap;
 	check(Tokens.Last().EndsWith(TEXT(".upipelinecache")));
 
 	for (int32 Index = 0; Index < Tokens.Num() - 1; Index++)
 	{
 		if (Tokens[Index].EndsWith(TEXT(".scl.csv")))
 		{
-			LoadStableSCL(StableMap, Tokens[Index]);
+			LoadStableSCL(StableMap, StableArray, Tokens[Index]);
 		}
 	}
 	if (UE_LOG_ACTIVE(LogShaderPipelineCacheTools, Verbose))
@@ -1032,7 +1034,7 @@ int32 BuildPSOSC(const TArray<FString>& Tokens)
 		UE_LOG(LogShaderPipelineCacheTools, Verbose, TEXT("    %s"), *FStableShaderKeyAndValue::HeaderLine());
 		for (const auto& Pair : StableMap)
 		{
-			FStableShaderKeyAndValue Temp(Pair.Key);
+			FStableShaderKeyAndValue Temp(StableArray[Pair.Key]);
 			Temp.OutputHash = Pair.Value;
 			UE_LOG(LogShaderPipelineCacheTools, Verbose, TEXT("    %s"), *Temp.ToString());
 		}
@@ -1040,7 +1042,7 @@ int32 BuildPSOSC(const TArray<FString>& Tokens)
 	//self test
 	for (const auto& Pair : StableMap)
 	{
-		FStableShaderKeyAndValue Item(Pair.Key);
+		FStableShaderKeyAndValue Item(StableArray[Pair.Key]);
 		Item.OutputHash = Pair.Value;
 		check(Pair.Value != FSHAHash());
 		FString TestString = Item.ToString();
@@ -1140,19 +1142,21 @@ int32 BuildPSOSC(const TArray<FString>& Tokens)
 
 				FSHAHash Match;
 				int32 Count = 0;
-				for (auto Iter = StableMap.CreateConstKeyIterator(Shader); Iter; ++Iter)
+				for (auto Iter : StableArray)
 				{
-					check(Iter.Value() != FSHAHash());
-					Match = Iter.Value();
-					if (TargetPlatform == NAME_None)
-					{
-						TargetPlatform = Iter.Key().TargetPlatform;
-					}
-					else
+					if (Iter == Shader)
 					{
-						check(TargetPlatform == Iter.Key().TargetPlatform);
+						Match = Iter.OutputHash;
+						if (TargetPlatform == NAME_None)
+						{
+							TargetPlatform = Iter.TargetPlatform;
+						}
+						else
+						{
+							check(TargetPlatform == Iter.TargetPlatform);
+						}
+						Count++;
 					}
-					Count++;
 				}
 
 				if (!Count)
