Fixed non-unity editor build

diff --git a/Engine/Source/Editor/UnrealEd/Private/Commandlets/ShaderPipelineCacheToolsCommandlet.cpp b/Engine/Source/Editor/UnrealEd/Private/Commandlets/ShaderPipelineCacheToolsCommandlet.cpp
--- a/Engine/Source/Editor/UnrealEd/Private/Commandlets/ShaderPipelineCacheToolsCommandlet.cpp
+++ b/Engine/Source/Editor/UnrealEd/Private/Commandlets/ShaderPipelineCacheToolsCommandlet.cpp
@@ -9,6 +9,7 @@
 #include "ShaderPipelineCache.h"
 #include "Serialization/MemoryReader.h"
 #include "Serialization/MemoryWriter.h"
+#include "HAL/PlatformFilemanager.h"
 
 DEFINE_LOG_CATEGORY_STATIC(LogShaderPipelineCacheTools, Log, All);
 
