Calm down statis analysis

diff --git a/Engine/Source/Editor/UnrealEd/Private/Subsystems/AssetEditorSubsystem.cpp b/Engine/Source/Editor/UnrealEd/Private/Subsystems/AssetEditorSubsystem.cpp
--- a/Engine/Source/Editor/UnrealEd/Private/Subsystems/AssetEditorSubsystem.cpp
+++ b/Engine/Source/Editor/UnrealEd/Private/Subsystems/AssetEditorSubsystem.cpp
@@ -274,6 +274,9 @@ bool UAssetEditorSubsystem::OpenEditorForAsset(UObject* Asset, const EToolkitMod
 		return false;
 	}
 
+	// Fix static analysis warning
+	check(Asset);
+
 	// @todo toolkit minor: When "Edit Here" happens in a different level editor from the one that an asset is already
 	//    being edited within, we should decide whether to disallow "Edit Here" in that case, or to close the old asset
 	//    editor and summon it in the new level editor, or to just foreground the old level editor (current behavior)
