Fix static analysis error
#rnx

diff --git a/Engine/Source/Editor/UnrealEd/Private/Subsystems/AssetEditorSubsystem.cpp b/Engine/Source/Editor/UnrealEd/Private/Subsystems/AssetEditorSubsystem.cpp
--- a/Engine/Source/Editor/UnrealEd/Private/Subsystems/AssetEditorSubsystem.cpp
+++ b/Engine/Source/Editor/UnrealEd/Private/Subsystems/AssetEditorSubsystem.cpp
@@ -376,22 +376,25 @@ bool UAssetEditorSubsystem::OpenEditorForAsset(UObject* Asset, const EToolkitMod
 	}
 	// Must check Asset here in addition to at the beginning of the function, because if the asset was destroyed and recreated it might not be found correctly
 	// Do not add to recently opened asset list if this is a level-associated asset like Level Blueprint or Built Data. Their naming is not compatible
-	if (Asset && Asset->IsAsset() && !Asset->IsA(UMapBuildDataRegistry::StaticClass()))
+	if (Asset)
 	{
-		FString AssetPath = Asset->GetOuter()->GetPathName();
-		FContentBrowserModule& CBModule = FModuleManager::LoadModuleChecked<FContentBrowserModule>(TEXT("ContentBrowser"));
-		FMainMRUFavoritesList* RecentlyOpenedAssets = CBModule.GetRecentlyOpenedAssets();
-		if (RecentlyOpenedAssets && FPackageName::IsValidLongPackageName(AssetPath))
+		if (Asset->IsAsset() && !Asset->IsA(UMapBuildDataRegistry::StaticClass()))
 		{
-			RecentlyOpenedAssets->AddMRUItem(AssetPath);
+			FString AssetPath = Asset->GetOuter()->GetPathName();
+			FContentBrowserModule& CBModule = FModuleManager::LoadModuleChecked<FContentBrowserModule>(TEXT("ContentBrowser"));
+			FMainMRUFavoritesList* RecentlyOpenedAssets = CBModule.GetRecentlyOpenedAssets();
+			if (RecentlyOpenedAssets && FPackageName::IsValidLongPackageName(AssetPath))
+			{
+				RecentlyOpenedAssets->AddMRUItem(AssetPath);
+			}
 		}
-	}
 
-	const double OpenTime = FStudioAnalytics::GetAnalyticSeconds() - OpenAssetStartTime;
-	FStudioAnalytics::FireEvent_Loading(TEXT("OpenAssetEditor"), OpenTime, {
-		FAnalyticsEventAttribute(TEXT("AssetPath"), Asset->GetFullName()),
-		FAnalyticsEventAttribute(TEXT("AssetType"), Asset->GetClass()->GetName())
-	});
+		const double OpenTime = FStudioAnalytics::GetAnalyticSeconds() - OpenAssetStartTime;
+		FStudioAnalytics::FireEvent_Loading(TEXT("OpenAssetEditor"), OpenTime, {
+			FAnalyticsEventAttribute(TEXT("AssetPath"), Asset->GetFullName()),
+			FAnalyticsEventAttribute(TEXT("AssetType"), Asset->GetClass()->GetName())
+		});
+	}
 
 	return true;
 }
