Temporarily reverting my change to the ShaderPipelineCache tools commandlet as they seem to be causing problems during cooking

diff --git a/Engine/Source/Editor/UnrealEd/Private/Commandlets/ShaderPipelineCacheToolsCommandlet.cpp b/Engine/Source/Editor/UnrealEd/Private/Commandlets/ShaderPipelineCacheToolsCommandlet.cpp
--- a/Engine/Source/Editor/UnrealEd/Private/Commandlets/ShaderPipelineCacheToolsCommandlet.cpp
+++ b/Engine/Source/Editor/UnrealEd/Private/Commandlets/ShaderPipelineCacheToolsCommandlet.cpp
@@ -9,7 +9,6 @@
 #include "ShaderPipelineCache.h"
 #include "Serialization/MemoryReader.h"
 #include "Serialization/MemoryWriter.h"
-#include "HAL/PlatformFilemanager.h"
 
 DEFINE_LOG_CATEGORY_STATIC(LogShaderPipelineCacheTools, Log, All);
 
@@ -17,16 +16,8 @@ const TCHAR* STABLE_CSV_EXT = TEXT("stablepc.csv");
 const TCHAR* STABLE_CSV_COMPRESSED_EXT = TEXT("stablepc.csv.compressed");
 const TCHAR* STABLE_COMPRESSED_EXT = TEXT(".compressed");
 const int32  STABLE_COMPRESSED_EXT_LEN = 11; // len of ".compressed";
-const int32  STABLE_COMPRESSED_VER = 2;
-const int64  STABLE_MAX_CHUNK_SIZE = MAX_int32 - 100 * 1024 * 1024;
+const int32  STABLE_COMPRESSED_VER = 1;
 
-struct FSCDataChunk
-{
-	FSCDataChunk() : UncomressedOutputLines(), OutputLinesAr(UncomressedOutputLines) {}
-
-	TArray<uint8> UncomressedOutputLines;
-	FMemoryWriter OutputLinesAr;
-};
 
 void ExpandWildcards(TArray<FString>& Parts)
 {
@@ -65,7 +56,7 @@ void ExpandWildcards(TArray<FString>& Parts)
 }
 
 
-void LoadStableSCL(TMultiMap<int32, FSHAHash>& StableMap, TArray<FStableShaderKeyAndValue>& StableArray, const FString& Filename)
+void LoadStableSCL(TMultiMap<FStableShaderKeyAndValue, FSHAHash>& StableMap, const FString& Filename)
 {
 	UE_LOG(LogShaderPipelineCacheTools, Display, TEXT("Loading %s...."), *Filename);
 	TArray<FString> SourceFileContents;
@@ -79,8 +70,7 @@ void LoadStableSCL(TMultiMap<int32, FSHAHash>& StableMap, TArray<FStableShaderKe
 			FMemory::Memzero(Item);
 			Item.ParseFromString(SourceFileContents[Index]);
 			check(!(Item.OutputHash == FSHAHash()));
-			int32 ItemIndex = StableArray.Add(Item);
-			StableMap.AddUnique(ItemIndex, Item.OutputHash);
+			StableMap.AddUnique(Item, Item.OutputHash);
 		}
 	}
 	if (SourceFileContents.Num() < 1)
@@ -91,7 +81,7 @@ void LoadStableSCL(TMultiMap<int32, FSHAHash>& StableMap, TArray<FStableShaderKe
 	UE_LOG(LogShaderPipelineCacheTools, Display, TEXT("Loaded %d shader info lines"), SourceFileContents.Num() - 1);
 }
 
-static bool LoadAndDecompressStableCSV(const FString& Filename, TArray<FString>& OutputLines)
+static bool LoadAndDecompressStableCSV(const FString& Filename, TArray<uint8>& UncompressedData)
 {
 	bool bResult = false;
 	FArchive* Ar = IFileManager::Get().CreateFileReader(*Filename);
@@ -100,41 +90,22 @@ static bool LoadAndDecompressStableCSV(const FString& Filename, TArray<FString>&
 		if (Ar->TotalSize() > 8)
 		{
 			int32 CompressedVersion = 0;
-			int32 NumChunks = 1;
+			int32 UncompressedSize = 0;
+			int32 CompressedSize = 0;
 			
 			Ar->Serialize(&CompressedVersion, sizeof(int32));
-			if (CompressedVersion > 1)
-			{
-				Ar->Serialize(&NumChunks, sizeof(int32));
-			}
+			Ar->Serialize(&UncompressedSize, sizeof(int32));
+			Ar->Serialize(&CompressedSize, sizeof(int32));
+			
+			TArray<uint8> CompressedData;
+			CompressedData.SetNumUninitialized(CompressedSize);
+			Ar->Serialize(CompressedData.GetData(), CompressedSize);
 
-			for (int32 Index = 0; Index < NumChunks; ++Index)
+			UncompressedData.SetNumUninitialized(UncompressedSize);
+			bResult = FCompression::UncompressMemory(NAME_Zlib, UncompressedData.GetData(), UncompressedSize, CompressedData.GetData(), CompressedSize);
+			if (!bResult)
 			{
-				int32 UncompressedSize = 0;
-				int32 CompressedSize = 0;
-
-				Ar->Serialize(&UncompressedSize, sizeof(int32));
-				Ar->Serialize(&CompressedSize, sizeof(int32));
-
-				TArray<uint8> CompressedData;
-				CompressedData.SetNumUninitialized(CompressedSize);
-				Ar->Serialize(CompressedData.GetData(), CompressedSize);
-
-				TArray<uint8> UncompressedData;
-				UncompressedData.SetNumUninitialized(UncompressedSize);
-				bResult = FCompression::UncompressMemory(NAME_Zlib, UncompressedData.GetData(), UncompressedSize, CompressedData.GetData(), CompressedSize);
-				if (!bResult)
-				{
-					UE_LOG(LogShaderPipelineCacheTools, Display, TEXT("Failed to decompress file %s"), *Filename);
-				}
-
-				FMemoryReader MemArchive(UncompressedData);
-				FString LineCSV;
-				while (!MemArchive.AtEnd())
-				{
-					MemArchive << LineCSV;
-					OutputLines.Add(LineCSV);
-				}
+				UE_LOG(LogShaderPipelineCacheTools, Display, TEXT("Failed to decompress file %s"), *Filename);
 			}
 		}
 		else
@@ -157,8 +128,16 @@ static bool LoadStableCSV(const FString& Filename, TArray<FString>& OutputLines)
 	bool bResult = false;
 	if (Filename.EndsWith(STABLE_CSV_COMPRESSED_EXT))
 	{
-		if (LoadAndDecompressStableCSV(Filename, OutputLines))
+		TArray<uint8> DecompressedData;
+		if (LoadAndDecompressStableCSV(Filename, DecompressedData))
 		{
+			FMemoryReader MemArchive(DecompressedData);
+			FString LineCSV;
+			while (!MemArchive.AtEnd())
+			{
+				MemArchive << LineCSV;
+				OutputLines.Add(LineCSV);
+			}
 			bResult = true;
 		}
 	}
@@ -170,69 +149,42 @@ static bool LoadStableCSV(const FString& Filename, TArray<FString>& OutputLines)
 	return bResult;
 }
 
-static int64 SaveStableCSV(const FString& Filename, const FSCDataChunk* DataChunks, int32 NumChunks)
+static int64 SaveStableCSV(const FString& Filename, const TArray<uint8>& UncompressedData)
 {
 	if (Filename.EndsWith(STABLE_CSV_COMPRESSED_EXT))
 	{
-		UE_LOG(LogShaderPipelineCacheTools, Display, TEXT("Compressing output, %d chunks"), NumChunks);
-
-		struct FSCCompressedChunk
-		{
-			FSCCompressedChunk(int32 UncompressedSize)
-			{
-				CompressedSize = FCompression::CompressMemoryBound(NAME_Zlib, UncompressedSize);
-				CompressedData.SetNumZeroed(CompressedSize);
-			}
-
-			TArray<uint8> CompressedData;
-			int32 CompressedSize;
-		};
-
-		TArray<FSCCompressedChunk> CompressedChunks;
+		int32 UncompressedSize = UncompressedData.Num();
+		UE_LOG(LogShaderPipelineCacheTools, Display, TEXT("Compressing output, size = %.1fKB"), UncompressedSize/1024.f);
+		int32 CompressedSize = FCompression::CompressMemoryBound(NAME_Zlib, UncompressedSize);
+		TArray<uint8> CompressedData;
+		CompressedData.SetNumZeroed(CompressedSize);
 
-		for (int32 Index = 0; Index < NumChunks; ++Index)
+		if (FCompression::CompressMemory(NAME_Zlib, CompressedData.GetData(), CompressedSize, UncompressedData.GetData(), UncompressedSize))
 		{
-			const FSCDataChunk& Chunk = DataChunks[Index];
-			CompressedChunks.Add(FSCCompressedChunk(Chunk.UncomressedOutputLines.Num()));
-
-			UE_LOG(LogShaderPipelineCacheTools, Display, TEXT("Compressing chunk %d, size = %.1fKB"), Index, Chunk.UncomressedOutputLines.Num() / 1024.f);
-			if (FCompression::CompressMemory(NAME_Zlib, CompressedChunks[Index].CompressedData.GetData(), CompressedChunks[Index].CompressedSize, Chunk.UncomressedOutputLines.GetData(), Chunk.UncomressedOutputLines.Num()) == false)
+			FArchive* Ar = IFileManager::Get().CreateFileWriter(*Filename);
+			if (!Ar)
 			{
-				UE_LOG(LogShaderPipelineCacheTools, Fatal, TEXT("Failed to compress chunk %d (%.1f KB)"), Index, Chunk.UncomressedOutputLines.Num() / 1024.f);
+				UE_LOG(LogShaderPipelineCacheTools, Fatal, TEXT("Failed to open %s"), *Filename);
+				return -1;
 			}
-		}
-
-		FArchive* Ar = IFileManager::Get().CreateFileWriter(*Filename);
-		if (!Ar)
-		{
-			UE_LOG(LogShaderPipelineCacheTools, Fatal, TEXT("Failed to open %s"), *Filename);
-			return -1;
-		}
 			
-		int32 CompressedVersion = STABLE_COMPRESSED_VER;
+			int32 CompressedVersion = STABLE_COMPRESSED_VER;
 			
-		Ar->Serialize(&CompressedVersion, sizeof(int32));
-		Ar->Serialize(&NumChunks, sizeof(int32));
-
-		for (int32 Index = 0; Index < NumChunks; ++Index)
-		{
-			int32 UncompressedSize = DataChunks[Index].UncomressedOutputLines.Num();
-			int32 CompressedSize = CompressedChunks[Index].CompressedSize;
+			Ar->Serialize(&CompressedVersion, sizeof(int32));
 			Ar->Serialize(&UncompressedSize, sizeof(int32));
 			Ar->Serialize(&CompressedSize, sizeof(int32));
-			Ar->Serialize(CompressedChunks[Index].CompressedData.GetData(), CompressedSize);
+			Ar->Serialize(CompressedData.GetData(), CompressedSize);
+			delete Ar;
+		}
+		else
+		{
+			UE_LOG(LogShaderPipelineCacheTools, Fatal, TEXT("Failed to compress (%.1f KB)"), UncompressedSize/1024.f);
+			return -1;
 		}
-
-		delete Ar;
 	}
 	else
 	{
-		if (NumChunks > 1)
-		{
-			UE_LOG(LogShaderPipelineCacheTools, Fatal, TEXT("SaveStableCSV does not support saving uncompressed files larger than 2GB."));
-		}
-
-		FMemoryReader MemArchive(DataChunks[0].UncomressedOutputLines);
+		FMemoryReader MemArchive(UncompressedData);
 		FString CombinedCSV;
 		FString LineCSV;
 		while (!MemArchive.AtEnd())
@@ -341,7 +293,7 @@ int32 DumpPSOSC(FString& Token)
 	return 0;
 }
 
-static void PrintShaders(const TMap<FSHAHash, TArray<int32>>& InverseMap, TArray<FStableShaderKeyAndValue>& StableArray, const FSHAHash& Shader, const TCHAR *Label)
+static void PrintShaders(const TMap<FSHAHash, TArray<FStableShaderKeyAndValue>>& InverseMap, const FSHAHash& Shader, const TCHAR *Label)
 {
 	UE_LOG(LogShaderPipelineCacheTools, Display, TEXT(" -- %s"), Label);
 
@@ -350,25 +302,25 @@ static void PrintShaders(const TMap<FSHAHash, TArray<int32>>& InverseMap, TArray
 		UE_LOG(LogShaderPipelineCacheTools, Display, TEXT("    null"));
 		return;
 	}
-	const TArray<int32>* Out = InverseMap.Find(Shader);
+	const TArray<FStableShaderKeyAndValue>* Out = InverseMap.Find(Shader);
 	if (!Out)
 	{
 		UE_LOG(LogShaderPipelineCacheTools, Display, TEXT("    No shaders found with hash %s"), *Shader.ToString());
 		return;
 	}
-	for (const int32& Item : *Out)
+	for (const FStableShaderKeyAndValue& Item : *Out)
 	{
-		UE_LOG(LogShaderPipelineCacheTools, Display, TEXT("    %s"), *StableArray[Item].ToString());
+		UE_LOG(LogShaderPipelineCacheTools, Display, TEXT("    %s"), *Item.ToString());
 	}
 }
 
-static bool GetStableShadersAndZeroHash(const TMap<FSHAHash, TArray<int32>>& InverseMap, TArray<FStableShaderKeyAndValue>& StableArray, const FSHAHash& Shader, TArray<int32>& StableShaders, bool& bOutAnyActiveButMissing)
+static bool GetStableShadersAndZeroHash(const TMap<FSHAHash, TArray<FStableShaderKeyAndValue>>& InverseMap, const FSHAHash& Shader, TArray<FStableShaderKeyAndValue>& StableShaders, bool& bOutAnyActiveButMissing)
 {
 	if (Shader == FSHAHash())
 	{
 		return false;
 	}
-	const TArray<int32>* Out = InverseMap.Find(Shader);
+	const TArray<FStableShaderKeyAndValue>* Out = InverseMap.Find(Shader);
 	if (!Out)
 	{
 		UE_LOG(LogShaderPipelineCacheTools, Warning, TEXT("No shaders found with hash %s"), *Shader.ToString());
@@ -377,17 +329,17 @@ static bool GetStableShadersAndZeroHash(const TMap<FSHAHash, TArray<int32>>& Inv
 		return false;
 	}
 	StableShaders.Reserve(Out->Num());
-	for (const int32& Item : *Out)
+	for (const FStableShaderKeyAndValue& Item : *Out)
 	{
-		FStableShaderKeyAndValue Temp = StableArray[Item];
+		FStableShaderKeyAndValue Temp = Item;
 		Temp.OutputHash = FSHAHash();
-		if (StableShaders.Contains(Item))
+		if (StableShaders.Contains(Temp))
 		{
 			UE_LOG(LogShaderPipelineCacheTools, Error, TEXT("Duplicate stable shader. This is bad because it means our stable key is not exhaustive."));
-			UE_LOG(LogShaderPipelineCacheTools, Error, TEXT(" %s"), *Temp.ToString());
+			UE_LOG(LogShaderPipelineCacheTools, Error, TEXT(" %s"), *Item.ToString());
 			continue;
 		}
-		StableShaders.Add(Item);
+		StableShaders.Add(Temp);
 	}
 	return true;
 }
@@ -437,13 +389,13 @@ bool CouldBeUsedTogether(const FStableShaderKeyAndValue& A, const FStableShaderK
 
 int32 DumpSCLCSV(const FString& Token)
 {
-	TArray<FStableShaderKeyAndValue> StableArray;
-	TMultiMap<int32, FSHAHash> StableMap;
-	LoadStableSCL(StableMap, StableArray, Token);
+
+	TMultiMap<FStableShaderKeyAndValue, FSHAHash> StableMap;
+	LoadStableSCL(StableMap, Token);
 	UE_LOG(LogShaderPipelineCacheTools, Display, TEXT("    %s"), *FStableShaderKeyAndValue::HeaderLine());
 	for (const auto& Pair : StableMap)
 	{
-		FStableShaderKeyAndValue Temp(StableArray[Pair.Key]);
+		FStableShaderKeyAndValue Temp(Pair.Key);
 		Temp.OutputHash = Pair.Value;
 		UE_LOG(LogShaderPipelineCacheTools, Display, TEXT("    %s"), *Temp.ToString());
 	}
@@ -467,7 +419,7 @@ struct FPermuation
 	FStableShaderKeyAndValue Slots[SF_NumFrequencies];
 };
 
-void GeneratePermuations(TArray<FPermuation>& Permutations, FPermuation& WorkingPerm, int32 SlotIndex , const TArray<int32> StableShadersPerSlot[SF_NumFrequencies], const TArray<FStableShaderKeyAndValue>& StableArray, const bool ActivePerSlot[SF_NumFrequencies])
+void GeneratePermuations(TArray<FPermuation>& Permutations, FPermuation& WorkingPerm, int32 SlotIndex , const TArray<FStableShaderKeyAndValue> StableShadersPerSlot[SF_NumFrequencies], const bool ActivePerSlot[SF_NumFrequencies])
 {
 	check(SlotIndex >= 0 && SlotIndex <= SF_NumFrequencies);
 	while (SlotIndex < SF_NumFrequencies && !ActivePerSlot[SlotIndex])
@@ -490,7 +442,7 @@ void GeneratePermuations(TArray<FPermuation>& Permutations, FPermuation& Working
 				continue;
 			}
 			check(SlotIndex != SF_Compute && SlotIndexInner != SF_Compute); // there is never any matching with compute shaders
-			if (!CouldBeUsedTogether(StableArray[StableShadersPerSlot[SlotIndex][StableIndex]], WorkingPerm.Slots[SlotIndexInner]))
+			if (!CouldBeUsedTogether(StableShadersPerSlot[SlotIndex][StableIndex], WorkingPerm.Slots[SlotIndexInner]))
 			{
 				bKeep = false;
 				break;
@@ -500,21 +452,20 @@ void GeneratePermuations(TArray<FPermuation>& Permutations, FPermuation& Working
 		{
 			continue;
 		}
-		WorkingPerm.Slots[SlotIndex] = StableArray[StableShadersPerSlot[SlotIndex][StableIndex]];
-		GeneratePermuations(Permutations, WorkingPerm, SlotIndex + 1, StableShadersPerSlot, StableArray, ActivePerSlot);
+		WorkingPerm.Slots[SlotIndex] = StableShadersPerSlot[SlotIndex][StableIndex];
+		GeneratePermuations(Permutations, WorkingPerm, SlotIndex + 1, StableShadersPerSlot, ActivePerSlot);
 	}
 }
 
 int32 ExpandPSOSC(const TArray<FString>& Tokens)
 {
-	TArray<FStableShaderKeyAndValue> StableArray;
-	TMultiMap<int32, FSHAHash> StableMap;
+	TMultiMap<FStableShaderKeyAndValue, FSHAHash> StableMap;
 	check(Tokens.Last().EndsWith(STABLE_CSV_EXT) || Tokens.Last().EndsWith(STABLE_CSV_COMPRESSED_EXT));
 	for (int32 Index = 0; Index < Tokens.Num() - 1; Index++)
 	{
 		if (Tokens[Index].EndsWith(TEXT(".scl.csv")))
 		{
-			LoadStableSCL(StableMap, StableArray, Tokens[Index]);
+			LoadStableSCL(StableMap, Tokens[Index]);
 		}
 	}
 	if (!StableMap.Num())
@@ -527,7 +478,7 @@ int32 ExpandPSOSC(const TArray<FString>& Tokens)
 		UE_LOG(LogShaderPipelineCacheTools, Verbose, TEXT("    %s"), *FStableShaderKeyAndValue::HeaderLine());
 		for (const auto& Pair : StableMap)
 		{
-			FStableShaderKeyAndValue Temp(StableArray[Pair.Key]);
+			FStableShaderKeyAndValue Temp(Pair.Key);
 			Temp.OutputHash = Pair.Value;
 			UE_LOG(LogShaderPipelineCacheTools, Verbose, TEXT("    %s"), *Temp.ToString());
 		}
@@ -535,7 +486,7 @@ int32 ExpandPSOSC(const TArray<FString>& Tokens)
 	//self test
 	for (const auto& Pair : StableMap)
 	{
-		FStableShaderKeyAndValue Item(StableArray[Pair.Key]);
+		FStableShaderKeyAndValue Item(Pair.Key);
 		Item.OutputHash = Pair.Value;
 		check(Pair.Value != FSHAHash());
 		FString TestString = Item.ToString();
@@ -617,7 +568,7 @@ int32 ExpandPSOSC(const TArray<FString>& Tokens)
 
 		for (const auto& Pair : StableMap)
 		{
-			FStableShaderKeyAndValue Temp(StableArray[Pair.Key]);
+			FStableShaderKeyAndValue Temp(Pair.Key);
 			Temp.OutputHash = Pair.Value;
 			InverseMap.FindOrAdd(Pair.Value).Add(Temp.ToString());
 		}
@@ -644,12 +595,13 @@ int32 ExpandPSOSC(const TArray<FString>& Tokens)
 			}
 		}
 	}
-
-	TMap<FSHAHash, TArray<int32>> InverseMap;
+	TMap<FSHAHash, TArray<FStableShaderKeyAndValue>> InverseMap;
 
 	for (const auto& Pair : StableMap)
 	{
-		InverseMap.FindOrAdd(Pair.Value).AddUnique(Pair.Key);
+		FStableShaderKeyAndValue Item(Pair.Key);
+		Item.OutputHash = Pair.Value;
+		InverseMap.FindOrAdd(Item.OutputHash).AddUnique(Item);
 	}
 
 	int32 TotalStablePSOs = 0;
@@ -679,22 +631,22 @@ int32 ExpandPSOSC(const TArray<FString>& Tokens)
 	{ 
 		NumExamined++;
 		check(SF_Vertex == 0 && SF_Compute == 5);
-		TArray<int32> StableShadersPerSlot[SF_NumFrequencies];
+		TArray<FStableShaderKeyAndValue> StableShadersPerSlot[SF_NumFrequencies];
 		bool ActivePerSlot[SF_NumFrequencies] = { false };
 
 		bool OutAnyActiveButMissing = false;
 
 		if (Item.Type == FPipelineCacheFileFormatPSO::DescriptorType::Compute)
 		{
-			ActivePerSlot[SF_Compute] = GetStableShadersAndZeroHash(InverseMap, StableArray, Item.ComputeDesc.ComputeShader, StableShadersPerSlot[SF_Compute], OutAnyActiveButMissing);
+			ActivePerSlot[SF_Compute] = GetStableShadersAndZeroHash(InverseMap, Item.ComputeDesc.ComputeShader, StableShadersPerSlot[SF_Compute], OutAnyActiveButMissing);
 		}
 		else
 		{
-			ActivePerSlot[SF_Vertex] = GetStableShadersAndZeroHash(InverseMap, StableArray, Item.GraphicsDesc.VertexShader, StableShadersPerSlot[SF_Vertex], OutAnyActiveButMissing);
-			ActivePerSlot[SF_Pixel] = GetStableShadersAndZeroHash(InverseMap, StableArray, Item.GraphicsDesc.FragmentShader, StableShadersPerSlot[SF_Pixel], OutAnyActiveButMissing);
-			ActivePerSlot[SF_Geometry] = GetStableShadersAndZeroHash(InverseMap, StableArray, Item.GraphicsDesc.GeometryShader, StableShadersPerSlot[SF_Geometry], OutAnyActiveButMissing);
-			ActivePerSlot[SF_Hull] = GetStableShadersAndZeroHash(InverseMap, StableArray, Item.GraphicsDesc.HullShader, StableShadersPerSlot[SF_Hull], OutAnyActiveButMissing);
-			ActivePerSlot[SF_Domain] = GetStableShadersAndZeroHash(InverseMap, StableArray, Item.GraphicsDesc.DomainShader, StableShadersPerSlot[SF_Domain], OutAnyActiveButMissing);
+			ActivePerSlot[SF_Vertex] = GetStableShadersAndZeroHash(InverseMap, Item.GraphicsDesc.VertexShader, StableShadersPerSlot[SF_Vertex], OutAnyActiveButMissing);
+			ActivePerSlot[SF_Pixel] = GetStableShadersAndZeroHash(InverseMap, Item.GraphicsDesc.FragmentShader, StableShadersPerSlot[SF_Pixel], OutAnyActiveButMissing);
+			ActivePerSlot[SF_Geometry] = GetStableShadersAndZeroHash(InverseMap, Item.GraphicsDesc.GeometryShader, StableShadersPerSlot[SF_Geometry], OutAnyActiveButMissing);
+			ActivePerSlot[SF_Hull] = GetStableShadersAndZeroHash(InverseMap, Item.GraphicsDesc.HullShader, StableShadersPerSlot[SF_Hull], OutAnyActiveButMissing);
+			ActivePerSlot[SF_Domain] = GetStableShadersAndZeroHash(InverseMap, Item.GraphicsDesc.DomainShader, StableShadersPerSlot[SF_Domain], OutAnyActiveButMissing);
 		}
 
 
@@ -703,16 +655,16 @@ int32 ExpandPSOSC(const TArray<FString>& Tokens)
 			UE_LOG(LogShaderPipelineCacheTools, Display, TEXT("PSO had an active shader slot that did not match any current shaders, ignored."));
 			if (Item.Type == FPipelineCacheFileFormatPSO::DescriptorType::Compute)
 			{
-				PrintShaders(InverseMap, StableArray, Item.ComputeDesc.ComputeShader, TEXT("ComputeShader"));
+				PrintShaders(InverseMap, Item.ComputeDesc.ComputeShader, TEXT("ComputeShader"));
 			}
 			else
 			{
 				UE_LOG(LogShaderPipelineCacheTools, Display, TEXT("   %s"), *Item.GraphicsDesc.StateToString());
-				PrintShaders(InverseMap, StableArray, Item.GraphicsDesc.VertexShader, TEXT("VertexShader"));
-				PrintShaders(InverseMap, StableArray, Item.GraphicsDesc.FragmentShader, TEXT("FragmentShader"));
-				PrintShaders(InverseMap, StableArray, Item.GraphicsDesc.GeometryShader, TEXT("GeometryShader"));
-				PrintShaders(InverseMap, StableArray, Item.GraphicsDesc.HullShader, TEXT("HullShader"));
-				PrintShaders(InverseMap, StableArray, Item.GraphicsDesc.DomainShader, TEXT("DomainShader"));
+				PrintShaders(InverseMap, Item.GraphicsDesc.VertexShader, TEXT("VertexShader"));
+				PrintShaders(InverseMap, Item.GraphicsDesc.FragmentShader, TEXT("FragmentShader"));
+				PrintShaders(InverseMap, Item.GraphicsDesc.GeometryShader, TEXT("GeometryShader"));
+				PrintShaders(InverseMap, Item.GraphicsDesc.HullShader, TEXT("HullShader"));
+				PrintShaders(InverseMap, Item.GraphicsDesc.DomainShader, TEXT("DomainShader"));
 			}
 			continue;
 		}
@@ -742,7 +694,7 @@ int32 ExpandPSOSC(const TArray<FString>& Tokens)
 						bool bFoundCompat = false;
 						for (int32 StableIndexInner = 0; StableIndexInner < StableShadersPerSlot[SlotIndexInner].Num(); StableIndexInner++)
 						{
-							if (CouldBeUsedTogether(StableArray[StableShadersPerSlot[SlotIndex][StableIndex]], StableArray[StableShadersPerSlot[SlotIndexInner][StableIndexInner]]))
+							if (CouldBeUsedTogether(StableShadersPerSlot[SlotIndex][StableIndex], StableShadersPerSlot[SlotIndexInner][StableIndexInner]))
 							{
 								bFoundCompat = true;
 								break;
@@ -776,11 +728,11 @@ int32 ExpandPSOSC(const TArray<FString>& Tokens)
 				UE_LOG(LogShaderPipelineCacheTools, Warning, TEXT("PSO did not create any stable PSOs! (no cross shader slot compatibility)"));
 				UE_LOG(LogShaderPipelineCacheTools, Warning, TEXT("   %s"), *Item.GraphicsDesc.StateToString());
 
-				PrintShaders(InverseMap, StableArray, Item.GraphicsDesc.VertexShader, TEXT("VertexShader"));
-				PrintShaders(InverseMap, StableArray, Item.GraphicsDesc.FragmentShader, TEXT("FragmentShader"));
-				PrintShaders(InverseMap, StableArray, Item.GraphicsDesc.GeometryShader, TEXT("GeometryShader"));
-				PrintShaders(InverseMap, StableArray, Item.GraphicsDesc.HullShader, TEXT("HullShader"));
-				PrintShaders(InverseMap, StableArray, Item.GraphicsDesc.DomainShader, TEXT("DomainShader"));
+				PrintShaders(InverseMap, Item.GraphicsDesc.VertexShader, TEXT("VertexShader"));
+				PrintShaders(InverseMap, Item.GraphicsDesc.FragmentShader, TEXT("FragmentShader"));
+				PrintShaders(InverseMap, Item.GraphicsDesc.GeometryShader, TEXT("GeometryShader"));
+				PrintShaders(InverseMap, Item.GraphicsDesc.HullShader, TEXT("HullShader"));
+				PrintShaders(InverseMap, Item.GraphicsDesc.DomainShader, TEXT("DomainShader"));
 
 				continue;
 			}
@@ -798,7 +750,7 @@ int32 ExpandPSOSC(const TArray<FString>& Tokens)
 
 		TArray<FPermuation>& Permutations(Current.Permutations);
 		FPermuation WorkingPerm;
-		GeneratePermuations(Permutations, WorkingPerm, 0, StableShadersPerSlot, StableArray, ActivePerSlot);
+		GeneratePermuations(Permutations, WorkingPerm, 0, StableShadersPerSlot, ActivePerSlot);
 		if (!Permutations.Num())
 		{
 			UE_LOG(LogShaderPipelineCacheTools, Error, TEXT("PSO did not create any stable PSOs! (somehow)"));
@@ -818,11 +770,10 @@ int32 ExpandPSOSC(const TArray<FString>& Tokens)
 		return 1;
 	}
 
-
 	int32 NumLines = 0;
-	FSCDataChunk DataChunks[16];
-	int32 CurrentChunk = 0;
-	TSet<uint32> DeDup;
+	TArray<uint8> UncomressedOutputLines;
+	FMemoryWriter OutputLinesAr(UncomressedOutputLines);
+	TSet<FString> DeDup;
 
 	{
 		FString PSOLine = FString::Printf(TEXT("\"%s\""), *FPipelineCacheFileFormatPSO::CommonHeaderLine());
@@ -832,7 +783,7 @@ int32 ExpandPSOSC(const TArray<FString>& Tokens)
 			PSOLine += FString::Printf(TEXT(",\"shaderslot%d: %s\""), SlotIndex, *FStableShaderKeyAndValue::HeaderLine());
 		}
 
-		DataChunks[CurrentChunk].OutputLinesAr << PSOLine;
+		OutputLinesAr << PSOLine;
 		NumLines++;
 	}
 
@@ -902,15 +853,11 @@ int32 ExpandPSOSC(const TArray<FString>& Tokens)
 				}
 			}
 
-			const uint32 PSOLineHash = FCrc::MemCrc32(PSOLine.GetCharArray().GetData(), sizeof(TCHAR) * PSOLine.Len());
-			if (!DeDup.Contains(PSOLineHash))
+
+			if (!DeDup.Contains(PSOLine))
 			{
-				DeDup.Add(PSOLineHash);
-				if (DataChunks[CurrentChunk].OutputLinesAr.TotalSize() + (int64)((PSOLine.Len() + 1) * sizeof(TCHAR)) >= STABLE_MAX_CHUNK_SIZE)
-				{
-					++CurrentChunk;
-				}
-				DataChunks[CurrentChunk].OutputLinesAr << PSOLine;
+				DeDup.Add(PSOLine);
+				OutputLinesAr << PSOLine;
 				NumLines++;
 			}
 		}
@@ -950,7 +897,7 @@ int32 ExpandPSOSC(const TArray<FString>& Tokens)
 		}
 	}
 
-	int64 FileSize = SaveStableCSV(OutputFilename, DataChunks, CurrentChunk + 1);
+	int64 FileSize = SaveStableCSV(OutputFilename, UncomressedOutputLines);
 	if (FileSize < 1)
 	{
 		return 1;
@@ -1019,15 +966,14 @@ void BuildDateSortedListOfFiles(const TArray<FString>& TokenList, FilenameFilter
 
 int32 BuildPSOSC(const TArray<FString>& Tokens)
 {
-	TArray<FStableShaderKeyAndValue> StableArray;
-	TMultiMap<int32, FSHAHash> StableMap;
+	TMultiMap<FStableShaderKeyAndValue, FSHAHash> StableMap;
 	check(Tokens.Last().EndsWith(TEXT(".upipelinecache")));
 
 	for (int32 Index = 0; Index < Tokens.Num() - 1; Index++)
 	{
 		if (Tokens[Index].EndsWith(TEXT(".scl.csv")))
 		{
-			LoadStableSCL(StableMap, StableArray, Tokens[Index]);
+			LoadStableSCL(StableMap, Tokens[Index]);
 		}
 	}
 	if (UE_LOG_ACTIVE(LogShaderPipelineCacheTools, Verbose))
@@ -1035,7 +981,7 @@ int32 BuildPSOSC(const TArray<FString>& Tokens)
 		UE_LOG(LogShaderPipelineCacheTools, Verbose, TEXT("    %s"), *FStableShaderKeyAndValue::HeaderLine());
 		for (const auto& Pair : StableMap)
 		{
-			FStableShaderKeyAndValue Temp(StableArray[Pair.Key]);
+			FStableShaderKeyAndValue Temp(Pair.Key);
 			Temp.OutputHash = Pair.Value;
 			UE_LOG(LogShaderPipelineCacheTools, Verbose, TEXT("    %s"), *Temp.ToString());
 		}
@@ -1043,7 +989,7 @@ int32 BuildPSOSC(const TArray<FString>& Tokens)
 	//self test
 	for (const auto& Pair : StableMap)
 	{
-		FStableShaderKeyAndValue Item(StableArray[Pair.Key]);
+		FStableShaderKeyAndValue Item(Pair.Key);
 		Item.OutputHash = Pair.Value;
 		check(Pair.Value != FSHAHash());
 		FString TestString = Item.ToString();
@@ -1143,21 +1089,19 @@ int32 BuildPSOSC(const TArray<FString>& Tokens)
 
 				FSHAHash Match;
 				int32 Count = 0;
-				for (auto Iter : StableArray)
+				for (auto Iter = StableMap.CreateConstKeyIterator(Shader); Iter; ++Iter)
 				{
-					if (Iter == Shader)
+					check(Iter.Value() != FSHAHash());
+					Match = Iter.Value();
+					if (TargetPlatform == NAME_None)
 					{
-						Match = Iter.OutputHash;
-						if (TargetPlatform == NAME_None)
-						{
-							TargetPlatform = Iter.TargetPlatform;
-						}
-						else
-						{
-							check(TargetPlatform == Iter.TargetPlatform);
-						}
-						Count++;
+						TargetPlatform = Iter.Key().TargetPlatform;
+					}
+					else
+					{
+						check(TargetPlatform == Iter.Key().TargetPlatform);
 					}
+					Count++;
 				}
 
 				if (!Count)
@@ -1460,7 +1404,7 @@ int32 DiffStable(const TArray<FString>& Tokens)
 
 int32 DecompressCSV(const TArray<FString>& Tokens)
 {
-	TArray<FString> DecompressedData;
+	TArray<uint8> DecompressedData;
 	for (int32 TokenIndex = 0; TokenIndex < Tokens.Num(); TokenIndex++)
 	{
 		const FString& CompressedFilename = Tokens[TokenIndex];
@@ -1473,22 +1417,17 @@ int32 DecompressCSV(const TArray<FString>& Tokens)
 		DecompressedData.Reset();
 		if (LoadAndDecompressStableCSV(CompressedFilename, DecompressedData))
 		{
-			FString FilenameCSV = CompressedFilename.LeftChop(STABLE_COMPRESSED_EXT_LEN);
-			FPlatformFileManager::Get().GetPlatformFile().DeleteFile(*FilenameCSV);
-
-			for (const FString& LineCSV : DecompressedData)
+			FMemoryReader MemArchive(DecompressedData);
+			FString LineCSV;
+			while (!MemArchive.AtEnd())
 			{
+				MemArchive << LineCSV;
 				CombinedCSV.Append(LineCSV);
 				CombinedCSV.Append(LINE_TERMINATOR);
-
-				if ((int64)(CombinedCSV.Len() * sizeof(TCHAR)) >= (int64)(MAX_int32 - 1024 * 1024))
-				{
-					FFileHelper::SaveStringToFile(CombinedCSV, *FilenameCSV, FFileHelper::EEncodingOptions::AutoDetect, &IFileManager::Get(), FILEWRITE_Append);
-					CombinedCSV.Empty();
-				}
 			}
 
-			FFileHelper::SaveStringToFile(CombinedCSV, *FilenameCSV, FFileHelper::EEncodingOptions::AutoDetect, &IFileManager::Get(), FILEWRITE_Append);
+			FString FilenameCSV = CompressedFilename.LeftChop(STABLE_COMPRESSED_EXT_LEN);
+			FFileHelper::SaveStringToFile(CombinedCSV, *FilenameCSV);
 		}
 	}
 
