Added mode to create a blueprint from a selection of actors in the level which creates a ChildActorComponent for each selected Actor.

diff --git a/Engine/Source/Editor/KismetWidgets/Private/CreateBlueprintFromActorDialog.cpp b/Engine/Source/Editor/KismetWidgets/Private/CreateBlueprintFromActorDialog.cpp
--- a/Engine/Source/Editor/KismetWidgets/Private/CreateBlueprintFromActorDialog.cpp
+++ b/Engine/Source/Editor/KismetWidgets/Private/CreateBlueprintFromActorDialog.cpp
@@ -57,38 +57,57 @@ void FCreateBlueprintFromActorDialog::OnCreateBlueprint(const FString& InAssetPa
 	switch (CreateMode) 
 	{
 		case ECreateBlueprintFromActorMode::Harvest:
-		{
-			TArray<AActor*> Actors;
+	{
+		TArray<AActor*> Actors;
 
-			USelection* SelectedActors = GEditor->GetSelectedActors();
-			for(FSelectionIterator Iter(*SelectedActors); Iter; ++Iter)
-			{
-				// We only care about actors that are referenced in the world for literals, and also in the same level as this blueprint
+		USelection* SelectedActors = GEditor->GetSelectedActors();
+		for(FSelectionIterator Iter(*SelectedActors); Iter; ++Iter)
+		{
+			// We only care about actors that are referenced in the world for literals, and also in the same level as this blueprint
 				if (AActor* Actor = Cast<AActor>(*Iter))
-				{
-					Actors.Add(Actor);
-				}
+			{
+				Actors.Add(Actor);
 			}
+		}
 
 			const bool bReplaceActor = true;
 			Blueprint = FKismetEditorUtilities::HarvestBlueprintFromActors(InAssetPath, Actors, bReplaceActor);
-		}
+	}
 		break;
 
 		case ECreateBlueprintFromActorMode::Subclass:
+	{
+		AActor* ActorToUse = ActorOverride.Get();
+
+		if( !ActorToUse )
 		{
-			AActor* ActorToUse = ActorOverride.Get();
+			TArray< UObject* > SelectedActors;
+			GEditor->GetSelectedActors()->GetSelectedObjects(AActor::StaticClass(), SelectedActors);
+				check(SelectedActors.Num() == 1);
+			ActorToUse =  Cast<AActor>(SelectedActors[0]);
+		}
+
+		const bool bReplaceActor = true;
+		Blueprint = FKismetEditorUtilities::CreateBlueprintFromActor(InAssetPath, ActorToUse, bReplaceActor);
+	}
+		break;
+
+		case ECreateBlueprintFromActorMode::ChildActor:
+		{
+			TArray<AActor*> Actors;
 
-			if (!ActorToUse)
+			USelection* SelectedActors = GEditor->GetSelectedActors();
+			for (FSelectionIterator Iter(*SelectedActors); Iter; ++Iter)
 			{
-				TArray< UObject* > SelectedActors;
-				GEditor->GetSelectedActors()->GetSelectedObjects(AActor::StaticClass(), SelectedActors);
-				check(SelectedActors.Num() == 1);
-				ActorToUse = Cast<AActor>(SelectedActors[0]);
+				// We only care about actors that are referenced in the world for literals, and also in the same level as this blueprint
+				if (AActor* Actor = Cast<AActor>(*Iter))
+				{
+					Actors.Add(Actor);
+				}
 			}
 
 			const bool bReplaceActor = true;
-			Blueprint = FKismetEditorUtilities::CreateBlueprintFromActor(InAssetPath, ActorToUse, bReplaceActor);
+			Blueprint = FKismetEditorUtilities::CreateBlueprintFromActors(InAssetPath, Actors, bReplaceActor);
 		}
 		break;
 	}
diff --git a/Engine/Source/Editor/KismetWidgets/Public/CreateBlueprintFromActorDialog.h b/Engine/Source/Editor/KismetWidgets/Public/CreateBlueprintFromActorDialog.h
--- a/Engine/Source/Editor/KismetWidgets/Public/CreateBlueprintFromActorDialog.h
+++ b/Engine/Source/Editor/KismetWidgets/Public/CreateBlueprintFromActorDialog.h
@@ -13,7 +13,8 @@ enum class ECreateBlueprintFromActorMode : uint8
 {
 	None,        // Indicates that a user is unable to create blueprints from the currently selected actor set
 	Harvest,     // Harvest all the components of the selected actors and create an actor blueprint with those components in it.
-	Subclass     // Create a subclass of the selected Actor's class with defaults from the selected Actor. Valid only when there is a single selected actor.
+	Subclass,     // Create a subclass of the selected Actor's class with defaults from the selected Actor. Valid only when there is a single selected actor.
+	ChildActor    // Create an actor blueprint with each selected actor as a child actor component
 };
 
 class FCreateBlueprintFromActorDialog
diff --git a/Engine/Source/Editor/LevelEditor/Private/LevelEditorActions.cpp b/Engine/Source/Editor/LevelEditor/Private/LevelEditorActions.cpp
--- a/Engine/Source/Editor/LevelEditor/Private/LevelEditorActions.cpp
+++ b/Engine/Source/Editor/LevelEditor/Private/LevelEditorActions.cpp
@@ -2336,20 +2336,73 @@ void FLevelEditorActionCallbacks::HarvestSelectedActorsIntoBlueprintClass()
 	FCreateBlueprintFromActorDialog::OpenDialog(ECreateBlueprintFromActorMode::Harvest);
 }
 
-bool FLevelEditorActionCallbacks::CanSubclassSelectedActorIntoBlueprintClass()
+ECreateBlueprintFromActorMode GetConvertSelectedActorsMode()
 {
-	bool bCanSubclass = GEditor->GetSelectedActorCount() == 1;
-	if (bCanSubclass)
+	if (GEditor->GetSelectedActorCount() == 0)
+	{
+		return ECreateBlueprintFromActorMode::None;
+	}
+
+	// Subclassing mode when we have 1 actor selected of a blueprint type
+	if (GEditor->GetSelectedActorCount() == 1)
 	{
 		AActor* Actor = Cast<AActor>(*GEditor->GetSelectedActorIterator());
-		bCanSubclass = FKismetEditorUtilities::CanCreateBlueprintOfClass(Actor->GetClass());
+		if (FKismetEditorUtilities::CanCreateBlueprintOfClass(Actor->GetClass()))
+		{
+			return ECreateBlueprintFromActorMode::Subclass;
+		}
+	}
+
+	// Child Actor mode when we have 1 or more selected objects that are all placeable
+	for (FSelectionIterator It(GEditor->GetSelectedActorIterator()); It; ++It)
+	{
+		AActor* Actor = Cast<AActor>(*It);
+		if (Actor->GetClass()->HasAnyClassFlags(CLASS_NotPlaceable))
+		{
+			return ECreateBlueprintFromActorMode::None;
+		}
+	}
+
+	return ECreateBlueprintFromActorMode::ChildActor;
+}
+
+bool FLevelEditorActionCallbacks::CanConvertSelectedActorsIntoBlueprintClass()
+{
+	const bool bCanConvert = GetConvertSelectedActorsMode() != ECreateBlueprintFromActorMode::None;
+	return bCanConvert;
+}
+
+void FLevelEditorActionCallbacks::ConvertSelectedActorsIntoBlueprintClass()
+{
+	ECreateBlueprintFromActorMode CreateMode = GetConvertSelectedActorsMode();
+	if (ensure(CreateMode != ECreateBlueprintFromActorMode::None))
+	{
+		FCreateBlueprintFromActorDialog::OpenDialog(CreateMode);
 	}
-	return bCanSubclass;
 }
 
-void FLevelEditorActionCallbacks::SubclassSelectedActorIntoBlueprintClass()
+FText FLevelEditorActionCallbacks::GetConvertSelectedActorsIntoBlueprintClassLabel()
+{
+	if (GetConvertSelectedActorsMode() == ECreateBlueprintFromActorMode::ChildActor)
 {
-	FCreateBlueprintFromActorDialog::OpenDialog(ECreateBlueprintFromActorMode::Subclass);
+		return LOCTEXT("ConvertSelectionToBlueprintViaChildActor_Label", "Convert Selected Actors to Blueprint Class...");
+	}
+	else
+	{
+		return LOCTEXT("ConvertSelectionToBlueprintViaSubclass_Label", "Convert Selected Actor to Blueprint Class...");
+	}
+}
+
+FText FLevelEditorActionCallbacks::GetConvertSelectedActorsIntoBlueprintClassTooltip()
+{
+	if (GetConvertSelectedActorsMode() == ECreateBlueprintFromActorMode::ChildActor)
+	{
+		return LOCTEXT("ConvertSelectionToBlueprintViaChildActor_Tooltip", "Replace the selected actors with a new Blueprint subclass based on Actor with each of the selected Actors as a Child Actor");
+	}
+	else
+	{
+		return LOCTEXT("ConvertSelectionToBlueprintViaSubclass_Tooltip", "Replace the selected actor with a new Blueprint subclass based on the class of the selected Actor");
+	}
 }
 
 void FLevelEditorActionCallbacks::CheckOutProjectSettingsConfig( )
diff --git a/Engine/Source/Editor/LevelEditor/Private/LevelEditorToolBar.cpp b/Engine/Source/Editor/LevelEditor/Private/LevelEditorToolBar.cpp
--- a/Engine/Source/Editor/LevelEditor/Private/LevelEditorToolBar.cpp
+++ b/Engine/Source/Editor/LevelEditor/Private/LevelEditorToolBar.cpp
@@ -2316,7 +2316,10 @@ void FLevelEditorToolBar::RegisterOpenBlueprintMenu()
 
 		// Convert selection to BP
 		Section.AddMenuEntry(FLevelEditorCommands::Get().ConvertSelectionToBlueprintViaHarvest);
-		Section.AddMenuEntry(FLevelEditorCommands::Get().ConvertSelectionToBlueprintViaSubclass);
+		Section.AddMenuEntry(
+			FLevelEditorCommands::Get().ConvertSelectionToBlueprintViaSubclass,
+			TAttribute<FText>::Create(FLevelEditorActionCallbacks::GetConvertSelectedActorsIntoBlueprintClassLabel),
+			TAttribute<FText>::Create(FLevelEditorActionCallbacks::GetConvertSelectedActorsIntoBlueprintClassTooltip));
 
 		// Open an existing Blueprint Class...
 		FSlateIcon OpenBPIcon(FEditorStyle::GetStyleSetName(), "LevelEditor.OpenClassBlueprint");
diff --git a/Engine/Source/Editor/LevelEditor/Private/SLevelEditor.cpp b/Engine/Source/Editor/LevelEditor/Private/SLevelEditor.cpp
--- a/Engine/Source/Editor/LevelEditor/Private/SLevelEditor.cpp
+++ b/Engine/Source/Editor/LevelEditor/Private/SLevelEditor.cpp
@@ -112,8 +112,8 @@ void SLevelEditor::BindCommands()
 
 	LevelEditorCommands->MapAction(
 		Actions.ConvertSelectionToBlueprintViaSubclass,
-		FExecuteAction::CreateStatic( &FLevelEditorActionCallbacks::SubclassSelectedActorIntoBlueprintClass ),
-		FCanExecuteAction::CreateStatic( &FLevelEditorActionCallbacks::CanSubclassSelectedActorIntoBlueprintClass ) );
+		FExecuteAction::CreateStatic( &FLevelEditorActionCallbacks::ConvertSelectedActorsIntoBlueprintClass ),
+		FCanExecuteAction::CreateStatic( &FLevelEditorActionCallbacks::CanConvertSelectedActorsIntoBlueprintClass ) );
 
 	LevelEditorCommands->MapAction(
 		Actions.OpenContentBrowser,
diff --git a/Engine/Source/Editor/LevelEditor/Public/LevelEditorActions.h b/Engine/Source/Editor/LevelEditor/Public/LevelEditorActions.h
--- a/Engine/Source/Editor/LevelEditor/Public/LevelEditorActions.h
+++ b/Engine/Source/Editor/LevelEditor/Public/LevelEditorActions.h
@@ -1184,11 +1184,17 @@ public:
 	/** Harvest all of the components in the selected actors into a Blueprint, and replace the instances with one new actor */
 	static void HarvestSelectedActorsIntoBlueprintClass();
 
-	/** Can call SubclassSelectedActorIntoBlueprintClass right now?  (is exactly one thing selected) */
-	static bool CanSubclassSelectedActorIntoBlueprintClass();
+	/** Can call SubclassSelectedActorIntoBlueprintClass right now? */
+	static bool CanConvertSelectedActorsIntoBlueprintClass();
 
-	/** Convert the selected actor into a Blueprint (via subclassing) */
-	static void SubclassSelectedActorIntoBlueprintClass();
+	/** Convert the selected actor into a Blueprint (via subclassing or child actors) */
+	static void ConvertSelectedActorsIntoBlueprintClass();
+
+	/** Get the label to display for the convert selected actors to blueprint menu option */
+	static FText GetConvertSelectedActorsIntoBlueprintClassLabel();
+
+	/** Get the tooltip to display for the convert selected actors to blueprint menu option */
+	static FText GetConvertSelectedActorsIntoBlueprintClassTooltip();
 
 	/** Shows only selected actors, hiding any unselected actors and unhiding any selected hidden actors. */
 	static void OnShowOnlySelectedActors();
diff --git a/Engine/Source/Editor/UnrealEd/Private/Kismet2/Kismet2.cpp b/Engine/Source/Editor/UnrealEd/Private/Kismet2/Kismet2.cpp
--- a/Engine/Source/Editor/UnrealEd/Private/Kismet2/Kismet2.cpp
+++ b/Engine/Source/Editor/UnrealEd/Private/Kismet2/Kismet2.cpp
@@ -1014,26 +1014,34 @@ bool FKismetEditorUtilities::CanCreateBlueprintOfClass(const UClass* Class)
 	return bCanCreateBlueprint;
 }
 
-UBlueprint* FKismetEditorUtilities::CreateBlueprintFromActor(const FString& Path, AActor* Actor, const bool bReplaceActor, bool bKeepMobility /*= false*/)
+UPackage* CreateBlueprintPackage(const FString& Path, FString& OutAssetName)
 {
-	UBlueprint* NewBlueprint = nullptr;
-
 	// Create a blueprint
-	FString PackageName = Path;
-	FString AssetName = FPackageName::GetLongPackageAssetName(Path);
+	FString PackageName;
+	OutAssetName = FPackageName::GetLongPackageAssetName(Path);
 
 	// If no AssetName was found, generate a unique asset name.
-	if(AssetName.Len() == 0)
+	if (OutAssetName.Len() == 0)
 	{
 		PackageName = FPackageName::GetLongPackagePath(Path);
 		FString BasePath = PackageName + TEXT("/") + LOCTEXT("BlueprintName_Default", "NewBlueprint").ToString();
 		FAssetToolsModule& AssetToolsModule = FModuleManager::LoadModuleChecked<FAssetToolsModule>("AssetTools");
-		AssetToolsModule.Get().CreateUniqueAssetName(BasePath, TEXT(""), PackageName, AssetName);
+		AssetToolsModule.Get().CreateUniqueAssetName(BasePath, TEXT(""), PackageName, OutAssetName);
+	}
+	else
+	{
+		PackageName = Path;
 	}
 
-	UPackage* Package = CreatePackage(NULL, *PackageName);
+	return CreatePackage(nullptr, *PackageName);
+}
+
+UBlueprint* FKismetEditorUtilities::CreateBlueprintFromActor(const FString& Path, AActor* Actor, const bool bReplaceActor, bool bKeepMobility /*= false*/)
+{
+	UBlueprint* NewBlueprint = nullptr;
+	FString AssetName;
 
-	if(Package)
+	if (UPackage* Package = CreateBlueprintPackage(Path, AssetName))
 	{
 		NewBlueprint = CreateBlueprintFromActor(FName(*AssetName), Package, Actor, bReplaceActor, bKeepMobility);
 	}
@@ -1041,7 +1049,20 @@ UBlueprint* FKismetEditorUtilities::CreateBlueprintFromActor(const FString& Path
 	return NewBlueprint;
 }
 
-void FKismetEditorUtilities::AddComponentsToBlueprint(UBlueprint* Blueprint, TArray<UActorComponent*> Components, bool bHarvesting, USCS_Node* OptionalNewRootNode, bool bKeepMobility /*= false*/)
+UBlueprint* FKismetEditorUtilities::CreateBlueprintFromActors(const FString& Path, const TArray<AActor*>& Actors, const bool bReplaceInWorld)
+{
+	UBlueprint* NewBlueprint = nullptr;
+	FString AssetName;
+
+	if (UPackage* Package = CreateBlueprintPackage(Path, AssetName))
+	{
+		NewBlueprint = CreateBlueprintFromActors(FName(*AssetName), Package, Actors, bReplaceInWorld);
+	}
+
+	return NewBlueprint;
+}
+
+void FKismetEditorUtilities::AddComponentsToBlueprint(UBlueprint* Blueprint, const TArray<UActorComponent*>& Components, bool bHarvesting, USCS_Node* OptionalNewRootNode, bool bKeepMobility /*= false*/)
 {
 	USimpleConstructionScript* SCS = Blueprint->SimpleConstructionScript;
 
@@ -1188,8 +1209,6 @@ void FKismetEditorUtilities::AddComponentsToBlueprint(UBlueprint* Blueprint, TAr
 	for (int32 CompIndex = 0; CompIndex < ActorComponents.Num(); ++CompIndex)
 	{
 		UActorComponent* ActorComponent = ActorComponents[CompIndex];
-		AActor* Actor = ActorComponent->GetOwner();
-		check(Actor);
 
 		USCS_Node* SCSNode = FAddComponentsToBlueprintImpl::MakeComponentCopy(ActorComponent, SCS, InstanceComponentToNodeMap, bKeepMobility);
 
@@ -1201,8 +1220,6 @@ void FKismetEditorUtilities::AddComponentsToBlueprint(UBlueprint* Blueprint, TAr
 	{
 		const FSceneComponentAndFirstParent& ComponentNode = SceneComponentNodes[CompIndex];
 		USceneComponent* SceneComponent = ComponentNode.SceneComponent;
-		AActor* Actor = SceneComponent->GetOwner();
-		check(Actor);
 
 		USCS_Node* SCSNode = FAddComponentsToBlueprintImpl::MakeComponentCopy(SceneComponent, SCS, InstanceComponentToNodeMap, bKeepMobility);
 
@@ -1233,7 +1250,8 @@ void FKismetEditorUtilities::AddComponentsToBlueprint(UBlueprint* Blueprint, TAr
 			}
 		}
 
-		if (SceneComponent == Actor->GetRootComponent() && (SceneComponent->GetAttachParent() == nullptr || FirstAttachParent == nullptr))
+		AActor* Actor = SceneComponent->GetOwner();
+		if (((Actor == nullptr) || (SceneComponent == Actor->GetRootComponent())) && (SceneComponent->GetAttachParent() == nullptr || FirstAttachParent == nullptr))
 		{
 			if (OptionalNewRootNode != nullptr)
 			{
@@ -1247,7 +1265,7 @@ void FKismetEditorUtilities::AddComponentsToBlueprint(UBlueprint* Blueprint, TAr
 		// If we're not attached to a blueprint component, add ourself to the root node or the SCS root component:
 		else if (SceneComponent->GetAttachParent() == nullptr)
 		{
-				AddChildToSCSRootNodeLambda(SCSNode);
+			AddChildToSCSRootNodeLambda(SCSNode);
 		}
 		// If we're attached to a blueprint component look it up as the variable name is the component name
 		else if (SceneComponent->GetAttachParent()->IsCreatedByConstructionScript())
@@ -1396,198 +1414,275 @@ UBlueprint* FKismetEditorUtilities::CreateBlueprintFromActor(const FName Bluepri
 	return NewBlueprint;
 }
 
-// This class cracks open the selected actors, harvests their components, and creates a new blueprint containing copies of them
-class FCreateConstructionScriptFromSelectedActors
+struct FBlueprintAssemblyProps
 {
-public:
-	FCreateConstructionScriptFromSelectedActors()
-		: Blueprint(nullptr)
-		, SCS(nullptr)
+	FBlueprintAssemblyProps(UBlueprint* InBlueprint, const TArray<AActor*>& InActors)
+		: Blueprint(InBlueprint)
+		, Actors(InActors)
 	{
 	}
 
-	UBlueprint* Execute(FString Path, TArray<AActor*> SelectedActors, bool bReplaceInWorld)
+	UBlueprint* Blueprint; 
+	const TArray<AActor*>& Actors;
+	TArray<AActor*> RootActors;
+	TMap<AActor*, TArray<AActor*>> AttachmentMap;
+	USCS_Node* RootNodeOverride = nullptr;
+};
+
+void CreateBlueprintFromActors_Internal(UBlueprint* Blueprint, const TArray<AActor*>& Actors, const bool bReplaceInWorld, TFunctionRef<void(const FBlueprintAssemblyProps&)> AssembleBlueprintFunc)
+{
+	check(Actors.Num());
+	check(Blueprint);
+	check(Blueprint->SimpleConstructionScript);
+	USimpleConstructionScript* SCS = Blueprint->SimpleConstructionScript;
+
+	FBlueprintAssemblyProps AssemblyProps(Blueprint, Actors);
+
+	TMap<AActor*, TArray<AActor*>>& AttachmentMap = AssemblyProps.AttachmentMap;
+	TArray<AActor*>& RootActors = AssemblyProps.RootActors;
+	RootActors = Actors;
+
+	int32 ConsideringActorIndex = 0;
+	while (ConsideringActorIndex < RootActors.Num())
 	{
-		if (SelectedActors.Num() > 0)
+		AActor* ConsideringActor = RootActors[ConsideringActorIndex];
+		bool bIsRoot = true;
+		for (int32 CheckIndex = RootActors.Num()-1; CheckIndex > ConsideringActorIndex; --CheckIndex)
 		{
-			// Create a blueprint
-			FString PackageName = Path;
-			FString AssetName = FPackageName::GetLongPackageAssetName(Path);
-			FString BasePath = PackageName + TEXT("/") + AssetName;
-
-			// If no AssetName was found, generate a unique asset name.
-			if(AssetName.Len() == 0)
+			AActor* ActorToCheck = RootActors[CheckIndex];
+			if (ActorToCheck->IsAttachedTo(ConsideringActor))
 			{
-				BasePath = PackageName + TEXT("/") + LOCTEXT("BlueprintName_Default", "NewBlueprint").ToString();
-				FAssetToolsModule& AssetToolsModule = FModuleManager::LoadModuleChecked<FAssetToolsModule>("AssetTools");
-				AssetToolsModule.Get().CreateUniqueAssetName(BasePath, TEXT(""), PackageName, AssetName);
+				AttachmentMap.FindOrAdd(ConsideringActor).Add(ActorToCheck);
+				RootActors.RemoveAtSwap(CheckIndex);
+			}
+			else if (ConsideringActor->IsAttachedTo(ActorToCheck))
+			{
+				AttachmentMap.FindOrAdd(ActorToCheck).Add(ConsideringActor);
+				bIsRoot = false;
+				break;
 			}
+		}
 
-			UPackage* Package = CreatePackage(nullptr, *PackageName);
-			Blueprint = FKismetEditorUtilities::CreateBlueprint(AActor::StaticClass(), Package, *AssetName, BPTYPE_Normal, UBlueprint::StaticClass(), UBlueprintGeneratedClass::StaticClass(), FName("HarvestFromActors"));
+		if (bIsRoot)
+		{
+			++ConsideringActorIndex;
+		}
+		else
+		{
+			RootActors.RemoveAtSwap(ConsideringActorIndex);
+		}
+	}
 
-			check(Blueprint->SimpleConstructionScript != NULL);
-			SCS = Blueprint->SimpleConstructionScript;
+	// If there is not one unique root actor then create a new scene component to serve as the shared root node
+	if (RootActors.Num() != 1)
+	{
+		AssemblyProps.RootNodeOverride = SCS->CreateNode(USceneComponent::StaticClass(), TEXT("SharedRoot"));
+		SCS->AddNode(AssemblyProps.RootNodeOverride);
+	}
 
-			// Create a common root if necessary
-			TArray<AActor*> RootActors = SelectedActors;
+	AssembleBlueprintFunc(AssemblyProps);
 
-			int32 ConsideringActorIndex = 0;
-			while (ConsideringActorIndex < RootActors.Num())
+	FTransform NewActorTransform = FTransform::Identity;
+	if (RootActors.Num() == 1)
+	{
+		NewActorTransform = RootActors[0]->GetTransform();
+	}
+	else if (RootActors.Num() > 1)
+	{
+		// Compute the average origin for all the actors, so it can be backed out when saving them in the blueprint
+		{
+			// Find average location of all selected actors
+			FVector AverageLocation = FVector::ZeroVector;
+			for (const AActor* Actor : RootActors)
 			{
-				AActor* ConsideringActor = RootActors[ConsideringActorIndex];
-				bool bIsRoot = true;
-				for (int32 CheckIndex = RootActors.Num()-1; CheckIndex > ConsideringActorIndex; --CheckIndex)
-				{
-					AActor* ActorToCheck = RootActors[CheckIndex];
-					if (ActorToCheck->IsAttachedTo(ConsideringActor))
-					{
-						RootActors.RemoveAtSwap(CheckIndex);
-					}
-					else if (ConsideringActor->IsAttachedTo(ActorToCheck))
-					{
-						bIsRoot = false;
-						break;
-					}
-				}
-
-				if (bIsRoot)
+				if (USceneComponent* RootComponent = Actor->GetRootComponent())
 				{
-					++ConsideringActorIndex;
+					AverageLocation += Actor->GetActorLocation();
 				}
-				else
-				{
-					RootActors.RemoveAtSwap(ConsideringActorIndex);
-				}
-			}		
-
-			// If there is not one unique root actor then create a new scene component to serve as the shared root node
-			USCS_Node* RootNodeOverride = nullptr;
-			if (RootActors.Num() != 1)
-			{
-				RootNodeOverride = SCS->CreateNode(USceneComponent::StaticClass(), TEXT("SharedRoot"));
-				SCS->AddNode(RootNodeOverride);
 			}
+			AverageLocation /= (float)RootActors.Num();
 
-			// Harvest the components from each actor and clone them into the SCS
-			TArray<UActorComponent*> AllSelectedComponents;
-			for (const AActor* Actor : SelectedActors)
+			// Spawn the new BP at that location
+			NewActorTransform.SetTranslation(AverageLocation);
+		}
+
+		// Reposition all of the children of the root node to recenter them around the new pivot
+		for (USCS_Node* TopLevelNode : SCS->GetRootNodes())
+		{
+			if (USceneComponent* TestRoot = Cast<USceneComponent>(TopLevelNode->ComponentTemplate))
 			{
-				for (UActorComponent* ComponentToConsider : Actor->GetComponents())
+				for (USCS_Node* ChildNode : TopLevelNode->GetChildNodes())
 				{
-					if (ComponentToConsider && !ComponentToConsider->IsVisualizationComponent())
+					if (USceneComponent* ChildComponent = Cast<USceneComponent>(ChildNode->ComponentTemplate))
 					{
-						AllSelectedComponents.Add(ComponentToConsider);
+						//The relative transform for those component was converted into the world space
+						const FTransform NewRelativeTransform = ChildComponent->GetRelativeTransform().GetRelativeTransform(NewActorTransform);
+						ChildComponent->SetRelativeTransform(NewRelativeTransform);
 					}
 				}
 			}
+		}
+	}
 
-			TArray<TPair<USceneComponent*, FTransform>> SceneComponentOldRelativeTransforms;
-			if (RootActors.Num() > 1)
-			{
-				SceneComponentOldRelativeTransforms.Reserve(RootActors.Num());
-				// Convert the components relative transform to world 
-				for (AActor* Actor : RootActors)
-				{
-					USceneComponent* SceneComponent = Actor->GetRootComponent();
-					SceneComponentOldRelativeTransforms.Emplace(SceneComponent, SceneComponent->GetRelativeTransform());
-					SceneComponent->SetRelativeTransform(SceneComponent->GetComponentTransform());
-				}
-			}
+	// Regenerate skeleton class as components have been added since initial generation
+	FKismetEditorUtilities::GenerateBlueprintSkeleton(Blueprint, /*bForceRegeneration=*/ true);
 
-			FKismetEditorUtilities::AddComponentsToBlueprint(Blueprint, MoveTemp(AllSelectedComponents), /*bHarvesting=*/ true, RootNodeOverride);
+	// Notify the asset registry
+	FAssetRegistryModule::AssetCreated(Blueprint);
 
-			// Replace the modified components to their relative transform
-			for (const TPair<USceneComponent*, FTransform >& Pair : SceneComponentOldRelativeTransforms)
-			{
-				Pair.Key->SetRelativeTransform(Pair.Value);
-			}
+	// Mark the package dirty
+	Blueprint->GetOutermost()->MarkPackageDirty();
+
+	// Delete the old actors and create a new instance in the map
+	if (bReplaceInWorld)
+	{
+		FVector Location = NewActorTransform.GetLocation();
+		FRotator Rotator = NewActorTransform.Rotator();
+
+		AActor* NewActor = FKismetEditorUtilities::CreateBlueprintInstanceFromSelection(Blueprint, Actors, Location, Rotator);
+		if (NewActor)
+		{
+			NewActor->SetActorScale3D(NewActorTransform.GetScale3D());
+		}
+	}
+
+	// Open the editor for the new blueprint
+	GEditor->GetEditorSubsystem<UAssetEditorSubsystem>()->OpenEditorForAsset(Blueprint);
+}
 
-			FTransform NewActorTransform = FTransform::Identity;
-			if (RootActors.Num() == 1)
+UBlueprint* FKismetEditorUtilities::CreateBlueprintFromActors(const FName BlueprintName, UPackage* Package, const TArray<AActor*>& Actors, const bool bReplaceInWorld)
+{
+	auto AssemblyFunction = [](const FBlueprintAssemblyProps& AssemblyProps)
+	{
+		USimpleConstructionScript* SCS = AssemblyProps.Blueprint->SimpleConstructionScript;
+
+		TArray<UActorComponent*> ChildActorComponents;
+
+		TFunction<void(AActor*, UChildActorComponent*)> RecursiveCreateChildActorTemplates = [SCS, &ChildActorComponents, &AssemblyProps, &RecursiveCreateChildActorTemplates](AActor* Actor, UChildActorComponent* ParentComponent)
+		{
+			const FName ComponentName = SCS->GenerateNewComponentName(UChildActorComponent::StaticClass(), Actor->GetFName());
+			UChildActorComponent* CAC = NewObject<UChildActorComponent>(GetTransientPackage(), ComponentName, RF_ArchetypeObject);
+			CAC->SetChildActorClass(Actor->GetClass(), Actor);
+
+			// Clear any properties that can't be on the template
+			CAC->GetChildActorTemplate()->SetActorRelativeTransform(FTransform::Identity);
+
+			CAC->SetWorldTransform(Actor->GetTransform());
+			if (ParentComponent)
 			{
-				NewActorTransform = RootActors[0]->GetTransform();
+				CAC->AttachToComponent(ParentComponent, FAttachmentTransformRules::KeepWorldTransform);
 			}
-			else if (RootActors.Num() > 1)
-			{
-				// Compute the average origin for all the actors, so it can be backed out when saving them in the blueprint
-				{
-					// Find average location of all selected actors
-					FVector AverageLocation = FVector::ZeroVector;
-					for (const AActor* Actor : RootActors)
-					{
-						if (USceneComponent* RootComponent = Actor->GetRootComponent())
-						{
-							AverageLocation += Actor->GetActorLocation();
-						}
-					}
-					AverageLocation /= (float)RootActors.Num();
 
-					// Spawn the new BP at that location
-					NewActorTransform.SetTranslation(AverageLocation);
-				}
+			ChildActorComponents.Add(CAC);
 
-				// Reposition all of the children of the root node to recenter them around the new pivot
-				for (USCS_Node* TopLevelNode : SCS->GetRootNodes())
+			if (const TArray<AActor*>* ChildActors = AssemblyProps.AttachmentMap.Find(Actor))
+			{
+				for (AActor* ChildActor : *ChildActors)
 				{
-					if (USceneComponent* TestRoot = Cast<USceneComponent>(TopLevelNode->ComponentTemplate))
-					{
-						for (USCS_Node* ChildNode : TopLevelNode->GetChildNodes())
-						{
-							if (USceneComponent* ChildComponent = Cast<USceneComponent>(ChildNode->ComponentTemplate))
-							{
-								//The relative transform for those component was converted into the world space
-								const FTransform NewRelativeTransform = ChildComponent->GetRelativeTransform().GetRelativeTransform(NewActorTransform);
-								ChildComponent->SetRelativeTransform(NewRelativeTransform);
-							}
-						}
-					}
+					RecursiveCreateChildActorTemplates(ChildActor, CAC);
 				}
 			}
+		};
 
-			// Regenerate skeleton class as components have been added since initial generation
-			FKismetEditorUtilities::GenerateBlueprintSkeleton(Blueprint, /*bForceRegeneration=*/ true); 
-			
-			// Notify the asset registry
-			FAssetRegistryModule::AssetCreated(Blueprint);
+		for (AActor* Actor : AssemblyProps.RootActors)
+		{
+			RecursiveCreateChildActorTemplates(Actor, nullptr);
+		}
 
-			// Mark the package dirty
-			Package->MarkPackageDirty();
+		FKismetEditorUtilities::AddComponentsToBlueprint(AssemblyProps.Blueprint, ChildActorComponents, /*bHarvesting=*/ true, AssemblyProps.RootNodeOverride);
+	};
+	
+	UBlueprint* Blueprint = nullptr;
 
-			// Delete the old actors and create a new instance in the map
-			if (bReplaceInWorld)
-			{
-				FVector Location = NewActorTransform.GetLocation();
-				FRotator Rotator = NewActorTransform.Rotator();
+	if (Actors.Num() > 0)
+	{
+		if (Package != nullptr)
+		{
+			// We don't have a factory, but we can still try to create a blueprint for this actor class
+			Blueprint = FKismetEditorUtilities::CreateBlueprint(AActor::StaticClass(), Package, BlueprintName, EBlueprintType::BPTYPE_Normal, UBlueprint::StaticClass(), UBlueprintGeneratedClass::StaticClass(), FName("CreateFromActors"));
+		}
 
-				AActor* NewActor = FKismetEditorUtilities::CreateBlueprintInstanceFromSelection(Blueprint, SelectedActors, Location, Rotator);
-				if (NewActor)
+		if (Blueprint != nullptr)
+		{
+			CreateBlueprintFromActors_Internal(Blueprint, Actors, bReplaceInWorld, AssemblyFunction);
+		}
+	}
+
+	return Blueprint;
+}
+
+UBlueprint* FKismetEditorUtilities::HarvestBlueprintFromActors(const FString& Path, const TArray<AActor*>& Actors, bool bReplaceInWorld)
+{
+	UBlueprint* NewBlueprint = nullptr;
+	FString AssetName;
+
+	if (UPackage* Package = CreateBlueprintPackage(Path, AssetName))
+	{
+		NewBlueprint = HarvestBlueprintFromActors(FName(*AssetName), Package, Actors, bReplaceInWorld);
+	}
+
+	return NewBlueprint;
+}
+
+UBlueprint* FKismetEditorUtilities::HarvestBlueprintFromActors(const FName BlueprintName, UPackage* Package, const TArray<AActor*>& Actors, bool bReplaceInWorld)
+{
+	auto AssemblyFunction = [](const FBlueprintAssemblyProps& AssemblyProps)
+	{
+		// Harvest the components from each actor and clone them into the SCS
+		TArray<UActorComponent*> AllSelectedComponents;
+		for (const AActor* Actor : AssemblyProps.Actors)
+		{
+			for (UActorComponent* ComponentToConsider : Actor->GetComponents())
+			{
+				if (ComponentToConsider && !ComponentToConsider->IsVisualizationComponent())
 				{
-					NewActor->SetActorScale3D(NewActorTransform.GetScale3D());
+					AllSelectedComponents.Add(ComponentToConsider);
 				}
 			}
+		}
 
-			// Open the editor for the new blueprint
-			GEditor->GetEditorSubsystem<UAssetEditorSubsystem>()->OpenEditorForAsset(Blueprint);
+		TArray<TPair<USceneComponent*, FTransform>> SceneComponentOldRelativeTransforms;
+		if (AssemblyProps.RootActors.Num() > 1)
+		{
+			SceneComponentOldRelativeTransforms.Reserve(AssemblyProps.RootActors.Num());
+			// Convert the components relative transform to world 
+			for (AActor* Actor : AssemblyProps.RootActors)
+			{
+				USceneComponent* SceneComponent = Actor->GetRootComponent();
+				SceneComponentOldRelativeTransforms.Emplace(SceneComponent, SceneComponent->GetRelativeTransform());
+				SceneComponent->SetRelativeTransform(SceneComponent->GetComponentTransform());
+			}
+		}
 
-			return Blueprint;
+		FKismetEditorUtilities::AddComponentsToBlueprint(AssemblyProps.Blueprint, AllSelectedComponents, /*bHarvesting=*/ true, AssemblyProps.RootNodeOverride);
+
+		// Replace the modified components to their relative transform
+		for (const TPair<USceneComponent*, FTransform >& Pair : SceneComponentOldRelativeTransforms)
+		{
+			Pair.Key->SetRelativeTransform(Pair.Value);
 		}
-		return nullptr;
-	}
+	};
 
-protected:
-	UBlueprint* Blueprint;
-	USimpleConstructionScript* SCS;
-};
+	UBlueprint* Blueprint = nullptr;
 
+	if (Actors.Num() > 0)
+	{
+		if (Package != nullptr)
+		{
+			// We don't have a factory, but we can still try to create a blueprint for this actor class
+			Blueprint = FKismetEditorUtilities::CreateBlueprint(AActor::StaticClass(), Package, BlueprintName, EBlueprintType::BPTYPE_Normal, UBlueprint::StaticClass(), UBlueprintGeneratedClass::StaticClass(), FName("HarvestFromActors"));
+		}
 
-UBlueprint* FKismetEditorUtilities::HarvestBlueprintFromActors(const FString& Path, const TArray<AActor*>& Actors, bool bReplaceInWorld)
-{
-	FCreateConstructionScriptFromSelectedActors Creator;
-	return Creator.Execute(Path, Actors, bReplaceInWorld);
+		if (Blueprint != nullptr)
+		{
+			CreateBlueprintFromActors_Internal(Blueprint, Actors, bReplaceInWorld, AssemblyFunction);
+		}
+	}
+
+	return Blueprint;
 }
 
-AActor* FKismetEditorUtilities::CreateBlueprintInstanceFromSelection(UBlueprint* Blueprint, TArray<AActor*>& SelectedActors, const FVector& Location, const FRotator& Rotator)
+AActor* FKismetEditorUtilities::CreateBlueprintInstanceFromSelection(UBlueprint* Blueprint, const TArray<AActor*>& SelectedActors, const FVector& Location, const FRotator& Rotator)
 {
 	check (SelectedActors.Num() > 0 );
 
@@ -1600,9 +1695,9 @@ AActor* FKismetEditorUtilities::CreateBlueprintInstanceFromSelection(UBlueprint*
 	GEditor->GetSelectedActors()->Modify();
 
 	ULayersSubsystem* Layers = GEditor->GetEditorSubsystem<ULayersSubsystem>();
-	for(auto It(SelectedActors.CreateIterator());It;++It)
+	for (AActor* Actor : SelectedActors)
 	{
-		if (AActor* Actor = *It)
+		if (Actor)
 		{
 			// Remove from active selection in editor
 			GEditor->SelectActor(Actor, /*bSelected=*/ false, /*bNotify=*/ false);
diff --git a/Engine/Source/Editor/UnrealEd/Public/Kismet2/KismetEditorUtilities.h b/Engine/Source/Editor/UnrealEd/Public/Kismet2/KismetEditorUtilities.h
--- a/Engine/Source/Editor/UnrealEd/Public/Kismet2/KismetEditorUtilities.h
+++ b/Engine/Source/Editor/UnrealEd/Public/Kismet2/KismetEditorUtilities.h
@@ -147,7 +147,7 @@ public:
 	static bool CanCreateBlueprintOfClass(const UClass* Class);
 
 	/** Take a list of components that belong to a single Actor and add them to a blueprint as SCSNodes */
-	static void AddComponentsToBlueprint(UBlueprint* Blueprint, TArray<UActorComponent*> Components, bool bHarvesting = false, class USCS_Node* OptionalNewRootNode = nullptr, bool bKeepMobility = false);
+	static void AddComponentsToBlueprint(UBlueprint* Blueprint, const TArray<UActorComponent*>& Components, bool bHarvesting = false, class USCS_Node* OptionalNewRootNode = nullptr, bool bKeepMobility = false);
 
 	/** 
 	 * Take an Actor and generate a blueprint based on it. Uses the Actors type as the parent class. 
@@ -168,8 +168,28 @@ public:
 	 * @param bKeepMobility			If true, The mobility of each actor components will be copy
 	 * @return The blueprint created from the actor
 	 */
-	static UBlueprint* CreateBlueprintFromActor(const FName BlueprintName, UObject* Outer, AActor* Actor, bool bReplaceActor, bool bKeepMobility = false);
+	static UBlueprint* CreateBlueprintFromActor(const FName BlueprintName, UObject* Outer, AActor* Actor, bool bReplaceInWorld, bool bKeepMobility = false);
 
+	/** 
+	 * Take a list of Actors and generate a blueprint based on it using the Actors as templates for child actor components.
+	 * @param Path					The path to use when creating the package for the new blueprint
+	 * @param Actors				The actors to use when creating child actor components
+	 * @param bReplaceActor			If true, replace the actor in the scene with one based on the created blueprint
+	 * @return The blueprint created from the actor
+	 */
+	static UBlueprint* CreateBlueprintFromActors(const FString& Path, const TArray<AActor*>& Actors, bool bReplaceActor);
+
+	/** 
+	 * Take a list of Actors and generate a blueprint based on it using the Actors as templates for child actor components.
+	 * @param BlueprintName			The name to use for the Blueprint
+	 * @param Outer					The package to create the blueprint within
+	 * @param Actors				The actors to use when creating child actor components
+	 * @param bReplaceActor			If true, replace the actor in the scene with one based on the created blueprint
+	 * @param bKeepMobility			If true, The mobility of each actor components will be copy
+	 * @return The blueprint created from the actor
+	 */
+	static UBlueprint* CreateBlueprintFromActors(const FName BlueprintName, UPackage* Package, const TArray<AActor*>& Actors, bool bReplaceInWorld);
+	
 	/** 
 	 * Take a list of Actors and generate a blueprint  by harvesting the components they have. Uses AActor as parent class type as the parent class. 
 	 * @param Path					The path to use when creating the package for the new blueprint
@@ -177,7 +197,17 @@ public:
 	 * @param bReplaceInWorld		If true, replace the selected actors in the scene with one based on the created blueprint
 	 * @return The blueprint created from the actors
 	 */
-	static UBlueprint* HarvestBlueprintFromActors(const FString& Path, const TArray<AActor*>& Actors, bool ReplaceInWorld);
+	static UBlueprint* HarvestBlueprintFromActors(const FString& Path, const TArray<AActor*>& Actors, bool bReplaceInWorld);
+
+	/**
+	 * Take a list of Actors and generate a blueprint  by harvesting the components they have. Uses AActor as parent class type as the parent class.
+	 * @param BlueprintName			The name to use for the Blueprint
+	 * @param Outer					The package to create the blueprint within
+	 * @param Actors				The actor list to use as the template for the new blueprint, typically this is the currently selected actors
+	 * @param bReplaceInWorld		If true, replace the selected actors in the scene with one based on the created blueprint
+	 * @return The blueprint created from the actors
+	 */
+	static UBlueprint* HarvestBlueprintFromActors(const FName BlueprintName, UPackage* Package, const TArray<AActor*>& Actors, bool bReplaceInWorld);
 
 	/** 
 	 * Creates a new blueprint instance and replaces the provided actor list with the new actor
@@ -186,7 +216,7 @@ public:
 	 * @param Location			The location of the newly created actor
 	 * @param Rotator			The rotation of the newly created actor
 	 */
-	static AActor* CreateBlueprintInstanceFromSelection(class UBlueprint* Blueprint, TArray<AActor*>& SelectedActors, const FVector& Location, const FRotator& Rotator);
+	static AActor* CreateBlueprintInstanceFromSelection(class UBlueprint* Blueprint, const TArray<AActor*>& SelectedActors, const FVector& Location, const FRotator& Rotator);
 
 	/** 
 	 * Create a new Blueprint from the supplied base class. Pops up window to let user select location and name.
