Tool Menus: Asset Editors, File, Asset, Help menu extenders changed to tool menu API
#rnx

diff --git a/Engine/Source/Editor/UnrealEd/Private/Toolkits/AssetEditorToolkit.cpp b/Engine/Source/Editor/UnrealEd/Private/Toolkits/AssetEditorToolkit.cpp
--- a/Engine/Source/Editor/UnrealEd/Private/Toolkits/AssetEditorToolkit.cpp
+++ b/Engine/Source/Editor/UnrealEd/Private/Toolkits/AssetEditorToolkit.cpp
@@ -904,20 +904,22 @@ TSharedRef<SDockTab> FAssetEditorToolkit::SpawnTab_Toolbar( const FSpawnTabArgs&
 
 
 
-void FAssetEditorToolkit::FillDefaultFileMenuCommands( FMenuBuilder& MenuBuilder )
+void FAssetEditorToolkit::FillDefaultFileMenuCommands(FToolMenuSection& InSection)
 {
-	MenuBuilder.AddMenuEntry(FAssetEditorCommonCommands::Get().SaveAsset, "SaveAsset", TAttribute<FText>(), TAttribute<FText>(), FSlateIcon(FEditorStyle::GetStyleSetName(), "AssetEditor.SaveAsset.Greyscale"));
+	const FToolMenuInsert InsertPosition(NAME_None, EToolMenuInsertType::First);
+
+	InSection.AddMenuEntry(FAssetEditorCommonCommands::Get().SaveAsset, TAttribute<FText>(), TAttribute<FText>(), FSlateIcon(FEditorStyle::GetStyleSetName(), "AssetEditor.SaveAsset.Greyscale")).InsertPosition = InsertPosition;
 	if( IsActuallyAnAsset() )
 	{
-		MenuBuilder.AddMenuEntry(FAssetEditorCommonCommands::Get().SaveAssetAs, "SaveAssetAs", TAttribute<FText>(), TAttribute<FText>(), FSlateIcon(FEditorStyle::GetStyleSetName(), "AssetEditor.SaveAssetAs.Small"));
+		InSection.AddMenuEntry(FAssetEditorCommonCommands::Get().SaveAssetAs, TAttribute<FText>(), TAttribute<FText>(), FSlateIcon(FEditorStyle::GetStyleSetName(), "AssetEditor.SaveAssetAs.Small")).InsertPosition = InsertPosition;
 	}
-	MenuBuilder.AddMenuSeparator();
+	InSection.AddMenuSeparator("DefaultFileMenuCommandsSeparator").InsertPosition = InsertPosition;;
 
 	if( IsWorldCentricAssetEditor() )
 	{
 		// @todo toolkit minor: It would be awesome if the user could just "tear off" the SToolkitDisplay to do SwitchToStandaloneEditor
 		//			Would need to probably drop at mouseup location though instead of using saved layout pos.
-		MenuBuilder.AddMenuEntry( FAssetEditorCommonCommands::Get().SwitchToStandaloneEditor );
+		InSection.AddMenuEntry( FAssetEditorCommonCommands::Get().SwitchToStandaloneEditor ).InsertPosition = InsertPosition;;
 	}
 	else
 	{
@@ -926,16 +928,16 @@ void FAssetEditorToolkit::FillDefaultFileMenuCommands( FMenuBuilder& MenuBuilder
 			// @todo toolkit checkin: Disabled temporarily until we have world-centric "ready to use"!
 			if( 0 )
 			{
-				MenuBuilder.AddMenuEntry( FAssetEditorCommonCommands::Get().SwitchToWorldCentricEditor );
+				InSection.AddMenuEntry( FAssetEditorCommonCommands::Get().SwitchToWorldCentricEditor ).InsertPosition = InsertPosition;;
 			}
 		}
 	}
 }
 
 
-void FAssetEditorToolkit::FillDefaultAssetMenuCommands( FMenuBuilder& MenuBuilder )
+void FAssetEditorToolkit::FillDefaultAssetMenuCommands(FToolMenuSection& InSection)
 {
-	MenuBuilder.AddMenuEntry(FGlobalEditorCommonCommands::Get().FindInContentBrowser, "FindInContentBrowser", LOCTEXT("FindInContentBrowser", "Find in Content Browser..."));
+	InSection.AddMenuEntry(FGlobalEditorCommonCommands::Get().FindInContentBrowser, LOCTEXT("FindInContentBrowser", "Find in Content Browser..."));
 
 	// Commands we only want to be accessible when editing an asset should go here 
 	if( IsActuallyAnAsset() )
@@ -957,20 +959,20 @@ void FAssetEditorToolkit::FillDefaultAssetMenuCommands( FMenuBuilder& MenuBuilde
 					const FName IconName = TEXT( "AssetEditor.Reimport" );
 					FUIAction UIAction;
 					UIAction.ExecuteAction.BindRaw( this, &FAssetEditorToolkit::Reimport_Execute, EditingObject );
-					MenuBuilder.AddMenuEntry( LabelText, ToolTipText, FSlateIcon(FEditorStyle::GetStyleSetName(), IconName), UIAction );
+					InSection.AddMenuEntry( NAME_None, LabelText, ToolTipText, FSlateIcon(FEditorStyle::GetStyleSetName(), IconName), UIAction );
 				}
 			}
 		}		
 	}
 }
 
-void FAssetEditorToolkit::FillDefaultHelpMenuCommands( FMenuBuilder& MenuBuilder )
+void FAssetEditorToolkit::FillDefaultHelpMenuCommands(FToolMenuSection& InSection)
 {
 	FFormatNamedArguments Args;
 	Args.Add(TEXT("Editor"), GetBaseToolkitName());
 	const FText ToolTip = FText::Format(LOCTEXT("BrowseDocumentationTooltip", "Browse {Editor} documentation..."), Args);
 
-	MenuBuilder.AddMenuEntry(FGlobalEditorCommonCommands::Get().OpenDocumentation, NAME_None, ToolTip);
+	InSection.AddMenuEntry(FGlobalEditorCommonCommands::Get().OpenDocumentation, ToolTip);
 }
 
 FName FAssetEditorToolkit::GetToolMenuAppName() const
diff --git a/Engine/Source/Editor/UnrealEd/Private/Toolkits/SStandaloneAssetEditorToolkitHost.cpp b/Engine/Source/Editor/UnrealEd/Private/Toolkits/SStandaloneAssetEditorToolkitHost.cpp
--- a/Engine/Source/Editor/UnrealEd/Private/Toolkits/SStandaloneAssetEditorToolkitHost.cpp
+++ b/Engine/Source/Editor/UnrealEd/Private/Toolkits/SStandaloneAssetEditorToolkitHost.cpp
@@ -1,6 +1,7 @@
 // Copyright Epic Games, Inc. All Rights Reserved.
 
 #include "Toolkits/SStandaloneAssetEditorToolkitHost.h"
+#include "Toolkits/AssetEditorToolkitMenuContext.h"
 #include "Framework/MultiBox/MultiBoxBuilder.h"
 #include "Widgets/Layout/SBorder.h"
 #include "Misc/ConfigCacheIni.h"
@@ -45,78 +46,79 @@ void SStandaloneAssetEditorToolkitHost::SetupInitialContent( const TSharedRef<FT
 	const FName AssetEditorMenuName = GetMenuName();
 	if (!UToolMenus::Get()->IsMenuRegistered(AssetEditorMenuName))
 	{
-		UToolMenus::Get()->RegisterMenu(AssetEditorMenuName, "MainFrame.MainMenu");
-	}
+		UToolMenu* Menu = UToolMenus::Get()->RegisterMenu(AssetEditorMenuName, "MainFrame.MainMenu");
 
-	if (bCreateDefaultStandaloneMenu)
-	{
-		struct Local
+		if (bCreateDefaultStandaloneMenu)
 		{
-			static void FillFileMenu( FMenuBuilder& MenuBuilder, TWeakPtr< FAssetEditorToolkit > AssetEditorToolkitWeak )
-			{
-				auto AssetEditorToolkit( AssetEditorToolkitWeak.Pin().ToSharedRef() );
-				
-				AssetEditorToolkit->FillDefaultFileMenuCommands( MenuBuilder );
-			}
+			CreateDefaultStandaloneMenuBar(Menu);
+		}
+	}
 
-			static void AddAssetMenu( FMenuBarBuilder& MenuBarBuilder, TWeakPtr< FAssetEditorToolkit > AssetEditorToolkitWeak )
-			{
-				MenuBarBuilder.AddPullDownMenu( 
-					LOCTEXT("AssetMenuLabel", "Asset"),		// @todo toolkit major: Either use "Asset", "File", or the asset type name e.g. "Blueprint" (Also update custom pull-down menus)
-					LOCTEXT("AssetMenuLabel_ToolTip", "Opens a menu with commands for managing this asset"),
-					FNewMenuDelegate::CreateStatic( &Local::FillAssetMenu, AssetEditorToolkitWeak ),
-					"Asset");
+	DefaultMenuWidget = SNullWidget::NullWidget;
 
-				auto AssetEditorToolkit( AssetEditorToolkitWeak.Pin().ToSharedRef() );
-			}
+	HostTabPtr = InHostTab;
 
-			static void FillAssetMenu( FMenuBuilder& MenuBuilder, TWeakPtr< FAssetEditorToolkit > AssetEditorToolkitWeak )
-			{
-				auto AssetEditorToolkit( AssetEditorToolkitWeak.Pin().ToSharedRef() );
-				
-				MenuBuilder.BeginSection("AssetEditorActions", LOCTEXT("ActionsHeading", "Actions") );
-				{
-					AssetEditorToolkit->FillDefaultAssetMenuCommands( MenuBuilder );
-				}
-				MenuBuilder.EndSection();
-			}
+	RestoreFromLayout(DefaultLayout);
+	GenerateMenus(bCreateDefaultStandaloneMenu);
+}
 
-			static void ExtendHelpMenu( FMenuBuilder& MenuBuilder, TWeakPtr< FAssetEditorToolkit > AssetEditorToolkitWeak )
+void SStandaloneAssetEditorToolkitHost::CreateDefaultStandaloneMenuBar(UToolMenu* MenuBar)
+{
+	struct Local
+	{
+		static void ExtendFileMenu(UToolMenu* InMenuBar)
+		{
+			const FName MenuName = *(InMenuBar->GetMenuName().ToString() + TEXT(".") + TEXT("File"));
+			UToolMenu* Menu = UToolMenus::Get()->ExtendMenu(MenuName);
+			FToolMenuSection& Section = Menu->FindOrAddSection("FileLoadAndSave");
+			Section.AddDynamicEntry(NAME_None, FNewToolMenuSectionDelegate::CreateLambda([](FToolMenuSection& InSection)
 			{
-				auto AssetEditorToolkit( AssetEditorToolkitWeak.Pin().ToSharedRef() );
-
-				MenuBuilder.BeginSection("HelpBrowse", NSLOCTEXT("MainHelpMenu", "Browse", "Browse"));
-				{
-					AssetEditorToolkit->FillDefaultHelpMenuCommands( MenuBuilder );
-				}
-				MenuBuilder.EndSection();
-			}
-		};
-
-		TSharedPtr<FExtender> MenuExtender = MakeShareable(new FExtender());
-
-		auto AssetEditorToolkit = HostedAssetEditorToolkit.ToSharedRef();
-
-		// Add asset-specific menu items to the top of the "File" menu
-		MenuExtender->AddMenuExtension( "FileLoadAndSave", EExtensionHook::First, AssetEditorToolkit->GetToolkitCommands(), FMenuExtensionDelegate::CreateStatic( &Local::FillFileMenu, TWeakPtr< FAssetEditorToolkit >( AssetEditorToolkit ) ) );
+				InSection.FindContext<UAssetEditorToolkitMenuContext>()->Toolkit.Pin()->FillDefaultFileMenuCommands(InSection);
+			}));
+		}
 
-		// Add the "Asset" menu, if we're editing an asset
-		if (AssetEditorToolkit->IsActuallyAnAsset())
+		static void FillAssetMenu(UToolMenu* InMenu)
 		{
-			MenuExtender->AddMenuBarExtension( "Edit", EExtensionHook::After, AssetEditorToolkit->GetToolkitCommands(), FMenuBarExtensionDelegate::CreateStatic( &Local::AddAssetMenu, TWeakPtr< FAssetEditorToolkit >( AssetEditorToolkit ) ) );
+			FToolMenuSection& Section = InMenu->AddSection("AssetEditorActions", LOCTEXT("ActionsHeading", "Actions"));
+			Section.AddDynamicEntry(NAME_None, FNewToolMenuSectionDelegate::CreateLambda([](FToolMenuSection& InSection)
+			{
+				InSection.FindContext<UAssetEditorToolkitMenuContext>()->Toolkit.Pin()->FillDefaultAssetMenuCommands(InSection);
+			}));
 		}
 
-		MenuExtender->AddMenuExtension( "HelpOnline", EExtensionHook::Before, AssetEditorToolkit->GetToolkitCommands(), FMenuExtensionDelegate::CreateStatic( &Local::ExtendHelpMenu, TWeakPtr< FAssetEditorToolkit >( AssetEditorToolkit ) ) );
-
-		MenuExtenders.Add(MenuExtender);
-	}
+		static void ExtendHelpMenu(UToolMenu* InMenuBar)
+		{
+			const FName MenuName = *(InMenuBar->GetMenuName().ToString() + TEXT(".") + TEXT("Help"));
+			UToolMenu* Menu = UToolMenus::Get()->ExtendMenu(MenuName);
+			FToolMenuSection& Section = Menu->AddSection("HelpBrowse", NSLOCTEXT("MainHelpMenu", "Browse", "Browse"));
+			Section.InsertPosition = FToolMenuInsert("HelpOnline", EToolMenuInsertType::Before);
+			Section.AddDynamicEntry(NAME_None, FNewToolMenuSectionDelegate::CreateLambda([](FToolMenuSection& InSection)
+			{
+				InSection.FindContext<UAssetEditorToolkitMenuContext>()->Toolkit.Pin()->FillDefaultHelpMenuCommands(InSection);
+			}));
+		}
+	};
 
-	DefaultMenuWidget = SNullWidget::NullWidget;
+	// Add asset-specific menu items to the top of the "File" menu
+	Local::ExtendFileMenu(MenuBar);
 
-	HostTabPtr = InHostTab;
+	// Add the "Asset" menu, if we're editing an asset
+	MenuBar->FindOrAddSection(NAME_None).AddDynamicEntry("DynamicAssetEntry", FNewToolMenuSectionDelegate::CreateLambda([](FToolMenuSection& InSection)
+	{
+		UAssetEditorToolkitMenuContext* Context = InSection.FindContext<UAssetEditorToolkitMenuContext>();
+		if (Context->Toolkit.Pin()->IsActuallyAnAsset())
+		{
+			InSection.AddSubMenu(
+				"Asset",
+				LOCTEXT("AssetMenuLabel", "Asset"),		// @todo toolkit major: Either use "Asset", "File", or the asset type name e.g. "Blueprint" (Also update custom pull-down menus)
+				LOCTEXT("AssetMenuLabel_ToolTip", "Opens a menu with commands for managing this asset"),
+				FNewToolMenuDelegate::CreateStatic(&Local::FillAssetMenu)
+				).InsertPosition = FToolMenuInsert("Edit", EToolMenuInsertType::After);
+		}
+	}));
 
-	RestoreFromLayout(DefaultLayout);
-	GenerateMenus(bCreateDefaultStandaloneMenu);
+	// Add asset-specific menu items to the "Help" menu
+	Local::ExtendHelpMenu(MenuBar);
 }
 
 
@@ -199,7 +201,10 @@ void SStandaloneAssetEditorToolkitHost::GenerateMenus(bool bForceCreateMenu)
 	if( bForceCreateMenu || DefaultMenuWidget != SNullWidget::NullWidget )
 	{
 		const FName AssetEditorMenuName = GetMenuName();
-		FToolMenuContext ToolMenuContext(HostedAssetEditorToolkit->GetToolkitCommands(), FExtender::Combine(MenuExtenders).ToSharedRef());
+
+		UAssetEditorToolkitMenuContext* ContextObject = NewObject<UAssetEditorToolkitMenuContext>();
+		ContextObject->Toolkit = HostedAssetEditorToolkit;
+		FToolMenuContext ToolMenuContext(HostedAssetEditorToolkit->GetToolkitCommands(), FExtender::Combine(MenuExtenders).ToSharedRef(), ContextObject);
 		IMainFrameModule& MainFrameModule = FModuleManager::LoadModuleChecked<IMainFrameModule>( "MainFrame" );
 		DefaultMenuWidget = MainFrameModule.MakeMainMenu( MyTabManager, AssetEditorMenuName, ToolMenuContext );
 
diff --git a/Engine/Source/Editor/UnrealEd/Private/Toolkits/SStandaloneAssetEditorToolkitHost.h b/Engine/Source/Editor/UnrealEd/Private/Toolkits/SStandaloneAssetEditorToolkitHost.h
--- a/Engine/Source/Editor/UnrealEd/Private/Toolkits/SStandaloneAssetEditorToolkitHost.h
+++ b/Engine/Source/Editor/UnrealEd/Private/Toolkits/SStandaloneAssetEditorToolkitHost.h
@@ -57,6 +57,7 @@ public:
 	virtual void OnToolkitHostingStarted( const TSharedRef< class IToolkit >& Toolkit ) override;
 	virtual void OnToolkitHostingFinished( const TSharedRef< class IToolkit >& Toolkit ) override;
 	virtual UWorld* GetWorld() const override;
+	void CreateDefaultStandaloneMenuBar(UToolMenu* MenuBar);
 
 	/** SWidget overrides */
 	virtual bool SupportsKeyboardFocus() const override
diff --git a/Engine/Source/Editor/UnrealEd/Public/Toolkits/AssetEditorToolkit.h b/Engine/Source/Editor/UnrealEd/Public/Toolkits/AssetEditorToolkit.h
--- a/Engine/Source/Editor/UnrealEd/Public/Toolkits/AssetEditorToolkit.h
+++ b/Engine/Source/Editor/UnrealEd/Public/Toolkits/AssetEditorToolkit.h
@@ -19,6 +19,7 @@
 class FAssetEditorModeManager;
 class FMenuBuilder;
 struct FToolMenuContext;
+struct FToolMenuSection;
 class SBorder;
 class SStandaloneAssetEditorToolkitHost;
 
@@ -145,21 +146,21 @@ public:
 	 *
 	 * @param	MenuBuilder		The menu to add commands to
 	 */
-	void FillDefaultFileMenuCommands( FMenuBuilder& MenuBuilder );
+	void FillDefaultFileMenuCommands(FToolMenuSection& InSection);
 
 	/**
 	 * Fills in the supplied menu with commands for modifying this asset that are generally common to most asset editors
 	 *
 	 * @param	MenuBuilder		The menu to add commands to
 	 */
-	void FillDefaultAssetMenuCommands( FMenuBuilder& MenuBuilder );
+	void FillDefaultAssetMenuCommands(FToolMenuSection& InSection);
 
 	/**
 	 * Fills in the supplied menu with commands for the help menu
 	 *
 	 * @param	MenuBuilder		The menu to add commands to
 	 */
-	void FillDefaultHelpMenuCommands( FMenuBuilder& MenuBuilder );
+	void FillDefaultHelpMenuCommands(FToolMenuSection& InSection);
 
 	/** @return	For standalone asset editing tool-kits, returns the toolkit host that was last hosting this asset editor before it was switched to standalone mode (if it's still valid.)  Returns null if these conditions aren't met. */
 	TSharedPtr< IToolkitHost > GetPreviousWorldCentricToolkitHost();
