Reuse the default scene root node when creating a blueprint from an actor that has a root node named the same as default scene root (mostl likely an "empty actor")

diff --git a/Engine/Source/Editor/UnrealEd/Private/Kismet2/Kismet2.cpp b/Engine/Source/Editor/UnrealEd/Private/Kismet2/Kismet2.cpp
--- a/Engine/Source/Editor/UnrealEd/Private/Kismet2/Kismet2.cpp
+++ b/Engine/Source/Editor/UnrealEd/Private/Kismet2/Kismet2.cpp
@@ -1316,8 +1316,6 @@ void FKismetEditorUtilities::AddComponentsToBlueprint(UBlueprint* Blueprint, con
 			}
 			else
 			{
-				// Unsure what this case is really for when harvesting, but to be consistent with previous behavior still do this
-				ensure(bHarvesting);
 				AddChildToSCSRootNodeLambda(SCSNode);
 			}
 		}
@@ -1366,7 +1364,29 @@ UBlueprint* FKismetEditorUtilities::CreateBlueprintFromActor(const FName Bluepri
 			// If the source Actor has Instance Components we need to translate these in to SCS Nodes
 			if (Actor->GetInstanceComponents().Num() > 0)
 			{
-				AddComponentsToBlueprint(NewBlueprint, Actor->GetInstanceComponents(), false,  (USCS_Node*)nullptr, bKeepMobility);
+				bool bUsedDefaultSceneRoot = false;
+				if (USceneComponent* RootComponent = Actor->GetRootComponent())
+				{
+					// If the instance root component of the actor being converted is a scene component named the same as the default scene root node
+					// then we'll use that from the SimpleConstructionScript rather than creating a new root node
+					if (   RootComponent->GetClass() == USceneComponent::StaticClass()
+					    && RootComponent->GetFName() == NewBlueprint->SimpleConstructionScript->GetDefaultSceneRootNode()->GetVariableName()
+					    && Actor->GetInstanceComponents().Contains(RootComponent))
+					{
+						bUsedDefaultSceneRoot = true;
+						if (Actor->GetInstanceComponents().Num() > 1)
+						{
+							TArray<UActorComponent*> InstanceComponents = Actor->GetInstanceComponents();
+							InstanceComponents.Remove(RootComponent);
+							AddComponentsToBlueprint(NewBlueprint, InstanceComponents, false, NewBlueprint->SimpleConstructionScript->GetDefaultSceneRootNode(), bKeepMobility);
+						}
+					}
+				}
+
+				if (!bUsedDefaultSceneRoot)
+				{
+					AddComponentsToBlueprint(NewBlueprint, Actor->GetInstanceComponents(), false, nullptr, bKeepMobility);
+				}
 			}
 
 			if (NewBlueprint->GeneratedClass != nullptr)
