Fix blueprint Diff menu
#rnx

diff --git a/Engine/Source/Editor/UnrealEd/Private/Toolkits/SStandaloneAssetEditorToolkitHost.cpp b/Engine/Source/Editor/UnrealEd/Private/Toolkits/SStandaloneAssetEditorToolkitHost.cpp
--- a/Engine/Source/Editor/UnrealEd/Private/Toolkits/SStandaloneAssetEditorToolkitHost.cpp
+++ b/Engine/Source/Editor/UnrealEd/Private/Toolkits/SStandaloneAssetEditorToolkitHost.cpp
@@ -205,6 +205,7 @@ void SStandaloneAssetEditorToolkitHost::GenerateMenus(bool bForceCreateMenu)
 		UAssetEditorToolkitMenuContext* ContextObject = NewObject<UAssetEditorToolkitMenuContext>();
 		ContextObject->Toolkit = HostedAssetEditorToolkit;
 		FToolMenuContext ToolMenuContext(HostedAssetEditorToolkit->GetToolkitCommands(), FExtender::Combine(MenuExtenders).ToSharedRef(), ContextObject);
+		HostedAssetEditorToolkit->InitToolMenuContext(ToolMenuContext);
 		IMainFrameModule& MainFrameModule = FModuleManager::LoadModuleChecked<IMainFrameModule>( "MainFrame" );
 		DefaultMenuWidget = MainFrameModule.MakeMainMenu( MyTabManager, AssetEditorMenuName, ToolMenuContext );
 
