Optimize ShaderPipelineCacheToolsCommandlet

- Optimize parsing of FVertexElement, FDepthStencilStateInitializerRHI, FBlendStateInitializerRHI, FPipelineFileCacheRasterizerState, FPipelineCacheFileFormatPSO using FStringView and TArray<FStringView, TInlineAllocator<N>> to avoid temporary allocations.
- Optimize parsing and string conversion of FCompactFullName, FStableShaderKeyAndValue using FStringView and T[Ansi]StringBuilder to avoid temporary allocations.
- Use TInlineAllocator<16> for FCompactFullName to eliminate dynamic allocations. The size was selected by profiling a title with many shaders.
- Serialize FEditorShaderStableInfo directly to UTF-8 instead of converting to it later to avoid converting every name from ASCII to UTF-16 to UTF-8 during the serialization process.
- Load input files in parallel now that millions of allocations have been eliminated.
- Make the PSOSC serialization self-test only run when verbose logging is enabled.
- Optimize ParseQuoteComma by using FStringView and TArray<FStringView, TInlineAllocator<N>> to avoid temporary allocations.

These changes reduced the execution time of the build command by 90% on a large title. The time savings affect full, iterative, and single-package cooks.

diff --git a/Engine/Source/Editor/UnrealEd/Private/Commandlets/ShaderPipelineCacheToolsCommandlet.cpp b/Engine/Source/Editor/UnrealEd/Private/Commandlets/ShaderPipelineCacheToolsCommandlet.cpp
--- a/Engine/Source/Editor/UnrealEd/Private/Commandlets/ShaderPipelineCacheToolsCommandlet.cpp
+++ b/Engine/Source/Editor/UnrealEd/Private/Commandlets/ShaderPipelineCacheToolsCommandlet.cpp
@@ -3,12 +3,17 @@
 #include "Commandlets/ShaderPipelineCacheToolsCommandlet.h"
 #include "Misc/Paths.h"
 
+#include "Algo/Accumulate.h"
+#include "Async/ParallelFor.h"
 #include "PipelineFileCache.h"
 #include "ShaderCodeLibrary.h"
 #include "Misc/FileHelper.h"
+#include "Misc/ScopeExit.h"
+#include "Misc/StringBuilder.h"
 #include "ShaderPipelineCache.h"
 #include "Serialization/MemoryReader.h"
 #include "Serialization/MemoryWriter.h"
+#include "String/ParseLines.h"
 
 DEFINE_LOG_CATEGORY_STATIC(LogShaderPipelineCacheTools, Log, All);
 
@@ -55,30 +60,50 @@ void ExpandWildcards(TArray<FString>& Parts)
 	Parts = NewParts;
 }
 
+static void LoadStableSCL(TArray<FStableShaderKeyAndValue>& StableArray, const FStringView& FileName)
+{
+	UE_LOG(LogShaderPipelineCacheTools, Display, TEXT("Loading %.*s..."), FileName.Len(), FileName.GetData());
 
-void LoadStableSCL(TMultiMap<FStableShaderKeyAndValue, FSHAHash>& StableMap, const FString& Filename)
+	FString SourceFileContents;
+	TArray<FStringView> SourceFileLines;
+	if (FFileHelper::LoadFileToString(SourceFileContents, *FString(FileName)))
+	{
+		UE::String::ParseLines(SourceFileContents, [&SourceFileLines](FStringView Line) { if (!Line.IsEmpty()) { SourceFileLines.Add(Line); } });
+	}
+
+	if (SourceFileLines.Num() < 1)
+	{
+		UE_LOG(LogShaderPipelineCacheTools, Fatal, TEXT("Could not load %.*s..."), FileName.Len(), FileName.GetData());
+		return;
+	}
+
+	const int32 StableArrayOffset = StableArray.Num();
+	StableArray.AddZeroed(SourceFileLines.Num() - 1);
+	ParallelFor(SourceFileLines.Num() - 1, [&StableArray, &SourceFileLines, StableArrayOffset](int32 Index)
+	{
+		FStableShaderKeyAndValue& Item = StableArray[StableArrayOffset + Index];
+		Item.ParseFromString(SourceFileLines[Index + 1]);
+		check(!(Item.OutputHash == FSHAHash()));
+	});
+
+	UE_LOG(LogShaderPipelineCacheTools, Display, TEXT("Loaded %d shader info lines from %.*s."), StableArray.Num(), FileName.Len(), FileName.GetData());
+}
+
+static void LoadStableSCLs(TMultiMap<FStableShaderKeyAndValue, FSHAHash>& StableMap, TArrayView<const FStringView> FileNames)
 {
-	UE_LOG(LogShaderPipelineCacheTools, Display, TEXT("Loading %s...."), *Filename);
-	TArray<FString> SourceFileContents;
+	TArray<TArray<FStableShaderKeyAndValue>> StableArrays;
+	StableArrays.AddDefaulted(FileNames.Num());
+	ParallelFor(FileNames.Num(), [&StableArrays, &FileNames](int32 Index) { LoadStableSCL(StableArrays[Index], FileNames[Index]); });
 
-	if (FFileHelper::LoadFileToStringArray(SourceFileContents, *Filename))
+	const int32 StableArrayCount = Algo::TransformAccumulate(StableArrays, &TArray<FStableShaderKeyAndValue>::Num, 0);
+	StableMap.Reserve(StableMap.Num() + StableArrayCount);
+	for (const TArray<FStableShaderKeyAndValue>& StableArray : StableArrays)
 	{
-		StableMap.Reserve(StableMap.Num() + SourceFileContents.Num() - 1);
-		for (int32 Index = 1; Index < SourceFileContents.Num(); Index++)
+		for (const FStableShaderKeyAndValue& Item : StableArray)
 		{
-			FStableShaderKeyAndValue Item;
-			FMemory::Memzero(Item);
-			Item.ParseFromString(SourceFileContents[Index]);
-			check(!(Item.OutputHash == FSHAHash()));
 			StableMap.AddUnique(Item, Item.OutputHash);
 		}
 	}
-	if (SourceFileContents.Num() < 1)
-	{
-		UE_LOG(LogShaderPipelineCacheTools, Fatal, TEXT("Could not load %s"), *Filename);
-		return;
-	}
-	UE_LOG(LogShaderPipelineCacheTools, Display, TEXT("Loaded %d shader info lines"), SourceFileContents.Num() - 1);
 }
 
 static bool LoadAndDecompressStableCSV(const FString& Filename, TArray<uint8>& UncompressedData)
@@ -123,30 +148,51 @@ static bool LoadAndDecompressStableCSV(const FString& Filename, TArray<uint8>& U
 	return bResult;
 }
 
-static bool LoadStableCSV(const FString& Filename, TArray<FString>& OutputLines)
+struct FRawStableCSV
 {
-	bool bResult = false;
-	if (Filename.EndsWith(STABLE_CSV_COMPRESSED_EXT))
+	TArray<uint8> SerializedData;
+	FString CSV;
+};
+
+static bool LoadStableCSV(const FString& FileName, FRawStableCSV& RawStableCSV)
+{
+	if (FileName.EndsWith(STABLE_CSV_COMPRESSED_EXT))
 	{
-		TArray<uint8> DecompressedData;
-		if (LoadAndDecompressStableCSV(Filename, DecompressedData))
+		return LoadAndDecompressStableCSV(FileName, RawStableCSV.SerializedData);
+	}
+	else
+	{
+		return FFileHelper::LoadFileToString(RawStableCSV.CSV, *FileName);
+	}
+}
+
+static void ReadStableCSV(const FRawStableCSV& RawStableCSV, const TFunctionRef<void(FStringView)>& LineVisitor)
+{
+	if (RawStableCSV.SerializedData.Num())
+	{
+		FMemoryReader MemArchive(RawStableCSV.SerializedData);
+		FString LineCSV;
+		while (!MemArchive.AtEnd())
 		{
-			FMemoryReader MemArchive(DecompressedData);
-			FString LineCSV;
-			while (!MemArchive.AtEnd())
-			{
-				MemArchive << LineCSV;
-				OutputLines.Add(LineCSV);
-			}
-			bResult = true;
+			MemArchive << LineCSV;
+			LineVisitor(LineCSV);
 		}
 	}
 	else
 	{
-		bResult = FFileHelper::LoadFileToStringArray(OutputLines, *Filename);
+		UE::String::ParseLines(RawStableCSV.CSV, LineVisitor);
 	}
-	
-	return bResult;
+}
+
+static bool LoadStableCSV(const FString& FileName, TArray<FString>& OutputLines)
+{
+	FRawStableCSV RawStableCSV;
+	if (LoadStableCSV(FileName, RawStableCSV))
+	{
+		ReadStableCSV(RawStableCSV, [&OutputLines](FStringView Line) { OutputLines.Emplace(Line); });
+		return true;
+	}
+	return false;
 }
 
 static int64 SaveStableCSV(const FString& Filename, const TArray<uint8>& UncompressedData)
@@ -344,6 +390,24 @@ static bool GetStableShadersAndZeroHash(const TMap<FSHAHash, TArray<FStableShade
 	return true;
 }
 
+static void StableShadersSerializationSelfTest(const TMultiMap<FStableShaderKeyAndValue, FSHAHash>& StableMap)
+{
+	TAnsiStringBuilder<384> TestString;
+	for (const auto& Pair : StableMap)
+	{
+		TestString.Reset();
+		FStableShaderKeyAndValue Item(Pair.Key);
+		Item.OutputHash = Pair.Value;
+		check(Pair.Value != FSHAHash());
+		Item.AppendString(TestString);
+		FStableShaderKeyAndValue TestItem;
+		TestItem.ParseFromString(UTF8_TO_TCHAR(TestString.ToString()));
+		check(Item == TestItem);
+		check(GetTypeHash(Item) == GetTypeHash(TestItem));
+		check(Item.OutputHash == TestItem.OutputHash);
+	}
+}
+
 // return true if these two shaders could be part of the same stable PSO
 // for example, if they come from two different vertex factories, we return false because that situation cannot occur
 bool CouldBeUsedTogether(const FStableShaderKeyAndValue& A, const FStableShaderKeyAndValue& B)
@@ -389,9 +453,9 @@ bool CouldBeUsedTogether(const FStableShaderKeyAndValue& A, const FStableShaderK
 
 int32 DumpSCLCSV(const FString& Token)
 {
-
+	const FStringView File = Token;
 	TMultiMap<FStableShaderKeyAndValue, FSHAHash> StableMap;
-	LoadStableSCL(StableMap, Token);
+	LoadStableSCLs(StableMap, MakeArrayView(&File, 1));
 	UE_LOG(LogShaderPipelineCacheTools, Display, TEXT("    %s"), *FStableShaderKeyAndValue::HeaderLine());
 	for (const auto& Pair : StableMap)
 	{
@@ -459,15 +523,19 @@ void GeneratePermuations(TArray<FPermuation>& Permutations, FPermuation& Working
 
 int32 ExpandPSOSC(const TArray<FString>& Tokens)
 {
-	TMultiMap<FStableShaderKeyAndValue, FSHAHash> StableMap;
 	check(Tokens.Last().EndsWith(STABLE_CSV_EXT) || Tokens.Last().EndsWith(STABLE_CSV_COMPRESSED_EXT));
+
+	TArray<FStringView, TInlineAllocator<16>> StableCSVs;
 	for (int32 Index = 0; Index < Tokens.Num() - 1; Index++)
 	{
 		if (Tokens[Index].EndsWith(TEXT(".scl.csv")))
 		{
-			LoadStableSCL(StableMap, Tokens[Index]);
+			StableCSVs.Add(Tokens[Index]);
 		}
 	}
+
+	TMultiMap<FStableShaderKeyAndValue, FSHAHash> StableMap;
+	LoadStableSCLs(StableMap, StableCSVs);
 	if (!StableMap.Num())
 	{
 		UE_LOG(LogShaderPipelineCacheTools, Warning, TEXT("No .scl.csv found or they were all empty. Nothing to do."));
@@ -482,21 +550,8 @@ int32 ExpandPSOSC(const TArray<FString>& Tokens)
 			Temp.OutputHash = Pair.Value;
 			UE_LOG(LogShaderPipelineCacheTools, Verbose, TEXT("    %s"), *Temp.ToString());
 		}
+		StableShadersSerializationSelfTest(StableMap);
 	}
-	//self test
-	for (const auto& Pair : StableMap)
-	{
-		FStableShaderKeyAndValue Item(Pair.Key);
-		Item.OutputHash = Pair.Value;
-		check(Pair.Value != FSHAHash());
-		FString TestString = Item.ToString();
-		FStableShaderKeyAndValue TestItem;
-		TestItem.ParseFromString(TestString);
-		check(Item == TestItem);
-		check(GetTypeHash(Item) == GetTypeHash(TestItem));
-		check(Item.OutputHash == TestItem.OutputHash);
-	}
-	// end self test
 	UE_LOG(LogShaderPipelineCacheTools, Display, TEXT("Loaded %d unique shader info lines total."), StableMap.Num());
 
 	TSet<FPipelineCacheFileFormatPSO> PSOs;
@@ -907,9 +962,10 @@ int32 ExpandPSOSC(const TArray<FString>& Tokens)
 	return 0;
 }
 
-void ParseQuoteComma(const FString& InLine, TArray<FString>& OutParts)
+template <uint32 InlineSize>
+static void ParseQuoteComma(const FStringView& InLine, TArray<FStringView, TInlineAllocator<InlineSize>>& OutParts)
 {
-	FString Line = InLine;
+	FStringView Line = InLine;
 	while (true)
 	{
 		int32 QuoteLoc = 0;
@@ -917,16 +973,173 @@ void ParseQuoteComma(const FString& InLine, TArray<FString>& OutParts)
 		{
 			break;
 		}
-		Line.RightChopInline(QuoteLoc + 1, false);
+		Line.RightChopInline(QuoteLoc + 1);
 		if (!Line.FindChar(TCHAR('\"'), QuoteLoc))
 		{
 			break;
 		}
 		OutParts.Add(Line.Left(QuoteLoc));
-		Line.RightChopInline(QuoteLoc + 1, false);
+		Line.RightChopInline(QuoteLoc + 1);
 	}
 }
 
+static TSet<FPipelineCacheFileFormatPSO> ParseStableCSV(const FString& FileName, const FRawStableCSV& RawStableCSV, const TMultiMap<FStableShaderKeyAndValue, FSHAHash>& StableMap, FName& TargetPlatform)
+{
+	TSet<FPipelineCacheFileFormatPSO> PSOs;
+
+	int32 LineIndex = 0;
+	bool bParsed = true;
+	ReadStableCSV(RawStableCSV, [&FileName, &StableMap, &TargetPlatform, &PSOs, &LineIndex, &bParsed](FStringView Line)
+	{
+		// Skip the header line.
+		if (LineIndex++ == 0)
+		{
+			return;
+		}
+
+		// Only attempt to parse the current line if previous lines succeeded.
+		if (!bParsed)
+		{
+			return;
+		}
+
+		TArray<FStringView, TInlineAllocator<2 + SF_Compute>> Parts;
+		ParseQuoteComma(Line, Parts);
+
+		if (Parts.Num() != 2 + SF_Compute) // SF_Compute here because the stablepc.csv file format does not have a compute slot
+		{
+			// Assume the rest of the file csv lines are are bad or are in an out of date format. If one is, they probably all are.
+			UE_LOG(LogShaderPipelineCacheTools, Warning, TEXT("File %s is not in the correct format ignoring the rest of its contents."), *FileName);
+			bParsed = false;
+			return;
+		}
+
+		FPipelineCacheFileFormatPSO PSO;
+		FMemory::Memzero(PSO);
+		PSO.Type = FPipelineCacheFileFormatPSO::DescriptorType::Graphics; // we will change this to compute later if needed
+		PSO.CommonFromString(Parts[0]);
+		bool bValidGraphicsDesc = PSO.GraphicsDesc.StateFromString(Parts[1]);
+		if (!bValidGraphicsDesc)
+		{
+			// Failed to parse graphics descriptor, most likely format was changed, skip whole file.
+			UE_LOG(LogShaderPipelineCacheTools, Warning, TEXT("File %s is not in the correct format (GraphicsDesc) ignoring the rest of its contents."), *FileName);
+			bParsed = false;
+			return;
+		}
+
+		// For backward compatibility, compute shaders are stored as a zeroed graphics desc with the shader in the hull shader slot.
+		static FName NAME_SF_Compute("SF_Compute");
+		for (int32 SlotIndex = 0; SlotIndex < SF_Compute; ++SlotIndex)
+		{
+			if (Parts[SlotIndex + 2].IsEmpty())
+			{
+				continue;
+			}
+
+			FStableShaderKeyAndValue Shader;
+			Shader.ParseFromString(Parts[SlotIndex + 2]);
+
+			int32 AdjustedSlotIndex = SlotIndex;
+			if (SlotIndex == SF_Hull)
+			{
+				if (Shader.TargetFrequency == NAME_SF_Compute)
+				{
+					PSO.Type = FPipelineCacheFileFormatPSO::DescriptorType::Compute;
+					AdjustedSlotIndex = SF_Compute;
+				}
+			}
+			else
+			{
+				check(Shader.TargetFrequency != NAME_SF_Compute);
+			}
+
+			FSHAHash Match;
+			int32 Count = 0;
+			for (auto Iter = StableMap.CreateConstKeyIterator(Shader); Iter; ++Iter)
+			{
+				check(Iter.Value() != FSHAHash());
+				Match = Iter.Value();
+				if (TargetPlatform == NAME_None)
+				{
+					TargetPlatform = Iter.Key().TargetPlatform;
+				}
+				else
+				{
+					check(TargetPlatform == Iter.Key().TargetPlatform);
+				}
+				++Count;
+			}
+
+			if (!Count)
+			{
+				UE_LOG(LogShaderPipelineCacheTools, Verbose, TEXT("Stable PSO not found, rejecting %s"), *Shader.ToString());
+				return;
+			}
+
+			if (Count > 1)
+			{
+				UE_LOG(LogShaderPipelineCacheTools, Error, TEXT("Stable PSO maps to multiple shaders. This is usually a bad thing and means you used .scl.csv files from multiple builds. Ignoring all but the last %s"), *Shader.ToString());
+			}
+
+			switch (AdjustedSlotIndex)
+			{
+			case SF_Vertex:
+				PSO.GraphicsDesc.VertexShader = Match;
+				break;
+			case SF_Pixel:
+				PSO.GraphicsDesc.FragmentShader = Match;
+				break;
+			case SF_Geometry:
+				PSO.GraphicsDesc.GeometryShader = Match;
+				break;
+			case SF_Hull:
+				PSO.GraphicsDesc.HullShader = Match;
+				break;
+			case SF_Domain:
+				PSO.GraphicsDesc.DomainShader = Match;
+				break;
+			case SF_Compute:
+				PSO.ComputeDesc.ComputeShader = Match;
+				break;
+			}
+		}
+
+		if (PSO.Type == FPipelineCacheFileFormatPSO::DescriptorType::Compute)
+		{
+			check(PSO.ComputeDesc.ComputeShader != FSHAHash() &&
+				PSO.GraphicsDesc.VertexShader == FSHAHash() &&
+				PSO.GraphicsDesc.FragmentShader == FSHAHash() &&
+				PSO.GraphicsDesc.GeometryShader == FSHAHash() &&
+				PSO.GraphicsDesc.HullShader == FSHAHash() &&
+				PSO.GraphicsDesc.DomainShader == FSHAHash());
+		}
+		else
+		{
+			check(PSO.ComputeDesc.ComputeShader == FSHAHash());
+		}
+
+		if (!PSO.Verify())
+		{
+			UE_LOG(LogShaderPipelineCacheTools, Warning, TEXT("Bad PSO found. Verify failed. PSO discarded [Line %d in: %s]"), LineIndex, *FileName);
+			return;
+		}
+
+		// Merge duplicate PSO lines together.
+		if (FPipelineCacheFileFormatPSO* ExistingPSO = PSOs.Find(PSO))
+		{
+			check(*ExistingPSO == PSO);
+			ExistingPSO->UsageMask |= PSO.UsageMask;
+			ExistingPSO->BindCount = FMath::Max(ExistingPSO->BindCount, PSO.BindCount);
+		}
+		else
+		{
+			PSOs.Add(PSO);
+		}
+	});
+
+	return PSOs;
+}
+
 typedef TFunction<bool(const FString&)> FilenameFilterFN;
 
 void BuildDateSortedListOfFiles(const TArray<FString>& TokenList, FilenameFilterFN FilterFn, TArray<FString>& Result)
@@ -966,280 +1179,172 @@ void BuildDateSortedListOfFiles(const TArray<FString>& TokenList, FilenameFilter
 
 int32 BuildPSOSC(const TArray<FString>& Tokens)
 {
-	TMultiMap<FStableShaderKeyAndValue, FSHAHash> StableMap;
 	check(Tokens.Last().EndsWith(TEXT(".upipelinecache")));
 
+	TArray<FStringView, TInlineAllocator<16>> StableSCLs;
+	TArray<FString> StablePipelineCacheFiles;
+
 	for (int32 Index = 0; Index < Tokens.Num() - 1; Index++)
 	{
 		if (Tokens[Index].EndsWith(TEXT(".scl.csv")))
 		{
-			LoadStableSCL(StableMap, Tokens[Index]);
+			StableSCLs.Add(Tokens[Index]);
 		}
 	}
-	if (UE_LOG_ACTIVE(LogShaderPipelineCacheTools, Verbose))
+
+	// Get the stable PC files in date order - least to most important(!?)
+	UE_LOG(LogShaderPipelineCacheTools, Display, TEXT("Sorting input stablepc.csv files into chronological order for merge processing..."));
+	FilenameFilterFN ExtensionFilterFn = [](const FString& Filename)
 	{
-		UE_LOG(LogShaderPipelineCacheTools, Verbose, TEXT("    %s"), *FStableShaderKeyAndValue::HeaderLine());
-		for (const auto& Pair : StableMap)
+		return Filename.EndsWith(STABLE_CSV_EXT) || Filename.EndsWith(STABLE_CSV_COMPRESSED_EXT);
+	};
+	BuildDateSortedListOfFiles(Tokens, ExtensionFilterFn, StablePipelineCacheFiles);
+
+	// Start populating the stable SCLs in a task.
+	TMultiMap<FStableShaderKeyAndValue, FSHAHash> StableMap;
+	FGraphEventRef StableMapTask = FFunctionGraphTask::CreateAndDispatchWhenReady([&StableSCLs, &StableMap]
+	{
+		LoadStableSCLs(StableMap, StableSCLs);
+		if (UE_LOG_ACTIVE(LogShaderPipelineCacheTools, Verbose))
 		{
-			FStableShaderKeyAndValue Temp(Pair.Key);
-			Temp.OutputHash = Pair.Value;
-			UE_LOG(LogShaderPipelineCacheTools, Verbose, TEXT("    %s"), *Temp.ToString());
+			UE_LOG(LogShaderPipelineCacheTools, Verbose, TEXT("    %s"), *FStableShaderKeyAndValue::HeaderLine());
+			for (const auto& Pair : StableMap)
+			{
+				FStableShaderKeyAndValue Temp(Pair.Key);
+				Temp.OutputHash = Pair.Value;
+				UE_LOG(LogShaderPipelineCacheTools, Verbose, TEXT("    %s"), *Temp.ToString());
+			}
+			StableShadersSerializationSelfTest(StableMap);
 		}
-	}
-	//self test
-	for (const auto& Pair : StableMap)
+		UE_LOG(LogShaderPipelineCacheTools, Display, TEXT("Loaded %d unique shader info lines total."), StableMap.Num());
+	}, TStatId());
+
+	// Read the stable PSO sets in parallel with the stable shaders.
+	FGraphEventArray LoadPSOTasks;
+	TArray<FRawStableCSV> RawStableCSVs;
+	LoadPSOTasks.Reserve(StablePipelineCacheFiles.Num());
+	RawStableCSVs.AddDefaulted(StablePipelineCacheFiles.Num());
+	for (int32 FileIndex = 0; FileIndex < StablePipelineCacheFiles.Num(); ++FileIndex)
 	{
-		FStableShaderKeyAndValue Item(Pair.Key);
-		Item.OutputHash = Pair.Value;
-		check(Pair.Value != FSHAHash());
-		FString TestString = Item.ToString();
-		FStableShaderKeyAndValue TestItem;
-		TestItem.ParseFromString(TestString);
-		check(Item == TestItem);
-		check(GetTypeHash(Item) == GetTypeHash(TestItem));
-		check(Item.OutputHash == TestItem.OutputHash);
+		LoadPSOTasks.Add(FFunctionGraphTask::CreateAndDispatchWhenReady([&RawStableCSV = RawStableCSVs[FileIndex], &FileName = StablePipelineCacheFiles[FileIndex]]
+		{
+			if (!LoadStableCSV(FileName, RawStableCSV))
+			{
+				UE_LOG(LogShaderPipelineCacheTools, Fatal, TEXT("Could not load %s"), *FileName);
+			}
+		}, TStatId()));
+	}
+
+	// Parse the stable PSO sets in parallel once both the stable shaders and the corresponding read are complete.
+	FGraphEventArray ParsePSOTasks;
+	TArray<TSet<FPipelineCacheFileFormatPSO>> PSOsByFile;
+	TArray<FName> TargetPlatformByFile;
+	ParsePSOTasks.Reserve(StablePipelineCacheFiles.Num());
+	PSOsByFile.AddDefaulted(StablePipelineCacheFiles.Num());
+	TargetPlatformByFile.AddDefaulted(StablePipelineCacheFiles.Num());
+	for (int32 FileIndex = 0; FileIndex < StablePipelineCacheFiles.Num(); ++FileIndex)
+	{
+		const FGraphEventArray PreReqs{StableMapTask, LoadPSOTasks[FileIndex]};
+		ParsePSOTasks.Add(FFunctionGraphTask::CreateAndDispatchWhenReady(
+			[&PSOs = PSOsByFile[FileIndex],
+			 &FileName = StablePipelineCacheFiles[FileIndex],
+			 &RawStableCSV = RawStableCSVs[FileIndex],
+			 &StableMap,
+			 &TargetPlatform = TargetPlatformByFile[FileIndex]]
+			{
+				PSOs = ParseStableCSV(FileName, RawStableCSV, StableMap, TargetPlatform);
+				RawStableCSV = FRawStableCSV();
+				UE_LOG(LogShaderPipelineCacheTools, Display, TEXT("Loaded %d stable PSO lines from %s."), PSOs.Num(), *FileName);
+			}, TStatId(), &PreReqs));
 	}
-	// end self test
-	UE_LOG(LogShaderPipelineCacheTools, Display, TEXT("Loaded %d unique shader info lines total."), StableMap.Num());
 
-	TSet<FPipelineCacheFileFormatPSO> PSOs;
-	TMap<uint32,int64> PSOAvgIterations;
-	FName TargetPlatform;
-	
-	// Get the stable PC files in date order - least to most important(!?)
-	UE_LOG(LogShaderPipelineCacheTools, Display, TEXT("Sorting input stablepc.csv files into chronological order for merge processing...."));
-	
-	FilenameFilterFN ExtenstionFilterFn = [](const FString& Filename)
+	// Always wait for these tasks before returning from this function.
+	// This is necessary if there is an error or if nothing consumes the stable map.
+	ON_SCOPE_EXIT
 	{
-		return Filename.EndsWith(STABLE_CSV_EXT) || Filename.EndsWith(STABLE_CSV_COMPRESSED_EXT);
+		FTaskGraphInterface::Get().WaitUntilTaskCompletes(StableMapTask);
+		FTaskGraphInterface::Get().WaitUntilTasksComplete(ParsePSOTasks);
 	};
 
-	TArray<FString> StablePiplineCacheFiles;
-	BuildDateSortedListOfFiles(Tokens, ExtenstionFilterFn, StablePiplineCacheFiles);
-
+	// Validate and merge the stable PSO sets sequentially as they finish.
+	TSet<FPipelineCacheFileFormatPSO> PSOs;
+	TMap<uint32,int64> PSOAvgIterations;
 	uint32 MergeCount = 0;
+	FName TargetPlatform;
 
-	for(int32 FileIndex = 0;FileIndex < StablePiplineCacheFiles.Num();++FileIndex)
+	for (int32 FileIndex = 0; FileIndex < StablePipelineCacheFiles.Num(); ++FileIndex)
 	{
-		FString const& FileName = StablePiplineCacheFiles[FileIndex];
-
-		UE_LOG(LogShaderPipelineCacheTools, Display, TEXT("Loading %s...."), *FileName);
-		TArray<FString> SourceFileContents;
+		FTaskGraphInterface::Get().WaitUntilTaskCompletes(ParsePSOTasks[FileIndex]);
 
-		if (!LoadStableCSV(FileName, SourceFileContents) || SourceFileContents.Num() < 2)
+		if (!PSOsByFile[FileIndex].Num())
 		{
-			UE_LOG(LogShaderPipelineCacheTools, Fatal, TEXT("Could not load %s"), *FileName);
 			return 1;
 		}
 
-		UE_LOG(LogShaderPipelineCacheTools, Display, TEXT("Loaded %d stable PSO lines."), SourceFileContents.Num() - 1);
+		check(TargetPlatform == NAME_None || TargetPlatform == TargetPlatformByFile[FileIndex]);
+		TargetPlatform = TargetPlatformByFile[FileIndex];
 
-		TSet<FPipelineCacheFileFormatPSO> CurrentFilePSOs;
-		for (int32 Index = 1; Index < SourceFileContents.Num(); Index++)
-		{
-			TArray<FString> Parts;
-			ParseQuoteComma(SourceFileContents[Index], Parts);
-			
-			if(Parts.Num() != 2 + SF_Compute) // SF_Compute here because the stablepc.csv file format does not have a compute slot
-			{
-				// Assume the rest of the file csv lines are are bad or are in an out of date format - if one is - they probably all are
-				UE_LOG(LogShaderPipelineCacheTools, Warning, TEXT("File %s is not in the correct format ignoring the rest of its contents."), *FileName);
-				break;
-			}
+		TSet<FPipelineCacheFileFormatPSO>& CurrentFilePSOs = PSOsByFile[FileIndex];
 
-			FPipelineCacheFileFormatPSO PSO;
-			FMemory::Memzero(PSO);
-			PSO.Type = FPipelineCacheFileFormatPSO::DescriptorType::Graphics; // we will change this to compute later if needed
-			PSO.CommonFromString(Parts[0]);
-			bool bValidGraphicsDesc = PSO.GraphicsDesc.StateFromString(Parts[1]);
-			if (!bValidGraphicsDesc)
-			{
-				// failed to parse graphics descriptor, most likely format was changed, skip whole file
-				UE_LOG(LogShaderPipelineCacheTools, Warning, TEXT("File %s is not in the correct format (GraphicsDesc) ignoring the rest of its contents."), *FileName);
-				break;
-			}
+		if (!CurrentFilePSOs.Num())
+		{
+			continue;
+		}
 
-			bool bValid = true;
+		// Now merge this file PSO set with main PSO set (this is going to be slow as we need to incrementally reprocess each existing PSO per file to get reasonable bindcount averages).
+		// Can't sum all and avg: A) Overflow and B) Later ones want to remain high so only start to get averaged from the point they are added onwards:
+		// 1) New PSO goes in with it's bindcount intact for this iteration - if it's the last file then it keeps it bindcount
+		// 2) Existing PSO from older file gets incrementally averaged with PSO bindcount from new file
+		// 3) Existing PSO from older file not in new file set gets incrementally averaged with zero - now less important
+		// 4) PSOs are incrementally averaged from the point they are seen - i.e. a PSO seen in an earler file will get averaged more times than one
+		//		seen in a later file using:  NewAvg = OldAvg + (NewValue - OldAvg) / CountFromPSOSeen
+		//
+		// Proof for incremental averaging:
+		//	DataSet = {25 65 95 128}; Standard Average = (sum(25, 65, 95, 128) / 4) = 78.25
+		//	Incremental:
+		//	=> 25
+		//	=> 25 + (65 - 25) / 2 = A 		==> 25 + (65 - 25) / 2 		= 45
+		//	=>  A + (95 -  A) / 3 = B 		==> 45 + (95 - 45) / 3 		= 61 2/3
+		//	=>  B + (128 - B) / 4 = Answer 	==> 61 2/3 + (128 - B) / 4 	= 78.25
+
+		for (FPipelineCacheFileFormatPSO& PSO : PSOs)
+		{
+			// Already existing PSO in the next file round - increase its average iteration
+			int64& PSOAvgIteration = PSOAvgIterations.FindChecked(GetTypeHash(PSO));
+			++PSOAvgIteration;
 
-			bool bLooksLikeAComputeShader = false;
+			// Default the bindcount
+			int64 NewBindCount = 0ll;
 
-			static FName NAME_SF_Compute("SF_Compute");
-			// because it is a CSV, and for backward compat, compute shaders will just be a zeroed graphics desc with the shader in the hull shader slot.
-			for (int32 SlotIndex = 0; SlotIndex < SF_Compute; SlotIndex++) // SF_Compute here because the stablepc.csv file format does not have a compute slot
+			// If you have the same PSO in the new file set
+			if (FPipelineCacheFileFormatPSO* NewFilePSO = CurrentFilePSOs.Find(PSO))
 			{
-				if (!Parts[SlotIndex + 2].Len())
-				{
-					continue;
-				}
+				// Sanity check!
+				check(*NewFilePSO == PSO);
 
-				FStableShaderKeyAndValue Shader;
-				Shader.ParseFromString(Parts[SlotIndex + 2]);
-
-				if (SlotIndex == SF_Hull)
-				{
-					if (Shader.TargetFrequency == NAME_SF_Compute)
-					{
-						bLooksLikeAComputeShader = true;
-					}
-				}
-				else
+				// Get More accurate stats by testing for diff - we could just merge and be done
+				if ((PSO.UsageMask & NewFilePSO->UsageMask) != NewFilePSO->UsageMask)
 				{
-					check(Shader.TargetFrequency != NAME_SF_Compute);
+					PSO.UsageMask |= NewFilePSO->UsageMask;
+					++MergeCount;
 				}
 
-				FSHAHash Match;
-				int32 Count = 0;
-				for (auto Iter = StableMap.CreateConstKeyIterator(Shader); Iter; ++Iter)
-				{
-					check(Iter.Value() != FSHAHash());
-					Match = Iter.Value();
-					if (TargetPlatform == NAME_None)
-					{
-						TargetPlatform = Iter.Key().TargetPlatform;
-					}
-					else
-					{
-						check(TargetPlatform == Iter.Key().TargetPlatform);
-					}
-					Count++;
-				}
-
-				if (!Count)
-				{
-					UE_LOG(LogShaderPipelineCacheTools, Verbose, TEXT("Stable PSO not found, rejecting %s"), *Shader.ToString());
-					bValid = false;
-					break;
-				}
-
-				if (Count > 1)
-				{
-					UE_LOG(LogShaderPipelineCacheTools, Error, TEXT("Stable PSO maps to multiple shaders. This is usually a bad thing and means you used .scl.csv files from multiple builds. Ignoring all but the last %s"), *Shader.ToString());
-				}
+				NewBindCount = NewFilePSO->BindCount;
 
-				switch (SlotIndex)
-				{
-				case SF_Vertex:
-					PSO.GraphicsDesc.VertexShader = Match;
-					break;
-				case SF_Pixel:
-					PSO.GraphicsDesc.FragmentShader = Match;
-					break;
-				case SF_Geometry:
-					PSO.GraphicsDesc.GeometryShader = Match;
-					break;
-				case SF_Hull:
-					PSO.GraphicsDesc.HullShader = Match;
-					break;
-				case SF_Domain:
-					PSO.GraphicsDesc.DomainShader = Match;
-					break;
-				}
-			}
-			if (bValid)
-			{
-				if (
-					PSO.GraphicsDesc.VertexShader == FSHAHash() &&
-					PSO.GraphicsDesc.FragmentShader == FSHAHash() &&
-					PSO.GraphicsDesc.GeometryShader == FSHAHash() &&
-					!(PSO.GraphicsDesc.HullShader == FSHAHash()) &&    // Compute shaders are stored in the hull slot
-					PSO.GraphicsDesc.DomainShader == FSHAHash() &&
-					bLooksLikeAComputeShader
-					)
-				{
-					// this is a compute shader
-					PSO.Type = FPipelineCacheFileFormatPSO::DescriptorType::Compute;
-					PSO.ComputeDesc.ComputeShader = PSO.GraphicsDesc.HullShader;
-					PSO.GraphicsDesc.HullShader = FSHAHash();
-				}
-				else
-				{
-					PSO.Type = FPipelineCacheFileFormatPSO::DescriptorType::Graphics;
-					check(!bLooksLikeAComputeShader);
-				}
-				
-				bValid = PSO.Verify();
-				if(!bValid)
-				{
-					UE_LOG(LogShaderPipelineCacheTools, Warning, TEXT("Bad PSO found - verify failed - PSO discarded [Line %d in: %s]"), Index, *FileName);
-				}
+				// Remove from current file set - it's already there and we don't want any 'overwrites'
+				CurrentFilePSOs.Remove(*NewFilePSO);
 			}
 
-			if (bValid)
-			{
-				// Merge duplicate PSO lines in the same file together - merge mask and Max bindcount
-				auto* ExistingPSO = CurrentFilePSOs.Find(PSO);
-				if(ExistingPSO != nullptr)
-				{
-					check(*ExistingPSO == PSO);
-					ExistingPSO->UsageMask |= PSO.UsageMask;
-					ExistingPSO->BindCount = FMath::Max(ExistingPSO->BindCount, PSO.BindCount);
-				}
-				else
-				{
-					CurrentFilePSOs.Add(PSO);
-				}
-				
-				if(!PSOAvgIterations.Contains(GetTypeHash(PSO)))
-				{
-					PSOAvgIterations.Add(GetTypeHash(PSO), 1ll);
-				}
-			}
+			// Incrementally average this PSO bindcount - if not found in this set then avg will be pulled down
+			PSO.BindCount += (NewBindCount - PSO.BindCount) / PSOAvgIteration;
 		}
-		
-		if(CurrentFilePSOs.Num())
-		{
-			// Now merge this file PSO set with main PSO set (this is going to be slow as we need to incrementally reprocess each existing PSO per file to get reasonable bindcount averages).
-			// Can't sum all and avg: A) Overflow and B) Later ones want to remain high so only start to get averaged from the point they are added onwards:
-			// 1) New PSO goes in with it's bindcount intact for this iteration - if it's the last file then it keeps it bindcount
-			// 2) Existing PSO from older file gets incrementally averaged with PSO bindcount from new file
-			// 3) Existing PSO from older file not in new file set gets incrementally averaged with zero - now less important
-			// 4) PSOs are incrementally averaged from the point they are seen - i.e. a PSO seen in an earler file will get averaged more times than one
-			//		seen in a later file using:  NewAvg = OldAvg + (NewValue - OldAvg) / CountFromPSOSeen
-			//
-			// Proof for incremental averaging:
-			//	DataSet = {25 65 95 128}; Standard Average = (sum(25, 65, 95, 128) / 4) = 78.25
-			//	Incremental:
-			//	=> 25
-			//	=> 25 + (65 - 25) / 2 = A 		==> 25 + (65 - 25) / 2 		= 45
-			//	=>  A + (95 -  A) / 3 = B 		==> 45 + (95 - 45) / 3 		= 61 2/3
-			//	=>  B + (128 - B) / 4 = Answer 	==> 61 2/3 + (128 - B) / 4 	= 78.25
-			
-			for( FPipelineCacheFileFormatPSO& PSO : PSOs)
-			{
-				// Already existing PSO in the next file round - increase it's average iteration
-				int64& PSOAvgIteration = PSOAvgIterations.FindChecked(GetTypeHash(PSO));
-				++PSOAvgIteration;
-				
-				// Default the bindcount
-				int64 NewBindCount = 0ll;
-				
-				// If you have the same PSO in the new file set
-				auto* NewFilePSO = CurrentFilePSOs.Find(PSO);
-				if(NewFilePSO != nullptr)
-				{
-					// Sanity check!
-					check(*NewFilePSO == PSO);
-
-					// Get More accurate stats by testing for diff - we could just merge and be done
-					if((PSO.UsageMask & NewFilePSO->UsageMask) != NewFilePSO->UsageMask)
-					{
-						PSO.UsageMask |= NewFilePSO->UsageMask;
-						++MergeCount;
-					}
-					
-					NewBindCount = NewFilePSO->BindCount;
 
-					// Remove from current file set - it's already there and we don't want any 'overwrites'
-					CurrentFilePSOs.Remove(*NewFilePSO);
-				}
-				
-				// Incrementally average this PSO bindcount - if not found in this set then avg will be pulled down
-				PSO.BindCount += (NewBindCount - PSO.BindCount) / PSOAvgIteration;
-			}
-			
-			// Just add any left over - their iterations will be 1 and not yet averaged
-			PSOs.Append(CurrentFilePSOs);
+		// Add the leftover PSOs from the current file and initialize their iteration count.
+		for (const FPipelineCacheFileFormatPSO& PSO : CurrentFilePSOs)
+		{
+			PSOAvgIterations.Add(GetTypeHash(PSO), 1ll);
 		}
+		PSOs.Append(MoveTemp(CurrentFilePSOs));
 	}
 	UE_LOG(LogShaderPipelineCacheTools, Display, TEXT("Re-deduplicated into %d binary PSOs [Usage Mask Merged = %d]."), PSOs.Num(), MergeCount);
 
@@ -1286,11 +1391,6 @@ int32 BuildPSOSC(const TArray<FString>& Tokens)
 			{
 				for (const FPipelineCacheFileFormatPSO& TestItem : KeptPSOs)
 				{
-					FSHAHash VertexShader;
-					FSHAHash FragmentShader;
-					FSHAHash GeometryShader;
-					FSHAHash HullShader;
-					FSHAHash DomainShader;
 					if (TestItem.Type == FPipelineCacheFileFormatPSO::DescriptorType::Graphics)
 					{
 						if (
diff --git a/Engine/Source/Runtime/RHI/Private/PipelineFileCache.cpp b/Engine/Source/Runtime/RHI/Private/PipelineFileCache.cpp
--- a/Engine/Source/Runtime/RHI/Private/PipelineFileCache.cpp
+++ b/Engine/Source/Runtime/RHI/Private/PipelineFileCache.cpp
@@ -19,6 +19,8 @@ PipelineFileCache.cpp: Pipeline state cache implementation.
 #include "HAL/PlatformFilemanager.h"
 #include "Misc/FileHelper.h"
 #include "ProfilingDebugging/CsvProfiler.h"
+#include "String/LexFromString.h"
+#include "String/ParseTokens.h"
 
 static FString JOURNAL_FILE_EXTENSION(TEXT(".jnl"));
 
@@ -267,31 +269,33 @@ FString FPipelineFileCacheRasterizerState::ToString() const
 	);
 }
 
-void FPipelineFileCacheRasterizerState::FromString(const FString& InSrc)
+void FPipelineFileCacheRasterizerState::FromString(const FStringView& Src)
 {
-	FString Src = InSrc;
-	Src.ReplaceInline(TEXT("\r"), TEXT(" "));
-	Src.ReplaceInline(TEXT("\n"), TEXT(" "));
-	Src.ReplaceInline(TEXT("\t"), TEXT(" "));
-	Src.ReplaceInline(TEXT("<"), TEXT(" "));
-	Src.ReplaceInline(TEXT(">"), TEXT(" "));
-	TArray<FString> Parts;
-	Src.TrimStartAndEnd().ParseIntoArray(Parts, TEXT(" "));
-
-	check(Parts.Num() == 6 && sizeof(FillMode) == 1 && sizeof(CullMode) == 1 && sizeof(bAllowMSAA) == 1 && sizeof(bEnableLineAA) == 1); //not a very robust parser
-	LexFromString(DepthBias, *Parts[0]);
-	LexFromString(SlopeScaleDepthBias, *Parts[1]);
-	LexFromString((uint8&)FillMode, *Parts[2]);
-	LexFromString((uint8&)CullMode, *Parts[3]);
-	LexFromString((uint8&)bAllowMSAA, *Parts[4]);
-	LexFromString((uint8&)bEnableLineAA, *Parts[5]);
+	constexpr int32 PartCount = 6;
+
+	TArray<FStringView, TInlineAllocator<PartCount>> Parts;
+	UE::String::ParseTokensMultiple(Src.TrimStartAndEnd(), {TEXT('\r'), TEXT('\n'), TEXT('\t'), TEXT('<'), TEXT('>'), TEXT(' ')},
+		[&Parts](FStringView Part) { if (!Part.IsEmpty()) { Parts.Add(Part); } });
+
+	check(Parts.Num() == PartCount && sizeof(FillMode) == 1 && sizeof(CullMode) == 1 && sizeof(bAllowMSAA) == 1 && sizeof(bEnableLineAA) == 1); //not a very robust parser
+	const FStringView* PartIt = Parts.GetData();
+
+	LexFromString(DepthBias, *PartIt++);
+	LexFromString(SlopeScaleDepthBias, *PartIt++);
+	LexFromString((uint8&)FillMode, *PartIt++);
+	LexFromString((uint8&)CullMode, *PartIt++);
+	LexFromString((uint8&)bAllowMSAA, *PartIt++);
+	LexFromString((uint8&)bEnableLineAA, *PartIt++);
+
+	check(Parts.GetData() + PartCount == PartIt);
 }
 
 FString FPipelineCacheFileFormatPSO::ComputeDescriptor::ToString() const
 {
 	return ComputeShader.ToString();
 }
-void FPipelineCacheFileFormatPSO::ComputeDescriptor::FromString(const FString& Src)
+
+void FPipelineCacheFileFormatPSO::ComputeDescriptor::FromString(const FStringView& Src)
 {
 	ComputeShader.FromString(Src.TrimStartAndEnd());
 }
@@ -315,18 +319,24 @@ FString FPipelineCacheFileFormatPSO::GraphicsDescriptor::ShadersToString() const
 
 	return Result;
 }
-void FPipelineCacheFileFormatPSO::GraphicsDescriptor::ShadersFromString(const FString& Src)
+
+void FPipelineCacheFileFormatPSO::GraphicsDescriptor::ShadersFromString(const FStringView& Src)
 {
-	TArray<FString> Parts;
-	Src.TrimStartAndEnd().ParseIntoArray(Parts, TEXT(","));
-	int32 PartIndex = 0;
-	check(Parts.Num() == 5); //not a very robust parser
-
-	VertexShader.FromString(Parts[PartIndex++]);
-	FragmentShader.FromString(Parts[PartIndex++]);
-	GeometryShader.FromString(Parts[PartIndex++]);
-	HullShader.FromString(Parts[PartIndex++]);
-	DomainShader.FromString(Parts[PartIndex++]);
+	constexpr int32 PartCount = 5;
+
+	TArray<FStringView, TInlineAllocator<PartCount>> Parts;
+	UE::String::ParseTokens(Src.TrimStartAndEnd(), TEXT(','), [&Parts](FStringView Part) { Parts.Add(Part); });
+
+	check(Parts.Num() == PartCount); //not a very robust parser
+	const FStringView* PartIt = Parts.GetData();
+
+	VertexShader.FromString(*PartIt++);
+	FragmentShader.FromString(*PartIt++);
+	GeometryShader.FromString(*PartIt++);
+	HullShader.FromString(*PartIt++);
+	DomainShader.FromString(*PartIt++);
+
+	check(Parts.GetData() + PartCount == PartIt);
 }
 
 FString FPipelineCacheFileFormatPSO::GraphicsDescriptor::ShaderHeaderLine()
@@ -397,74 +407,80 @@ FString FPipelineCacheFileFormatPSO::GraphicsDescriptor::StateToString() const
 	return Result.Left(Result.Len() - 1); // remove trailing comma
 }
 
-bool FPipelineCacheFileFormatPSO::GraphicsDescriptor::StateFromString(const FString& Src)
+bool FPipelineCacheFileFormatPSO::GraphicsDescriptor::StateFromString(const FStringView& Src)
 {
-	TArray<FString> Parts;
-	Src.TrimStartAndEnd().ParseIntoArray(Parts, TEXT(","));
-	int32 PartIndex = 0;
+	constexpr int32 PartCount = FPipelineCacheGraphicsDescPartsNum;
+
+	TArray<FStringView, TInlineAllocator<PartCount>> Parts;
+	UE::String::ParseTokens(Src.TrimStartAndEnd(), TEXT(','), [&Parts](FStringView Part) { Parts.Add(Part); });
 
 	// check if we have expected number of parts
-	if (Parts.Num() != FPipelineCacheGraphicsDescPartsNum)
+	if (Parts.Num() != PartCount)
 	{
 		// instead of crashing let caller handle this case
 		return false;
 	}
 
-	check(Parts.Num() - PartIndex >= 3); //not a very robust parser
-	BlendState.FromString(Parts[PartIndex++]);
-	RasterizerState.FromString(Parts[PartIndex++]);
-	DepthStencilState.FromString(Parts[PartIndex++]);
+	const FStringView* PartIt = Parts.GetData();
+	const FStringView* PartEnd = PartIt + PartCount;
+
+	check(PartEnd - PartIt >= 3); //not a very robust parser
+	BlendState.FromString(*PartIt++);
+	RasterizerState.FromString(*PartIt++);
+	DepthStencilState.FromString(*PartIt++);
 
-	check(Parts.Num() - PartIndex >= 3 && sizeof(EPixelFormat) == sizeof(uint32)); //not a very robust parser
-	LexFromString(MSAASamples, *Parts[PartIndex++]);
-	LexFromString((uint32&)DepthStencilFormat, *Parts[PartIndex++]);
-	LexFromString(DepthStencilFlags, *Parts[PartIndex++]);
+	check(PartEnd - PartIt >= 3 && sizeof(EPixelFormat) == sizeof(uint32)); //not a very robust parser
+	LexFromString(MSAASamples, *PartIt++);
+	LexFromString((uint32&)DepthStencilFormat, *PartIt++);
+	LexFromString(DepthStencilFlags, *PartIt++);
 
-	check(Parts.Num() - PartIndex >= 5 && sizeof(DepthLoad) == 1 && sizeof(StencilLoad) == 1 && sizeof(DepthStore) == 1 && sizeof(StencilStore) == 1 && sizeof(PrimitiveType) == 4); //not a very robust parser
-	LexFromString((uint32&)DepthLoad, *Parts[PartIndex++]);
-	LexFromString((uint32&)StencilLoad, *Parts[PartIndex++]);
-	LexFromString((uint32&)DepthStore, *Parts[PartIndex++]);
-	LexFromString((uint32&)StencilStore, *Parts[PartIndex++]);
-	LexFromString((uint32&)PrimitiveType, *Parts[PartIndex++]);
+	check(PartEnd - PartIt >= 5 && sizeof(DepthLoad) == 1 && sizeof(StencilLoad) == 1 && sizeof(DepthStore) == 1 && sizeof(StencilStore) == 1 && sizeof(PrimitiveType) == 4); //not a very robust parser
+	LexFromString((uint32&)DepthLoad, *PartIt++);
+	LexFromString((uint32&)StencilLoad, *PartIt++);
+	LexFromString((uint32&)DepthStore, *PartIt++);
+	LexFromString((uint32&)StencilStore, *PartIt++);
+	LexFromString((uint32&)PrimitiveType, *PartIt++);
 
-	check(Parts.Num() - PartIndex >= 1); //not a very robust parser
-	LexFromString(RenderTargetsActive, *Parts[PartIndex++]);
+	check(PartEnd - PartIt >= 1); //not a very robust parser
+	LexFromString(RenderTargetsActive, *PartIt++);
 
 	for (int32 Index = 0; Index < MaxSimultaneousRenderTargets; Index++)
 	{
-		check(Parts.Num() - PartIndex >= 4 && sizeof(ERenderTargetLoadAction) == 1 && sizeof(ERenderTargetStoreAction) == 1 && sizeof(EPixelFormat) == sizeof(uint32)); //not a very robust parser
-		LexFromString((uint32&)(RenderTargetFormats[Index]), *Parts[PartIndex++]);
-		LexFromString(RenderTargetFlags[Index], *Parts[PartIndex++]);
+		check(PartEnd - PartIt >= 4 && sizeof(ERenderTargetLoadAction) == 1 && sizeof(ERenderTargetStoreAction) == 1 && sizeof(EPixelFormat) == sizeof(uint32)); //not a very robust parser
+		LexFromString((uint32&)(RenderTargetFormats[Index]), *PartIt++);
+		LexFromString(RenderTargetFlags[Index], *PartIt++);
 		uint8 Load, Store;
-		LexFromString(Load, *Parts[PartIndex++]);
-		LexFromString(Store, *Parts[PartIndex++]);
+		LexFromString(Load, *PartIt++);
+		LexFromString(Store, *PartIt++);
 	}
 
 	// parse sub-pass information
 	{
 		uint32 LocalSubpassHint = 0;
 		uint32 LocalSubpassIndex = 0;
-		check(Parts.Num() - PartIndex >= 2);
-		LexFromString(LocalSubpassHint, *Parts[PartIndex++]);
-		LexFromString(LocalSubpassIndex, *Parts[PartIndex++]);
+		check(PartEnd - PartIt >= 2);
+		LexFromString(LocalSubpassHint, *PartIt++);
+		LexFromString(LocalSubpassIndex, *PartIt++);
 		SubpassHint = LocalSubpassHint;
 		SubpassIndex = LocalSubpassIndex;
 	}
-	
-	check(Parts.Num() - PartIndex >= 1); //not a very robust parser
+
+	check(PartEnd - PartIt >= 1); //not a very robust parser
 	int32 VertDescNum = 0;
-	LexFromString(VertDescNum, *Parts[PartIndex++]);
+	LexFromString(VertDescNum, *PartIt++);
 	check(VertDescNum >= 0 && VertDescNum <= MaxVertexElementCount);
 
 	VertexDescriptor.Empty(VertDescNum);
 	VertexDescriptor.AddZeroed(VertDescNum);
 
-	check(Parts.Num() - PartIndex == MaxVertexElementCount); //not a very robust parser
+	check(PartEnd - PartIt == MaxVertexElementCount); //not a very robust parser
 	for (int32 Index = 0; Index < VertDescNum; Index++)
 	{
-		VertexDescriptor[Index].FromString(Parts[PartIndex++]);
+		VertexDescriptor[Index].FromString(*PartIt++);
 	}
-	
+
+	check(PartIt + MaxVertexElementCount == PartEnd + VertDescNum);
+
 	VertexDescriptor.Sort([](FVertexElement const& A, FVertexElement const& B)
 	  {
 		  if (A.StreamIndex < B.StreamIndex)
@@ -554,17 +570,21 @@ FString FPipelineCacheFileFormatPSO::GraphicsDescriptor::ToString() const
 	return FString::Printf(TEXT("%s,%s"), *ShadersToString(), *StateToString());
 }
 
-bool FPipelineCacheFileFormatPSO::GraphicsDescriptor::FromString(const FString& Src)
+bool FPipelineCacheFileFormatPSO::GraphicsDescriptor::FromString(const FStringView& Src)
 {
-	static const int32 NumShaderParts = 5;
-	TArray<FString> Parts;
-	Src.TrimStartAndEnd().ParseIntoArray(Parts, TEXT(","));
-	check(Parts.Num() > NumShaderParts);
-	TArray<FString> StateParts(Parts);
-	StateParts.RemoveAt(0, NumShaderParts);
-	Parts.RemoveAt(NumShaderParts, Parts.Num() - 5);
-	ShadersFromString(FString::Join(Parts, TEXT(",")));
-	return StateFromString(FString::Join(StateParts, TEXT(",")));
+	constexpr int32 NumShaderParts = 5;
+
+	int32 StateOffset = 0;
+	for (int32 CommaCount = 0; CommaCount < NumShaderParts; ++CommaCount)
+	{
+		int32 CommaOffset = 0;
+		bool FoundComma = Src.RightChop(StateOffset).FindChar(TEXT(','), CommaOffset);
+		check(FoundComma);
+		StateOffset += CommaOffset + 1;
+	}
+
+	ShadersFromString(Src.Left(StateOffset - 1));
+	return StateFromString(Src.RightChop(StateOffset));
 }
 
 FString FPipelineCacheFileFormatPSO::GraphicsDescriptor::HeaderLine()
@@ -589,20 +609,20 @@ FString FPipelineCacheFileFormatPSO::CommonToString() const
 	return FString::Printf(TEXT("\"%d,%llu\""), Count, Mask);
 }
 
-void FPipelineCacheFileFormatPSO::CommonFromString(const FString& Src)
+void FPipelineCacheFileFormatPSO::CommonFromString(const FStringView& Src)
 {
 #if PSO_COOKONLY_DATA
-    TArray<FString> Parts;
-    Src.TrimStartAndEnd().ParseIntoArray(Parts, TEXT(","));
-	
+    TArray<FStringView, TInlineAllocator<2>> Parts;
+	UE::String::ParseTokens(Src.TrimStartAndEnd(), TEXT(','), [&Parts](FStringView Part) { Parts.Add(Part); });
+
 	if (Parts.Num() == 1)
 	{
-		LexFromString(UsageMask, *Parts[0]);
+		LexFromString(UsageMask, Parts[0]);
 	}
 	else if(Parts.Num() > 1)
 	{
-		LexFromString(BindCount, *Parts[0]);
-		LexFromString(UsageMask, *Parts[1]);
+		LexFromString(BindCount, Parts[0]);
+		LexFromString(UsageMask, Parts[1]);
 	}
 #endif
 }
diff --git a/Engine/Source/Runtime/RHI/Private/RHI.cpp b/Engine/Source/Runtime/RHI/Private/RHI.cpp
--- a/Engine/Source/Runtime/RHI/Private/RHI.cpp
+++ b/Engine/Source/Runtime/RHI/Private/RHI.cpp
@@ -10,6 +10,8 @@
 #include "Misc/MessageDialog.h"
 #include "RHIShaderFormatDefinitions.inl"
 #include "ProfilingDebugging/CsvProfiler.h"
+#include "String/LexFromString.h"
+#include "String/ParseTokens.h"
 
 IMPLEMENT_MODULE(FDefaultModuleImpl, RHI);
 
@@ -138,22 +140,26 @@ FString FVertexElement::ToString() const
 
 void FVertexElement::FromString(const FString& InSrc)
 {
-	FString Src = InSrc;
-	Src.ReplaceInline(TEXT("\r"), TEXT(" "));
-	Src.ReplaceInline(TEXT("\n"), TEXT(" "));
-	Src.ReplaceInline(TEXT("\t"), TEXT(" "));
-	Src.ReplaceInline(TEXT("<"), TEXT(" "));
-	Src.ReplaceInline(TEXT(">"), TEXT(" "));
-	TArray<FString> Parts;
-	Src.TrimStartAndEnd().ParseIntoArray(Parts, TEXT(" "));
-
-	check(Parts.Num() == 6 && sizeof(Type) == 1); //not a very robust parser
-	LexFromString(StreamIndex, *Parts[0]);
-	LexFromString(Offset, *Parts[1]);
-	LexFromString((uint8&)Type, *Parts[2]);
-	LexFromString(AttributeIndex, *Parts[3]);
-	LexFromString(Stride, *Parts[4]);
-	LexFromString(bUseInstanceIndex, *Parts[5]);
+	FromString(FStringView(InSrc));
+}
+
+void FVertexElement::FromString(const FStringView& InSrc)
+{
+	constexpr int32 PartCount = 6;
+
+	TArray<FStringView, TInlineAllocator<PartCount>> Parts;
+	UE::String::ParseTokensMultiple(InSrc.TrimStartAndEnd(), {TEXT('\r'), TEXT('\n'), TEXT('\t'), TEXT('<'), TEXT('>'), TEXT(' ')},
+		[&Parts](FStringView Part) { if (!Part.IsEmpty()) { Parts.Add(Part); } });
+
+	check(Parts.Num() == PartCount && sizeof(Type) == 1); //not a very robust parser
+	const FStringView* PartIt = Parts.GetData();
+	LexFromString(StreamIndex, *PartIt++);
+	LexFromString(Offset, *PartIt++);
+	LexFromString((uint8&)Type, *PartIt++);
+	LexFromString(AttributeIndex, *PartIt++);
+	LexFromString(Stride, *PartIt++);
+	LexFromString(bUseInstanceIndex, *PartIt++);
+	check(Parts.GetData() + PartCount == PartIt);
 }
 
 FString FDepthStencilStateInitializerRHI::ToString() const
@@ -182,36 +188,43 @@ FString FDepthStencilStateInitializerRHI::ToString() const
 			, uint32(StencilWriteMask)
 		);
 }
+
 void FDepthStencilStateInitializerRHI::FromString(const FString& InSrc)
 {
-	FString Src = InSrc;
-	Src.ReplaceInline(TEXT("\r"), TEXT(" "));
-	Src.ReplaceInline(TEXT("\n"), TEXT(" "));
-	Src.ReplaceInline(TEXT("\t"), TEXT(" "));
-	Src.ReplaceInline(TEXT("<"), TEXT(" "));
-	Src.ReplaceInline(TEXT(">"), TEXT(" "));
-	TArray<FString> Parts;
-	Src.TrimStartAndEnd().ParseIntoArray(Parts, TEXT(" "));
-
-	check(Parts.Num() == 14 && sizeof(bool) == 1 && sizeof(FrontFaceStencilFailStencilOp) == 1 && sizeof(BackFaceStencilTest) == 1 && sizeof(BackFaceDepthFailStencilOp) == 1); //not a very robust parser
-
-	LexFromString((uint8&)bEnableDepthWrite, *Parts[0]);
-	LexFromString((uint8&)DepthTest, *Parts[1]);
-
-	LexFromString((uint8&)bEnableFrontFaceStencil, *Parts[2]);
-	LexFromString((uint8&)FrontFaceStencilTest, *Parts[3]);
-	LexFromString((uint8&)FrontFaceStencilFailStencilOp, *Parts[4]);
-	LexFromString((uint8&)FrontFaceDepthFailStencilOp, *Parts[5]);
-	LexFromString((uint8&)FrontFacePassStencilOp, *Parts[6]);
-
-	LexFromString((uint8&)bEnableBackFaceStencil, *Parts[7]);
-	LexFromString((uint8&)BackFaceStencilTest, *Parts[8]);
-	LexFromString((uint8&)BackFaceStencilFailStencilOp, *Parts[9]);
-	LexFromString((uint8&)BackFaceDepthFailStencilOp, *Parts[10]);
-	LexFromString((uint8&)BackFacePassStencilOp, *Parts[11]);
-
-	LexFromString(StencilReadMask, *Parts[12]);
-	LexFromString(StencilWriteMask, *Parts[13]);
+	FromString(FStringView(InSrc));
+}
+
+void FDepthStencilStateInitializerRHI::FromString(const FStringView& InSrc)
+{
+	constexpr int32 PartCount = 14;
+
+	TArray<FStringView, TInlineAllocator<PartCount>> Parts;
+	UE::String::ParseTokensMultiple(InSrc.TrimStartAndEnd(), {TEXT('\r'), TEXT('\n'), TEXT('\t'), TEXT('<'), TEXT('>'), TEXT(' ')},
+		[&Parts](FStringView Part) { if (!Part.IsEmpty()) { Parts.Add(Part); } });
+
+	check(Parts.Num() == PartCount && sizeof(bool) == 1 && sizeof(FrontFaceStencilFailStencilOp) == 1 && sizeof(BackFaceStencilTest) == 1 && sizeof(BackFaceDepthFailStencilOp) == 1); //not a very robust parser
+
+	const FStringView* PartIt = Parts.GetData();
+
+	LexFromString((uint8&)bEnableDepthWrite, *PartIt++);
+	LexFromString((uint8&)DepthTest, *PartIt++);
+
+	LexFromString((uint8&)bEnableFrontFaceStencil, *PartIt++);
+	LexFromString((uint8&)FrontFaceStencilTest, *PartIt++);
+	LexFromString((uint8&)FrontFaceStencilFailStencilOp, *PartIt++);
+	LexFromString((uint8&)FrontFaceDepthFailStencilOp, *PartIt++);
+	LexFromString((uint8&)FrontFacePassStencilOp, *PartIt++);
+
+	LexFromString((uint8&)bEnableBackFaceStencil, *PartIt++);
+	LexFromString((uint8&)BackFaceStencilTest, *PartIt++);
+	LexFromString((uint8&)BackFaceStencilFailStencilOp, *PartIt++);
+	LexFromString((uint8&)BackFaceDepthFailStencilOp, *PartIt++);
+	LexFromString((uint8&)BackFacePassStencilOp, *PartIt++);
+
+	LexFromString(StencilReadMask, *PartIt++);
+	LexFromString(StencilWriteMask, *PartIt++);
+
+	check(Parts.GetData() + PartCount == PartIt);
 }
 
 FString FBlendStateInitializerRHI::ToString() const
@@ -227,22 +240,26 @@ FString FBlendStateInitializerRHI::ToString() const
 
 void FBlendStateInitializerRHI::FromString(const FString& InSrc)
 {
-	FString Src = InSrc;
-	Src.ReplaceInline(TEXT("\r"), TEXT(" "));
-	Src.ReplaceInline(TEXT("\n"), TEXT(" "));
-	Src.ReplaceInline(TEXT("\t"), TEXT(" "));
-	Src.ReplaceInline(TEXT("<"), TEXT(" "));
-	Src.ReplaceInline(TEXT(">"), TEXT(" "));
-	TArray<FString> Parts;
-	Src.TrimStartAndEnd().ParseIntoArray(Parts, TEXT(" "));
+	FromString(FStringView(InSrc));
+}
+
+void FBlendStateInitializerRHI::FromString(const FStringView& InSrc)
+{
+	constexpr int32 PartCount = MaxSimultaneousRenderTargets * FRenderTarget::NUM_STRING_FIELDS + 1;
 
+	TArray<FStringView, TInlineAllocator<PartCount>> Parts;
+	UE::String::ParseTokensMultiple(InSrc.TrimStartAndEnd(), {TEXT('\r'), TEXT('\n'), TEXT('\t'), TEXT('<'), TEXT('>'), TEXT(' ')},
+		[&Parts](FStringView Part) { if (!Part.IsEmpty()) { Parts.Add(Part); } });
 
-	check(Parts.Num() == MaxSimultaneousRenderTargets * FRenderTarget::NUM_STRING_FIELDS + 1 && sizeof(bool) == 1); //not a very robust parser
+	check(Parts.Num() == PartCount && sizeof(bool) == 1); //not a very robust parser
+	const FStringView* PartIt = Parts.GetData();
 	for (int32 Index = 0; Index < MaxSimultaneousRenderTargets; Index++)
 	{
-		RenderTargets[Index].FromString(Parts, FRenderTarget::NUM_STRING_FIELDS * Index);
+		RenderTargets[Index].FromString(MakeArrayView(PartIt, FRenderTarget::NUM_STRING_FIELDS));
+		PartIt += FRenderTarget::NUM_STRING_FIELDS;
 	}
-	LexFromString((int8&)bUseIndependentRenderTargetBlendStates, *Parts[MaxSimultaneousRenderTargets * FRenderTarget::NUM_STRING_FIELDS]);
+	LexFromString((int8&)bUseIndependentRenderTargetBlendStates, *PartIt++);
+	check(Parts.GetData() + PartCount == PartIt);
 }
 
 
@@ -271,6 +288,19 @@ void FBlendStateInitializerRHI::FRenderTarget::FromString(const TArray<FString>&
 	LexFromString((uint8&)ColorWriteMask, *Parts[Index++]);
 }
 
+void FBlendStateInitializerRHI::FRenderTarget::FromString(TArrayView<const FStringView> Parts)
+{
+	check(Parts.Num() == NUM_STRING_FIELDS);
+	const FStringView* PartIt = Parts.GetData();
+	LexFromString((uint8&)ColorBlendOp, *PartIt++);
+	LexFromString((uint8&)ColorSrcBlend, *PartIt++);
+	LexFromString((uint8&)ColorDestBlend, *PartIt++);
+	LexFromString((uint8&)AlphaBlendOp, *PartIt++);
+	LexFromString((uint8&)AlphaSrcBlend, *PartIt++);
+	LexFromString((uint8&)AlphaDestBlend, *PartIt++);
+	LexFromString((uint8&)ColorWriteMask, *PartIt++);
+}
+
 bool FRHIResource::Bypass()
 {
 	return GRHICommandList.Bypass();
diff --git a/Engine/Source/Runtime/RHI/Public/PipelineFileCache.h b/Engine/Source/Runtime/RHI/Public/PipelineFileCache.h
--- a/Engine/Source/Runtime/RHI/Public/PipelineFileCache.h
+++ b/Engine/Source/Runtime/RHI/Public/PipelineFileCache.h
@@ -4,6 +4,7 @@
 
 #include "CoreMinimal.h"
 #include "Containers/List.h"
+#include "Containers/StringView.h"
 #include "RHI.h"
 
 DECLARE_STATS_GROUP(TEXT("ShaderPipelineCache"),STATGROUP_PipelineStateCache, STATCAT_Advanced);
@@ -71,7 +72,7 @@ struct RHI_API FPipelineFileCacheRasterizerState
 		return KeyHash;
 	}
 	FString ToString() const;
-	void FromString(const FString& Src);
+	void FromString(const FStringView& Src);
 };
 
 /**
@@ -111,7 +112,7 @@ struct RHI_API FPipelineCacheFileFormatPSO
 
 		FString ToString() const;
 		static FString HeaderLine();
-		void FromString(const FString& Src);
+		void FromString(const FStringView& Src);
 	};
 	struct RHI_API GraphicsDescriptor
 	{
@@ -145,15 +146,15 @@ struct RHI_API FPipelineCacheFileFormatPSO
 		
 		FString ToString() const;
 		static FString HeaderLine();
-		bool FromString(const FString& Src);
+		bool FromString(const FStringView& Src);
 
 		FString ShadersToString() const;
 		static FString ShaderHeaderLine();
-		void ShadersFromString(const FString& Src);
+		void ShadersFromString(const FStringView& Src);
 
 		FString StateToString() const;
 		static FString StateHeaderLine();
-		bool StateFromString(const FString& Src);
+		bool StateFromString(const FStringView& Src);
 	};
 	enum class DescriptorType : uint32
 	{
@@ -187,7 +188,7 @@ struct RHI_API FPipelineCacheFileFormatPSO
 	
 	FString CommonToString() const;
 	static FString CommonHeaderLine();
-	void CommonFromString(const FString& Src);
+	void CommonFromString(const FStringView& Src);
 	
 	// Potential cases for seperating verify logic if requiired: RunTime-Logging, RunTime-UserCaching, RunTime-PreCompile, CommandLet-Cooking
 	bool Verify() const;
diff --git a/Engine/Source/Runtime/RHI/Public/RHI.h b/Engine/Source/Runtime/RHI/Public/RHI.h
--- a/Engine/Source/Runtime/RHI/Public/RHI.h
+++ b/Engine/Source/Runtime/RHI/Public/RHI.h
@@ -15,9 +15,11 @@
 #define RHI_COMMAND_LIST_DEBUG_TRACES 0
 #endif
 
+template <typename T> class TArrayView;
+
 class FResourceArrayInterface;
 class FResourceBulkDataInterface;
-
+class FStringView;
 
 /** Alignment of the shader parameters struct is required to be 16-byte boundaries. */
 #define SHADER_PARAMETER_STRUCT_ALIGNMENT 16
@@ -764,6 +766,7 @@ struct FVertexElement
 	}
 	RHI_API FString ToString() const;
 	RHI_API void FromString(const FString& Src);
+	RHI_API void FromString(const FStringView& Src);
 };
 
 typedef TArray<FVertexElement,TFixedAllocator<MaxVertexElementCount> > FVertexDeclarationElementList;
@@ -952,8 +955,7 @@ struct FDepthStencilStateInitializerRHI
 	}
 	RHI_API FString ToString() const;
 	RHI_API void FromString(const FString& Src);
-
-
+	RHI_API void FromString(const FStringView& Src);
 };
 
 class FBlendStateInitializerRHI
@@ -1005,8 +1007,7 @@ public:
 		}
 		RHI_API FString ToString() const;
 		RHI_API void FromString(const TArray<FString>& Parts, int32 Index);
-
-
+		RHI_API void FromString(TArrayView<const FStringView> Parts);
 	};
 
 	FBlendStateInitializerRHI() {}
@@ -1040,8 +1041,7 @@ public:
 	}
 	RHI_API FString ToString() const;
 	RHI_API void FromString(const FString& Src);
-
-
+	RHI_API void FromString(const FStringView& Src);
 };
 
 /**
diff --git a/Engine/Source/Runtime/RenderCore/Private/ShaderCodeLibrary.cpp b/Engine/Source/Runtime/RenderCore/Private/ShaderCodeLibrary.cpp
--- a/Engine/Source/Runtime/RenderCore/Private/ShaderCodeLibrary.cpp
+++ b/Engine/Source/Runtime/RenderCore/Private/ShaderCodeLibrary.cpp
@@ -6,6 +6,8 @@ ShaderCodeLibrary.cpp: Bound shader state cache implementation.
 
 #include "ShaderCodeLibrary.h"
 #include "Shader.h"
+#include "Algo/Replace.h"
+#include "Containers/StringView.h"
 #include "Misc/SecureHash.h"
 #include "Misc/Paths.h"
 #include "Math/UnitConversion.h"
@@ -14,6 +16,8 @@ ShaderCodeLibrary.cpp: Bound shader state cache implementation.
 #include "HAL/PlatformSplash.h"
 #include "Misc/ScopeLock.h"
 #include "Misc/ScopeRWLock.h"
+#include "Misc/StringBuilder.h"
+#include "String/ParseTokens.h"
 #include "Async/AsyncFileHandle.h"
 #include "PipelineFileCache.h"
 #include "Interfaces/IPluginManager.h"
@@ -121,68 +125,125 @@ static void FShaderLibraryHelperCompressCode(EShaderPlatform Platform, const TAr
 }
 
 
-FORCEINLINE FName ParseFNameCached(const FString& Src, TMap<uint32,FName>& NameCache)
+FORCEINLINE FName ParseFNameCached(const FStringView& Src, TMap<uint32,FName>& NameCache)
 {
-	uint32 SrcHash = CityHash32((char*)Src.GetCharArray().GetData(), Src.GetCharArray().GetAllocatedSize());
+	uint32 SrcHash = CityHash32(reinterpret_cast<const char*>(Src.GetData()), Src.Len() * sizeof(TCHAR));
 	if (FName* Name = NameCache.Find(SrcHash))
 	{
 		return *Name;
 	}
 	else
 	{
-		return NameCache.Emplace(SrcHash, FName(*Src));
+		return NameCache.Emplace(SrcHash, FName(Src.Len(), Src.GetData()));
 	}
 }
 
+static void AppendFNameAsUTF8(FAnsiStringBuilderBase& Out, const FName& InName)
+{
+	if (!InName.TryAppendAnsiString(Out))
+	{
+		TStringBuilder<128> WideName;
+		InName.AppendString(WideName);
+		Out << TCHAR_TO_UTF8(WideName.ToString());
+	}
+}
+
+static void AppendSanitizedFNameAsUTF8(FAnsiStringBuilderBase& Out, const FName& InName, ANSICHAR Delim)
+{
+	const int32 Offset = Out.Len();
+	AppendFNameAsUTF8(Out, InName);
+	Algo::Replace(MakeArrayView(Out).Slice(Offset, Out.Len() - Offset), Delim, ' ');
+}
+
+static void AppendSanitizedFName(FStringBuilderBase& Out, const FName& InName, TCHAR Delim)
+{
+	const int32 Offset = Out.Len();
+	InName.AppendString(Out);
+	Algo::Replace(MakeArrayView(Out).Slice(Offset, Out.Len() - Offset), Delim, TEXT(' '));
+}
+
 FString FCompactFullName::ToString() const
 {
-	FString RetString;
-	RetString.Reserve(256);
-	if (!ObjectClassAndPath.Num())
+	TStringBuilder<256> RetString;
+	AppendString(RetString);
+	return FString(FStringView(RetString));
+}
+
+void FCompactFullName::AppendString(FStringBuilderBase& Out) const
+{
+	const int32 ObjectClassAndPathCount = ObjectClassAndPath.Num();
+	if (!ObjectClassAndPathCount)
+	{
+		Out << TEXT("empty");
+	}
+	else
+	{
+		for (int32 NameIdx = 0; NameIdx < ObjectClassAndPathCount; NameIdx++)
+		{
+			Out << ObjectClassAndPath[NameIdx];
+			if (NameIdx == 0)
+			{
+				Out << TEXT(' ');
+			}
+			else if (NameIdx < ObjectClassAndPathCount - 1)
+			{
+				if (NameIdx == ObjectClassAndPathCount - 2)
+				{
+					Out << TEXT('.');
+				}
+				else
+				{
+					Out << TEXT('/');
+				}
+			}
+		}
+	}
+}
+
+void FCompactFullName::AppendString(FAnsiStringBuilderBase& Out) const
+{
+	const int32 ObjectClassAndPathCount = ObjectClassAndPath.Num();
+	if (!ObjectClassAndPathCount)
 	{
-		RetString += TEXT("empty");
+		Out << "empty";
 	}
 	else
 	{
-		for (int32 NameIdx = 0; NameIdx < ObjectClassAndPath.Num(); NameIdx++)
+		for (int32 NameIdx = 0; NameIdx < ObjectClassAndPathCount; NameIdx++)
 		{
-			RetString += ObjectClassAndPath[NameIdx].ToString();
+			AppendFNameAsUTF8(Out, ObjectClassAndPath[NameIdx]);
 			if (NameIdx == 0)
 			{
-				RetString += TEXT(" ");
+				Out << ' ';
 			}
-			else if (NameIdx < ObjectClassAndPath.Num() - 1)
+			else if (NameIdx < ObjectClassAndPathCount - 1)
 			{
-				if (NameIdx == ObjectClassAndPath.Num() - 2)
+				if (NameIdx == ObjectClassAndPathCount - 2)
 				{
-					RetString += TEXT(".");
+					Out << '.';
 				}
 				else
 				{
-					RetString += TEXT("/");
+					Out << '/';
 				}
 			}
 		}
 	}
-	return RetString;
 }
 
-void FCompactFullName::ParseFromString(const FString& InSrc)
+void FCompactFullName::ParseFromString(const FStringView& InSrc)
 {
-	FString Src = InSrc;
-	Src.ReplaceInline(TEXT("\t"), TEXT(" "));
-	Src.ReplaceInline(TEXT("."), TEXT(" "));
-	Src.ReplaceInline(TEXT("/"), TEXT(" "));
-	TArray<FString> Fields;
-	Src.TrimStartAndEnd().ParseIntoArray(Fields, TEXT(" "), true);
-	if (Fields.Num() == 1 && Fields[0] == TEXT("empty"))
+	TArray<FStringView, TInlineAllocator<64>> Fields;
+	UE::String::ParseTokensMultiple(InSrc.TrimStartAndEnd(), {TEXT(' '), TEXT('.'), TEXT('/'), TEXT('\t')},
+		[&Fields](FStringView Field) { if (!Field.IsEmpty()) { Fields.Add(Field); } });
+	if (Fields.Num() == 1 && Fields[0] == TEXT("empty"_SV))
 	{
 		Fields.Empty();
 	}
 	ObjectClassAndPath.Empty(Fields.Num());
-	for (const FString& Item : Fields)
+	for (const FStringView& Item : Fields)
 	{
-		ObjectClassAndPath.Add(FName(*Item));
+		ObjectClassAndPath.Emplace(Item);
 	}
 }
 
@@ -237,10 +298,10 @@ void FStableShaderKeyAndValue::ComputeKeyHash()
 	KeyHash = HashCombine(KeyHash, GetTypeHash(PipelineHash));
 }
 
-void FStableShaderKeyAndValue::ParseFromString(const FString& Src)
+void FStableShaderKeyAndValue::ParseFromString(const FStringView& Src)
 {
-	TArray<FString> Fields;
-	Src.TrimStartAndEnd().ParseIntoArray(Fields, TEXT(","), false);
+	TArray<FStringView, TInlineAllocator<12>> Fields;
+	UE::String::ParseTokens(Src.TrimStartAndEnd(), TEXT(','), [&Fields](FStringView Field) { Fields.Add(Field); });
 
 	/* disabled, should not be happening since 1H 2018
 	if (Fields.Num() > 12)
@@ -256,17 +317,17 @@ void FStableShaderKeyAndValue::ParseFromString(const FString& Src)
 	int32 Index = 0;
 	ClassNameAndObjectPath.ParseFromString(Fields[Index++]);
 
-	ShaderType = FName(*Fields[Index++]);
-	ShaderClass = FName(*Fields[Index++]);
-	MaterialDomain = FName(*Fields[Index++]);
-	FeatureLevel = FName(*Fields[Index++]);
+	ShaderType = FName(Fields[Index++]);
+	ShaderClass = FName(Fields[Index++]);
+	MaterialDomain = FName(Fields[Index++]);
+	FeatureLevel = FName(Fields[Index++]);
 
-	QualityLevel = FName(*Fields[Index++]);
-	TargetFrequency = FName(*Fields[Index++]);
-	TargetPlatform = FName(*Fields[Index++]);
+	QualityLevel = FName(Fields[Index++]);
+	TargetFrequency = FName(Fields[Index++]);
+	TargetPlatform = FName(Fields[Index++]);
 
-	VFType = FName(*Fields[Index++]);
-	PermutationId = FName(*Fields[Index++]);
+	VFType = FName(Fields[Index++]);
+	PermutationId = FName(Fields[Index++]);
 
 	OutputHash.FromString(Fields[Index++]);
 
@@ -285,10 +346,10 @@ void FStableShaderKeyAndValue::ParseFromString(const FString& Src)
 }
 
 
-void FStableShaderKeyAndValue::ParseFromStringCached(const FString& Src, TMap<uint32, FName>& NameCache)
+void FStableShaderKeyAndValue::ParseFromStringCached(const FStringView& Src, TMap<uint32, FName>& NameCache)
 {
-	TArray<FString> Fields;
-	Src.TrimStartAndEnd().ParseIntoArray(Fields, TEXT(","), false);
+	TArray<FStringView, TInlineAllocator<12>> Fields;
+	UE::String::ParseTokens(Src.TrimStartAndEnd(), TEXT(','), [&Fields](FStringView Field) { Fields.Add(Field); });
 
 	/* disabled, should not be happening since 1H 2018
 	if (Fields.Num() > 11)
@@ -297,7 +358,7 @@ void FStableShaderKeyAndValue::ParseFromStringCached(const FString& Src, TMap<ui
 		FixupUnsanitizedNames(Src, Fields);
 	}
 	*/
-	
+
 	// for a while, accept old .scl.csv without pipelinehash
 	check(Fields.Num() == 11 || Fields.Num() == 12);
 
@@ -343,38 +404,65 @@ FString FStableShaderKeyAndValue::ToString() const
 
 void FStableShaderKeyAndValue::ToString(FString& OutResult) const
 {
-	const TCHAR* Delim = TEXT(",");
-
-	OutResult.Reset(255);
+	TStringBuilder<384> Out;
+	const TCHAR Delim = TEXT(',');
 
-	OutResult += ClassNameAndObjectPath.ToString().Replace(Delim, TEXT(" "));
-	OutResult += Delim;
+	const int32 ClassNameAndObjectPathOffset = Out.Len();
+	ClassNameAndObjectPath.AppendString(Out);
+	Algo::Replace(MakeArrayView(Out).Slice(ClassNameAndObjectPathOffset, Out.Len() - ClassNameAndObjectPathOffset), Delim, TEXT(' '));
+	Out << Delim;
 
-	OutResult += ShaderType.ToString().Replace(Delim, TEXT(" "));
-	OutResult += Delim;
-	OutResult += ShaderClass.ToString().Replace(Delim, TEXT(" "));
-	OutResult += Delim;
-	OutResult += MaterialDomain.ToString();
-	OutResult += Delim;
-	OutResult += FeatureLevel.ToString();
-	OutResult += Delim;
+	AppendSanitizedFName(Out, ShaderType, Delim);
+	Out << Delim;
+	AppendSanitizedFName(Out, ShaderClass, Delim);
+	Out << Delim;
 
-	OutResult += QualityLevel.ToString();
-	OutResult += Delim;
-	OutResult += TargetFrequency.ToString();
-	OutResult += Delim;
-	OutResult += TargetPlatform.ToString();
-	OutResult += Delim;
+	Out << MaterialDomain << Delim;
+	Out << FeatureLevel << Delim;
+	Out << QualityLevel << Delim;
+	Out << TargetFrequency << Delim;
+	Out << TargetPlatform << Delim;
+	Out << VFType << Delim;
+	Out << PermutationId << Delim;
 
-	OutResult += VFType.ToString();
-	OutResult += Delim;
-	OutResult += PermutationId.ToString();
-	OutResult += Delim;
+	Out << OutputHash << Delim;
+	Out << PipelineHash;
 
-	OutResult += OutputHash.ToString();
-	OutResult += Delim;
+	OutResult = FStringView(Out);
+}
 
-	OutResult += PipelineHash.ToString();
+void FStableShaderKeyAndValue::AppendString(FAnsiStringBuilderBase& Out) const
+{
+	const ANSICHAR Delim = ',';
+
+	const int32 ClassNameAndObjectPathOffset = Out.Len();
+	ClassNameAndObjectPath.AppendString(Out);
+	Algo::Replace(MakeArrayView(Out).Slice(ClassNameAndObjectPathOffset, Out.Len() - ClassNameAndObjectPathOffset), Delim, ' ');
+	Out << Delim;
+
+	AppendSanitizedFNameAsUTF8(Out, ShaderType, Delim);
+	Out << Delim;
+	AppendSanitizedFNameAsUTF8(Out, ShaderClass, Delim);
+	Out << Delim;
+
+	AppendFNameAsUTF8(Out, MaterialDomain);
+	Out << Delim;
+	AppendFNameAsUTF8(Out, FeatureLevel);
+	Out << Delim;
+	AppendFNameAsUTF8(Out, QualityLevel);
+	Out << Delim;
+	AppendFNameAsUTF8(Out, TargetFrequency);
+	Out << Delim;
+	AppendFNameAsUTF8(Out, TargetPlatform);
+	Out << Delim;
+	AppendFNameAsUTF8(Out, VFType);
+	Out << Delim;
+	AppendFNameAsUTF8(Out, PermutationId);
+	Out << Delim;
+
+	Out << OutputHash;
+	Out << Delim;
+	Out << PipelineHash;
 }
 
 FString FStableShaderKeyAndValue::HeaderLine()
@@ -1512,7 +1600,6 @@ struct FEditorShaderStableInfo
 		check(LibraryName.Len() == 0);
 		check(Name.Len() > 0);
 		LibraryName = Name;
-		Offset = 0;
 		StableMap.Empty();
 	}
 
@@ -1595,15 +1682,14 @@ struct FEditorShaderStableInfo
 
 				IntermediateFormatAr->Serialize((ANSICHAR*)HeaderSrc.Get(), HeaderSrc.Length() * sizeof(ANSICHAR));
 
-				FString LineBuffer;
-				LineBuffer.Reserve(512);
+				TAnsiStringBuilder<512> LineBuffer;
 
 				for (const FStableShaderKeyAndValue& Item : StableMap)
 				{
-					Item.ToString(LineBuffer);
-					LineBuffer += TCHAR('\n');
-					auto LineConverted = StringCast<ANSICHAR>(*LineBuffer, LineBuffer.Len());
-					IntermediateFormatAr->Serialize((ANSICHAR*)LineConverted.Get(), LineConverted.Length() * sizeof(ANSICHAR));
+					LineBuffer.Reset();
+					Item.AppendString(LineBuffer);
+					LineBuffer << '\n';
+					IntermediateFormatAr->Serialize(const_cast<ANSICHAR*>(LineBuffer.ToString()), LineBuffer.Len() * sizeof(ANSICHAR));
 				}
 			}
 
@@ -1633,7 +1719,6 @@ private:
 	FName FormatName;
 	FString LibraryName;
 	TSet<FStableShaderKeyAndValue> StableMap;
-	uint64 Offset;
 };
 
 struct FShaderCodeStats
diff --git a/Engine/Source/Runtime/RenderCore/Public/ShaderCodeLibrary.h b/Engine/Source/Runtime/RenderCore/Public/ShaderCodeLibrary.h
--- a/Engine/Source/Runtime/RenderCore/Public/ShaderCodeLibrary.h
+++ b/Engine/Source/Runtime/RenderCore/Public/ShaderCodeLibrary.h
@@ -11,7 +11,9 @@
 
 DECLARE_LOG_CATEGORY_EXTERN(LogShaderLibrary, Log, All);
 
+class FAnsiStringBuilderBase;
 class FShaderPipeline;
+class FStringView;
 
 struct RENDERCORE_API FShaderCodeLibraryPipeline
 {
@@ -56,7 +58,7 @@ struct RENDERCORE_API FShaderCodeLibraryPipeline
 
 struct RENDERCORE_API FCompactFullName
 {
-	TArray<FName> ObjectClassAndPath;
+	TArray<FName, TInlineAllocator<16>> ObjectClassAndPath;
 
 	bool operator==(const FCompactFullName& Other) const
 	{
@@ -64,7 +66,9 @@ struct RENDERCORE_API FCompactFullName
 	}
 
 	FString ToString() const;
-	void ParseFromString(const FString& Src);
+	void AppendString(FStringBuilderBase& Out) const;
+	void AppendString(FAnsiStringBuilderBase& Out) const;
+	void ParseFromString(const FStringView& Src);
 	friend RENDERCORE_API uint32 GetTypeHash(const FCompactFullName& A);
 };
 
@@ -93,10 +97,11 @@ struct RENDERCORE_API FStableShaderKeyAndValue
 	}
 
 	void ComputeKeyHash();
-	void ParseFromString(const FString& Src);
-	void ParseFromStringCached(const FString& Src, class TMap<uint32, FName>& NameCache);
+	void ParseFromString(const FStringView& Src);
+	void ParseFromStringCached(const FStringView& Src, class TMap<uint32, FName>& NameCache);
 	FString ToString() const;
 	void ToString(FString& OutResult) const;
+	void AppendString(FAnsiStringBuilderBase& Out) const;
 	static FString HeaderLine();
 
 	/** Computes pipeline hash from the passed pipeline. Pass nullptr to clear */
