Add an export to function FEditorFileUtils::SaveMap()

#rnx

diff --git a/Engine/Source/Editor/UnrealEd/Public/FileHelpers.h b/Engine/Source/Editor/UnrealEd/Public/FileHelpers.h
--- a/Engine/Source/Editor/UnrealEd/Public/FileHelpers.h
+++ b/Engine/Source/Editor/UnrealEd/Public/FileHelpers.h
@@ -248,7 +248,7 @@ public:
 	 *
 	 * @return					true if the map was saved successfully.
 	 */
-	static bool SaveMap(UWorld* World, const FString& Filename );
+	static UNREALED_API bool SaveMap(UWorld* World, const FString& Filename );
 
 	/**
 	 * Saves the specified level.  SaveAs is performed as necessary.
