Fix PlayLevel compiler error. Mistake in edit merge.

diff --git a/Engine/Source/Editor/UnrealEd/Private/PlayLevel.cpp b/Engine/Source/Editor/UnrealEd/Private/PlayLevel.cpp
--- a/Engine/Source/Editor/UnrealEd/Private/PlayLevel.cpp
+++ b/Engine/Source/Editor/UnrealEd/Private/PlayLevel.cpp
@@ -2476,7 +2476,7 @@ void UEditorEngine::StartPlayInEditorSession(FRequestPlaySessionParams& InReques
 	}
 
 	// Flush all audio sources from the editor world
-	if (FAudioDevice* AudioDevice = EditorWorld->GetAudioDevice())
+	if (FAudioDeviceHandle AudioDevice = EditorWorld->GetAudioDevice())
 	{
 		AudioDevice->Flush(EditorWorld);
 		AudioDevice->ResetInterpolation();
@@ -2728,7 +2728,7 @@ UGameInstance* UEditorEngine::CreateInnerProcessPIEGameInstance(FRequestPlaySess
 
 		if (!InParams.EditorPlaySettings->EnableGameSound)
 		{
-			if (FAudioDevice* GameInstanceAudioDevice = GameInstance->GetWorld()->GetAudioDevice())
+			if (FAudioDeviceHandle GameInstanceAudioDevice = GameInstance->GetWorld()->GetAudioDevice())
 			{
 				GameInstanceAudioDevice->SetTransientMasterVolume(0.0f);
 			}
