Add support for text assets to UEditorEngine::PackageIsAMapFile

diff --git a/Engine/Source/Editor/UnrealEd/Private/EditorServer.cpp b/Engine/Source/Editor/UnrealEd/Private/EditorServer.cpp
--- a/Engine/Source/Editor/UnrealEd/Private/EditorServer.cpp
+++ b/Engine/Source/Editor/UnrealEd/Private/EditorServer.cpp
@@ -160,6 +160,9 @@
 #include "Subsystems/BrushEditingSubsystem.h"
 #include "EdMode.h"
 
+#include "Serialization/StructuredArchive.h"
+#include "Serialization/Formatters/JsonArchiveInputFormatter.h"
+
 DEFINE_LOG_CATEGORY_STATIC(LogEditorServer, Log, All);
 
 /** Used for the "tagsounds" and "checksounds" commands only			*/
@@ -2305,7 +2308,21 @@ bool UEditorEngine::PackageIsAMapFile( const TCHAR* PackageFilename, FText& OutN
 	if( CheckMapPackageFile )
 	{
 		FPackageFileSummary Summary;
-		( *CheckMapPackageFile ) << Summary;
+
+#if WITH_TEXT_ARCHIVE_SUPPORT
+		if (FPackageName::IsTextPackageExtension(*FPaths::GetExtension(PackageFilename)))
+		{
+			FJsonArchiveInputFormatter Formatter(*CheckMapPackageFile);
+			FStructuredArchive Archive(Formatter);
+			Archive.Open().EnterRecord() << SA_VALUE(TEXT("Summary"), Summary);
+			Archive.Close();
+		}
+		else
+#endif
+		{
+			(*CheckMapPackageFile) << Summary;
+		}
+
 		delete CheckMapPackageFile;
 
 		// Check flag.
