Fixed reinstancing crash with custom serialization

GetSerializedProperty can return null outside of tagged property serialization

diff --git a/Engine/Source/Editor/UnrealEd/Private/Kismet2/KismetReinstanceUtilities.cpp b/Engine/Source/Editor/UnrealEd/Private/Kismet2/KismetReinstanceUtilities.cpp
--- a/Engine/Source/Editor/UnrealEd/Private/Kismet2/KismetReinstanceUtilities.cpp
+++ b/Engine/Source/Editor/UnrealEd/Private/Kismet2/KismetReinstanceUtilities.cpp
@@ -1445,23 +1445,25 @@ namespace InstancedPropertyUtils
 		{
 			if (Obj == nullptr)
 			{
-				FProperty* SerializingProperty = GetSerializedProperty();
-				if (UObject* const* OldInstancedObjPtr = OldInstancedSubObjects.Find(SerializingProperty->GetFName()))
+				if (FProperty* SerializingProperty = GetSerializedProperty())
 				{
-					const UObject* OldInstancedObj = *OldInstancedObjPtr;
-					check(SerializingProperty->HasAnyPropertyFlags(CPF_PersistentInstance));
-
-					UClass* TargetClass = TargetCDO->GetClass();
-					// @TODO: Handle nested instances when we have more time to flush this all out  
-					if ( TargetClass->IsChildOf(SerializingProperty->GetOwnerClass()) )
+					if (UObject* const* OldInstancedObjPtr = OldInstancedSubObjects.Find(SerializingProperty->GetFName()))
 					{
-						FObjectPropertyBase* SerializingObjProperty = CastFieldChecked<FObjectPropertyBase>(SerializingProperty);
-						// being extra careful, not to create our own instanced version when we expect one from the CDO
-						if (SerializingObjProperty->GetObjectPropertyValue_InContainer(TargetCDO) == nullptr)
+						const UObject* OldInstancedObj = *OldInstancedObjPtr;
+						check(SerializingProperty->HasAnyPropertyFlags(CPF_PersistentInstance));
+
+						UClass* TargetClass = TargetCDO->GetClass();
+						// @TODO: Handle nested instances when we have more time to flush this all out  
+						if (TargetClass->IsChildOf(SerializingProperty->GetOwnerClass()))
 						{
-							// @TODO: What if the instanced object is of the same type 
-							//        that we're currently reinstancing
-							Obj = StaticDuplicateObject(OldInstancedObj, Target);// NewObject<UObject>(Target, OldInstancedObj->GetClass()->GetAuthoritativeClass(), OldInstancedObj->GetFName());
+							FObjectPropertyBase* SerializingObjProperty = CastFieldChecked<FObjectPropertyBase>(SerializingProperty);
+							// being extra careful, not to create our own instanced version when we expect one from the CDO
+							if (SerializingObjProperty->GetObjectPropertyValue_InContainer(TargetCDO) == nullptr)
+							{
+								// @TODO: What if the instanced object is of the same type 
+								//        that we're currently reinstancing
+								Obj = StaticDuplicateObject(OldInstancedObj, Target);// NewObject<UObject>(Target, OldInstancedObj->GetClass()->GetAuthoritativeClass(), OldInstancedObj->GetFName());
+							}
 						}
 					}
 				}
