Flag packages containing text with an invalid package ID as stale when generating the gather cache reports

diff --git a/Engine/Source/Editor/UnrealEd/Private/Commandlets/GatherTextFromAssetsCommandlet.cpp b/Engine/Source/Editor/UnrealEd/Private/Commandlets/GatherTextFromAssetsCommandlet.cpp
--- a/Engine/Source/Editor/UnrealEd/Private/Commandlets/GatherTextFromAssetsCommandlet.cpp
+++ b/Engine/Source/Editor/UnrealEd/Private/Commandlets/GatherTextFromAssetsCommandlet.cpp
@@ -795,7 +795,8 @@ int32 UGatherTextFromAssetsCommandlet::Main(const FString& Params)
 				{
 					// Look for any structurally significant changes (missing, added, or changed texts) in the cache
 					// Ignore insignificant things (like source changes caused by assets moving or being renamed)
-					if (!IsGatherableTextDataIdentical(GatherableTextDataArray, PackagePendingGather.GatherableTextDataArray))
+					if (EnumHasAnyFlags(GatherableTextResultFlags, EPropertyLocalizationGathererResultFlags::HasTextWithInvalidPackageLocalizationID) 
+						|| !IsGatherableTextDataIdentical(GatherableTextDataArray, PackagePendingGather.GatherableTextDataArray))
 					{
 						PackagesWithStaleGatherCache.Add(PackagePendingGather.PackageName);
 						
