Fix bad merge
#rnx

diff --git a/Engine/Source/Editor/UnrealEd/Private/PlayLevel.cpp b/Engine/Source/Editor/UnrealEd/Private/PlayLevel.cpp
--- a/Engine/Source/Editor/UnrealEd/Private/PlayLevel.cpp
+++ b/Engine/Source/Editor/UnrealEd/Private/PlayLevel.cpp
@@ -2355,7 +2355,7 @@ void UEditorEngine::StartPlayInEditorSession(FRequestPlaySessionParams& InReques
 
 	// Broadcast PreBeginPIE before checks that might block PIE below (BeginPIE is broadcast below after the checks)
 	FEditorDelegates::PreBeginPIE.Broadcast(InRequestParams.WorldType == EPlaySessionWorldType::SimulateInEditor);
-	const double PIEStartTime = FStudioAnalytics::GetAnalyticSeconds()
+	const double PIEStartTime = FStudioAnalytics::GetAnalyticSeconds();
 	const FScopedBusyCursor BusyCursor;
 
 	// Block PIE when there is a transaction recording into the undo buffer. This is generally avoided
