Fixed deprecation warning version number

#rnx

diff --git a/Engine/Source/Editor/UnrealEd/Classes/ThumbnailRendering/ThumbnailRenderer.h b/Engine/Source/Editor/UnrealEd/Classes/ThumbnailRendering/ThumbnailRenderer.h
--- a/Engine/Source/Editor/UnrealEd/Classes/ThumbnailRendering/ThumbnailRenderer.h
+++ b/Engine/Source/Editor/UnrealEd/Classes/ThumbnailRendering/ThumbnailRenderer.h
@@ -52,7 +52,7 @@ public:
 	 * @param Viewport the viewport being drawn in
 	 * @param Canvas the render interface to draw with
 	 */
-	UE_DEPRECATED(4.24, "Please override the other prototype of the Draw function.")
+	UE_DEPRECATED(4.25, "Please override the other prototype of the Draw function.")
 	virtual void Draw(UObject* Object, int32 X, int32 Y, uint32 Width, uint32 Height, FRenderTarget* Viewport, FCanvas* Canvas) { Draw(Object, X, Y, Width, Height, Viewport, Canvas, false); }
 
 	/**
