Ensure redirectors to map packages are using the .umap extension
Before this, moving a map asset, let's say "/Game/Old/Test.umap", to a new location "/Game/New/Test.umap" would leave "/Game/Old/Test.umap" behind, create a "/Game/Old/Test.uasset" redirector and a new "/Game/New/Test.umap" file
Now, the redirector is created as "/Game/Old/Test.umap"

diff --git a/Engine/Source/Editor/UnrealEd/Private/ObjectTools.cpp b/Engine/Source/Editor/UnrealEd/Private/ObjectTools.cpp
--- a/Engine/Source/Editor/UnrealEd/Private/ObjectTools.cpp
+++ b/Engine/Source/Editor/UnrealEd/Private/ObjectTools.cpp
@@ -3426,17 +3426,18 @@ namespace ObjectTools
 				Redirector = NULL;
 			}
 
+			UPackage* OldPackage = Object->GetOutermost();
 			UPackage* NewPackage = CreatePackage( NULL, *PkgName );
+
 			// if this object is being renamed out of the MyLevel package into a content package, we need to mark it RF_Standalone
 			// so that it will be saved (UWorld::CleanupWorld() clears this flag for all objects inside the package)
 			if (!Object->HasAnyFlags(RF_Standalone)
-				&&	Object->GetOutermost()->ContainsMap()
+				&&	OldPackage->ContainsMap()
 				&&	!NewPackage->GetOutermost()->ContainsMap() )
 			{
 				Object->SetFlags(RF_Standalone);
 			}
 
-			UPackage *OldPackage = Object->GetOutermost();
 			FString OldObjectFullName = Object->GetFullName();
 			FString OldObjectPathName = Object->GetPathName();
 			GEditor->RenameObject( Object, NewPackage, *ObjName, bLeaveRedirector ? REN_None : REN_DontCreateRedirectors );
@@ -3458,6 +3459,12 @@ namespace ObjectTools
 			UObjectRedirector* NewRedirector = FindObject<UObjectRedirector>(NULL, *OldObjectPathName);
 			if ( NewRedirector )
 			{
+				// If we created a redirector to a map asset, ensure the redirector package is flagged as containing a map for it to have the correct file extension.
+				if (NewPackage->ContainsMap())
+				{
+					NewRedirector->GetOutermost()->ThisContainsMap();
+				}
+
 				FAssetRegistryModule::AssetCreated(NewRedirector);
 			}
 
diff --git a/Engine/Source/Runtime/AssetRegistry/Private/AssetRegistry.cpp b/Engine/Source/Runtime/AssetRegistry/Private/AssetRegistry.cpp
--- a/Engine/Source/Runtime/AssetRegistry/Private/AssetRegistry.cpp
+++ b/Engine/Source/Runtime/AssetRegistry/Private/AssetRegistry.cpp
@@ -1500,7 +1500,6 @@ void UAssetRegistryImpl::AssetCreated(UObject* NewAsset)
 		NewPackage->SetPackageFlags(PKG_NewlyCreated);
 
 		const FString NewPackageName = NewPackage->GetName();
-		const FString Filename = FPackageName::LongPackageNameToFilename(NewPackageName, FPackageName::GetAssetPackageExtension());
 
 		// This package not empty, in case it ever was
 		RemoveEmptyPackage(NewPackage->GetFName());
diff --git a/Engine/Source/Runtime/CoreUObject/Private/UObject/SavePackage.cpp b/Engine/Source/Runtime/CoreUObject/Private/UObject/SavePackage.cpp
--- a/Engine/Source/Runtime/CoreUObject/Private/UObject/SavePackage.cpp
+++ b/Engine/Source/Runtime/CoreUObject/Private/UObject/SavePackage.cpp
@@ -4753,6 +4753,13 @@ FSavePackageResultStruct UPackage::Save(UPackage* InOuter, UObject* Base, EObjec
 					for( int32 i=0; i<Linker->ExportMap.Num(); i++ )
 					{
 						UObject* Object = Linker->ExportMap[i].Object;
+						
+						// Consider redirectors to world/levels as map packages too.
+						if (UObjectRedirector* Redirector = Cast<UObjectRedirector>(Object))
+						{
+							Object = Redirector->DestinationObject;
+						}
+							
 						if ( Object )
 						{
 							const FString ExportClassName = Object->GetClass()->GetName();
