Do not create AI Systems on inactive world

#rnx

diff --git a/Engine/Source/Editor/UnrealEd/Private/EditorEngine.cpp b/Engine/Source/Editor/UnrealEd/Private/EditorEngine.cpp
--- a/Engine/Source/Editor/UnrealEd/Private/EditorEngine.cpp
+++ b/Engine/Source/Editor/UnrealEd/Private/EditorEngine.cpp
@@ -7151,9 +7151,12 @@ void UEditorEngine::InitializeNewlyCreatedInactiveWorld(UWorld* World)
 
 		// Create the world without a physics scene because creating too many physics scenes causes deadlock issues in PhysX. The scene will be created when it is opened in the level editor.
 		// Also, don't create an FXSystem because it consumes too much video memory. This is also created when the level editor opens this world.
+		// Do not create AISystem/Navigation for inactive world. These ones will also be created when the level editor opens this world. if required.
 		World->InitWorld(GetEditorWorldInitializationValues()
 			.CreatePhysicsScene(false)
 			.CreateFXSystem(false)
+			.CreateAISystem(false)
+			.CreateNavigation(false)
 			);
 
 		// Update components so the scene is populated
