Don't write empty catch-all loc chunks

#rnx

diff --git a/Engine/Source/Editor/UnrealEd/Private/LocalizationChunkDataGenerator.cpp b/Engine/Source/Editor/UnrealEd/Private/LocalizationChunkDataGenerator.cpp
--- a/Engine/Source/Editor/UnrealEd/Private/LocalizationChunkDataGenerator.cpp
+++ b/Engine/Source/Editor/UnrealEd/Private/LocalizationChunkDataGenerator.cpp
@@ -27,8 +27,7 @@ void FLocalizationChunkDataGenerator::GenerateChunkDataFiles(const int32 InChunk
 	const bool bIsCatchAllChunk = InChunkId == CatchAllChunkId;
 
 	// We can skip empty non-primary and non-catch-all chunks
-	const bool bCanSkipChunk = !bIsPrimaryChunk && !bIsCatchAllChunk;
-	if (bCanSkipChunk && InPackagesInChunk.Num() == 0)
+	if (!bIsPrimaryChunk && !bIsCatchAllChunk && InPackagesInChunk.Num() == 0)
 	{
 		return;
 	}
@@ -93,8 +92,8 @@ void FLocalizationChunkDataGenerator::GenerateChunkDataFiles(const int32 InChunk
 			return true; // continue enumeration
 		}, false);
 
-		// We can skip empty non-primary and non-catch-all chunks
-		if (bCanSkipChunk && !bChunkHasText)
+		// We can skip empty non-primary chunks
+		if (!bIsPrimaryChunk && !bChunkHasText)
 		{
 			continue;
 		}
