When importing animation from FBX files, only merge anim layers as-needed. Unreal requires all animations be merged to a single layer, so we can't skip the merge entirely. This change only performs the merge when necessary, keeping the curve fidelity as long as its part of an AnimStack with a single member.

FbxCurveAPI is now using the same code has the fbx animation curve importer
The curve import now support "tangent break" and "User tangent"

Remove bad unit conversion for curve tangent

diff --git a/Engine/Source/Editor/MovieSceneTools/Private/MovieSceneToolHelpers.cpp b/Engine/Source/Editor/MovieSceneTools/Private/MovieSceneToolHelpers.cpp
--- a/Engine/Source/Editor/MovieSceneTools/Private/MovieSceneToolHelpers.cpp
+++ b/Engine/Source/Editor/MovieSceneTools/Private/MovieSceneToolHelpers.cpp
@@ -882,7 +882,7 @@ bool ImportFBXProperty(FString NodeName, FString AnimatedPropertyName, FGuid Obj
 				const int32 CompositeIndex = 0;
 				FRichCurve Source;
 				const bool bNegative = false;
-				CurveAPI.GetCurveData(NodeName, AnimatedPropertyName, ChannelIndex, CompositeIndex, Source, bNegative);
+				CurveAPI.GetCurveDataForSequencer(NodeName, AnimatedPropertyName, ChannelIndex, CompositeIndex, Source, bNegative);
 
 				FMovieSceneFloatChannel* Channel = FloatSection->GetChannelProxy().GetChannel<FMovieSceneFloatChannel>(0);
 				TMovieSceneChannelData<FMovieSceneFloatValue> ChannelData = Channel->GetData();
@@ -1042,7 +1042,8 @@ bool ImportFBXTransform(FString NodeName, FGuid ObjectBinding, UnFbx::FFbxCurves
 	FRichCurve EulerRotation[3];
 	FRichCurve Scale[3];
 	FTransform DefaultTransform;
-	CurveAPI.GetConvertedTransformCurveData(NodeName, Translation[0], Translation[1], Translation[2], EulerRotation[0], EulerRotation[1], EulerRotation[2], Scale[0], Scale[1], Scale[2], DefaultTransform);
+	const bool bUseSequencerCurve = true;
+	CurveAPI.GetConvertedTransformCurveData(NodeName, Translation[0], Translation[1], Translation[2], EulerRotation[0], EulerRotation[1], EulerRotation[2], Scale[0], Scale[1], Scale[2], DefaultTransform, bUseSequencerCurve);
 
  	UMovieScene3DTransformTrack* TransformTrack = InMovieScene->FindTrack<UMovieScene3DTransformTrack>(ObjectBinding); 
 	if (!TransformTrack)
diff --git a/Engine/Source/Editor/UnrealEd/Private/Fbx/FbxCurveDataImport.cpp b/Engine/Source/Editor/UnrealEd/Private/Fbx/FbxCurveDataImport.cpp
--- a/Engine/Source/Editor/UnrealEd/Private/Fbx/FbxCurveDataImport.cpp
+++ b/Engine/Source/Editor/UnrealEd/Private/Fbx/FbxCurveDataImport.cpp
@@ -208,10 +208,9 @@ namespace UnFbx {
 		}
 	}
 
-	//Similar to function UnFbx::FFbxImporter::ImportCurve in SkeletalMeshEdit but with weighted tangent support.
-	void FFbxCurvesAPI::GetCurveData(const FFbxAnimCurveHandle &CurveHandle, FRichCurve& RichCurve, bool bNegative) const
+	void FFbxCurvesAPI::GetCurveDataForSequencer(const FFbxAnimCurveHandle &CurveHandle, FRichCurve& RichCurve, bool bNegative) const
 	{
-		static float DefaultCurveWeight = FbxAnimCurveDef::sDEFAULT_WEIGHT;
+		const float DefaultCurveWeight = FbxAnimCurveDef::sDEFAULT_WEIGHT;
 		FbxAnimCurve* FbxCurve = CurveHandle.AnimCurve;
 		if (FbxCurve)
 		{
@@ -233,8 +232,8 @@ namespace UnFbx {
 
 				float LeaveTangent = 0.f;
 				float ArriveTangent = 0.f;
-				float LeaveTangentWeight = 0.f;
-				float ArriveTangentWeight = 0.f;
+				float LeaveTangentWeight = DefaultCurveWeight;
+				float ArriveTangentWeight = DefaultCurveWeight;
 				float ArriveTimeDiff = 0.f;
 				float LeaveTimeDiff = 0.f;
 
@@ -348,6 +347,19 @@ namespace UnFbx {
 		}
 	}
 
+	void FFbxCurvesAPI::GetCurveData(const FFbxAnimCurveHandle &CurveHandle, FRichCurve& RichCurve, bool bNegative) const
+	{
+		FbxAnimCurve* FbxCurve = CurveHandle.AnimCurve;
+		if (FbxCurve)
+		{
+			RichCurve.Reset();
+			//Send a neutral timespan 0 to infinite, the ImportCurve, offset the keytime by doing: (KeyTime - TimeSpan.start), setting TimeSpan.start at zero will not affect the keys time value.
+			FbxTimeSpan AnimTimeSpan(FBXSDK_TIME_ZERO, FBXSDK_TIME_INFINITE);
+			const bool bAutoSetTangents = false;
+			UnFbx::FFbxImporter::ImportCurve(FbxCurve, RichCurve, AnimTimeSpan, 1.0f, bAutoSetTangents);
+		}
+	}
+
 
 	void FFbxCurvesAPI::GetBakeCurveData(const FFbxAnimCurveHandle &CurveHandle, TArray<float>& CurveData, float PeriodTime, float StartTime /*= 0.0f*/, float StopTime /*= -1.0f*/, bool bNegative /*= false*/) const
 	{
@@ -392,6 +404,20 @@ namespace UnFbx {
 		}
 	}
 
+	void FFbxCurvesAPI::GetCurveDataForSequencer(const FString& NodeName, const FString& PropertyName, int32 ChannelIndex, int32 CompositeIndex, FRichCurve& RichCurve, bool bNegative) const
+	{
+		FFbxAnimCurveHandle CurveHandle;
+		GetCurveHandle(NodeName, PropertyName, ChannelIndex, CompositeIndex, CurveHandle);
+		if (CurveHandle.AnimCurve != nullptr)
+		{
+			GetCurveDataForSequencer(CurveHandle, RichCurve, bNegative);
+		}
+		else
+		{
+			RichCurve.Reset();
+		}
+	}
+
 	void FFbxCurvesAPI::GetCurveData(const FString& NodeName, const FString& PropertyName, int32 ChannelIndex, int32 CompositeIndex, FRichCurve& RichCurve, bool bNegative) const
 	{
 		FFbxAnimCurveHandle CurveHandle;
@@ -606,11 +632,10 @@ namespace UnFbx {
 	}
 
 
-
 	void FFbxCurvesAPI::GetConvertedTransformCurveData(const FString& NodeName, FRichCurve& TranslationX, FRichCurve& TranslationY, FRichCurve& TranslationZ,
 		FRichCurve& EulerRotationX, FRichCurve& EulerRotationY, FRichCurve& EulerRotationZ,
 		FRichCurve& ScaleX, FRichCurve& ScaleY, FRichCurve& ScaleZ,
-		FTransform& DefaultTransform) const
+		FTransform& DefaultTransform, bool bUseSequencerCurve) const
 	{
 
 		for (TPair< uint64, FFbxAnimNodeHandle> AnimNodeKvp : CurvesData)
@@ -632,18 +657,35 @@ namespace UnFbx {
 						}
 					}
 				}
+				const bool bNegate = true;
+				if (bUseSequencerCurve)
+				{
+					GetCurveDataForSequencer(TransformCurves[0], TranslationX, !bNegate);
+					GetCurveDataForSequencer(TransformCurves[1], TranslationY, bNegate);
+					GetCurveDataForSequencer(TransformCurves[2], TranslationZ, !bNegate);
 
-				GetCurveData(TransformCurves[0], TranslationX, false);
-				GetCurveData(TransformCurves[1], TranslationY, true);
-				GetCurveData(TransformCurves[2], TranslationZ, false);
+					GetCurveDataForSequencer(TransformCurves[3], EulerRotationX, !bNegate);
+					GetCurveDataForSequencer(TransformCurves[4], EulerRotationY, bNegate);
+					GetCurveDataForSequencer(TransformCurves[5], EulerRotationZ, bNegate);
 
-				GetCurveData(TransformCurves[3], EulerRotationX, false);
-				GetCurveData(TransformCurves[4], EulerRotationY, true);
-				GetCurveData(TransformCurves[5], EulerRotationZ, true);
+					GetCurveDataForSequencer(TransformCurves[6], ScaleX, !bNegate);
+					GetCurveDataForSequencer(TransformCurves[7], ScaleY, !bNegate);
+					GetCurveDataForSequencer(TransformCurves[8], ScaleZ, !bNegate);
+				}
+				else
+				{
+					GetCurveData(TransformCurves[0], TranslationX, !bNegate);
+					GetCurveData(TransformCurves[1], TranslationY, bNegate);
+					GetCurveData(TransformCurves[2], TranslationZ, !bNegate);
 
-				GetCurveData(TransformCurves[6], ScaleX, false);
-				GetCurveData(TransformCurves[7], ScaleY, false);
-				GetCurveData(TransformCurves[8], ScaleZ, false);
+					GetCurveData(TransformCurves[3], EulerRotationX, !bNegate);
+					GetCurveData(TransformCurves[4], EulerRotationY, bNegate);
+					GetCurveData(TransformCurves[5], EulerRotationZ, bNegate);
+
+					GetCurveData(TransformCurves[6], ScaleX, !bNegate);
+					GetCurveData(TransformCurves[7], ScaleY, !bNegate);
+					GetCurveData(TransformCurves[8], ScaleZ, !bNegate);
+				}
 
 				if (bIsCamera || bIsLight)
 				{
diff --git a/Engine/Source/Editor/UnrealEd/Private/Fbx/FbxMainImport.cpp b/Engine/Source/Editor/UnrealEd/Private/Fbx/FbxMainImport.cpp
--- a/Engine/Source/Editor/UnrealEd/Private/Fbx/FbxMainImport.cpp
+++ b/Engine/Source/Editor/UnrealEd/Private/Fbx/FbxMainImport.cpp
@@ -1307,8 +1307,11 @@ void FFbxImporter::ConvertScene()
 	for (int32 AnimStackIndex = 0; AnimStackIndex < AnimStackCount; AnimStackIndex++)
 	{
 		FbxAnimStack* CurAnimStack = Scene->GetSrcObject<FbxAnimStack>(AnimStackIndex);
-		int32 ResampleRate = GetGlobalAnimStackSampleRate(CurAnimStack);
-		MergeAllLayerAnimation(CurAnimStack, ResampleRate);
+		if (CurAnimStack->GetMemberCount() > 1)
+		{
+			int32 ResampleRate = GetGlobalAnimStackSampleRate(CurAnimStack);
+			MergeAllLayerAnimation(CurAnimStack, ResampleRate);
+		}
 	}
 
 	//Set the original file information
diff --git a/Engine/Source/Editor/UnrealEd/Private/SkeletalMeshEdit.cpp b/Engine/Source/Editor/UnrealEd/Private/SkeletalMeshEdit.cpp
--- a/Engine/Source/Editor/UnrealEd/Private/SkeletalMeshEdit.cpp
+++ b/Engine/Source/Editor/UnrealEd/Private/SkeletalMeshEdit.cpp
@@ -990,9 +990,9 @@ bool UnFbx::FFbxImporter::ValidateAnimStack(TArray<FbxNode*>& SortedLinks, TArra
 	return bValidAnimStack;
 }
 
-bool UnFbx::FFbxImporter::ImportCurve(const FbxAnimCurve* FbxCurve, FRichCurve& RichCurve, const FbxTimeSpan &AnimTimeSpan, const float ValueScale/*=1.f*/) const
+bool UnFbx::FFbxImporter::ImportCurve(const FbxAnimCurve* FbxCurve, FRichCurve& RichCurve, const FbxTimeSpan &AnimTimeSpan, const float ValueScale/*=1.f*/, const bool bAutoSetTangents /*= true*/)
 {
-	static float DefaultCurveWeight = FbxAnimCurveDef::sDEFAULT_WEIGHT;
+	const float DefaultCurveWeight = FbxAnimCurveDef::sDEFAULT_WEIGHT;
 
 	if ( FbxCurve )
 	{
@@ -1013,8 +1013,8 @@ bool UnFbx::FFbxImporter::ImportCurve(const FbxAnimCurve* FbxCurve, FRichCurve&
 
 			float LeaveTangent = 0.f; 
 			float ArriveTangent = 0.f;
-			float LeaveTangentWeight = 0.f;
-			float ArriveTangentWeight = 0.f;
+			float LeaveTangentWeight = DefaultCurveWeight;
+			float ArriveTangentWeight = DefaultCurveWeight;
 
 			switch (KeyInterpMode)
 			{
@@ -1044,70 +1044,68 @@ bool UnFbx::FFbxImporter::ImportCurve(const FbxAnimCurve* FbxCurve, FRichCurve&
 				break;
 			}
 
-			// break or any other tangent mode doesn't work well with DCC
-			// it's because we don't support tangent weights, break with tangent weights won't work
-			// I added new ticket to support this, but meanwhile, we'll have to just import using auto. 
-			// @Todo: fix me: UE-20414
-			// when we import tangent, we only support break or user
-			// since it's modified by DCC and we only assume these two are valid
-			// auto does our own stuff, which doesn't work with what you see in DCC
-// 			if (KeyTangentMode & FbxAnimCurveDef::eTangentGenericBreak)
-// 			{
-// 				NewTangentMode = RCTM_Break;
-// 			}
-// 			else
-// 			{
-// 				NewTangentMode = RCTM_User;
-// 			}
-
-			// @fix me : weight of tangent is not used, but we'll just save this for future where we might use it. 
-			switch (KeyTangentWeightMode)
+			//Gather if we want to use auto mode for tangent weight
+			bool bSetDefaultWeight = (KeyTangentMode & FbxAnimCurveDef::eTangentAuto);
+ 			if (KeyTangentMode & FbxAnimCurveDef::eTangentGenericBreak)
+ 			{
+ 				NewTangentMode = RCTM_Break;
+ 			}
+  			else if(KeyTangentMode & FbxAnimCurveDef::eTangentUser)
+  			{
+  				NewTangentMode = RCTM_User;
+  			}
+			//Anything else will be set to auto, we do not support eTangentTCB (Tension, Continuity, Bias)
+
+			// @fix me : weight of tangent is not used, but we'll just save this for future where we might use it.
+			if (!bSetDefaultWeight)
 			{
-			case FbxAnimCurveDef::eWeightedNone://! Tangent has default weights of 0.333; we define this state as not weighted.
-				LeaveTangentWeight = ArriveTangentWeight = DefaultCurveWeight;
-				NewTangentWeightMode = RCTWM_WeightedNone;
-				break;
-			case FbxAnimCurveDef::eWeightedRight: //! Right tangent is weighted.
-				NewTangentWeightMode = RCTWM_WeightedLeave;
-				LeaveTangentWeight = Key.GetDataFloat(FbxAnimCurveDef::eRightWeight);
-				ArriveTangentWeight = DefaultCurveWeight;
-				break;
-			case FbxAnimCurveDef::eWeightedNextLeft://! Left tangent is weighted.
-				NewTangentWeightMode = RCTWM_WeightedArrive;
-				LeaveTangentWeight = DefaultCurveWeight;
-				if ( KeyIndex > 0 )
+				switch (KeyTangentWeightMode)
 				{
-					FbxAnimCurveKey PrevKey = FbxCurve->KeyGet(KeyIndex-1);
-					ArriveTangentWeight = PrevKey.GetDataFloat(FbxAnimCurveDef::eNextLeftWeight);
-				}
-				else
-				{
-					ArriveTangentWeight = 0.f;
-				}
-				break;
-			case FbxAnimCurveDef::eWeightedAll://! Both left and right tangents are weighted.
-				NewTangentWeightMode = RCTWM_WeightedBoth;
-				LeaveTangentWeight = Key.GetDataFloat(FbxAnimCurveDef::eRightWeight);
-				if ( KeyIndex > 0 )
-				{
-					FbxAnimCurveKey PrevKey = FbxCurve->KeyGet(KeyIndex-1);
-					ArriveTangentWeight = PrevKey.GetDataFloat(FbxAnimCurveDef::eNextLeftWeight);
-				}
-				else
-				{
-					ArriveTangentWeight = 0.f;
+				case FbxAnimCurveDef::eWeightedNone://! Tangent has default weights of 0.333; we define this state as not weighted.
+					LeaveTangentWeight = ArriveTangentWeight = DefaultCurveWeight;
+					NewTangentWeightMode = RCTWM_WeightedNone;
+					break;
+				case FbxAnimCurveDef::eWeightedRight: //! Right tangent is weighted.
+					NewTangentWeightMode = RCTWM_WeightedLeave;
+					LeaveTangentWeight = Key.GetDataFloat(FbxAnimCurveDef::eRightWeight);
+					ArriveTangentWeight = DefaultCurveWeight;
+					break;
+				case FbxAnimCurveDef::eWeightedNextLeft://! Left tangent is weighted.
+					NewTangentWeightMode = RCTWM_WeightedArrive;
+					LeaveTangentWeight = DefaultCurveWeight;
+					if (KeyIndex > 0)
+					{
+						FbxAnimCurveKey PrevKey = FbxCurve->KeyGet(KeyIndex - 1);
+						ArriveTangentWeight = PrevKey.GetDataFloat(FbxAnimCurveDef::eNextLeftWeight);
+					}
+					else
+					{
+						ArriveTangentWeight = 0.f;
+					}
+					break;
+				case FbxAnimCurveDef::eWeightedAll://! Both left and right tangents are weighted.
+					NewTangentWeightMode = RCTWM_WeightedBoth;
+					LeaveTangentWeight = Key.GetDataFloat(FbxAnimCurveDef::eRightWeight);
+					if (KeyIndex > 0)
+					{
+						FbxAnimCurveKey PrevKey = FbxCurve->KeyGet(KeyIndex - 1);
+						ArriveTangentWeight = PrevKey.GetDataFloat(FbxAnimCurveDef::eNextLeftWeight);
+					}
+					else
+					{
+						ArriveTangentWeight = 0.f;
+					}
+					break;
 				}
-				break;
 			}
 
-			RichCurve.SetKeyInterpMode(NewKeyHandle, NewInterpMode);
-			RichCurve.SetKeyTangentMode(NewKeyHandle, NewTangentMode);
-			RichCurve.SetKeyTangentWeightMode(NewKeyHandle, NewTangentWeightMode);
+			RichCurve.SetKeyInterpMode(NewKeyHandle, NewInterpMode, bAutoSetTangents);
+			RichCurve.SetKeyTangentMode(NewKeyHandle, NewTangentMode, bAutoSetTangents);
+			RichCurve.SetKeyTangentWeightMode(NewKeyHandle, NewTangentWeightMode, bAutoSetTangents);
 
 			FRichCurveKey& NewKey = RichCurve.GetKey(NewKeyHandle);
-			// apply 1/100 - that seems like the tangent unit difference with FBX
-			NewKey.ArriveTangent = ArriveTangent * 0.01f;
-			NewKey.LeaveTangent = LeaveTangent * 0.01f;
+			NewKey.ArriveTangent = ArriveTangent;
+			NewKey.LeaveTangent = LeaveTangent;
 			NewKey.ArriveTangentWeight = ArriveTangentWeight;
 			NewKey.LeaveTangentWeight = LeaveTangentWeight;
 		}
diff --git a/Engine/Source/Editor/UnrealEd/Public/FbxImporter.h b/Engine/Source/Editor/UnrealEd/Public/FbxImporter.h
--- a/Engine/Source/Editor/UnrealEd/Public/FbxImporter.h
+++ b/Engine/Source/Editor/UnrealEd/Public/FbxImporter.h
@@ -357,6 +357,9 @@ public:
 	UE_DEPRECATED(4.21, "Please use FRichCurve version instead to get tangent weight support")
 	UNREALED_API void GetCurveData(const FString& NodeName, const FString& PropertyName, int32 ChannelIndex, int32 CompositeIndex, FInterpCurveFloat& CurveData, bool bNegative) const;
 	
+	//This one should be use only by the sequencer the key tangents data is transform to fit the expected data we have in the old matinee code
+	UNREALED_API void GetCurveDataForSequencer(const FString& NodeName, const FString& PropertyName, int32 ChannelIndex, int32 CompositeIndex, FRichCurve& RichCurve, bool bNegative) const;
+
 	UNREALED_API void GetCurveData(const FString& NodeName, const FString& PropertyName, int32 ChannelIndex, int32 CompositeIndex, FRichCurve& CurveData, bool bNegative) const;
 
 	
@@ -369,6 +372,9 @@ public:
 	UE_DEPRECATED(4.21, "Please use FRichCurve version instead to get tangent weight support")
 	UNREALED_API void GetCurveData(const FFbxAnimCurveHandle &CurveHandle, FInterpCurveFloat& CurveData, bool bNegative) const;
 	
+	//This one should be use only by the sequencer the key tangents data is transform to fit the expected data we have in the old matinee code
+	UNREALED_API void GetCurveDataForSequencer(const FFbxAnimCurveHandle &CurveHandle, FRichCurve& RichCurve, bool bNegative) const;
+
 	UNREALED_API void GetCurveData(const FFbxAnimCurveHandle &CurveHandle, FRichCurve& CurveData, bool bNegative) const;
 
 	UNREALED_API void GetBakeCurveData(const FFbxAnimCurveHandle &CurveHandle, TArray<float>& CurveData, float PeriodTime, float StartTime = 0.0f, float StopTime = -1.0f, bool bNegative = false) const;
@@ -383,7 +389,7 @@ public:
 	UNREALED_API void GetConvertedTransformCurveData(const FString& NodeName, FRichCurve& TranslationX, FRichCurve& TranslationY, FRichCurve& TranslationZ,
 		FRichCurve& EulerRotationX, FRichCurve& EulerRotationY, FRichCurve& EulerRotationZ,
 		FRichCurve& ScaleX, FRichCurve& ScaleY, FRichCurve& ScaleZ,
-		FTransform& DefaultTransform) const;
+		FTransform& DefaultTransform, bool bUseSequencerCurve) const;
 
 	FbxScene* Scene;
 	TMap<uint64, FFbxAnimNodeHandle> CurvesData;
@@ -1085,7 +1091,7 @@ public:
 	/**
 	 * Import FbxCurve to Curve
 	 */
-	bool ImportCurve(const FbxAnimCurve* FbxCurve, FRichCurve& RichCurve, const FbxTimeSpan &AnimTimeSpan, const float ValueScale = 1.f) const;
+	static bool ImportCurve(const FbxAnimCurve* FbxCurve, FRichCurve& RichCurve, const FbxTimeSpan &AnimTimeSpan, const float ValueScale = 1.f, const bool bAutoSetTangents = true);
 
 	/**
 	 * Merge all layers of one AnimStack to one layer.
diff --git a/Engine/Source/Runtime/Engine/Classes/Curves/RichCurve.h b/Engine/Source/Runtime/Engine/Classes/Curves/RichCurve.h
--- a/Engine/Source/Runtime/Engine/Classes/Curves/RichCurve.h
+++ b/Engine/Source/Runtime/Engine/Classes/Curves/RichCurve.h
@@ -265,12 +265,13 @@ public:
 
 	/** Set the interp mode of the specified key */
 	virtual void SetKeyInterpMode(FKeyHandle KeyHandle, ERichCurveInterpMode NewInterpMode) final override;
+	void SetKeyInterpMode(FKeyHandle KeyHandle, ERichCurveInterpMode NewInterpMode, bool bAutoSetTangents);
 
 	/** Set the tangent mode of the specified key */
-	void SetKeyTangentMode(FKeyHandle KeyHandle, ERichCurveTangentMode NewTangentMode);
+	void SetKeyTangentMode(FKeyHandle KeyHandle, ERichCurveTangentMode NewTangentMode, bool bAutoSetTangents = true);
 
 	/** Set the tangent weight mode of the specified key */
-	void SetKeyTangentWeightMode(FKeyHandle KeyHandle, ERichCurveTangentWeightMode NewTangentWeightMode);
+	void SetKeyTangentWeightMode(FKeyHandle KeyHandle, ERichCurveTangentWeightMode NewTangentWeightMode, bool bAutoSetTangents = true);
 
 	/** Get the interp mode of the specified key */
 	virtual ERichCurveInterpMode GetKeyInterpMode(FKeyHandle KeyHandle) const final override;
diff --git a/Engine/Source/Runtime/Engine/Private/Curves/RichCurve.cpp b/Engine/Source/Runtime/Engine/Private/Curves/RichCurve.cpp
--- a/Engine/Source/Runtime/Engine/Private/Curves/RichCurve.cpp
+++ b/Engine/Source/Runtime/Engine/Private/Curves/RichCurve.cpp
@@ -391,6 +391,11 @@ TPair<float, float> FRichCurve::GetKeyTimeValuePair(FKeyHandle KeyHandle) const
 }
 
 void FRichCurve::SetKeyInterpMode(FKeyHandle KeyHandle, ERichCurveInterpMode NewInterpMode)
+{
+	SetKeyInterpMode(KeyHandle, NewInterpMode, true);
+}
+
+void FRichCurve::SetKeyInterpMode(FKeyHandle KeyHandle, ERichCurveInterpMode NewInterpMode, bool bAutoSetTangents)
 {
 	if (!IsKeyHandleValid(KeyHandle))
 	{
@@ -398,11 +403,14 @@ void FRichCurve::SetKeyInterpMode(FKeyHandle KeyHandle, ERichCurveInterpMode New
 	}
 
 	GetKey(KeyHandle).InterpMode = NewInterpMode;
-	AutoSetTangents();
+	if (bAutoSetTangents)
+	{
+		AutoSetTangents();
+	}
 }
 
 
-void FRichCurve::SetKeyTangentMode(FKeyHandle KeyHandle, ERichCurveTangentMode NewTangentMode)
+void FRichCurve::SetKeyTangentMode(FKeyHandle KeyHandle, ERichCurveTangentMode NewTangentMode, bool bAutoSetTangents /*= true*/)
 {
 	if (!IsKeyHandleValid(KeyHandle))
 	{
@@ -410,11 +418,14 @@ void FRichCurve::SetKeyTangentMode(FKeyHandle KeyHandle, ERichCurveTangentMode N
 	}
 
 	GetKey(KeyHandle).TangentMode = NewTangentMode;
-	AutoSetTangents();
+	if (bAutoSetTangents)
+	{
+		AutoSetTangents();
+	}
 }
 
 
-void FRichCurve::SetKeyTangentWeightMode(FKeyHandle KeyHandle, ERichCurveTangentWeightMode NewTangentWeightMode)
+void FRichCurve::SetKeyTangentWeightMode(FKeyHandle KeyHandle, ERichCurveTangentWeightMode NewTangentWeightMode, bool bAutoSetTangents /*= true*/)
 {
 	if (!IsKeyHandleValid(KeyHandle))
 	{
@@ -422,7 +433,10 @@ void FRichCurve::SetKeyTangentWeightMode(FKeyHandle KeyHandle, ERichCurveTangent
 	}
 
 	GetKey(KeyHandle).TangentWeightMode = NewTangentWeightMode;
-	AutoSetTangents();
+	if (bAutoSetTangents)
+	{
+		AutoSetTangents();
+	}
 }
 
 
