Multiple improvements to the "View Mode" menu of the Viewport, including a feature request from JIRA UE-81469. Most of these improvements are related to the submenu of "View Mode" called "Buffer Visualization":

Issue 1/6 (the JIRA itself): The "View Mode" UI displays as its title the name on the selected category or subcategory chosen by the user (or otherwise the default one). However, this behavior was not consistent across different categories. For the "Buffer Visualization" category, if any of its subcategories were chosen, "View Mode" was only displaying "Buffer Visualization" as title, rather than the particular subcategory chosen. This was not the case for the other 2 categories ("Level of Detail" and "Optimization Viewmode").
Solution 1: The displayed name of "View Mode" is now forced to be the exact chosen (sub)category, regardless of the category ("Buffer Visualization", "Level of Detail", or "Optimization Viewmode").

Issue 2/6: The name displayed as the selected "View Mode" was being assigned in parallel with the displayed name of the buttons representing each (sub)category. This resulted in some categories and subcategories having different names in the buttons than in the "View Mode" label itself. The behavior was also not consistent across the different menus inside "View Mode".
Solution 2: The displayed name of the chosen label is forced to match the displayed name of the button itself.

Issue 3/6: "Optimization Viewmodes" was the only category without an icon.
Solution 3: Icon assigned to one of the icons of its subcategories. In particular, the one from "Quad Overdraw".

Issue 4/6: "Buffer visualization" was the only category without a proper tooltip, it was displaying a FName, rather than a more human-readable FText.
Solution 4: Assigned its tooltip to the FText field of DisplayedName.

Issue 5/6: When the subcategory "Overview" (of "Buffer Visualization") was chosen, the displayed name on each individual subwindow that are generated corresponded with limited FNames (rathern than a more human-readable text).
Solution 5: Assigned their value to the FTexts corresponding to their DisplayedName value.

Issue 6/6: Inside the possible Buffer Visualization Materials, `LightingModel` and `ShadingModel` are being used interchangeably all over the code, leading to missmatches and confusion for the user. E.g., `UMaterialInterface` and `UMaterial` use `LightingModel`, while `FBufferVisualizationData` (thus `FBufferVisualizationMenuCommands` and `FEditorViewportClient`) use `ShadingModel`.
Solution 6: Assigned all of the above classes to the same value.

Issue and solution 7 added in previous CL # 10228324.

diff --git a/Engine/Source/Editor/CommonMenuExtensions/Private/BufferVisualizationMenuCommands.cpp b/Engine/Source/Editor/CommonMenuExtensions/Private/BufferVisualizationMenuCommands.cpp
--- a/Engine/Source/Editor/CommonMenuExtensions/Private/BufferVisualizationMenuCommands.cpp
+++ b/Engine/Source/Editor/CommonMenuExtensions/Private/BufferVisualizationMenuCommands.cpp
@@ -67,7 +67,8 @@ void FBufferVisualizationMenuCommands::CreateOverviewCommand()
 
 	FBufferVisualizationRecord& OverviewRecord = CommandMap.Add(CommandName, FBufferVisualizationRecord());
 	OverviewRecord.Name = NAME_None;
-	OverviewRecord.Command = FUICommandInfoDecl(this->AsShared(), CommandName, LOCTEXT("BufferVisualization", "Overview"), LOCTEXT("BufferVisualization", "Overview"))
+	const FText MaterialDisplayDefaultName = FBufferVisualizationData::GetMaterialDefaultDisplayName();
+	OverviewRecord.Command = FUICommandInfoDecl(this->AsShared(), CommandName, MaterialDisplayDefaultName, MaterialDisplayDefaultName)
 		.UserInterfaceType(EUserInterfaceActionType::RadioButton)
 		.DefaultChord(FInputChord());
 }
@@ -80,8 +81,7 @@ void FBufferVisualizationMenuCommands::CreateVisualizationCommands()
 
 		FBufferVisualizationRecord& Record = CommandMap.Add(CommandName, FBufferVisualizationRecord());
 		Record.Name = *InMaterialName;
-		const FText MaterialNameText = FText::FromString(InMaterialName);
-		Record.Command = FUICommandInfoDecl(this->AsShared(), CommandName, MaterialNameText, MaterialNameText)
+		Record.Command = FUICommandInfoDecl(this->AsShared(), CommandName, InDisplayName, InDisplayName)
 			.UserInterfaceType(EUserInterfaceActionType::RadioButton)
 			.DefaultChord(FInputChord());
 	});
@@ -106,7 +106,7 @@ void FBufferVisualizationMenuCommands::AddOverviewCommandToMenu(FMenuBuilder& Me
 {
 	check(CommandMap.Num() > 0);
 
-	Menu.AddMenuEntry(OverviewCommand().Command, NAME_None, LOCTEXT("BufferVisualization", "Overview"));
+	Menu.AddMenuEntry(OverviewCommand().Command, NAME_None, FBufferVisualizationData::GetMaterialDefaultDisplayName());
 }
 
 void FBufferVisualizationMenuCommands::AddVisualizationCommandsToMenu(FMenuBuilder& Menu) const
diff --git a/Engine/Source/Editor/LevelEditor/Private/SLevelViewportToolBar.cpp b/Engine/Source/Editor/LevelEditor/Private/SLevelViewportToolBar.cpp
--- a/Engine/Source/Editor/LevelEditor/Private/SLevelViewportToolBar.cpp
+++ b/Engine/Source/Editor/LevelEditor/Private/SLevelViewportToolBar.cpp
@@ -1296,7 +1296,6 @@ void SLevelViewportToolBar::CreateViewMenuExtensions(FMenuBuilder& MenuBuilder)
 				check(ViewportClient.IsValid());
 				return ViewportClient->IsViewModeEnabled(VMI_VisualizeBuffer);
 			})),
-			//FIsActionChecked::CreateStatic(&FBufferVisualizationMenuCommands::IsVisualisationMenuSelected, ViewportClient)),
 		/* InExtensionHook = */ NAME_None, EUserInterfaceActionType::RadioButton,
 		/* bInOpenSubMenuOnClick = */ false, FSlateIcon(FEditorStyle::GetStyleSetName(), "EditorViewport.VisualizeBufferMode"));
 
diff --git a/Engine/Source/Editor/UnrealEd/Private/EditorViewportClient.cpp b/Engine/Source/Editor/UnrealEd/Private/EditorViewportClient.cpp
--- a/Engine/Source/Editor/UnrealEd/Private/EditorViewportClient.cpp
+++ b/Engine/Source/Editor/UnrealEd/Private/EditorViewportClient.cpp
@@ -56,6 +56,7 @@
 #include "RayTracingDebugVisualizationMenuCommands.h"
 #include "Misc/ScopedSlowTask.h"
 #include "UnrealEngine.h"
+#include "BufferVisualizationData.h"
 
 #define LOCTEXT_NAMESPACE "EditorViewportClient"
 
@@ -2509,6 +2510,13 @@ bool FEditorViewportClient::IsBufferVisualizationModeSelected( FName InName ) co
 	return IsViewModeEnabled( VMI_VisualizeBuffer ) && CurrentBufferVisualizationMode == InName;
 }
 
+FText FEditorViewportClient::GetCurrentBufferVisualizationModeDisplayName() const
+{
+	checkf(IsViewModeEnabled(VMI_VisualizeBuffer), TEXT("In order to call GetCurrentBufferVisualizationMode(), first you must set ViewMode to VMI_VisualizeBuffer."));
+	return (CurrentBufferVisualizationMode.IsNone()
+		? FBufferVisualizationData::GetMaterialDefaultDisplayName() : GetBufferVisualizationData().GetMaterialDisplayName(CurrentBufferVisualizationMode));
+}
+
 void FEditorViewportClient::ChangeRayTracingDebugVisualizationMode(FName InName)
 {
 	SetViewMode(VMI_RayTracingDebug);
diff --git a/Engine/Source/Editor/UnrealEd/Private/SEditorViewportViewMenu.cpp b/Engine/Source/Editor/UnrealEd/Private/SEditorViewportViewMenu.cpp
--- a/Engine/Source/Editor/UnrealEd/Private/SEditorViewportViewMenu.cpp
+++ b/Engine/Source/Editor/UnrealEd/Private/SEditorViewportViewMenu.cpp
@@ -32,108 +32,18 @@ FText SEditorViewportViewMenu::GetViewMenuLabel() const
 	TSharedPtr< SEditorViewport > PinnedViewport = Viewport.Pin();
 	if( PinnedViewport.IsValid() )
 	{
-		switch( PinnedViewport->GetViewportClient()->GetViewMode() )
+		const TSharedPtr<FEditorViewportClient> ViewportClient = PinnedViewport->GetViewportClient();
+		check(ViewportClient.IsValid());
+		const EViewModeIndex ViewMode = ViewportClient->GetViewMode();
+		// If VMI_VisualizeBuffer, return its subcategory name
+		if (ViewMode == VMI_VisualizeBuffer)
 		{
-			case VMI_BrushWireframe:
-				Label = LOCTEXT("ViewMenuTitle_BrushWireframe", "Wireframe");
-				break;
-
-			case VMI_Wireframe:
-				Label = LOCTEXT("ViewMenuTitle_Wireframe", "Wireframe");
-				break;
-
-			case VMI_Unlit:
-				Label = LOCTEXT("ViewMenuTitle_Unlit", "Unlit");
-				break;
-
-			case VMI_Lit:
-				Label = LOCTEXT("ViewMenuTitle_Lit", "Lit");
-				break;
-
-#if RHI_RAYTRACING
-			case VMI_PathTracing:
-				Label = LOCTEXT("ViewMenuTitle_PathTracing", "Path Tracing");
-				break;
-
-			case VMI_RayTracingDebug:
-				Label = LOCTEXT("ViewMenuTitle_RayTracingDebug", "Ray Tracing Debug");
-				break;
-#endif
-
-			case VMI_Lit_DetailLighting:
-				Label = LOCTEXT("ViewMenuTitle_DetailLighting", "Detail Lighting");
-				break;
-
-			case VMI_LightingOnly:
-				Label = LOCTEXT("ViewMenuTitle_LightingOnly", "Lighting Only");
-				break;
-
-			case VMI_LightComplexity:
-				Label = LOCTEXT("ViewMenuTitle_LightComplexity", "Light Complexity");
-				break;
-
-			case VMI_ShaderComplexity:
-				Label = LOCTEXT("ViewMenuTitle_ShaderComplexity", "Shader Complexity");
-				break;
-
-			case VMI_QuadOverdraw:
-				Label = LOCTEXT("ViewMenuTitle_QuadOverdraw", "Quad Overdraw");
-				break;
-
-			case VMI_ShaderComplexityWithQuadOverdraw:
-				Label = LOCTEXT("ViewMenuTitle_ShaderComplexityWithQuadOverdraw", "Shader Complexity & Quads");
-				break;
-
-			case VMI_PrimitiveDistanceAccuracy:
-				Label = LOCTEXT("ViewMenuTitle_PrimitiveDistanceAccuracy", "Primitive Distance Accuracy");
-				break;
-
-			case VMI_MeshUVDensityAccuracy:
-				Label = LOCTEXT("ViewMenuTitle_MeshUVDensityAccuracy", "Mesh UV Densities Accuracy");
-				break;
-
-			case VMI_MaterialTextureScaleAccuracy:
-				Label = LOCTEXT("ViewMenuTitle_MaterialTextureScaleAccuracy", "Material Texture Scales Accuracy");
-				break;
-
-			case VMI_RequiredTextureResolution:
-				Label = LOCTEXT("ViewMenuTitle_Required Texture Resolution", "Required Texture Resolution");
-				break;
-
-			case VMI_StationaryLightOverlap:
-				Label = LOCTEXT("ViewMenuTitle_StationaryLightOverlap", "Stationary Light Overlap");
-				break;
-
-			case VMI_LightmapDensity:
-				Label = LOCTEXT("ViewMenuTitle_LightmapDensity", "Lightmap Density");
-				break;
-
-			case VMI_ReflectionOverride:
-				Label = LOCTEXT("ViewMenuTitle_ReflectionOverride", "Reflections");
-				break;
-
-			case VMI_VisualizeBuffer:
-				Label = LOCTEXT("ViewMenuTitle_VisualizeBuffer", "Buffer Visualization");
-				break;
-
-			case VMI_CollisionPawn:
-				Label = LOCTEXT("ViewMenuTitle_CollisionPawn", "Player Collision");
-				break;
-
-			case VMI_CollisionVisibility:
-				Label = LOCTEXT("ViewMenuTitle_CollisionVisibility", "Visibility Collision");
-				break;
-
-			case VMI_LitLightmapDensity:
-				break;
-
-			case VMI_LODColoration:
-				Label = LOCTEXT("ViewMenuTitle_LODColoration", "LOD Coloration");
-				break;
-
-			case VMI_HLODColoration:
-				Label = LOCTEXT("ViewMenuTitle_HLODColoration", "HLOD Coloration");
-				break;
+			Label = ViewportClient->GetCurrentBufferVisualizationModeDisplayName();
+		}
+		// For any other category, return its own name
+		else
+		{
+			Label = UViewModeUtils::GetViewModeDisplayName(ViewMode);
 		}
 	}
 
@@ -291,21 +201,21 @@ TSharedRef<SWidget> SEditorViewportViewMenu::GenerateViewMenuContent() const
 		{
 			ViewMenuBuilder.BeginSection("ViewMode", LOCTEXT("ViewModeHeader", "View Mode") );
 			{
-				ViewMenuBuilder.AddMenuEntry(BaseViewportActions.LitMode, NAME_None, LOCTEXT("LitViewModeDisplayName", "Lit"));
-				ViewMenuBuilder.AddMenuEntry(BaseViewportActions.UnlitMode, NAME_None, LOCTEXT("UnlitViewModeDisplayName", "Unlit"));
-				ViewMenuBuilder.AddMenuEntry(BaseViewportActions.WireframeMode, NAME_None, LOCTEXT("BrushWireframeViewModeDisplayName", "Wireframe"));
-				ViewMenuBuilder.AddMenuEntry(BaseViewportActions.DetailLightingMode, NAME_None, LOCTEXT("DetailLightingViewModeDisplayName", "Detail Lighting"));
-				ViewMenuBuilder.AddMenuEntry(BaseViewportActions.LightingOnlyMode, NAME_None, LOCTEXT("LightingOnlyViewModeDisplayName", "Lighting Only"));
-				ViewMenuBuilder.AddMenuEntry(BaseViewportActions.ReflectionOverrideMode, NAME_None, LOCTEXT("ReflectionOverrideViewModeDisplayName", "Reflections"));
-				ViewMenuBuilder.AddMenuEntry(BaseViewportActions.CollisionPawn, NAME_None, LOCTEXT("CollisionPawnViewModeDisplayName", "Player Collision"));
-				ViewMenuBuilder.AddMenuEntry(BaseViewportActions.CollisionVisibility, NAME_None, LOCTEXT("CollisionVisibilityViewModeDisplayName", "Visibility Collision"));
+				ViewMenuBuilder.AddMenuEntry(BaseViewportActions.LitMode, NAME_None, UViewModeUtils::GetViewModeDisplayName(VMI_Lit));
+				ViewMenuBuilder.AddMenuEntry(BaseViewportActions.UnlitMode, NAME_None, UViewModeUtils::GetViewModeDisplayName(VMI_Unlit));
+				ViewMenuBuilder.AddMenuEntry(BaseViewportActions.WireframeMode, NAME_None, UViewModeUtils::GetViewModeDisplayName(VMI_BrushWireframe));
+				ViewMenuBuilder.AddMenuEntry(BaseViewportActions.DetailLightingMode, NAME_None, UViewModeUtils::GetViewModeDisplayName(VMI_Lit_DetailLighting));
+				ViewMenuBuilder.AddMenuEntry(BaseViewportActions.LightingOnlyMode, NAME_None, UViewModeUtils::GetViewModeDisplayName(VMI_LightingOnly));
+				ViewMenuBuilder.AddMenuEntry(BaseViewportActions.ReflectionOverrideMode, NAME_None, UViewModeUtils::GetViewModeDisplayName(VMI_ReflectionOverride));
+				ViewMenuBuilder.AddMenuEntry(BaseViewportActions.CollisionPawn, NAME_None, UViewModeUtils::GetViewModeDisplayName(VMI_CollisionPawn));
+				ViewMenuBuilder.AddMenuEntry(BaseViewportActions.CollisionVisibility, NAME_None, UViewModeUtils::GetViewModeDisplayName(VMI_CollisionVisibility));
 			}
 
 #if RHI_RAYTRACING
 			if (IsRayTracingEnabled())
 			{
 				const FEditorViewportCommands& BaseViewportCommands = FEditorViewportCommands::Get();
-				ViewMenuBuilder.AddMenuEntry(BaseViewportCommands.PathTracingMode, NAME_None, LOCTEXT("PathTracingViewModeDisplayName", "Path Tracing"));
+				ViewMenuBuilder.AddMenuEntry(BaseViewportCommands.PathTracingMode, NAME_None, UViewModeUtils::GetViewModeDisplayName(VMI_PathTracing));
 			}
 #endif
 
@@ -324,20 +234,20 @@ TSharedRef<SWidget> SEditorViewportViewMenu::GenerateViewMenuContent() const
 						{
 							if (FeatureLevel == ERHIFeatureLevel::SM5)
 							{
-								Menu.AddMenuEntry(BaseViewportCommands.LightComplexityMode, NAME_None, LOCTEXT("LightComplexityViewModeDisplayName", "Light Complexity"));
-								Menu.AddMenuEntry(BaseViewportCommands.LightmapDensityMode, NAME_None, LOCTEXT("LightmapDensityViewModeDisplayName", "Lightmap Density"));
-								Menu.AddMenuEntry(BaseViewportCommands.StationaryLightOverlapMode, NAME_None, LOCTEXT("StationaryLightOverlapViewModeDisplayName", "Stationary Light Overlap"));
+								Menu.AddMenuEntry(BaseViewportCommands.LightComplexityMode, NAME_None, UViewModeUtils::GetViewModeDisplayName(VMI_LightComplexity));
+								Menu.AddMenuEntry(BaseViewportCommands.LightmapDensityMode, NAME_None, UViewModeUtils::GetViewModeDisplayName(VMI_LightmapDensity));
+								Menu.AddMenuEntry(BaseViewportCommands.StationaryLightOverlapMode, NAME_None, UViewModeUtils::GetViewModeDisplayName(VMI_StationaryLightOverlap));
 							}
 
-							Menu.AddMenuEntry(BaseViewportCommands.ShaderComplexityMode, NAME_None, LOCTEXT("ShaderComplexityViewModeDisplayName", "Shader Complexity"));
+							Menu.AddMenuEntry(BaseViewportCommands.ShaderComplexityMode, NAME_None, UViewModeUtils::GetViewModeDisplayName(VMI_ShaderComplexity));
 
 							if (AllowDebugViewShaderMode(DVSM_ShaderComplexityContainedQuadOverhead, GMaxRHIShaderPlatform, FeatureLevel))
 							{
-								Menu.AddMenuEntry(BaseViewportCommands.ShaderComplexityWithQuadOverdrawMode, NAME_None, LOCTEXT("ShaderComplexityWithQuadOverdrawViewModeDisplayName", "Shader Complexity & Quads"));
+								Menu.AddMenuEntry(BaseViewportCommands.ShaderComplexityWithQuadOverdrawMode, NAME_None, UViewModeUtils::GetViewModeDisplayName(VMI_ShaderComplexityWithQuadOverdraw));
 							}
 							if (AllowDebugViewShaderMode(DVSM_QuadComplexity, GMaxRHIShaderPlatform, FeatureLevel))
 							{
-								Menu.AddMenuEntry(BaseViewportCommands.QuadOverdrawMode, NAME_None, LOCTEXT("QuadOverdrawViewModeDisplayName", "Quad Overdraw"));
+								Menu.AddMenuEntry(BaseViewportCommands.QuadOverdrawMode, NAME_None, UViewModeUtils::GetViewModeDisplayName(VMI_QuadOverdraw));
 							}
 						}
 						Menu.EndSection();
@@ -345,20 +255,20 @@ TSharedRef<SWidget> SEditorViewportViewMenu::GenerateViewMenuContent() const
 						Menu.BeginSection("TextureStreaming", LOCTEXT("TextureStreamingHeader", "Texture Streaming Accuracy") );
 						if ( AllowDebugViewShaderMode(DVSM_PrimitiveDistanceAccuracy, GMaxRHIShaderPlatform, FeatureLevel) && (!InParentToolBar.IsValid() || InParentToolBar.Pin()->IsViewModeSupported(VMI_PrimitiveDistanceAccuracy)) )
 						{
-							Menu.AddMenuEntry( BaseViewportCommands.TexStreamAccPrimitiveDistanceMode, NAME_None, LOCTEXT("TexStreamAccPrimitiveDistanceViewModeDisplayName", "Primitive Distance") );
+							Menu.AddMenuEntry(BaseViewportCommands.TexStreamAccPrimitiveDistanceMode, NAME_None, UViewModeUtils::GetViewModeDisplayName(VMI_PrimitiveDistanceAccuracy));
 						}
 						if ( AllowDebugViewShaderMode(DVSM_MeshUVDensityAccuracy, GMaxRHIShaderPlatform, FeatureLevel) && (!InParentToolBar.IsValid() || InParentToolBar.Pin()->IsViewModeSupported(VMI_MeshUVDensityAccuracy)) )
 						{
-							Menu.AddMenuEntry(BaseViewportCommands.TexStreamAccMeshUVDensityMode, NAME_None, LOCTEXT("TexStreamAccMeshUVDensityViewModeDisplayName", "Mesh UV Densities"));
+							Menu.AddMenuEntry(BaseViewportCommands.TexStreamAccMeshUVDensityMode, NAME_None, UViewModeUtils::GetViewModeDisplayName(VMI_MeshUVDensityAccuracy));
 						}
 						// TexCoordScale accuracy viewmode requires shaders that are only built in the TextureStreamingBuild, which requires the new metrics to be enabled.
 						if ( AllowDebugViewShaderMode(DVSM_MaterialTextureScaleAccuracy, GMaxRHIShaderPlatform, FeatureLevel) && CVarStreamingUseNewMetrics.GetValueOnAnyThread() != 0 && (!InParentToolBar.IsValid() || InParentToolBar.Pin()->IsViewModeSupported(VMI_MaterialTextureScaleAccuracy)) )
 						{
-							Menu.AddMenuEntry(BaseViewportCommands.TexStreamAccMaterialTextureScaleMode, NAME_None, LOCTEXT("TexStreamAccMaterialTextureScaleViewModeDisplayName", "Material Texture Scales"));
+							Menu.AddMenuEntry(BaseViewportCommands.TexStreamAccMaterialTextureScaleMode, NAME_None, UViewModeUtils::GetViewModeDisplayName(VMI_MaterialTextureScaleAccuracy));
 						}
 						if ( AllowDebugViewShaderMode(DVSM_RequiredTextureResolution, GMaxRHIShaderPlatform, FeatureLevel) && (!InParentToolBar.IsValid() || InParentToolBar.Pin()->IsViewModeSupported(VMI_MaterialTextureScaleAccuracy)) )
 						{
-							Menu.AddMenuEntry(BaseViewportCommands.RequiredTextureResolutionMode, NAME_None, LOCTEXT("RequiredTextureResolutionModeDisplayName", "Required Texture Resolution"));
+							Menu.AddMenuEntry(BaseViewportCommands.RequiredTextureResolutionMode, NAME_None, UViewModeUtils::GetViewModeDisplayName(VMI_RequiredTextureResolution));
 						}
 						Menu.EndSection();
 					}
@@ -409,8 +319,8 @@ TSharedRef<SWidget> SEditorViewportViewMenu::GenerateViewMenuContent() const
 					{
 						Menu.BeginSection("LevelViewportLODColoration", LOCTEXT("LODModesHeader", "Level of Detail Coloration"));
 						{
-							Menu.AddMenuEntry(FEditorViewportCommands::Get().LODColorationMode, NAME_None, LOCTEXT("LODColorationModeDisplayName", "Mesh LODs"));
-							Menu.AddMenuEntry(FEditorViewportCommands::Get().HLODColorationMode, NAME_None, LOCTEXT("HLODColorationModeDisplayName", "Hierarchical LODs"));
+							Menu.AddMenuEntry(FEditorViewportCommands::Get().LODColorationMode, NAME_None, UViewModeUtils::GetViewModeDisplayName(VMI_LODColoration));
+							Menu.AddMenuEntry(FEditorViewportCommands::Get().HLODColorationMode, NAME_None, UViewModeUtils::GetViewModeDisplayName(VMI_HLODColoration));
 						}
 						Menu.EndSection();
 					}
diff --git a/Engine/Source/Editor/UnrealEd/Public/EditorViewportClient.h b/Engine/Source/Editor/UnrealEd/Public/EditorViewportClient.h
--- a/Engine/Source/Editor/UnrealEd/Public/EditorViewportClient.h
+++ b/Engine/Source/Editor/UnrealEd/Public/EditorViewportClient.h
@@ -1136,6 +1136,11 @@ public:
 	 */
 	bool IsBufferVisualizationModeSelected( FName InName ) const;
 
+	/**
+	 * It returns the FText display name associate with CurrentBufferVisualizationMode
+	 */
+	FText GetCurrentBufferVisualizationModeDisplayName() const;
+
 	/**
 	 * Changes the ray tracing debug visualization mode for this viewport
 	 *
diff --git a/Engine/Source/Runtime/Engine/Classes/Engine/EngineBaseTypes.h b/Engine/Source/Runtime/Engine/Classes/Engine/EngineBaseTypes.h
--- a/Engine/Source/Runtime/Engine/Classes/Engine/EngineBaseTypes.h
+++ b/Engine/Source/Runtime/Engine/Classes/Engine/EngineBaseTypes.h
@@ -911,7 +911,7 @@ enum ENetMode
  * Don't change the order, the ID is serialized with the editor
  */
 UENUM()
-enum EViewModeIndex 
+enum EViewModeIndex
 {
 	/** Wireframe w/ brushes. */
 	VMI_BrushWireframe = 0 UMETA(DisplayName = "Brush Wireframe"),
@@ -970,9 +970,25 @@ enum EViewModeIndex
 
 	VMI_Max UMETA(Hidden),
 
+	// VMI_Unknown - The value assigned to VMI_Unknown must be the highest possible of any member of EViewModeIndex, or GetViewModeName might seg-fault
 	VMI_Unknown = 255 UMETA(Hidden),
 };
 
+/**
+ * Class containing a static util function to help with EViewModeIndex
+ */
+UCLASS(config = Engine)
+class ENGINE_API UViewModeUtils : public UObject
+{
+	GENERATED_BODY()
+
+public:
+	/**
+	 * Get the display name associated with a particular EViewModeIndex
+	 */
+	static FText GetViewModeDisplayName(const EViewModeIndex ViewModeIndex);
+};
+
 
 /** Settings to allow designers to override the automatic expose */
 USTRUCT()
diff --git a/Engine/Source/Runtime/Engine/Private/BufferVisualizationData.cpp b/Engine/Source/Runtime/Engine/Private/BufferVisualizationData.cpp
--- a/Engine/Source/Runtime/Engine/Private/BufferVisualizationData.cpp
+++ b/Engine/Source/Runtime/Engine/Private/BufferVisualizationData.cpp
@@ -4,6 +4,8 @@
 #include "Materials/Material.h"
 #include "Misc/ConfigCacheIni.h"
 
+#define LOCTEXT_NAMESPACE "FBufferVisualizationData"
+
 static FBufferVisualizationData GBufferVisualizationData;
 
 static TAutoConsoleVariable<int32> BufferVisualizationDumpFramesAsHDR(
@@ -29,6 +31,7 @@ void FBufferVisualizationData::Initialize()
 		if (AllowDebugViewmodes())
 		{
 			check(MaterialMap.Num() == 0);
+			check(MaterialMapFromMaterialName.Num() == 0);
 
 			FConfigSection* MaterialSection = GConfig->GetSectionPrivate( TEXT("Engine.BufferVisualizationMaterials"), false, true, GEngineIni );
 
@@ -50,6 +53,7 @@ void FBufferVisualizationData::Initialize()
 							FText DisplayName;
 							FParse::Value( *It.Value().GetValue(), TEXT("Name="), DisplayName, TEXT("Engine.BufferVisualizationMaterials") );
 							Rec.DisplayName = DisplayName;
+							MaterialMapFromMaterialName.Add(Material->GetFName(), Rec);
 						}
 					}
 				}
@@ -104,17 +108,53 @@ void FBufferVisualizationData::ConfigureConsoleCommand()
 
 UMaterial* FBufferVisualizationData::GetMaterial(FName InMaterialName)
 {
-	Record* Result = MaterialMap.Find(InMaterialName);
-	if (Result)
+	// Get UMaterial from the TMaterialMap FName
+	if (const Record* Result = MaterialMap.Find(InMaterialName))
 	{
 		return Result->Material;
 	}
+	// Get UMaterial from the UMaterial FName
+	// Almost all BufferVisualizationData variables contain a FMaterial with its same FName.
+	// But not all, e.g., ShadingModel uses the FMaterial LightingModel. This could confuse the developer depending on which one he is using
+	// This "else if" case handles the case in which we look for "LightingModel" rather than "ShadingModel", returning the same value
+	else if (const Record* ResultFromMaterialName = MaterialMapFromMaterialName.Find(InMaterialName))
+	{
+		return ResultFromMaterialName->Material;
+	}
+	// Not found
 	else
 	{
-		return NULL;
+		return nullptr;
 	}
 }
 
+FText FBufferVisualizationData::GetMaterialDisplayName(FName InMaterialName) const
+{
+	// Get UMaterial from the TMaterialMap FName
+	if (const Record* Result = MaterialMap.Find(InMaterialName))
+	{
+		return Result->DisplayName;
+	}
+	// Get UMaterial from the UMaterial FName
+	// Almost all BufferVisualizationData variables contain a FMaterial with its same FName.
+	// But not all, e.g., ShadingModel uses the FMaterial LightingModel. This could confuse the developer depending on which one he is using
+	// This "else if" case handles the case in which we look for "LightingModel" rather than "ShadingModel", returning the same value
+	else if (const Record* ResultFromMaterialName = MaterialMapFromMaterialName.Find(InMaterialName))
+	{
+		return ResultFromMaterialName->DisplayName;
+	}
+	// Not found
+	else
+	{
+		return FText::GetEmpty();
+	}
+}
+
+FText FBufferVisualizationData::GetMaterialDefaultDisplayName()
+{
+	return LOCTEXT("BufferVisualization", "Overview");
+}
+
 void FBufferVisualizationData::SetCurrentOverviewMaterialNames(const FString& InNameList)
 {
 	CurrentOverviewMaterialNames = InNameList;
@@ -139,3 +179,5 @@ FBufferVisualizationData& GetBufferVisualizationData()
 
 	return GBufferVisualizationData;
 }
+
+#undef LOCTEXT_NAMESPACE
diff --git a/Engine/Source/Runtime/Engine/Private/ViewModeNames.cpp b/Engine/Source/Runtime/Engine/Private/ViewModeNames.cpp
new file mode 100644
--- /dev/null
+++ b/Engine/Source/Runtime/Engine/Private/ViewModeNames.cpp
@@ -0,0 +1,179 @@
+// Copyright 1998-2019 Epic Games, Inc. All Rights Reserved.
+
+#include "CoreMinimal.h"
+#include "Engine/EngineBaseTypes.h"
+
+#define LOCTEXT_NAMESPACE "UViewModeUtils"
+
+TArray<FText> FillViewModeDisplayNames()
+{
+	// Allocate size
+	TArray<FText> ViewModeDisplayNames;
+	ViewModeDisplayNames.Reserve(VMI_Unknown);
+	// Fill ViewModeIndex, VMI_Unknown+1 to include VMI_Unknown too
+	for (int32 Index = 0; Index < VMI_Unknown + 1; ++Index)
+	{
+		const EViewModeIndex ViewModeIndex = (EViewModeIndex)Index;
+
+		// Wireframe w/ brushes
+		if (ViewModeIndex == VMI_BrushWireframe)
+		{
+			ViewModeDisplayNames.Emplace(LOCTEXT("UViewModeUtils_VMI_BrushWireframe", "Brush Wireframe"));
+		}
+		// Wireframe w/ BSP
+		else if (ViewModeIndex == VMI_Wireframe)
+		{
+			ViewModeDisplayNames.Emplace(LOCTEXT("UViewModeUtils_VMI_Wireframe", "Wireframe"));
+		}
+		// Unlit
+		else if (ViewModeIndex == VMI_Unlit)
+		{
+			ViewModeDisplayNames.Emplace(LOCTEXT("UViewModeUtils_VMI_Unlit", "Unlit"));
+		}
+		// Lit
+		else if (ViewModeIndex == VMI_Lit)
+		{
+			ViewModeDisplayNames.Emplace(LOCTEXT("UViewModeUtils_VMI_Lit", "Lit"));
+		}
+		else if (ViewModeIndex == VMI_Lit_DetailLighting)
+		{
+			ViewModeDisplayNames.Emplace(LOCTEXT("UViewModeUtils_VMI_Lit_DetailLighting", "Detail Lighting"));
+		}
+		// Lit wo/ materials
+		else if (ViewModeIndex == VMI_LightingOnly)
+		{
+			ViewModeDisplayNames.Emplace(LOCTEXT("UViewModeUtils_VMI_LightingOnly", "Lighting Only"));
+		}
+		// Colored according to light count
+		else if (ViewModeIndex == VMI_LightComplexity)
+		{
+			ViewModeDisplayNames.Emplace(LOCTEXT("UViewModeUtils_VMI_LightComplexity", "Light Complexity"));
+		}
+		// Colored according to shader complexity
+		else if (ViewModeIndex == VMI_ShaderComplexity)
+		{
+			ViewModeDisplayNames.Emplace(LOCTEXT("UViewModeUtils_VMI_ShaderComplexity", "Shader Complexity"));
+		}
+		// Colored according to world-space LightMap texture density
+		else if (ViewModeIndex == VMI_LightmapDensity)
+		{
+			ViewModeDisplayNames.Emplace(LOCTEXT("UViewModeUtils_VMI_LightmapDensity", "Lightmap Density"));
+		}
+		// Colored according to light count - showing lightmap texel density on texture mapped objects
+		else if (ViewModeIndex == VMI_LitLightmapDensity)
+		{
+			ViewModeDisplayNames.Emplace(LOCTEXT("UViewModeUtils_VMI_LitLightmapDensity", "Lit Lightmap Density"));
+		}
+		else if (ViewModeIndex == VMI_ReflectionOverride)
+		{
+			ViewModeDisplayNames.Emplace(LOCTEXT("UViewModeUtils_VMI_ReflectionOverride", "Reflections"));
+		}
+		else if (ViewModeIndex == VMI_VisualizeBuffer)
+		{
+			ViewModeDisplayNames.Emplace(LOCTEXT("UViewModeUtils_VMI_VisualizeBuffer", "Buffer Visualization5"));
+		}
+		//	VMI_VoxelLighting = 13,
+	
+		// Colored according to stationary light overlap
+		else if (ViewModeIndex == VMI_StationaryLightOverlap)
+		{
+			ViewModeDisplayNames.Emplace(LOCTEXT("UViewModeUtils_VMI_StationaryLightOverlap", "Stationary Light Overlap"));
+		}
+	
+		else if (ViewModeIndex == VMI_CollisionPawn)
+		{
+			ViewModeDisplayNames.Emplace(LOCTEXT("UViewModeUtils_VMI_CollisionPawn", "Player Collision"));
+		}
+		else if (ViewModeIndex == VMI_CollisionVisibility)
+		{
+			ViewModeDisplayNames.Emplace(LOCTEXT("UViewModeUtils_VMI_CollisionVisibility", "Visibility Collision"));
+		}
+		//VMI_UNUSED = 17,
+		// Colored according to the current LOD index
+		else if (ViewModeIndex == VMI_LODColoration)
+		{
+			ViewModeDisplayNames.Emplace(LOCTEXT("UViewModeUtils_VMI_LODColoration", "Mesh LOD Coloration"));
+		}
+		// Colored according to the quad coverage
+		else if (ViewModeIndex == VMI_QuadOverdraw)
+		{
+			ViewModeDisplayNames.Emplace(LOCTEXT("UViewModeUtils_VMI_QuadOverdraw", "Quad Overdraw"));
+		}
+		// Visualize the accuracy of the primitive distance computed for texture streaming
+		else if (ViewModeIndex == VMI_PrimitiveDistanceAccuracy)
+		{
+			ViewModeDisplayNames.Emplace(LOCTEXT("UViewModeUtils_VMI_PrimitiveDistanceAccuracy", "Primitive Distance"));
+		}
+		// Visualize the accuracy of the mesh UV densities computed for texture streaming
+		else if (ViewModeIndex == VMI_MeshUVDensityAccuracy)
+		{
+			ViewModeDisplayNames.Emplace(LOCTEXT("UViewModeUtils_VMI_MeshUVDensityAccuracy", "Mesh UV Density"));
+		}
+		// Colored according to shader complexity, including quad overdraw
+		else if (ViewModeIndex == VMI_ShaderComplexityWithQuadOverdraw)
+		{
+			ViewModeDisplayNames.Emplace(LOCTEXT("UViewModeUtils_VMI_ShaderComplexityWithQuadOverdraw", "Shader Complexity & Quads"));
+		}
+		// Colored according to the current HLOD index
+		else if (ViewModeIndex == VMI_HLODColoration)
+		{
+			ViewModeDisplayNames.Emplace(LOCTEXT("UViewModeUtils_VMI_HLODColoration", "Hierarchical LOD Coloration"));
+		}
+		// Group item for LOD and HLOD coloration*/
+		else if (ViewModeIndex == VMI_GroupLODColoration)
+		{
+			ViewModeDisplayNames.Emplace(LOCTEXT("UViewModeUtils_VMI_GroupLODColoration", "Group LOD Coloration"));
+		}
+		// Visualize the accuracy of the material texture scales used for texture streaming
+		else if (ViewModeIndex == VMI_MaterialTextureScaleAccuracy)
+		{
+			ViewModeDisplayNames.Emplace(LOCTEXT("UViewModeUtils_VMI_MaterialTextureScaleAccuracy", "Material Texture Scales"));
+		}
+		// Compare the required texture resolution to the actual resolution
+		else if (ViewModeIndex == VMI_RequiredTextureResolution)
+		{
+			ViewModeDisplayNames.Emplace(LOCTEXT("UViewModeUtils_VMI_RequiredTextureResolution", "Required Texture Resolution"));
+		}
+	
+		// Ray tracing modes
+		// Run path tracing pipeline
+		else if (ViewModeIndex == VMI_PathTracing)
+		{
+			ViewModeDisplayNames.Emplace(LOCTEXT("UViewModeUtils_VMI_PathTracing", "Path Tracing"));
+		}
+		// Run ray tracing debug pipeline 
+		else if (ViewModeIndex == VMI_RayTracingDebug)
+		{
+			ViewModeDisplayNames.Emplace(LOCTEXT("UViewModeUtils_VMI_RayTracingDebug", "Ray Tracing Debug"));
+		}
+
+		// VMI_Max
+		else if (ViewModeIndex == VMI_Max)
+		{
+			ViewModeDisplayNames.Emplace(LOCTEXT("UViewModeUtils_VMI_Max", "Max EViewModeIndex value"));
+		}
+		// VMI_Unknown
+		else if (ViewModeIndex == VMI_Unknown)
+		{
+			ViewModeDisplayNames.Emplace(LOCTEXT("UViewModeUtils_VMI_Unknown", "Unknown EViewModeIndex value"));
+		}
+		// Not considered case
+		else
+		{
+			ViewModeDisplayNames.Emplace(FText::GetEmpty());
+		}
+	}
+	// Return ViewModeIndex
+	return ViewModeDisplayNames;
+}
+
+const static TArray<FText> GViewModeDisplayNames = FillViewModeDisplayNames();
+
+FText UViewModeUtils::GetViewModeDisplayName(const EViewModeIndex ViewModeIndex)
+{
+	const FText ViewModeName = GViewModeDisplayNames[ViewModeIndex];
+	ensureMsgf(!ViewModeName.IsEmpty(), TEXT("Used an unknown value of EViewModeIndex (with value %d). Consider adding this new value in UViewModeUtils::GetViewModeName"), ViewModeIndex);
+	return ViewModeName;
+}
+
+#undef LOCTEXT_NAMESPACE
diff --git a/Engine/Source/Runtime/Engine/Public/BufferVisualizationData.h b/Engine/Source/Runtime/Engine/Public/BufferVisualizationData.h
--- a/Engine/Source/Runtime/Engine/Public/BufferVisualizationData.h
+++ b/Engine/Source/Runtime/Engine/Public/BufferVisualizationData.h
@@ -24,6 +24,12 @@ public:
 	/** Get a named material from the available material map **/
 	ENGINE_API UMaterial* GetMaterial(FName InMaterialName);
 
+	/** Get the display name of a named material from the available material map **/
+	ENGINE_API FText GetMaterialDisplayName(FName InMaterialName) const;
+
+	/** Get the default display name if no material is used and the Overview window is being used **/
+	ENGINE_API static FText GetMaterialDefaultDisplayName();
+
 	/** We cache the overview material name list from the console command here, so all dynamically created views can re-use the existing cached list of materials */
 	void SetCurrentOverviewMaterialNames(const FString& InNameList);
 	bool IsDifferentToCurrentOverviewMaterialNames(const FString& InNameList);
@@ -62,6 +68,9 @@ private:
 
 	/** The name->material mapping table */
 	TMaterialMap MaterialMap;
+
+	/** The UMaterial.name->material mapping table */
+	TMaterialMap MaterialMapFromMaterialName;
 	
 	/** List of material names to use in the buffer visualization overview */
 	FString CurrentOverviewMaterialNames;
diff --git a/Engine/Source/Runtime/Renderer/Private/PostProcess/PostProcessVisualizeBuffer.cpp b/Engine/Source/Runtime/Renderer/Private/PostProcess/PostProcessVisualizeBuffer.cpp
--- a/Engine/Source/Runtime/Renderer/Private/PostProcess/PostProcessVisualizeBuffer.cpp
+++ b/Engine/Source/Runtime/Renderer/Private/PostProcess/PostProcessVisualizeBuffer.cpp
@@ -9,6 +9,7 @@
 #include "ImageWriteTask.h"
 #include "ImageWriteQueue.h"
 #include "HighResScreenshot.h"
+#include "BufferVisualizationData.h"
 
 class FVisualizeBufferPS : public FGlobalShader
 {
@@ -414,7 +415,7 @@ FScreenPassTexture AddVisualizeGBufferOverviewPass(
 
 			FVisualizeBufferTile Tile;
 			Tile.Input = AddDownsamplePass(GraphBuilder, View, DownsampleInputs);
-			Tile.Label = MaterialName;
+			Tile.Label = GetBufferVisualizationData().GetMaterialDisplayName(FName(*MaterialName)).ToString();
 #if !(UE_BUILD_SHIPPING || UE_BUILD_TEST)
 			Tile.bSelected = 
 				PostProcessSettings.bBufferVisualizationOverviewTargetIsSelected &&
