IsModeActive now also considers UEdModes

diff --git a/Engine/Source/Editor/UnrealEd/Private/EditorModeManager.cpp b/Engine/Source/Editor/UnrealEd/Private/EditorModeManager.cpp
--- a/Engine/Source/Editor/UnrealEd/Private/EditorModeManager.cpp
+++ b/Engine/Source/Editor/UnrealEd/Private/EditorModeManager.cpp
@@ -1130,7 +1130,7 @@ bool FEditorModeTools::EnsureNotInMode(FEditorModeID ModeID, const FText& ErrorM
 	return bInASafeMode;
 }
 
-UEdMode* FEditorModeTools::GetActiveScriptableMode(FEditorModeID InID)
+UEdMode* FEditorModeTools::GetActiveScriptableMode(FEditorModeID InID) const
 {
 	for (auto& Mode : ActiveScriptableModes)
 	{
@@ -2034,7 +2034,15 @@ const FModeTool* FEditorModeTools::GetActiveTool( FEditorModeID InID ) const
 
 bool FEditorModeTools::IsModeActive( FEditorModeID InID ) const
 {
-	return GetActiveMode( InID ) != nullptr;
+	if (GetActiveMode(InID) != nullptr)
+	{
+		return true;
+	}
+	else if (GetActiveScriptableMode(InID) != nullptr)
+	{
+		return true;
+	}
+	return false;
 }
 
 bool FEditorModeTools::IsDefaultModeActive() const
diff --git a/Engine/Source/Editor/UnrealEd/Public/EditorModeManager.h b/Engine/Source/Editor/UnrealEd/Public/EditorModeManager.h
--- a/Engine/Source/Editor/UnrealEd/Public/EditorModeManager.h
+++ b/Engine/Source/Editor/UnrealEd/Public/EditorModeManager.h
@@ -125,7 +125,7 @@ public:
 	UE_DEPRECATED(4.24, "Use GetActiveMode instead.")
 	FEdMode* FindMode( FEditorModeID InID );
 
-	UEdMode* GetActiveScriptableMode(FEditorModeID InID);
+	UEdMode* GetActiveScriptableMode(FEditorModeID InID) const;
 
 	/**
 	 * Returns true if the current mode is not the specified ModeID.  Also optionally warns the user.
