Fixing default UI for UEdMode

diff --git a/Engine/Source/Editor/UnrealEd/Private/Toolkits/BaseToolkit.cpp b/Engine/Source/Editor/UnrealEd/Private/Toolkits/BaseToolkit.cpp
--- a/Engine/Source/Editor/UnrealEd/Private/Toolkits/BaseToolkit.cpp
+++ b/Engine/Source/Editor/UnrealEd/Private/Toolkits/BaseToolkit.cpp
@@ -355,13 +355,6 @@ UEdMode* FModeToolkit::GetScriptableEditorMode() const
 TSharedPtr<SWidget> FModeToolkit::GetInlineContent() const
 {
 	return SNew(SVerticalBox)
-		+SVerticalBox::Slot()
-		.HAlign(HAlign_Center)
-		.AutoHeight()
-		[
-			SNew(STextBlock)
-			.Text(GetBaseToolkitName())
-		]
 		+ SVerticalBox::Slot()
 		.AutoHeight()
 		[
