Moved editor mode toolbar functionality off of the base IToolkit and into FModeToolkit as it only applies to modes

diff --git a/Engine/Source/Editor/UnrealEd/Public/Toolkits/BaseToolkit.h b/Engine/Source/Editor/UnrealEd/Public/Toolkits/BaseToolkit.h
--- a/Engine/Source/Editor/UnrealEd/Public/Toolkits/BaseToolkit.h
+++ b/Engine/Source/Editor/UnrealEd/Public/Toolkits/BaseToolkit.h
@@ -131,8 +131,24 @@ public:
 	virtual class FEdMode* GetEditorMode() const override;
 	virtual UEdMode* GetScriptableEditorMode() const;
 	virtual TSharedPtr<SWidget> GetInlineContent() const override;
-	virtual void BuildToolPalette(FName PaletteName, class FToolBarBuilder& ToolbarBuilder) override;
-	virtual void OnToolPaletteChanged(FName PaletteName) override;
+
+	/** Returns the number of Mode specific tabs in the mode toolbar **/
+	virtual void GetToolPaletteNames(TArray<FName>& PaletteNames) const {}
+
+	/**
+	 * @param PaletteIndex      The index of the ToolPalette to build
+	 * @returns the name of Tool Palette
+	 **/
+	virtual FText GetToolPaletteDisplayName(FName Palette) const { return FText(); }
+
+	/**
+	 * @param PaletteIndex      The index of the ToolPalette to build
+	 * @param ToolbarBuilder    The builder to use for given PaletteIndex
+	**/
+	virtual void BuildToolPalette(FName Palette, class FToolBarBuilder& ToolbarBuilder);
+
+
+	virtual void OnToolPaletteChanged(FName PaletteName);
 
 	void SetModeSettingsObject(UObject* InSettingsObject);
 
diff --git a/Engine/Source/Editor/UnrealEd/Public/Toolkits/IToolkit.h b/Engine/Source/Editor/UnrealEd/Public/Toolkits/IToolkit.h
--- a/Engine/Source/Editor/UnrealEd/Public/Toolkits/IToolkit.h
+++ b/Engine/Source/Editor/UnrealEd/Public/Toolkits/IToolkit.h
@@ -136,22 +136,5 @@ public:
 
 	/* Returns true if this is the default generic asset editor used by multiple classes */
 	virtual bool IsSimpleAssetEditor() const { return false; }
-
-	/** Returns the number of Mode specific tabs in the mode toolbar **/ 
-	virtual void GetToolPaletteNames(TArray<FName>& PaletteNames) const {}
-
-	/** 
-	 * @param PaletteIndex      The index of the ToolPalette to build
-	 * @returns the name of Tool Palette 
-	 **/ 
-	virtual FText GetToolPaletteDisplayName(FName Palette) const { return FText(); }
-
-	/** 
-	 * @param PaletteIndex      The index of the ToolPalette to build
-	 * @param ToolbarBuilder    The builder to use for given PaletteIndex
-	**/ 
-	virtual void BuildToolPalette(FName Palette, class FToolBarBuilder& ToolbarBuilder) {}
-
-	virtual void OnToolPaletteChanged(FName NewActivePalette) {}
 };
 
