Make sure we log an error if there is an issue importing a collision model, this is a git hub PR
#5018

diff --git a/Engine/Source/Editor/UnrealEd/Private/Fbx/FbxStaticMeshImport.cpp b/Engine/Source/Editor/UnrealEd/Private/Fbx/FbxStaticMeshImport.cpp
--- a/Engine/Source/Editor/UnrealEd/Private/Fbx/FbxStaticMeshImport.cpp
+++ b/Engine/Source/Editor/UnrealEd/Private/Fbx/FbxStaticMeshImport.cpp
@@ -2402,7 +2402,7 @@ extern void AddConvexGeomFromVertices( const TArray<FVector>& Verts, FKAggregate
 extern void AddSphereGeomFromVerts(const TArray<FVector>& Verts, FKAggregateGeom* AggGeom, const TCHAR* ObjName);
 extern void AddCapsuleGeomFromVerts(const TArray<FVector>& Verts, FKAggregateGeom* AggGeom, const TCHAR* ObjName);
 extern void AddBoxGeomFromTris(const TArray<FPoly>& Tris, FKAggregateGeom* AggGeom, const TCHAR* ObjName);
-extern void DecomposeUCXMesh( const TArray<FVector>& CollisionVertices, const TArray<int32>& CollisionFaceIdx, UBodySetup* BodySetup );
+extern bool DecomposeUCXMesh( const TArray<FVector>& CollisionVertices, const TArray<int32>& CollisionFaceIdx, UBodySetup* BodySetup );
 
 bool UnFbx::FFbxImporter::ImportCollisionModels(UStaticMesh* StaticMesh, const FbxString& InNodeName)
 {
@@ -2515,7 +2515,10 @@ bool UnFbx::FFbxImporter::ImportCollisionModels(UStaticMesh* StaticMesh, const F
 		{
 			if( !ImportOptions->bOneConvexHullPerUCX )
 			{
-				DecomposeUCXMesh( CollisionVertices, CollisionFaceIdx, StaticMesh->BodySetup );
+				if (!DecomposeUCXMesh(CollisionVertices, CollisionFaceIdx, StaticMesh->BodySetup)) {
+					FString CollisionModelName = UTF8_TO_TCHAR(ModelName.Buffer());
+					AddTokenizedErrorMessage(FTokenizedMessage::Create(EMessageSeverity::Warning, FText::Format(LOCTEXT("Error_DecomposingCollisionMesh", "Could not decompose collision mesh [{0}]."), FText::FromString(CollisionModelName))), FFbxErrors::StaticMesh_BuildError);
+				}
 			}
 			else
 			{
diff --git a/Engine/Source/Editor/UnrealEd/Private/StaticMeshEdit.cpp b/Engine/Source/Editor/UnrealEd/Private/StaticMeshEdit.cpp
--- a/Engine/Source/Editor/UnrealEd/Private/StaticMeshEdit.cpp
+++ b/Engine/Source/Editor/UnrealEd/Private/StaticMeshEdit.cpp
@@ -274,7 +274,7 @@ private:
 	}
 };
 
-void DecomposeUCXMesh( const TArray<FVector>& CollisionVertices, const TArray<int32>& CollisionFaceIdx, UBodySetup* BodySetup )
+bool DecomposeUCXMesh( const TArray<FVector>& CollisionVertices, const TArray<int32>& CollisionFaceIdx, UBodySetup* BodySetup )
 {
 	// We keep no ref to this Model, so it will be GC'd at some point after the import.
 	auto TempModel = NewObject<UModel>();
@@ -294,6 +294,7 @@ void DecomposeUCXMesh( const TArray<FVector>& CollisionVertices, const TArray<in
 	ConnectivityBuilder.CreateConnectivityGroups();
 
 	// For each valid group build BSP and extract convex hulls
+	bool bSuccess = true;
 	for ( int32 i=0; i<ConnectivityBuilder.Groups.Num(); i++ )
 	{
 		const FMeshConnectivityGroup &Group = ConnectivityBuilder.Groups[ i ];
@@ -331,8 +332,11 @@ void DecomposeUCXMesh( const TArray<FVector>& CollisionVertices, const TArray<in
 
 		// Convert collision model into a collection of convex hulls.
 		// Generated convex hulls will be added to existing ones
-		BodySetup->CreateFromModel( TempModel, false );
+		bSuccess = BodySetup->CreateFromModel( TempModel, false ) && bSuccess;
 	}
+
+	// Could all meshes be properly decomposed?
+	return bSuccess;
 }
 
 /** 
diff --git a/Engine/Source/Runtime/Engine/Classes/PhysicsEngine/BodySetup.h b/Engine/Source/Runtime/Engine/Classes/PhysicsEngine/BodySetup.h
--- a/Engine/Source/Runtime/Engine/Classes/PhysicsEngine/BodySetup.h
+++ b/Engine/Source/Runtime/Engine/Classes/PhysicsEngine/BodySetup.h
@@ -390,7 +390,7 @@ public:
 	 * @param		bRemoveExisting			If true, clears any pre-existing collision
 	 * @return								true on success, false on failure because of vertex count overflow.
 	 */
-	ENGINE_API void CreateFromModel(class UModel* InModel, bool bRemoveExisting);
+	ENGINE_API bool CreateFromModel(class UModel* InModel, bool bRemoveExisting);
 
 	/**
 	 * Updates the tri mesh collision with new positions, and refits the BVH to match. 
diff --git a/Engine/Source/Runtime/Engine/Private/PhysicsEngine/PhysUtils.cpp b/Engine/Source/Runtime/Engine/Private/PhysicsEngine/PhysUtils.cpp
--- a/Engine/Source/Runtime/Engine/Private/PhysicsEngine/PhysUtils.cpp
+++ b/Engine/Source/Runtime/Engine/Private/PhysicsEngine/PhysUtils.cpp
@@ -107,7 +107,7 @@ static bool ModelToHullsWorker(FKAggregateGeom* outGeom,
 	return true;
 }
 
-void UBodySetup::CreateFromModel(UModel* InModel, bool bRemoveExisting)
+bool UBodySetup::CreateFromModel(UModel* InModel, bool bRemoveExisting)
 {
 	if ( bRemoveExisting )
 	{
@@ -116,10 +116,12 @@ void UBodySetup::CreateFromModel(UModel* InModel, bool bRemoveExisting)
 
 	const int32 NumHullsAtStart = AggGeom.ConvexElems.Num();
 	
+	bool bSuccess = false;
+
 	if( InModel != NULL && InModel->Nodes.Num() > 0)
 	{
 		TArray<FPlane>	Planes;
-		bool bSuccess = ModelToHullsWorker(&AggGeom, InModel, 0, InModel->RootOutside, Planes);
+		bSuccess = ModelToHullsWorker(&AggGeom, InModel, 0, InModel->RootOutside, Planes);
 		if ( !bSuccess )
 		{
 			// ModelToHullsWorker failed.  Clear out anything that may have been created.
@@ -129,6 +131,7 @@ void UBodySetup::CreateFromModel(UModel* InModel, bool bRemoveExisting)
 
 	// Create new GUID
 	InvalidatePhysicsData();
+	return bSuccess;
 }
 
 //////////////////////////////////////////////////////////////////////////
