Fix scriptable editor modes not updating the icon in the level editor toolbox tab when the mode is activated

diff --git a/Engine/Source/Editor/LevelEditor/Private/SLevelEditorToolBox.cpp b/Engine/Source/Editor/LevelEditor/Private/SLevelEditorToolBox.cpp
--- a/Engine/Source/Editor/LevelEditor/Private/SLevelEditorToolBox.cpp
+++ b/Engine/Source/Editor/LevelEditor/Private/SLevelEditorToolBox.cpp
@@ -168,14 +168,10 @@ void SLevelEditorToolBox::UpdateInlineContent(const TSharedPtr<IToolkit>& Toolki
 	TabName = NSLOCTEXT("LevelEditor", "ToolsTabTitle", "Toolbox");
 	TabIcon = FEditorStyle::Get().GetBrush("LevelEditor.Tabs.Modes");
 
-	if (Toolkit.IsValid() && Toolkit->GetEditorMode())
+	if (Toolkit.IsValid())
 	{
-		TabName = Toolkit->GetEditorMode()->GetModeInfo().Name;
-		TabIcon = Toolkit->GetEditorMode()->GetModeInfo().IconBrush.GetSmallIcon();
-	}
-	else if (Toolkit.IsValid())
-	{
-		TabName = Toolkit->GetToolkitName();
+		TabName = Toolkit->GetEditorModeDisplayName();
+		TabIcon = Toolkit->GetEditorModeIcon().GetSmallIcon();
 	}
 
 	TSharedPtr<SDockTab> ParentTabPinned = ParentTab.Pin();
diff --git a/Engine/Source/Editor/UnrealEd/Private/Toolkits/AssetEditorToolkit.cpp b/Engine/Source/Editor/UnrealEd/Private/Toolkits/AssetEditorToolkit.cpp
--- a/Engine/Source/Editor/UnrealEd/Private/Toolkits/AssetEditorToolkit.cpp
+++ b/Engine/Source/Editor/UnrealEd/Private/Toolkits/AssetEditorToolkit.cpp
@@ -400,7 +400,17 @@ FText FAssetEditorToolkit::GetToolTipTextForObject(const UObject* InObject)
 
 class FEdMode* FAssetEditorToolkit::GetEditorMode() const
 {
-	return NULL;
+	return nullptr;
+}
+
+FText FAssetEditorToolkit::GetEditorModeDisplayName() const
+{
+	return FText::GetEmpty();
+}
+
+FSlateIcon FAssetEditorToolkit::GetEditorModeIcon() const
+{
+	return FSlateIcon();
 }
 
 const TArray< UObject* >* FAssetEditorToolkit::GetObjectsCurrentlyBeingEdited() const
diff --git a/Engine/Source/Editor/UnrealEd/Private/Toolkits/BaseToolkit.cpp b/Engine/Source/Editor/UnrealEd/Private/Toolkits/BaseToolkit.cpp
--- a/Engine/Source/Editor/UnrealEd/Private/Toolkits/BaseToolkit.cpp
+++ b/Engine/Source/Editor/UnrealEd/Private/Toolkits/BaseToolkit.cpp
@@ -347,6 +347,34 @@ class FEdMode* FModeToolkit::GetEditorMode() const
 	return nullptr; 
 }
 
+FText FModeToolkit::GetEditorModeDisplayName() const
+{
+	if (FEdMode* EdMode = GetEditorMode())
+	{
+		return EdMode->GetModeInfo().Name;
+	}
+	else if (UEdMode* ScriptableMode = GetScriptableEditorMode())
+	{
+		return ScriptableMode->GetModeInfo().Name;
+	}
+
+	return FText::GetEmpty();
+}
+
+FSlateIcon FModeToolkit::GetEditorModeIcon() const
+{
+	if (FEdMode* EdMode = GetEditorMode())
+	{
+		return EdMode->GetModeInfo().IconBrush;
+	}
+	else if (UEdMode* ScriptableMode = GetScriptableEditorMode())
+	{
+		return ScriptableMode->GetModeInfo().IconBrush;
+	}
+
+	return FSlateIcon();
+}
+
 UEdMode* FModeToolkit::GetScriptableEditorMode() const
 {
 	return nullptr;
diff --git a/Engine/Source/Editor/UnrealEd/Public/Toolkits/AssetEditorToolkit.h b/Engine/Source/Editor/UnrealEd/Public/Toolkits/AssetEditorToolkit.h
--- a/Engine/Source/Editor/UnrealEd/Public/Toolkits/AssetEditorToolkit.h
+++ b/Engine/Source/Editor/UnrealEd/Public/Toolkits/AssetEditorToolkit.h
@@ -127,6 +127,8 @@ public:
 	virtual FText GetToolkitToolTipText() const override;
 	virtual FString GetWorldCentricTabPrefix() const override = 0;	// Must implement in derived class!
 	virtual class FEdMode* GetEditorMode() const override;
+	virtual FText GetEditorModeDisplayName() const override;
+	virtual FSlateIcon GetEditorModeIcon() const override;
 
 	/** IAssetEditorInstance interface */
 	virtual FName GetEditorName() const override;
diff --git a/Engine/Source/Editor/UnrealEd/Public/Toolkits/BaseToolkit.h b/Engine/Source/Editor/UnrealEd/Public/Toolkits/BaseToolkit.h
--- a/Engine/Source/Editor/UnrealEd/Public/Toolkits/BaseToolkit.h
+++ b/Engine/Source/Editor/UnrealEd/Public/Toolkits/BaseToolkit.h
@@ -129,6 +129,9 @@ public:
 	virtual const TArray<UObject*>* GetObjectsCurrentlyBeingEdited() const override;
 	virtual FLinearColor GetWorldCentricTabColorScale() const override;
 	virtual class FEdMode* GetEditorMode() const override;
+	virtual FText GetEditorModeDisplayName() const override;
+	virtual FSlateIcon GetEditorModeIcon() const override;
+
 	virtual UEdMode* GetScriptableEditorMode() const;
 	virtual TSharedPtr<SWidget> GetInlineContent() const override;
 
diff --git a/Engine/Source/Editor/UnrealEd/Public/Toolkits/IToolkit.h b/Engine/Source/Editor/UnrealEd/Public/Toolkits/IToolkit.h
--- a/Engine/Source/Editor/UnrealEd/Public/Toolkits/IToolkit.h
+++ b/Engine/Source/Editor/UnrealEd/Public/Toolkits/IToolkit.h
@@ -125,6 +125,12 @@ public:
 	/** @returns the editor mode this toolkit is used for, or null if not relevant. */
 	virtual FEdMode* GetEditorMode() const = 0;
 
+	/** @return the editor mode display name if there is an editor mode associated with this toolkit.  An empty FText is returned otherwise */
+	virtual FText GetEditorModeDisplayName() const = 0;
+
+	/** @return the editor mode icon if there is an editor mode associated with this toolkit.  An empty icon is returned otherwise */
+	virtual FSlateIcon GetEditorModeIcon() const = 0;
+
 	/** @returns the inline content that this toolkit returns if it is an editor mode */
 	virtual TSharedPtr<SWidget> GetInlineContent() const = 0;
 
@@ -138,3 +144,4 @@ public:
 	virtual bool IsSimpleAssetEditor() const { return false; }
 };
 
+
