Fix incremental cis

diff --git a/Engine/Source/Editor/UnrealEd/Private/Toolkits/BaseToolkit.cpp b/Engine/Source/Editor/UnrealEd/Private/Toolkits/BaseToolkit.cpp
--- a/Engine/Source/Editor/UnrealEd/Private/Toolkits/BaseToolkit.cpp
+++ b/Engine/Source/Editor/UnrealEd/Private/Toolkits/BaseToolkit.cpp
@@ -12,6 +12,7 @@
 #include "InteractiveTool.h"
 #include "Tools/UEdMode.h"
 #include "Framework/MultiBox/MultiBoxBuilder.h"
+#include "EdMode.h"
 
 #define LOCTEXT_NAMESPACE "BaseToolkit"
 
