CIS: Fixed default generated copy constructor emitting a deprecation warning.

#rnx

diff --git a/Engine/Source/Editor/UnrealEd/Public/PlayInEditorDataTypes.h b/Engine/Source/Editor/UnrealEd/Public/PlayInEditorDataTypes.h
--- a/Engine/Source/Editor/UnrealEd/Public/PlayInEditorDataTypes.h
+++ b/Engine/Source/Editor/UnrealEd/Public/PlayInEditorDataTypes.h
@@ -77,10 +77,16 @@ struct FPieLoginStruct
 	int32 PIEInstanceIndex;
 
 	PRAGMA_DISABLE_DEPRECATION_WARNINGS
-	FPieLoginStruct& operator = (const FPieLoginStruct &) = default;
-	FPieLoginStruct& operator = (FPieLoginStruct &&) = default;
-	FPieLoginStruct(const FPieLoginStruct&) = default; // Copy
-	FPieLoginStruct(FPieLoginStruct&&) = default; // Move
+	FPieLoginStruct(const FPieLoginStruct& InOther)
+	{
+		WorldContextHandle = InOther.WorldContextHandle;
+		SettingsIndex = InOther.SettingsIndex;
+		NextX = InOther.NextX;
+		NextY = InOther.NextY;
+		PIEStartTime = InOther.PIEStartTime;
+		GameInstancePIEParameters = InOther.GameInstancePIEParameters;
+		PIEInstanceIndex = InOther.PIEInstanceIndex;
+	}
 	PRAGMA_ENABLE_DEPRECATION_WARNINGS
 
 	FPieLoginStruct()
