Add a configurable threshold for morph target position delta. Also make the other threshold available in the skeletalmesh build settings

diff --git a/Engine/Config/BaseEditorPerProjectUserSettings.ini b/Engine/Config/BaseEditorPerProjectUserSettings.ini
--- a/Engine/Config/BaseEditorPerProjectUserSettings.ini
+++ b/Engine/Config/BaseEditorPerProjectUserSettings.ini
@@ -686,6 +686,7 @@ bImportMorphTargets=False
 ThresholdPosition=0.00002
 ThresholdTangentNormal=0.00002
 ThresholdUV=0.0009765625
+MorphThresholdPosition=0.015
 
 [/Script/UnrealEd.FbxAnimSequenceImportData]
 bImportMeshesInBoneHierarchy=True
diff --git a/Engine/Plugins/Tests/EditorTests/Source/EditorTests/Private/UnrealEd/FbxAutomationTests.cpp b/Engine/Plugins/Tests/EditorTests/Source/EditorTests/Private/UnrealEd/FbxAutomationTests.cpp
--- a/Engine/Plugins/Tests/EditorTests/Source/EditorTests/Private/UnrealEd/FbxAutomationTests.cpp
+++ b/Engine/Plugins/Tests/EditorTests/Source/EditorTests/Private/UnrealEd/FbxAutomationTests.cpp
@@ -326,6 +326,7 @@ bool FFbxImportAssetsAutomationTest::RunTest(const FString& Parameters)
 					ImportData->bImportMorphTargets = TestPlan->ImportUI->SkeletalMeshImportData->bImportMorphTargets;
 					ImportData->ThresholdPosition = TestPlan->ImportUI->SkeletalMeshImportData->ThresholdPosition;
 					ImportData->ThresholdTangentNormal = TestPlan->ImportUI->SkeletalMeshImportData->ThresholdTangentNormal;
+					ImportData->MorphThresholdPosition = TestPlan->ImportUI->SkeletalMeshImportData->MorphThresholdPosition;
 					ImportData->ThresholdUV = TestPlan->ImportUI->SkeletalMeshImportData->ThresholdUV;
 					ImportData->bPreserveSmoothingGroups = TestPlan->ImportUI->SkeletalMeshImportData->bPreserveSmoothingGroups;
 					ImportData->bUpdateSkeletonReferencePose = TestPlan->ImportUI->SkeletalMeshImportData->bUpdateSkeletonReferencePose;
@@ -423,6 +424,7 @@ bool FFbxImportAssetsAutomationTest::RunTest(const FString& Parameters)
 					ImportData->ThresholdPosition = TestPlan->ImportUI->SkeletalMeshImportData->ThresholdPosition;
 					ImportData->ThresholdTangentNormal = TestPlan->ImportUI->SkeletalMeshImportData->ThresholdTangentNormal;
 					ImportData->ThresholdUV = TestPlan->ImportUI->SkeletalMeshImportData->ThresholdUV;
+					ImportData->MorphThresholdPosition = TestPlan->ImportUI->SkeletalMeshImportData->MorphThresholdPosition;
 					ImportData->bPreserveSmoothingGroups = TestPlan->ImportUI->SkeletalMeshImportData->bPreserveSmoothingGroups;
 					ImportData->bUpdateSkeletonReferencePose = TestPlan->ImportUI->SkeletalMeshImportData->bUpdateSkeletonReferencePose;
 					ImportData->bUseT0AsRefPose = TestPlan->ImportUI->SkeletalMeshImportData->bUseT0AsRefPose;
diff --git a/Engine/Source/Developer/MeshBuilder/Private/SkeletalMeshBuilder.cpp b/Engine/Source/Developer/MeshBuilder/Private/SkeletalMeshBuilder.cpp
--- a/Engine/Source/Developer/MeshBuilder/Private/SkeletalMeshBuilder.cpp
+++ b/Engine/Source/Developer/MeshBuilder/Private/SkeletalMeshBuilder.cpp
@@ -122,7 +122,7 @@ bool FSkeletalMeshBuilder::Build(USkeletalMesh* SkeletalMesh, const int32 LODInd
 		SlowTask.EnterProgressFrame(1.0f, NSLOCTEXT("SkeltalMeshBuilder", "RebuildMorphTarget", "Rebuilding morph targets..."));
 		if (SkeletalMeshImportData.MorphTargetNames.Num() > 0)
 		{
-			FLODUtilities::BuildMorphTargets(SkeletalMesh, SkeletalMeshImportData, LODIndex, !Options.bComputeNormals, !Options.bComputeTangents, Options.bUseMikkTSpace);
+			FLODUtilities::BuildMorphTargets(SkeletalMesh, SkeletalMeshImportData, LODIndex, !Options.bComputeNormals, !Options.bComputeTangents, Options.bUseMikkTSpace, Options.OverlappingThresholds);
 		}
 
 		//Re-apply the alternate skinning it must be after the inline reduction
diff --git a/Engine/Source/Developer/MeshUtilities/Public/MeshUtilities.h b/Engine/Source/Developer/MeshUtilities/Public/MeshUtilities.h
--- a/Engine/Source/Developer/MeshUtilities/Public/MeshUtilities.h
+++ b/Engine/Source/Developer/MeshUtilities/Public/MeshUtilities.h
@@ -234,6 +234,7 @@ public:
 			OverlappingThresholds.ThresholdPosition = SkeletalMeshBuildSettings.ThresholdPosition;
 			OverlappingThresholds.ThresholdTangentNormal = SkeletalMeshBuildSettings.ThresholdTangentNormal;
 			OverlappingThresholds.ThresholdUV = SkeletalMeshBuildSettings.ThresholdUV;
+			OverlappingThresholds.MorphThresholdPosition = SkeletalMeshBuildSettings.MorphThresholdPosition;
 			bComputeNormals = SkeletalMeshBuildSettings.bRecomputeNormals;
 			bComputeTangents = SkeletalMeshBuildSettings.bRecomputeTangents;
 			bUseMikkTSpace = SkeletalMeshBuildSettings.bUseMikkTSpace;
diff --git a/Engine/Source/Developer/SkeletalMeshUtilitiesCommon/Private/LODUtilities.cpp b/Engine/Source/Developer/SkeletalMeshUtilitiesCommon/Private/LODUtilities.cpp
--- a/Engine/Source/Developer/SkeletalMeshUtilitiesCommon/Private/LODUtilities.cpp
+++ b/Engine/Source/Developer/SkeletalMeshUtilitiesCommon/Private/LODUtilities.cpp
@@ -2104,7 +2104,7 @@ public:
 		TArray< FMorphTargetDelta >& InMorphDeltas, TArray<uint32>& InBaseIndexData, TArray< uint32 >& InBaseWedgePointIndices,
 		TMap<uint32, uint32>& InWedgePointToVertexIndexMap, const FOverlappingCorners& InOverlappingCorners,
 		const TSet<uint32> InModifiedPoints, const TMultiMap< int32, int32 >& InWedgeToFaces, const FMeshDataBundle& InMeshDataBundle, const TArray<FVector>& InTangentZ,
-		bool InShouldImportNormals, bool InShouldImportTangents, bool InbUseMikkTSpace)
+		bool InShouldImportNormals, bool InShouldImportTangents, bool InbUseMikkTSpace, const FOverlappingThresholds InThresholds)
 		: LODModel(InLODModel)
 		, RefSkeleton(InRefSkeleton)
 		, BaseImportData(InBaseImportData)
@@ -2122,6 +2122,7 @@ public:
 		, ShouldImportNormals(InShouldImportNormals)
 		, ShouldImportTangents(InShouldImportTangents)
 		, bUseMikkTSpace(InbUseMikkTSpace)
+		, Thresholds(InThresholds)
 	{
 		MeshUtilities = &FModuleManager::Get().LoadModuleChecked<IMeshUtilities>("MeshUtilities");
 	}
@@ -2250,7 +2251,7 @@ public:
 						}
 
 						// check if position actually changed much
-						if (PositionDelta.SizeSquared() > FMath::Square(THRESH_POINTS_ARE_NEAR) ||
+						if (PositionDelta.SizeSquared() > FMath::Square(Thresholds.MorphThresholdPosition) ||
 							// since we can't get imported morphtarget normal from FBX
 							// we can't compare normal unless it's calculated
 							// this is special flag to ignore normal diff
@@ -2326,9 +2327,10 @@ private:
 	bool ShouldImportNormals;
 	bool ShouldImportTangents;
 	bool bUseMikkTSpace;
+	const FOverlappingThresholds Thresholds;
 };
 
-void FLODUtilities::BuildMorphTargets(USkeletalMesh* BaseSkelMesh, FSkeletalMeshImportData &BaseImportData, int32 LODIndex, bool ShouldImportNormals, bool ShouldImportTangents, bool bUseMikkTSpace)
+void FLODUtilities::BuildMorphTargets(USkeletalMesh* BaseSkelMesh, FSkeletalMeshImportData &BaseImportData, int32 LODIndex, bool ShouldImportNormals, bool ShouldImportTangents, bool bUseMikkTSpace, const FOverlappingThresholds& Thresholds)
 {
 	bool bComputeNormals = !ShouldImportNormals || !BaseImportData.bHasNormals;
 	bool bComputeTangents = !ShouldImportTangents || !BaseImportData.bHasTangents;
@@ -2456,7 +2458,7 @@ void FLODUtilities::BuildMorphTargets(USkeletalMesh* BaseSkelMesh, FSkeletalMesh
 
 			FAsyncTask<FAsyncImportMorphTargetWork>* NewWork = new FAsyncTask<FAsyncImportMorphTargetWork>(&BaseLODModel, BaseSkelMesh->RefSkeleton, BaseImportData,
 				MoveTemp(ShapeImportData.Points), *Deltas, BaseIndexData, BaseWedgePointIndices, WedgePointToVertexIndexMap, OverlappingVertices, MoveTemp(ModifiedPoints), WedgeToFaces, MeshDataBundle, TangentZ,
-				ShouldImportNormals, ShouldImportTangents, bUseMikkTSpace);
+				ShouldImportNormals, ShouldImportTangents, bUseMikkTSpace, Thresholds);
 			PendingWork.Add(NewWork);
 
 			NewWork->StartBackgroundTask(GLargeThreadPool);
@@ -2491,7 +2493,7 @@ void FLODUtilities::BuildMorphTargets(USkeletalMesh* BaseSkelMesh, FSkeletalMesh
 		GWarn->StatusUpdate(Index + 1, MorphTargets.Num(), FText::Format(LOCTEXT("BuildingMorphTargetRenderDataStatus", "Building Morph Target Render Data: {NumCompleted} of {NumTasks}"), Args));
 
 		UMorphTarget* MorphTarget = MorphTargets[Index];
-		MorphTarget->PopulateDeltas(*Results[Index], LODIndex, BaseLODModel.Sections, ShouldImportNormals == false);
+		MorphTarget->PopulateDeltas(*Results[Index], LODIndex, BaseLODModel.Sections, ShouldImportNormals == false, false, Thresholds.MorphThresholdPosition);
 
 		// register does mark package as dirty
 		if (MorphTarget->HasValidData())
diff --git a/Engine/Source/Developer/SkeletalMeshUtilitiesCommon/Public/LODUtilities.h b/Engine/Source/Developer/SkeletalMeshUtilitiesCommon/Public/LODUtilities.h
--- a/Engine/Source/Developer/SkeletalMeshUtilitiesCommon/Public/LODUtilities.h
+++ b/Engine/Source/Developer/SkeletalMeshUtilitiesCommon/Public/LODUtilities.h
@@ -96,7 +96,7 @@ public:
 	/*
 	 * Build the morph targets for the specified LOD. The function use the Morph target data stored in the FSkeletalMeshImportData ImportData structure
 	 */
-	static void BuildMorphTargets(USkeletalMesh* SkeletalMesh, class FSkeletalMeshImportData &ImportData, int32 LODIndex, bool ShouldImportNormals, bool ShouldImportTangents, bool bUseMikkTSpace);
+	static void BuildMorphTargets(USkeletalMesh* SkeletalMesh, class FSkeletalMeshImportData &ImportData, int32 LODIndex, bool ShouldImportNormals, bool ShouldImportTangents, bool bUseMikkTSpace, const FOverlappingThresholds& Thresholds);
 
 	/**
 	 *	This function apply the skinning weights from asource skeletal mesh to the destination skeletal mesh.
diff --git a/Engine/Source/Editor/Persona/Private/AnimationEditorPreviewScene.h b/Engine/Source/Editor/Persona/Private/AnimationEditorPreviewScene.h
--- a/Engine/Source/Editor/Persona/Private/AnimationEditorPreviewScene.h
+++ b/Engine/Source/Editor/Persona/Private/AnimationEditorPreviewScene.h
@@ -76,6 +76,24 @@ public:
 		OnLODChanged.RemoveAll(Thing);
 	}
 
+	virtual void RegisterOnMorphTargetsChanged(const FSimpleDelegate& Delegate) override
+	{
+		OnMorphTargetsChanged.Add(Delegate);
+	}
+
+	virtual void UnregisterOnMorphTargetsChanged(void* Thing) override
+	{
+		OnMorphTargetsChanged.RemoveAll(Thing);
+	}
+
+	virtual void BroadcastOnMorphTargetsChanged() override
+	{
+		if (OnMorphTargetsChanged.IsBound())
+		{
+			OnMorphTargetsChanged.Broadcast();
+		}
+	}
+
 	virtual void RegisterOnInvalidateViews(const FSimpleDelegate& Delegate) override
 	{
 		OnInvalidateViews.Add(Delegate);
@@ -383,6 +401,9 @@ private:
 	/** LOD changed delegate */
 	FSimpleMulticastDelegate OnLODChanged;
 
+	/** Morph target changed delegate */
+	FSimpleMulticastDelegate OnMorphTargetsChanged;
+
 	/** View invalidation delegate */
 	FSimpleMulticastDelegate OnInvalidateViews;
 
diff --git a/Engine/Source/Editor/Persona/Private/PersonaMeshDetails.cpp b/Engine/Source/Editor/Persona/Private/PersonaMeshDetails.cpp
--- a/Engine/Source/Editor/Persona/Private/PersonaMeshDetails.cpp
+++ b/Engine/Source/Editor/Persona/Private/PersonaMeshDetails.cpp
@@ -1421,6 +1421,49 @@ void FSkeletalMeshBuildSettingsLayout::GenerateChildContent(IDetailChildrenBuild
 		];
 	}
 
+	{
+		FDetailWidgetRow& PositionThreshold = AddFloatRow(ChildrenBuilder,
+			LOCTEXT("PositionThreshold_Row", "ThresholdPositionRow"),
+			LOCTEXT("PositionThreshold", "Threshold Position"),
+			LOCTEXT("PositionThreshold_ToolTip", "Threshold to compare vertex position equality."),
+			0.0f,
+			10000.0f,
+			FGetFloatDelegate::CreateRaw(this, &FSkeletalMeshBuildSettingsLayout::GetThresholdPosition),
+			FSetFloatDelegate::CreateRaw(this, &FSkeletalMeshBuildSettingsLayout::SetThresholdPosition));
+	}
+
+	{
+		FDetailWidgetRow& TangentNormalThreshold = AddFloatRow(ChildrenBuilder,
+			LOCTEXT("TangentNormalThreshold_Row", "ThresholdTangentNormalRow"),
+			LOCTEXT("TangentNormalThreshold", "Threshold Tangent and Normal"),
+			LOCTEXT("TangentNormalThreshold_ToolTip", "Threshold to compare normal, tangent or bi-normal equality."),
+			0.0f,
+			1.0f,
+			FGetFloatDelegate::CreateRaw(this, &FSkeletalMeshBuildSettingsLayout::GetThresholdTangentNormal),
+			FSetFloatDelegate::CreateRaw(this, &FSkeletalMeshBuildSettingsLayout::SetThresholdTangentNormal));
+	}
+
+	{
+		FDetailWidgetRow& UVThreshold = AddFloatRow(ChildrenBuilder,
+			LOCTEXT("UVThreshold_Row", "ThresholdUVRow"),
+			LOCTEXT("UVThreshold", "Threshold UV"),
+			LOCTEXT("UVThreshold_ToolTip", "Threshold to compare vertex UV equality."),
+			0.0f,
+			1.0f,
+			FGetFloatDelegate::CreateRaw(this, &FSkeletalMeshBuildSettingsLayout::GetThresholdUV),
+			FSetFloatDelegate::CreateRaw(this, &FSkeletalMeshBuildSettingsLayout::SetThresholdUV));
+	}
+
+	{
+		FDetailWidgetRow& MorphPositionThreshold = AddFloatRow(ChildrenBuilder,
+			LOCTEXT("MorphPositionThreshold_Row", "MorphThresholdPositionRow"),
+			LOCTEXT("MorphPositionThreshold", "Morph Threshold Position"),
+			LOCTEXT("MorphPositionThreshold_ToolTip", "Threshold to compare vertex position equality when computing morph target deltas."),
+			0.0f,
+			1.0f,
+			FGetFloatDelegate::CreateRaw(this, &FSkeletalMeshBuildSettingsLayout::GetMorphThresholdPosition),
+			FSetFloatDelegate::CreateRaw(this, &FSkeletalMeshBuildSettingsLayout::SetMorphThresholdPosition));
+	}
 }
 
 bool FSkeletalMeshBuildSettingsLayout::IsBuildEnabled() const
@@ -1428,6 +1471,141 @@ bool FSkeletalMeshBuildSettingsLayout::IsBuildEnabled() const
 	return IsBuildSettingsEnabledDelegate.Execute(LODIndex);
 }
 
+FDetailWidgetRow& FSkeletalMeshBuildSettingsLayout::AddFloatRow(IDetailChildrenBuilder& ChildrenBuilder, const FText RowTitleText, const FText RowNameContentText, const FText RowNameContentTootlipText, const float MinSliderValue, const float MaxSliderValue, FGetFloatDelegate GetterDelegate, FSetFloatDelegate SetterDelegate)
+{
+	int32 SliderDataIndex = SliderStateDataArray.Num();
+	FSliderStateData& SliderData = SliderStateDataArray.AddDefaulted_GetRef();
+	SliderData.bSliderActiveMode = false;
+
+	auto BeginSliderMovementHelperFunc = [GetterDelegate, SliderDataIndex, this]()
+	{
+		check(SliderStateDataArray.IsValidIndex(SliderDataIndex));
+		SliderStateDataArray[SliderDataIndex].bSliderActiveMode = true;
+		SliderStateDataArray[SliderDataIndex].MovementValueFloat = GetterDelegate.IsBound() ? GetterDelegate.Execute() : 0.0f;
+	};
+
+	auto EndSliderMovementHelperFunc = [SetterDelegate, SliderDataIndex, this](float Value)
+	{
+		check(SliderStateDataArray.IsValidIndex(SliderDataIndex));
+		SliderStateDataArray[SliderDataIndex].bSliderActiveMode = false;
+		SliderStateDataArray[SliderDataIndex].MovementValueFloat = 0.0f;
+		SetterDelegate.ExecuteIfBound(Value);
+	};
+
+	auto SetValueHelperFunc = [SetterDelegate, SliderDataIndex, this](float Value)
+	{
+		check(SliderStateDataArray.IsValidIndex(SliderDataIndex));
+		if (SliderStateDataArray[SliderDataIndex].bSliderActiveMode)
+		{
+			SliderStateDataArray[SliderDataIndex].MovementValueFloat = Value;
+		}
+		else
+		{
+			SetterDelegate.ExecuteIfBound(Value);
+		}
+	};
+
+	auto GetValueHelperFunc = [GetterDelegate, SliderDataIndex, this]()
+	{
+		check(SliderStateDataArray.IsValidIndex(SliderDataIndex));
+		if (SliderStateDataArray[SliderDataIndex].bSliderActiveMode)
+		{
+			return SliderStateDataArray[SliderDataIndex].MovementValueFloat;
+		}
+		return GetterDelegate.IsBound() ? GetterDelegate.Execute() : 0.0f;
+	};
+
+	FDetailWidgetRow& Row = ChildrenBuilder.AddCustomRow(RowTitleText)
+	.NameContent()
+	[
+		SNew(STextBlock)
+		.Font(IDetailLayoutBuilder::GetDetailFont())
+		.Text(RowNameContentText)
+		.ToolTipText(RowNameContentTootlipText)
+	]
+	.ValueContent()
+	[
+		SNew(SSpinBox<float>)
+		.Font(IDetailLayoutBuilder::GetDetailFont())
+		.MinValue(MinSliderValue)
+		.MaxValue(MaxSliderValue)
+		.Value_Lambda(GetValueHelperFunc)
+		.OnValueChanged_Lambda(SetValueHelperFunc)
+		.OnBeginSliderMovement_Lambda(BeginSliderMovementHelperFunc)
+		.OnEndSliderMovement_Lambda(EndSliderMovementHelperFunc)
+		.IsEnabled(this, &FSkeletalMeshBuildSettingsLayout::IsBuildEnabled)
+	];
+	return Row;
+}
+
+float FSkeletalMeshBuildSettingsLayout::GetThresholdPosition() const
+{
+	return BuildSettings.ThresholdPosition;
+}
+
+void FSkeletalMeshBuildSettingsLayout::SetThresholdPosition(float Value)
+{
+	if (BuildSettings.ThresholdPosition != Value)
+	{
+		FText TransactionText = FText::Format(LOCTEXT("PersonaSetThresholdPositionLOD", "LOD{0} build settings: threshold for position changed"), LODIndex);
+		FScopedTransaction Transaction(TransactionText);
+		ModifyMeshLODSettingsDelegate.ExecuteIfBound(LODIndex);
+
+		BuildSettings.ThresholdPosition = Value;
+	}
+}
+
+float FSkeletalMeshBuildSettingsLayout::GetThresholdTangentNormal() const
+{
+	return BuildSettings.ThresholdTangentNormal;
+}
+
+void FSkeletalMeshBuildSettingsLayout::SetThresholdTangentNormal(float Value)
+{
+	if (BuildSettings.ThresholdTangentNormal != Value)
+	{
+		FText TransactionText = FText::Format(LOCTEXT("PersonaSetThresholdTangentNormalLOD", "LOD{0} build settings: threshold for tangent and normal changed"), LODIndex);
+		FScopedTransaction Transaction(TransactionText);
+		ModifyMeshLODSettingsDelegate.ExecuteIfBound(LODIndex);
+
+		BuildSettings.ThresholdTangentNormal = Value;
+	}
+}
+
+float FSkeletalMeshBuildSettingsLayout::GetThresholdUV() const
+{
+	return BuildSettings.ThresholdUV;
+}
+
+void FSkeletalMeshBuildSettingsLayout::SetThresholdUV(float Value)
+{
+	if (BuildSettings.ThresholdUV != Value)
+	{
+		FText TransactionText = FText::Format(LOCTEXT("PersonaSetThresholdUVLOD", "LOD{0} build settings: threshold for UVs changed"), LODIndex);
+		FScopedTransaction Transaction(TransactionText);
+		ModifyMeshLODSettingsDelegate.ExecuteIfBound(LODIndex);
+
+		BuildSettings.ThresholdUV = Value;
+	}
+}
+
+float FSkeletalMeshBuildSettingsLayout::GetMorphThresholdPosition() const
+{
+	return BuildSettings.MorphThresholdPosition;
+}
+
+void FSkeletalMeshBuildSettingsLayout::SetMorphThresholdPosition(float Value)
+{
+	if (BuildSettings.MorphThresholdPosition != Value)
+	{
+		FText TransactionText = FText::Format(LOCTEXT("PersonaSetMorphThresholdPositionLOD", "LOD{0} build settings: threshold for morph position changed"), LODIndex);
+		FScopedTransaction Transaction(TransactionText);
+		ModifyMeshLODSettingsDelegate.ExecuteIfBound(LODIndex);
+
+		BuildSettings.MorphThresholdPosition = Value;
+	}
+}
+
 ECheckBoxState FSkeletalMeshBuildSettingsLayout::ShouldRecomputeNormals() const
 {
 	return BuildSettings.bRecomputeNormals ? ECheckBoxState::Checked : ECheckBoxState::Unchecked;
@@ -2962,6 +3140,10 @@ FReply FPersonaMeshDetails::ApplyLODChanges(int32 LODIndex)
 							SKImportData->NormalGenerationMethod = LODInfo->BuildSettings.bUseMikkTSpace ? EFBXNormalGenerationMethod::MikkTSpace : EFBXNormalGenerationMethod::BuiltIn;
 						}
 						SKImportData->bComputeWeightedNormals = LODInfo->BuildSettings.bComputeWeightedNormals;
+						SKImportData->ThresholdPosition = LODInfo->BuildSettings.ThresholdPosition;
+						SKImportData->ThresholdTangentNormal = LODInfo->BuildSettings.ThresholdTangentNormal;
+						SKImportData->ThresholdUV = LODInfo->BuildSettings.ThresholdUV;
+						SKImportData->MorphThresholdPosition = LODInfo->BuildSettings.MorphThresholdPosition;
 					}
 				}
 			}
@@ -2975,6 +3157,11 @@ FReply FPersonaMeshDetails::ApplyLODChanges(int32 LODIndex)
 		SkelMesh->MarkPackageDirty();
 	}
 	MeshDetailLayout->ForceRefreshDetails();
+	if (SkelMesh->MorphTargets.Num() > 0)
+	{
+		GetPersonaToolkit()->GetPreviewScene()->BroadcastOnMorphTargetsChanged();
+	}
+	
 	return FReply::Handled();
 }
 
@@ -3226,6 +3413,12 @@ void FPersonaMeshDetails::ApplyChanges()
 
 		//PostEditChange will be call when going out of scope
 	}
+	//Update the morph target list since we have rebuild the asset
+	if (SkelMesh->MorphTargets.Num() > 0)
+	{
+		GetPersonaToolkit()->GetPreviewScene()->BroadcastOnMorphTargetsChanged();
+	}
+
 	MeshDetailLayout->ForceRefreshDetails();
 }
 
diff --git a/Engine/Source/Editor/Persona/Private/PersonaMeshDetails.h b/Engine/Source/Editor/Persona/Private/PersonaMeshDetails.h
--- a/Engine/Source/Editor/Persona/Private/PersonaMeshDetails.h
+++ b/Engine/Source/Editor/Persona/Private/PersonaMeshDetails.h
@@ -228,6 +228,9 @@ public:
 	FSkeletalMeshBuildSettingsLayout(FSkeletalMeshBuildSettings& InBuildSettings, int32 InLODIndex, FIsLODSettingsEnabledDelegate InIsBuildSettingsEnabledDelegate, FModifyMeshLODSettingsDelegate InModifyMeshLODSettingsDelegate);
 	virtual ~FSkeletalMeshBuildSettingsLayout() {};
 
+	DECLARE_DELEGATE_RetVal(float, FGetFloatDelegate);
+	DECLARE_DELEGATE_OneParam(FSetFloatDelegate, float);
+
 private:
 	/** IDetailCustomNodeBuilder Interface*/
 	virtual void SetOnRebuildChildren(FSimpleDelegate InOnRegenerateChildren) override {}
@@ -239,6 +242,22 @@ private:
 	virtual bool InitiallyCollapsed() const override { return true; }
 
 	bool IsBuildEnabled() const;
+
+		//Custom Row Add utilities
+	FDetailWidgetRow& AddFloatRow(IDetailChildrenBuilder& ChildrenBuilder, const FText RowTitleText, const FText RowNameContentText, const FText RowNameContentTootlipText, const float MinSliderValue, const float MaxSliderValue, FGetFloatDelegate GetterDelegate, FSetFloatDelegate SetterDelegate);
+
+	float GetThresholdPosition() const;
+	void SetThresholdPosition(float Value);
+	
+	float GetThresholdTangentNormal() const;
+	void SetThresholdTangentNormal(float Value);
+	
+	float GetThresholdUV() const;
+	void SetThresholdUV(float Value);
+
+	float GetMorphThresholdPosition() const;
+	void SetMorphThresholdPosition(float Value);
+
 	ECheckBoxState ShouldRecomputeNormals() const;
 	ECheckBoxState ShouldRecomputeTangents() const;
 	ECheckBoxState ShouldUseMikkTSpace() const;
@@ -256,11 +275,22 @@ private:
 	void OnUseHighPrecisionTangentBasisChanged(ECheckBoxState NewState);
 	void OnUseFullPrecisionUVsChanged(ECheckBoxState NewState);
 	void OnBuildAdjacencyBufferChanged(ECheckBoxState NewState);
+
 private:
 	FSkeletalMeshBuildSettings& BuildSettings;
 	int32 LODIndex;
 	FIsLODSettingsEnabledDelegate IsBuildSettingsEnabledDelegate;
 	FModifyMeshLODSettingsDelegate ModifyMeshLODSettingsDelegate;
+
+	//Use this data to keep a valid reference so the helper lambda can have persistent data
+	//Helper lambda are use with spinner to not do a transaction when spinning with mouse movement
+	struct FSliderStateData
+	{
+		float MovementValueFloat = 0.0f;
+		int32 MovementValueInt = 0;
+		bool bSliderActiveMode = false;
+	};
+	TArray<FSliderStateData> SliderStateDataArray;
 };
 
 class FPersonaMeshDetails : public IDetailCustomization
diff --git a/Engine/Source/Editor/Persona/Private/SMorphTargetViewer.cpp b/Engine/Source/Editor/Persona/Private/SMorphTargetViewer.cpp
--- a/Engine/Source/Editor/Persona/Private/SMorphTargetViewer.cpp
+++ b/Engine/Source/Editor/Persona/Private/SMorphTargetViewer.cpp
@@ -330,6 +330,7 @@ void SMorphTargetViewer::Construct(const FArguments& InArgs, const TSharedRef<IP
 
 	SkeletalMesh = InPreviewScene->GetPreviewMeshComponent()->SkeletalMesh;
 	InPreviewScene->RegisterOnPreviewMeshChanged( FOnPreviewMeshChanged::CreateSP( this, &SMorphTargetViewer::OnPreviewMeshChanged ) );
+	InPreviewScene->RegisterOnMorphTargetsChanged(FSimpleDelegate::CreateSP(this, &SMorphTargetViewer::OnMorphTargetsChanged));
 	OnPostUndo.Add(FSimpleDelegate::CreateSP(this, &SMorphTargetViewer::OnPostUndo));
 
 	const FText SkeletalMeshName = SkeletalMesh ? FText::FromString( SkeletalMesh->GetName() ) : LOCTEXT( "MorphTargetMeshNameLabel", "No Skeletal Mesh Present" );
@@ -397,6 +398,11 @@ void SMorphTargetViewer::OnPreviewMeshChanged(class USkeletalMesh* OldPreviewMes
 	CreateMorphTargetList( NameFilterBox->GetText().ToString() );
 }
 
+void SMorphTargetViewer::OnMorphTargetsChanged()
+{
+	CreateMorphTargetList(NameFilterBox->GetText().ToString());
+}
+
 void SMorphTargetViewer::OnFilterTextChanged( const FText& SearchText )
 {
 	FilterText = SearchText;
diff --git a/Engine/Source/Editor/Persona/Private/SMorphTargetViewer.h b/Engine/Source/Editor/Persona/Private/SMorphTargetViewer.h
--- a/Engine/Source/Editor/Persona/Private/SMorphTargetViewer.h
+++ b/Engine/Source/Editor/Persona/Private/SMorphTargetViewer.h
@@ -78,6 +78,11 @@ public:
 	*/
 	void OnPreviewMeshChanged(class USkeletalMesh* OldPreviewMesh, class USkeletalMesh* NewPreviewMesh);
 
+	/**
+	* Is registered with Persona to handle when its preview mesh morph targets has changed.
+	*/
+	void OnMorphTargetsChanged();
+
 	/**
 	* Filters the SListView when the user changes the search text box (NameFilterBox)
 	*
diff --git a/Engine/Source/Editor/Persona/Public/IPersonaPreviewScene.h b/Engine/Source/Editor/Persona/Public/IPersonaPreviewScene.h
--- a/Engine/Source/Editor/Persona/Public/IPersonaPreviewScene.h
+++ b/Engine/Source/Editor/Persona/Public/IPersonaPreviewScene.h
@@ -156,6 +156,15 @@ public:
 	/** Unregisters a delegate to be called when the preview mesh's LOD has changed */
 	virtual void UnregisterOnLODChanged(void* Thing) = 0;
 
+	/** Registers a delegate to be called when the preview mesh's morph targets has changed */
+	virtual void RegisterOnMorphTargetsChanged(const FSimpleDelegate& Delegate) = 0;
+
+	/** Unregisters a delegate to be called when the preview mesh's morph targets has changed */
+	virtual void UnregisterOnMorphTargetsChanged(void* Thing) = 0;
+
+	/** Broadcasts that the preview mesh morph targets has changed */
+	virtual void BroadcastOnMorphTargetsChanged() = 0;
+
 	/** Registers a delegate to be called when the view is invalidated */
 	virtual void RegisterOnInvalidateViews(const FSimpleDelegate& Delegate) = 0;
 
diff --git a/Engine/Source/Editor/UnrealEd/Classes/Factories/FbxSceneImportOptionsSkeletalMesh.h b/Engine/Source/Editor/UnrealEd/Classes/Factories/FbxSceneImportOptionsSkeletalMesh.h
--- a/Engine/Source/Editor/UnrealEd/Classes/Factories/FbxSceneImportOptionsSkeletalMesh.h
+++ b/Engine/Source/Editor/UnrealEd/Classes/Factories/FbxSceneImportOptionsSkeletalMesh.h
@@ -54,6 +54,10 @@ class UFbxSceneImportOptionsSkeletalMesh : public UObject
 	UPROPERTY(EditAnywhere, config, Category = "SkeletalMesh|Thresholds", meta = (NoSpinbox = "true", ClampMin = "0.0", ClampMax = "1.0"))
 	float ThresholdUV;
 
+	/** Threshold to compare vertex position equality when computing morph target deltas. */
+	UPROPERTY(EditAnywhere, config, Category = "SkeletalMesh|Thresholds", meta = (NoSpinbox = "true", ClampMin = "0.0"))
+	float MorphThresholdPosition;
+
 	//////////////////////////////////////////////////////////////////////////
 	// Animation section
 
diff --git a/Engine/Source/Editor/UnrealEd/Classes/Factories/FbxSkeletalMeshImportData.h b/Engine/Source/Editor/UnrealEd/Classes/Factories/FbxSkeletalMeshImportData.h
--- a/Engine/Source/Editor/UnrealEd/Classes/Factories/FbxSkeletalMeshImportData.h
+++ b/Engine/Source/Editor/UnrealEd/Classes/Factories/FbxSkeletalMeshImportData.h
@@ -79,6 +79,10 @@ public:
 	UPROPERTY(EditAnywhere, config, Category="Mesh", meta = (ImportType = "SkeletalMesh|GeoOnly", SubCategory = "Thresholds", NoSpinbox = "true", ClampMin = "0.0", ClampMax = "1.0"))
 	float ThresholdUV;
 
+	/** Threshold to compare vertex position equality when computing morph target deltas. */
+	UPROPERTY(EditAnywhere, config, Category = "Mesh", meta = (editcondition = "bImportMorphTargets", ImportType = "SkeletalMesh|GeoOnly", SubCategory = "Thresholds", NoSpinbox = "true", ClampMin = "0.0", ClampMax = "1.0"))
+	float MorphThresholdPosition;
+
 	/** Gets or creates fbx import data for the specified skeletal mesh */
 	static UFbxSkeletalMeshImportData* GetImportDataForSkeletalMesh(USkeletalMesh* SkeletalMesh, UFbxSkeletalMeshImportData* TemplateForCreation);
 
diff --git a/Engine/Source/Editor/UnrealEd/Private/Fbx/FbxMainImport.cpp b/Engine/Source/Editor/UnrealEd/Private/Fbx/FbxMainImport.cpp
--- a/Engine/Source/Editor/UnrealEd/Private/Fbx/FbxMainImport.cpp
+++ b/Engine/Source/Editor/UnrealEd/Private/Fbx/FbxMainImport.cpp
@@ -352,38 +352,54 @@ FBXImportOptions* GetImportOptions( UnFbx::FFbxImporter* FbxImporter, UFbxImport
 void ApplyImportUIToImportOptions(UFbxImportUI* ImportUI, FBXImportOptions& InOutImportOptions)
 {
 	check(ImportUI);
-	InOutImportOptions.bImportMaterials = ImportUI->bImportMaterials;
-	InOutImportOptions.bInvertNormalMap = ImportUI->TextureImportData->bInvertNormalMaps;
-	InOutImportOptions.MaterialSearchLocation = ImportUI->TextureImportData->MaterialSearchLocation;
-	UMaterialInterface* BaseMaterialInterface = Cast<UMaterialInterface>(ImportUI->TextureImportData->BaseMaterialName.TryLoad());
-	if (BaseMaterialInterface) {
-		InOutImportOptions.BaseMaterial = BaseMaterialInterface;
-		InOutImportOptions.BaseColorName = ImportUI->TextureImportData->BaseColorName;
-		InOutImportOptions.BaseDiffuseTextureName = ImportUI->TextureImportData->BaseDiffuseTextureName;
-		InOutImportOptions.BaseNormalTextureName = ImportUI->TextureImportData->BaseNormalTextureName;
-		InOutImportOptions.BaseEmmisiveTextureName = ImportUI->TextureImportData->BaseEmmisiveTextureName;
-		InOutImportOptions.BaseSpecularTextureName = ImportUI->TextureImportData->BaseSpecularTextureName;
-		InOutImportOptions.BaseEmissiveColorName = ImportUI->TextureImportData->BaseEmissiveColorName;
-	}
-	InOutImportOptions.bImportTextures = ImportUI->bImportTextures;
-	InOutImportOptions.bUsedAsFullName = ImportUI->bOverrideFullName;
-	InOutImportOptions.bImportAnimations = ImportUI->bImportAnimations;
-	InOutImportOptions.SkeletonForAnimation = ImportUI->Skeleton;
-	InOutImportOptions.ImportType = ImportUI->MeshTypeToImport;
-
-	InOutImportOptions.bAutoComputeLodDistances = ImportUI->bAutoComputeLodDistances;
-	InOutImportOptions.LodDistances.Empty(8);
-	InOutImportOptions.LodDistances.Add(ImportUI->LodDistance0);
-	InOutImportOptions.LodDistances.Add(ImportUI->LodDistance1);
-	InOutImportOptions.LodDistances.Add(ImportUI->LodDistance2);
-	InOutImportOptions.LodDistances.Add(ImportUI->LodDistance3);
-	InOutImportOptions.LodDistances.Add(ImportUI->LodDistance4);
-	InOutImportOptions.LodDistances.Add(ImportUI->LodDistance5);
-	InOutImportOptions.LodDistances.Add(ImportUI->LodDistance6);
-	InOutImportOptions.LodDistances.Add(ImportUI->LodDistance7);
-	InOutImportOptions.LodNumber = ImportUI->LodNumber;
-	InOutImportOptions.MinimumLodNumber = ImportUI->MinimumLodNumber;
 
+	//General options
+	{
+		InOutImportOptions.bUsedAsFullName			= ImportUI->bOverrideFullName;
+		InOutImportOptions.ImportType				= ImportUI->MeshTypeToImport;
+
+		InOutImportOptions.bAutoComputeLodDistances	= ImportUI->bAutoComputeLodDistances;
+		InOutImportOptions.LodDistances.Empty(8);
+		InOutImportOptions.LodDistances.Add(ImportUI->LodDistance0);
+		InOutImportOptions.LodDistances.Add(ImportUI->LodDistance1);
+		InOutImportOptions.LodDistances.Add(ImportUI->LodDistance2);
+		InOutImportOptions.LodDistances.Add(ImportUI->LodDistance3);
+		InOutImportOptions.LodDistances.Add(ImportUI->LodDistance4);
+		InOutImportOptions.LodDistances.Add(ImportUI->LodDistance5);
+		InOutImportOptions.LodDistances.Add(ImportUI->LodDistance6);
+		InOutImportOptions.LodDistances.Add(ImportUI->LodDistance7);
+		InOutImportOptions.LodNumber				= ImportUI->LodNumber;
+		InOutImportOptions.MinimumLodNumber			= ImportUI->MinimumLodNumber;
+	}
+
+	//Animation and skeletal mesh options
+	{
+		InOutImportOptions.bImportAnimations		= ImportUI->bImportAnimations;
+		InOutImportOptions.SkeletonForAnimation		= ImportUI->Skeleton;
+		InOutImportOptions.bCreatePhysicsAsset		= ImportUI->bCreatePhysicsAsset;
+		InOutImportOptions.PhysicsAsset				= ImportUI->PhysicsAsset;
+		InOutImportOptions.AnimationName			= ImportUI->OverrideAnimationName;
+	}
+
+	//Material options
+	{
+		InOutImportOptions.bImportMaterials			= ImportUI->bImportMaterials;
+		InOutImportOptions.bInvertNormalMap			= ImportUI->TextureImportData->bInvertNormalMaps;
+		InOutImportOptions.MaterialSearchLocation	= ImportUI->TextureImportData->MaterialSearchLocation;
+		UMaterialInterface* BaseMaterialInterface	= Cast<UMaterialInterface>(ImportUI->TextureImportData->BaseMaterialName.TryLoad());
+		if (BaseMaterialInterface) {
+			InOutImportOptions.BaseMaterial			= BaseMaterialInterface;
+			InOutImportOptions.BaseColorName		= ImportUI->TextureImportData->BaseColorName;
+			InOutImportOptions.BaseDiffuseTextureName = ImportUI->TextureImportData->BaseDiffuseTextureName;
+			InOutImportOptions.BaseNormalTextureName = ImportUI->TextureImportData->BaseNormalTextureName;
+			InOutImportOptions.BaseEmmisiveTextureName = ImportUI->TextureImportData->BaseEmmisiveTextureName;
+			InOutImportOptions.BaseSpecularTextureName = ImportUI->TextureImportData->BaseSpecularTextureName;
+			InOutImportOptions.BaseEmissiveColorName = ImportUI->TextureImportData->BaseEmissiveColorName;
+		}
+		InOutImportOptions.bImportTextures			= ImportUI->bImportTextures;
+	}
+
+	//Some options are overlap between static mesh, skeletal mesh and anim sequence. We must set them according to the import type to set them only once
 	if ( ImportUI->MeshTypeToImport == FBXIT_StaticMesh )
 	{
 		UFbxStaticMeshImportData* StaticMeshData	= ImportUI->StaticMeshImportData;
@@ -406,8 +422,6 @@ void ApplyImportUIToImportOptions(UFbxImportUI* ImportUI, FBXImportOptions& InOu
 	else if ( ImportUI->MeshTypeToImport == FBXIT_SkeletalMesh )
 	{
 		UFbxSkeletalMeshImportData* SkeletalMeshData	= ImportUI->SkeletalMeshImportData;
-		InOutImportOptions.bImportAsSkeletalGeometry    = SkeletalMeshData->ImportContentType == EFBXImportContentType::FBXICT_Geometry;
-		InOutImportOptions.bImportAsSkeletalSkinning	= SkeletalMeshData->ImportContentType == EFBXImportContentType::FBXICT_SkinningWeights;
 		InOutImportOptions.NormalImportMethod			= SkeletalMeshData->NormalImportMethod;
 		InOutImportOptions.NormalGenerationMethod		= SkeletalMeshData->NormalGenerationMethod;
 		InOutImportOptions.bComputeWeightedNormals		= SkeletalMeshData->bComputeWeightedNormals;
@@ -439,7 +453,7 @@ void ApplyImportUIToImportOptions(UFbxImportUI* ImportUI, FBXImportOptions& InOu
 	else
 	{
 		UFbxAnimSequenceImportData* AnimData	= ImportUI->AnimSequenceImportData;
-		InOutImportOptions.NormalImportMethod = FBXNIM_ComputeNormals;
+		InOutImportOptions.NormalImportMethod	= FBXNIM_ComputeNormals;
 		InOutImportOptions.bComputeWeightedNormals = true;
 		InOutImportOptions.ImportTranslation	= AnimData->ImportTranslation;
 		InOutImportOptions.ImportRotation		= AnimData->ImportRotation;
@@ -449,43 +463,53 @@ void ApplyImportUIToImportOptions(UFbxImportUI* ImportUI, FBXImportOptions& InOu
 		InOutImportOptions.bConvertSceneUnit	= AnimData->bConvertSceneUnit;
 	}
 
-	InOutImportOptions.bImportMorph = ImportUI->SkeletalMeshImportData->bImportMorphTargets;
-	InOutImportOptions.bUpdateSkeletonReferencePose = ImportUI->SkeletalMeshImportData->bUpdateSkeletonReferencePose;
-	InOutImportOptions.bImportRigidMesh = ImportUI->OriginalImportType == FBXIT_StaticMesh && ImportUI->MeshTypeToImport == FBXIT_SkeletalMesh;
-	InOutImportOptions.bUseT0AsRefPose = ImportUI->SkeletalMeshImportData->bUseT0AsRefPose;
-	InOutImportOptions.bPreserveSmoothingGroups = ImportUI->SkeletalMeshImportData->bPreserveSmoothingGroups;
-	InOutImportOptions.OverlappingThresholds.ThresholdPosition = ImportUI->SkeletalMeshImportData->ThresholdPosition;
-	InOutImportOptions.OverlappingThresholds.ThresholdTangentNormal = ImportUI->SkeletalMeshImportData->ThresholdTangentNormal;
-	InOutImportOptions.OverlappingThresholds.ThresholdUV = ImportUI->SkeletalMeshImportData->ThresholdUV;
-	InOutImportOptions.bCombineToSingle = ImportUI->StaticMeshImportData->bCombineMeshes;
+	//Skeletal mesh unshared options
+	{
+		InOutImportOptions.bImportAsSkeletalGeometry	= ImportUI->SkeletalMeshImportData->ImportContentType == EFBXImportContentType::FBXICT_Geometry;
+		InOutImportOptions.bImportAsSkeletalSkinning	= ImportUI->SkeletalMeshImportData->ImportContentType == EFBXImportContentType::FBXICT_SkinningWeights;
+		InOutImportOptions.bImportMorph					= ImportUI->SkeletalMeshImportData->bImportMorphTargets;
+		InOutImportOptions.bUpdateSkeletonReferencePose = ImportUI->SkeletalMeshImportData->bUpdateSkeletonReferencePose;
+		InOutImportOptions.bImportRigidMesh				= ImportUI->OriginalImportType == FBXIT_StaticMesh && ImportUI->MeshTypeToImport == FBXIT_SkeletalMesh;
+		InOutImportOptions.bUseT0AsRefPose				= ImportUI->SkeletalMeshImportData->bUseT0AsRefPose;
+		InOutImportOptions.bPreserveSmoothingGroups		= ImportUI->SkeletalMeshImportData->bPreserveSmoothingGroups;
+		InOutImportOptions.OverlappingThresholds.ThresholdPosition = ImportUI->SkeletalMeshImportData->ThresholdPosition;
+		InOutImportOptions.OverlappingThresholds.ThresholdTangentNormal = ImportUI->SkeletalMeshImportData->ThresholdTangentNormal;
+		InOutImportOptions.OverlappingThresholds.ThresholdUV = ImportUI->SkeletalMeshImportData->ThresholdUV;
+		InOutImportOptions.OverlappingThresholds.MorphThresholdPosition = ImportUI->SkeletalMeshImportData->MorphThresholdPosition;
+		InOutImportOptions.bImportMeshesInBoneHierarchy = ImportUI->SkeletalMeshImportData->bImportMeshesInBoneHierarchy;
+	}
+
+	//Static mesh unshared options
+	{
+		InOutImportOptions.bCombineToSingle				= ImportUI->StaticMeshImportData->bCombineMeshes;
+		InOutImportOptions.bRemoveDegenerates			= ImportUI->StaticMeshImportData->bRemoveDegenerates;
+		InOutImportOptions.bBuildAdjacencyBuffer		= ImportUI->StaticMeshImportData->bBuildAdjacencyBuffer;
+		InOutImportOptions.bBuildReversedIndexBuffer	= ImportUI->StaticMeshImportData->bBuildReversedIndexBuffer;
+		InOutImportOptions.bGenerateLightmapUVs			= ImportUI->StaticMeshImportData->bGenerateLightmapUVs;
+		InOutImportOptions.bOneConvexHullPerUCX			= ImportUI->StaticMeshImportData->bOneConvexHullPerUCX;
+		InOutImportOptions.bAutoGenerateCollision		= ImportUI->StaticMeshImportData->bAutoGenerateCollision;
+		InOutImportOptions.StaticMeshLODGroup			= ImportUI->StaticMeshImportData->StaticMeshLODGroup;
+	}
 	
-	InOutImportOptions.bRemoveDegenerates = ImportUI->StaticMeshImportData->bRemoveDegenerates;
-	InOutImportOptions.bBuildAdjacencyBuffer = ImportUI->StaticMeshImportData->bBuildAdjacencyBuffer;
-	InOutImportOptions.bBuildReversedIndexBuffer = ImportUI->StaticMeshImportData->bBuildReversedIndexBuffer;
-	InOutImportOptions.bGenerateLightmapUVs = ImportUI->StaticMeshImportData->bGenerateLightmapUVs;
-	InOutImportOptions.bOneConvexHullPerUCX = ImportUI->StaticMeshImportData->bOneConvexHullPerUCX;
-	InOutImportOptions.bAutoGenerateCollision = ImportUI->StaticMeshImportData->bAutoGenerateCollision;
-	InOutImportOptions.StaticMeshLODGroup = ImportUI->StaticMeshImportData->StaticMeshLODGroup;
-	InOutImportOptions.bImportMeshesInBoneHierarchy = ImportUI->SkeletalMeshImportData->bImportMeshesInBoneHierarchy;
-	InOutImportOptions.bCreatePhysicsAsset = ImportUI->bCreatePhysicsAsset;
-	InOutImportOptions.PhysicsAsset = ImportUI->PhysicsAsset;
-	// animation options
-	InOutImportOptions.AnimationLengthImportType = ImportUI->AnimSequenceImportData->AnimationLength;
-	InOutImportOptions.AnimationRange.X = ImportUI->AnimSequenceImportData->FrameImportRange.Min;
-	InOutImportOptions.AnimationRange.Y = ImportUI->AnimSequenceImportData->FrameImportRange.Max;
-	InOutImportOptions.AnimationName = ImportUI->OverrideAnimationName;
-	// only re-sample if they don't want to use default sample rate
-	InOutImportOptions.bResample = !ImportUI->AnimSequenceImportData->bUseDefaultSampleRate;
-	InOutImportOptions.ResampleRate = ImportUI->AnimSequenceImportData->CustomSampleRate;
-	InOutImportOptions.bPreserveLocalTransform = ImportUI->AnimSequenceImportData->bPreserveLocalTransform;
-	InOutImportOptions.bDeleteExistingMorphTargetCurves = ImportUI->AnimSequenceImportData->bDeleteExistingMorphTargetCurves;
-	InOutImportOptions.bRemoveRedundantKeys = ImportUI->AnimSequenceImportData->bRemoveRedundantKeys;
-	InOutImportOptions.bDoNotImportCurveWithZero = ImportUI->AnimSequenceImportData->bDoNotImportCurveWithZero;
-	InOutImportOptions.bImportCustomAttribute = ImportUI->AnimSequenceImportData->bImportCustomAttribute;
-	InOutImportOptions.bDeleteExistingCustomAttributeCurves = ImportUI->AnimSequenceImportData->bDeleteExistingCustomAttributeCurves;
-	InOutImportOptions.bImportBoneTracks = ImportUI->AnimSequenceImportData->bImportBoneTracks;
-	InOutImportOptions.bSetMaterialDriveParameterOnCustomAttribute = ImportUI->AnimSequenceImportData->bSetMaterialDriveParameterOnCustomAttribute;
-	InOutImportOptions.MaterialCurveSuffixes = ImportUI->AnimSequenceImportData->MaterialCurveSuffixes;
+	// animation unshared options
+	{
+		
+		InOutImportOptions.AnimationLengthImportType	= ImportUI->AnimSequenceImportData->AnimationLength;
+		InOutImportOptions.AnimationRange.X				= ImportUI->AnimSequenceImportData->FrameImportRange.Min;
+		InOutImportOptions.AnimationRange.Y				= ImportUI->AnimSequenceImportData->FrameImportRange.Max;
+		// only re-sample if they don't want to use default sample rate
+		InOutImportOptions.bResample					= !ImportUI->AnimSequenceImportData->bUseDefaultSampleRate;
+		InOutImportOptions.ResampleRate					= ImportUI->AnimSequenceImportData->CustomSampleRate;
+		InOutImportOptions.bPreserveLocalTransform		= ImportUI->AnimSequenceImportData->bPreserveLocalTransform;
+		InOutImportOptions.bDeleteExistingMorphTargetCurves = ImportUI->AnimSequenceImportData->bDeleteExistingMorphTargetCurves;
+		InOutImportOptions.bRemoveRedundantKeys			= ImportUI->AnimSequenceImportData->bRemoveRedundantKeys;
+		InOutImportOptions.bDoNotImportCurveWithZero	= ImportUI->AnimSequenceImportData->bDoNotImportCurveWithZero;
+		InOutImportOptions.bImportCustomAttribute		= ImportUI->AnimSequenceImportData->bImportCustomAttribute;
+		InOutImportOptions.bDeleteExistingCustomAttributeCurves = ImportUI->AnimSequenceImportData->bDeleteExistingCustomAttributeCurves;
+		InOutImportOptions.bImportBoneTracks			= ImportUI->AnimSequenceImportData->bImportBoneTracks;
+		InOutImportOptions.bSetMaterialDriveParameterOnCustomAttribute = ImportUI->AnimSequenceImportData->bSetMaterialDriveParameterOnCustomAttribute;
+		InOutImportOptions.MaterialCurveSuffixes		= ImportUI->AnimSequenceImportData->MaterialCurveSuffixes;
+	}
 }
 
 void FImportedMaterialData::AddImportedMaterial( FbxSurfaceMaterial& FbxMaterial, UMaterialInterface& UnrealMaterial )
@@ -1456,6 +1480,7 @@ bool FFbxImporter::ImportFromFile(const FString& Filename, const FString& Type,
 				  * @EventParam ThresholdPosition float Returns the threshold delta to weld vertices
 				  * @EventParam ThresholdTangentNormal float Returns the threshold delta to weld tangents and normals
 				  * @EventParam ThresholdUV float Returns the threshold delta to weld UVs
+				  * @EventParam MorphThresholdPosition float Returns the morph target threshold delta to compute deltas
 				  * @EventParam AutoComputeLodDistances boolean Returns whether the importer should set the auto compute LOD distance
 				  * @EventParam LodNumber integer Returns the LOD number we should have after the import
 				  * @EventParam BuildAdjacencyBuffer boolean Returns whether the importer should fill the adjacency buffer when building the static mesh
@@ -1553,6 +1578,7 @@ bool FFbxImporter::ImportFromFile(const FString& Filename, const FString& Type,
 							Attribs.Add(FAnalyticsEventAttribute(TEXT("SkeletalMeshOpt OverlappingThresholds.ThresholdPosition"), CaptureImportOptions->OverlappingThresholds.ThresholdPosition));
 							Attribs.Add(FAnalyticsEventAttribute(TEXT("SkeletalMeshOpt OverlappingThresholds.ThresholdTangentNormal"), CaptureImportOptions->OverlappingThresholds.ThresholdTangentNormal));
 							Attribs.Add(FAnalyticsEventAttribute(TEXT("SkeletalMeshOpt OverlappingThresholds.ThresholdUV"), CaptureImportOptions->OverlappingThresholds.ThresholdUV));
+							Attribs.Add(FAnalyticsEventAttribute(TEXT("SkeletalMeshOpt OverlappingThresholds.MorphThresholdPosition"), CaptureImportOptions->OverlappingThresholds.MorphThresholdPosition));
 						};
 
 						auto AddSMAnalytic = [&Attribs, &CaptureImportOptions]()
diff --git a/Engine/Source/Editor/UnrealEd/Private/Fbx/FbxSceneImportData.cpp b/Engine/Source/Editor/UnrealEd/Private/Fbx/FbxSceneImportData.cpp
--- a/Engine/Source/Editor/UnrealEd/Private/Fbx/FbxSceneImportData.cpp
+++ b/Engine/Source/Editor/UnrealEd/Private/Fbx/FbxSceneImportData.cpp
@@ -124,7 +124,7 @@ UnFbx::FBXImportOptions *JSONToFbxOption(TSharedPtr<FJsonValue> OptionJsonValue,
 	
 	if (OptionObj->TryGetObjectField("OverlappingThresholds", DataObj))
 	{
-		double Pos, NTB, UV;
+		double Pos, NTB, UV, MorphPos;
 		if ((*DataObj)->TryGetNumberField("ThresholdPosition", Pos))
 		{
 			Option->OverlappingThresholds.ThresholdPosition = (float)Pos;
@@ -137,6 +137,10 @@ UnFbx::FBXImportOptions *JSONToFbxOption(TSharedPtr<FJsonValue> OptionJsonValue,
 		{
 			Option->OverlappingThresholds.ThresholdUV = (float)UV;
 		}
+		if ((*DataObj)->TryGetNumberField("MorphThresholdPosition", MorphPos))
+		{
+			Option->OverlappingThresholds.MorphThresholdPosition = (float)MorphPos;
+		}
 	}
 
 	FString MaterialBasePath;
@@ -201,7 +205,7 @@ FString FbxOptionToJSON(FString OptionName, UnFbx::FBXImportOptions *Option)
 		);
 
 	//Skeletal mesh options
-	JsonString += FString::Printf(TEXT("\"bUpdateSkeletonReferencePose\" : \"%d\", \"bUseT0AsRefPose\" : \"%d\", \"bPreserveSmoothingGroups\" : \"%d\", \"bImportMeshesInBoneHierarchy\" : \"%d\", \"bImportMorphTargets\" : \"%d\", \"OverlappingThresholds\" : {\"ThresholdPosition\" : \"%f\", \"ThresholdTangentNormal\" : \"%f\", \"ThresholdUV\" : \"%f\"},"),
+	JsonString += FString::Printf(TEXT("\"bUpdateSkeletonReferencePose\" : \"%d\", \"bUseT0AsRefPose\" : \"%d\", \"bPreserveSmoothingGroups\" : \"%d\", \"bImportMeshesInBoneHierarchy\" : \"%d\", \"bImportMorphTargets\" : \"%d\", \"OverlappingThresholds\" : {\"ThresholdPosition\" : \"%f\", \"ThresholdTangentNormal\" : \"%f\", \"ThresholdUV\" : \"%f\", \"MorphThresholdPosition\" : \"%f\"},"),
 		Option->bUpdateSkeletonReferencePose ? 1 : 0,
 		Option->bUseT0AsRefPose ? 1 : 0,
 		Option->bPreserveSmoothingGroups ? 1 : 0,
@@ -209,7 +213,8 @@ FString FbxOptionToJSON(FString OptionName, UnFbx::FBXImportOptions *Option)
 		Option->bImportMorph ? 1 : 0,
 		Option->OverlappingThresholds.ThresholdPosition,
 		Option->OverlappingThresholds.ThresholdTangentNormal,
-		Option->OverlappingThresholds.ThresholdUV
+		Option->OverlappingThresholds.ThresholdUV,
+		Option->OverlappingThresholds.MorphThresholdPosition
 		);
 
 	JsonString += FString::Printf(TEXT("\"MaterialBasePath\" : \"%s\"}"), *(Option->MaterialBasePath.ToString()));
diff --git a/Engine/Source/Editor/UnrealEd/Private/Fbx/FbxSceneImportOptionsSkeletalMesh.cpp b/Engine/Source/Editor/UnrealEd/Private/Fbx/FbxSceneImportOptionsSkeletalMesh.cpp
--- a/Engine/Source/Editor/UnrealEd/Private/Fbx/FbxSceneImportOptionsSkeletalMesh.cpp
+++ b/Engine/Source/Editor/UnrealEd/Private/Fbx/FbxSceneImportOptionsSkeletalMesh.cpp
@@ -17,6 +17,7 @@ UFbxSceneImportOptionsSkeletalMesh::UFbxSceneImportOptionsSkeletalMesh(const FOb
 	, ThresholdPosition(THRESH_POINTS_ARE_SAME)
 	, ThresholdTangentNormal(THRESH_NORMALS_ARE_SAME)
 	, ThresholdUV(THRESH_UVS_ARE_SAME)
+	, MorphThresholdPosition(THRESH_POINTS_ARE_NEAR)
 	, bImportAnimations(true)
 	, AnimationLength(EFbxSceneVertexColorImportOption::Replace)
 	, FrameImportRange(0, 0)
@@ -37,6 +38,7 @@ void UFbxSceneImportOptionsSkeletalMesh::FillSkeletalMeshInmportData(UFbxSkeleta
 	SkeletalMeshImportData->ThresholdPosition = ThresholdPosition;
 	SkeletalMeshImportData->ThresholdTangentNormal = ThresholdTangentNormal;
 	SkeletalMeshImportData->ThresholdUV = ThresholdUV;
+	SkeletalMeshImportData->MorphThresholdPosition = MorphThresholdPosition;
 	SkeletalMeshImportData->bPreserveSmoothingGroups = bPreserveSmoothingGroups;
 	SkeletalMeshImportData->bUpdateSkeletonReferencePose = bUpdateSkeletonReferencePose;
 	SkeletalMeshImportData->bUseT0AsRefPose = bUseT0AsRefPose;
diff --git a/Engine/Source/Editor/UnrealEd/Private/Fbx/FbxSkeletalMeshImport.cpp b/Engine/Source/Editor/UnrealEd/Private/Fbx/FbxSkeletalMeshImport.cpp
--- a/Engine/Source/Editor/UnrealEd/Private/Fbx/FbxSkeletalMeshImport.cpp
+++ b/Engine/Source/Editor/UnrealEd/Private/Fbx/FbxSkeletalMeshImport.cpp
@@ -1609,7 +1609,7 @@ bool UnFbx::FFbxImporter::GatherPointsForMorphTarget(FSkeletalMeshImportData* Ou
 	{
 		int32 OriginalPointIdx = OutData->PointToRawMap[ PointIdx ];
 
-		if ( ( NewImportData.Points[ OriginalPointIdx ] - OutData->Points[ PointIdx ] ).SizeSquared() > FMath::Square( THRESH_VECTORS_ARE_NEAR ) ) // THRESH_POINTS_ARE_NEAR is too big, we might not be recomputing some normals that can have changed significantly
+		if ( ( NewImportData.Points[ OriginalPointIdx ] - OutData->Points[ PointIdx ] ).SizeSquared() > FMath::Square(THRESH_POINTS_ARE_SAME) )
 		{
 			ModifiedPoints.Add( PointIdx );
 		}
@@ -1812,9 +1812,32 @@ USkeletalMesh* UnFbx::FFbxImporter::ImportSkeletalMesh(FImportSkeletalMeshArgs &
 		return nullptr;
 	}
 
+	FSkeletalMeshBuildSettings BuildOptions;
+	//Make sure the build option change in the re-import ui is reconduct
+	BuildOptions.bBuildAdjacencyBuffer = true;
+	BuildOptions.bRecomputeNormals = !ImportOptions->ShouldImportNormals() || !SkelMeshImportDataPtr->bHasNormals;
+	BuildOptions.bRecomputeTangents = !ImportOptions->ShouldImportTangents() || !SkelMeshImportDataPtr->bHasTangents;
+	BuildOptions.bUseMikkTSpace = (ImportOptions->NormalGenerationMethod == EFBXNormalGenerationMethod::MikkTSpace) && (!ImportOptions->ShouldImportNormals() || !ImportOptions->ShouldImportTangents());
+	BuildOptions.bComputeWeightedNormals = ImportOptions->bComputeWeightedNormals;
+	BuildOptions.bRemoveDegenerates = ImportOptions->bRemoveDegenerates;
+	BuildOptions.ThresholdPosition = ImportOptions->OverlappingThresholds.ThresholdPosition;
+	BuildOptions.ThresholdTangentNormal = ImportOptions->OverlappingThresholds.ThresholdTangentNormal;
+	BuildOptions.ThresholdUV = ImportOptions->OverlappingThresholds.ThresholdUV;
+	BuildOptions.MorphThresholdPosition = ImportOptions->OverlappingThresholds.MorphThresholdPosition;
+	
 	if (ExistingSkelMesh)
 	{
 		ExistingSkelMesh->PreEditChange(NULL);
+		//We update the build settings before saving the asset data we want to restore after the re-import
+		int32 SafeReimportLODIndex = ImportSkeletalMeshArgs.LodIndex < 0 ? 0 : ImportSkeletalMeshArgs.LodIndex;
+		FSkeletalMeshLODInfo* LODInfoPtr = ExistingSkelMesh->GetLODInfo(SafeReimportLODIndex);
+		if (LODInfoPtr)
+		{
+			//Adjacency buffer cannot be change in the re-import options, it must be backup before and restored after the copy.
+			bool bBuildAdjacencyBuffer = LODInfoPtr->BuildSettings.bBuildAdjacencyBuffer;
+			LODInfoPtr->BuildSettings = BuildOptions;
+			LODInfoPtr->BuildSettings.bBuildAdjacencyBuffer = bBuildAdjacencyBuffer;
+		}
 		//The backup of the skeletal mesh data empty the LOD array in the ImportedResource of the skeletal mesh
 		//If the import fail after this step the editor can crash when updating the bone later since the LODModel will not exist anymore
 		ExistSkelMeshDataPtr = SaveExistingSkelMeshData(ExistingSkelMesh, !ImportOptions->bImportMaterials, ImportSkeletalMeshArgs.LodIndex);
@@ -1904,17 +1927,6 @@ USkeletalMesh* UnFbx::FFbxImporter::ImportSkeletalMesh(FImportSkeletalMeshArgs &
 		TArray<int32> LODPointToRawMap;
 		SkelMeshImportDataPtr->CopyLODImportData(LODPoints,LODWedges,LODFaces,LODInfluences,LODPointToRawMap);
 
-		FSkeletalMeshBuildSettings BuildOptions;
-		BuildOptions.bBuildAdjacencyBuffer = true;
-		BuildOptions.bRecomputeNormals = !ImportOptions->ShouldImportNormals() || !SkelMeshImportDataPtr->bHasNormals;
-		BuildOptions.bRecomputeTangents = !ImportOptions->ShouldImportTangents() || !SkelMeshImportDataPtr->bHasTangents;
-		BuildOptions.bUseMikkTSpace = (ImportOptions->NormalGenerationMethod == EFBXNormalGenerationMethod::MikkTSpace) && (!ImportOptions->ShouldImportNormals() || !ImportOptions->ShouldImportTangents());
-		BuildOptions.bComputeWeightedNormals = ImportOptions->bComputeWeightedNormals;
-		BuildOptions.bRemoveDegenerates = ImportOptions->bRemoveDegenerates;
-		BuildOptions.ThresholdPosition = ImportOptions->OverlappingThresholds.ThresholdPosition;
-		BuildOptions.ThresholdTangentNormal = ImportOptions->OverlappingThresholds.ThresholdTangentNormal;
-		BuildOptions.ThresholdUV = ImportOptions->OverlappingThresholds.ThresholdUV;
-
 		bool bUseLegacyMeshUtilities = false;
 		bool bBuildSuccess = false;
 		if (bUseLegacyMeshUtilities)
@@ -4325,7 +4337,9 @@ void UnFbx::FFbxImporter::ImportMorphTargetsInternal( TArray<FbxNode*>& SkelMesh
 				LODIndex,
 				ImportOptions->ShouldImportNormals(),
 				ImportOptions->ShouldImportTangents(),
-				(ImportOptions->NormalGenerationMethod == EFBXNormalGenerationMethod::MikkTSpace));
+				(ImportOptions->NormalGenerationMethod == EFBXNormalGenerationMethod::MikkTSpace),
+				ImportOptions->OverlappingThresholds
+				);
 		}
 	}
 
diff --git a/Engine/Source/Editor/UnrealEd/Private/Fbx/SFbxSceneOptionWindow.cpp b/Engine/Source/Editor/UnrealEd/Private/Fbx/SFbxSceneOptionWindow.cpp
--- a/Engine/Source/Editor/UnrealEd/Private/Fbx/SFbxSceneOptionWindow.cpp
+++ b/Engine/Source/Editor/UnrealEd/Private/Fbx/SFbxSceneOptionWindow.cpp
@@ -1497,6 +1497,7 @@ void SFbxSceneOptionWindow::CopySkeletalMeshOptionsToFbxOptions(UnFbx::FBXImport
 	ImportSettings->OverlappingThresholds.ThresholdPosition = SkeletalMeshOptions->ThresholdPosition;
 	ImportSettings->OverlappingThresholds.ThresholdTangentNormal = SkeletalMeshOptions->ThresholdTangentNormal;
 	ImportSettings->OverlappingThresholds.ThresholdUV = SkeletalMeshOptions->ThresholdUV;
+	ImportSettings->OverlappingThresholds.MorphThresholdPosition = SkeletalMeshOptions->MorphThresholdPosition;
 	ImportSettings->bPreserveSmoothingGroups = SkeletalMeshOptions->bPreserveSmoothingGroups;
 	ImportSettings->bUpdateSkeletonReferencePose = SkeletalMeshOptions->bUpdateSkeletonReferencePose;
 	ImportSettings->bUseT0AsRefPose = SkeletalMeshOptions->bUseT0AsRefPose;
@@ -1521,6 +1522,7 @@ void SFbxSceneOptionWindow::CopyFbxOptionsToSkeletalMeshOptions(UnFbx::FBXImport
 	SkeletalMeshOptions->ThresholdPosition = ImportSettings->OverlappingThresholds.ThresholdPosition;
 	SkeletalMeshOptions->ThresholdTangentNormal = ImportSettings->OverlappingThresholds.ThresholdTangentNormal;
 	SkeletalMeshOptions->ThresholdUV = ImportSettings->OverlappingThresholds.ThresholdUV;
+	SkeletalMeshOptions->MorphThresholdPosition = ImportSettings->OverlappingThresholds.MorphThresholdPosition;
 	SkeletalMeshOptions->bPreserveSmoothingGroups = ImportSettings->bPreserveSmoothingGroups;
 	SkeletalMeshOptions->bUpdateSkeletonReferencePose = ImportSettings->bUpdateSkeletonReferencePose;
 	SkeletalMeshOptions->bUseT0AsRefPose = ImportSettings->bUseT0AsRefPose;
diff --git a/Engine/Source/Editor/UnrealEd/Private/FbxMeshUtils.cpp b/Engine/Source/Editor/UnrealEd/Private/FbxMeshUtils.cpp
--- a/Engine/Source/Editor/UnrealEd/Private/FbxMeshUtils.cpp
+++ b/Engine/Source/Editor/UnrealEd/Private/FbxMeshUtils.cpp
@@ -383,6 +383,7 @@ namespace FbxMeshUtils
 					LODInfo->BuildSettings.ThresholdPosition = ImportOptions->OverlappingThresholds.ThresholdPosition;
 					LODInfo->BuildSettings.ThresholdTangentNormal = ImportOptions->OverlappingThresholds.ThresholdTangentNormal;
 					LODInfo->BuildSettings.ThresholdUV = ImportOptions->OverlappingThresholds.ThresholdUV;
+					LODInfo->BuildSettings.MorphThresholdPosition = ImportOptions->OverlappingThresholds.MorphThresholdPosition;
 				}
 			}
 
diff --git a/Engine/Source/Editor/UnrealEd/Private/SkinWeightsUtilities.cpp b/Engine/Source/Editor/UnrealEd/Private/SkinWeightsUtilities.cpp
--- a/Engine/Source/Editor/UnrealEd/Private/SkinWeightsUtilities.cpp
+++ b/Engine/Source/Editor/UnrealEd/Private/SkinWeightsUtilities.cpp
@@ -318,6 +318,7 @@ bool FSkinWeightsUtilities::RemoveSkinnedWeightProfileData(USkeletalMesh* Skelet
 	BuildOptions.OverlappingThresholds.ThresholdPosition = OriginalSkeletalMeshImportData->ThresholdPosition;
 	BuildOptions.OverlappingThresholds.ThresholdTangentNormal = OriginalSkeletalMeshImportData->ThresholdTangentNormal;
 	BuildOptions.OverlappingThresholds.ThresholdUV = OriginalSkeletalMeshImportData->ThresholdUV;
+	BuildOptions.OverlappingThresholds.MorphThresholdPosition = OriginalSkeletalMeshImportData->MorphThresholdPosition;
 	BuildOptions.bComputeNormals = !bShouldImportNormals || !ImportDataDest.bHasNormals;
 	BuildOptions.bComputeTangents = !bShouldImportTangents || !ImportDataDest.bHasTangents;
 	BuildOptions.bUseMikkTSpace = (OriginalSkeletalMeshImportData->NormalGenerationMethod == EFBXNormalGenerationMethod::MikkTSpace) && (!bShouldImportNormals || !bShouldImportTangents);
diff --git a/Engine/Source/Runtime/Engine/Classes/Animation/MorphTarget.h b/Engine/Source/Runtime/Engine/Classes/Animation/MorphTarget.h
--- a/Engine/Source/Runtime/Engine/Classes/Animation/MorphTarget.h
+++ b/Engine/Source/Runtime/Engine/Classes/Animation/MorphTarget.h
@@ -131,7 +131,7 @@ public:
 
 #if WITH_EDITOR
 	/** Populates the given morph target LOD model with the provided deltas */
-	ENGINE_API void PopulateDeltas(const TArray<FMorphTargetDelta>& Deltas, const int32 LODIndex, const TArray<struct FSkelMeshSection>& Sections, const bool bCompareNormal = false, const bool bGeneratedByReductionSetting = false);
+	ENGINE_API void PopulateDeltas(const TArray<FMorphTargetDelta>& Deltas, const int32 LODIndex, const TArray<struct FSkelMeshSection>& Sections, const bool bCompareNormal = false, const bool bGeneratedByReductionSetting = false, const float PositionThreshold = THRESH_POINTS_ARE_NEAR);
 	/** Remove empty LODModels */
 	ENGINE_API void RemoveEmptyMorphTargets();
 #endif // WITH_EDITOR
diff --git a/Engine/Source/Runtime/Engine/Classes/Engine/EngineTypes.h b/Engine/Source/Runtime/Engine/Classes/Engine/EngineTypes.h
--- a/Engine/Source/Runtime/Engine/Classes/Engine/EngineTypes.h
+++ b/Engine/Source/Runtime/Engine/Classes/Engine/EngineTypes.h
@@ -2727,6 +2727,10 @@ struct FSkeletalMeshBuildSettings
 	UPROPERTY(EditAnywhere, Category = BuildSettings)
 	float ThresholdUV;
 
+	/** Threshold to compare vertex position equality when computing morph target deltas. */
+	UPROPERTY(EditAnywhere, Category = BuildSettings)
+	float MorphThresholdPosition;
+
 	/** Default settings. */
 	FSkeletalMeshBuildSettings()
 		: bRecomputeNormals(true)
@@ -2740,6 +2744,7 @@ struct FSkeletalMeshBuildSettings
 		, ThresholdPosition(0.00002)
 		, ThresholdTangentNormal(0.00002)
 		, ThresholdUV(0.0009765625)
+		, MorphThresholdPosition(0.015f)
 	{}
 
 	/** Equality operator. */
@@ -2755,7 +2760,8 @@ struct FSkeletalMeshBuildSettings
 			&& bBuildAdjacencyBuffer == Other.bBuildAdjacencyBuffer
 			&& ThresholdPosition == Other.ThresholdPosition
 			&& ThresholdTangentNormal == Other.ThresholdTangentNormal
-			&& ThresholdUV == Other.ThresholdUV;
+			&& ThresholdUV == Other.ThresholdUV
+			&& MorphThresholdPosition == Other.MorphThresholdPosition;
 	}
 
 	/** Inequality. */
diff --git a/Engine/Source/Runtime/Engine/Private/MorphTools.cpp b/Engine/Source/Runtime/Engine/Private/MorphTools.cpp
--- a/Engine/Source/Runtime/Engine/Private/MorphTools.cpp
+++ b/Engine/Source/Runtime/Engine/Private/MorphTools.cpp
@@ -54,7 +54,7 @@ bool UMorphTarget::HasValidData() const
 
 #if WITH_EDITOR
 
-void UMorphTarget::PopulateDeltas(const TArray<FMorphTargetDelta>& Deltas, const int32 LODIndex, const TArray<FSkelMeshSection>& Sections, const bool bCompareNormal, const bool bGeneratedByReductionSetting)
+void UMorphTarget::PopulateDeltas(const TArray<FMorphTargetDelta>& Deltas, const int32 LODIndex, const TArray<FSkelMeshSection>& Sections, const bool bCompareNormal, const bool bGeneratedByReductionSetting, const float PositionThreshold)
 {
 	// create the LOD entry if it doesn't already exist
 	if (LODIndex >= MorphLODModels.Num())
@@ -80,7 +80,7 @@ void UMorphTarget::PopulateDeltas(const TArray<FMorphTargetDelta>& Deltas, const
 	// Still keep this (could remove in long term due to incoming data)
 	for (const FMorphTargetDelta& Delta : Deltas)
 	{
-		if (Delta.PositionDelta.SizeSquared() > FMath::Square(THRESH_POINTS_ARE_NEAR) || 
+		if (Delta.PositionDelta.SizeSquared() > FMath::Square(PositionThreshold) ||
 			( bCompareNormal && Delta.TangentZDelta.SizeSquared() > 0.01f))
 		{
 			MorphModel.Vertices.Add(Delta);
diff --git a/Engine/Source/Runtime/Engine/Private/SkeletalMesh.cpp b/Engine/Source/Runtime/Engine/Private/SkeletalMesh.cpp
--- a/Engine/Source/Runtime/Engine/Private/SkeletalMesh.cpp
+++ b/Engine/Source/Runtime/Engine/Private/SkeletalMesh.cpp
@@ -3814,6 +3814,7 @@ static void SerializeBuildSettingsForDDC(FArchive& Ar, FSkeletalMeshBuildSetting
 	Ar << BuildSettings.ThresholdPosition;
 	Ar << BuildSettings.ThresholdTangentNormal;
 	Ar << BuildSettings.ThresholdUV;
+	Ar << BuildSettings.MorphThresholdPosition;
 }
 
 
diff --git a/Engine/Source/Runtime/Engine/Private/SkeletalMeshRenderData.cpp b/Engine/Source/Runtime/Engine/Private/SkeletalMeshRenderData.cpp
--- a/Engine/Source/Runtime/Engine/Private/SkeletalMeshRenderData.cpp
+++ b/Engine/Source/Runtime/Engine/Private/SkeletalMeshRenderData.cpp
@@ -58,8 +58,8 @@ static void SerializeLODInfoForDDC(USkeletalMesh* SkeletalMesh, FString& KeySuff
 // If skeletal mesh derived data needs to be rebuilt (new format, serialization
 // differences, etc.) replace the version GUID below with a new one.
 // In case of merge conflicts with DDC versions, you MUST generate a new GUID
-// and set this new GUID as the version.                                       
-#define SKELETALMESH_DERIVEDDATA_VER TEXT("29B218D50FE444AF8EFC8E4234F5D4E1")
+// and set this new GUID as the version.
+#define SKELETALMESH_DERIVEDDATA_VER TEXT("7208E3E18CB047589DD75CD8E66EA1D6")
 
 static const FString& GetSkeletalMeshDerivedDataVersion()
 {
@@ -151,7 +151,6 @@ void FSkeletalMeshRenderData::Cache(USkeletalMesh* Owner)
 		if (GetDerivedDataCacheRef().GetSynchronous(*DerivedDataKey, DerivedData))
 		{
 			COOK_STAT(Timer.AddHit(DerivedData.Num()));
-			
 
 			FMemoryReader Ar(DerivedData, /*bIsPersistent=*/ true);
 
@@ -163,6 +162,30 @@ void FSkeletalMeshRenderData::Cache(USkeletalMesh* Owner)
 			{
 				FSkeletalMeshModel* SkelMeshModel = Owner->GetImportedModel();
 				check(SkelMeshModel);
+
+				int32 MorphTargetNumber = 0;
+				Ar << MorphTargetNumber;
+				Owner->MorphTargets.Empty();
+				//Rebuild the MorphTarget object
+				//We cannot serialize directly the UMorphTarget with a FMemoryArchive. This is not supported.
+				for (int32 MorphTargetIndex = 0; MorphTargetIndex < MorphTargetNumber; ++MorphTargetIndex)
+				{
+					FName MorphTargetName = NAME_None;
+					Ar << MorphTargetName;
+					UMorphTarget* MorphTarget = NewObject<UMorphTarget>(Owner, MorphTargetName);
+					Owner->MorphTargets.Add(MorphTarget);
+					check(MorphTargetIndex == Owner->MorphTargets.Num() - 1);
+					int32 MorphLODModelNumber = 0;
+					Ar << MorphLODModelNumber;
+					Owner->MorphTargets[MorphTargetIndex]->MorphLODModels.AddDefaulted(MorphLODModelNumber);
+					for (int32 MorphDataIndex = 0; MorphDataIndex < MorphLODModelNumber; ++MorphDataIndex)
+					{
+						Ar << Owner->MorphTargets[MorphTargetIndex]->MorphLODModels[MorphDataIndex];
+					}
+				}
+				//Rebuild the mapping and rehook the curve data
+				Owner->InitMorphTargets();
+
 				//Serialize the LODModel sections since they are dependent on the reduction
 				for (int32 LODIndex = 0; LODIndex < SkelMeshModel->LODModels.Num(); LODIndex++)
 				{
@@ -254,6 +277,21 @@ void FSkeletalMeshRenderData::Cache(USkeletalMesh* Owner)
 			//So we do not serialize the LODModel sections.
 			if (!Owner->UseLegacyMeshDerivedDataKey)
 			{
+				int32 MorphTargetNumber = Owner->MorphTargets.Num();
+				Ar << MorphTargetNumber;
+				for (int32 MorphTargetIndex = 0; MorphTargetIndex < MorphTargetNumber; ++MorphTargetIndex)
+				{
+					FName MorphTargetName = Owner->MorphTargets[MorphTargetIndex]->GetFName();
+					Ar << MorphTargetName;
+					int32 MorphLODModelNumber = Owner->MorphTargets[MorphTargetIndex]->MorphLODModels.Num();
+					Ar << MorphLODModelNumber;
+					for (int32 MorphIndex = 0; MorphIndex < MorphLODModelNumber; ++MorphIndex)
+					{
+						Ar << Owner->MorphTargets[MorphTargetIndex]->MorphLODModels[MorphIndex];
+					}
+				}
+				//No need to serialize the morph target mapping since we will rebuild the mapping when loading a ddc
+
 				//Serialize the LODModel sections since they are dependent on the reduction
 				for (int32 LODIndex = 0; LODIndex < SkelMeshModel->LODModels.Num(); LODIndex++)
 				{
diff --git a/Engine/Source/Runtime/Engine/Public/MeshBuild.h b/Engine/Source/Runtime/Engine/Public/MeshBuild.h
--- a/Engine/Source/Runtime/Engine/Public/MeshBuild.h
+++ b/Engine/Source/Runtime/Engine/Public/MeshBuild.h
@@ -17,6 +17,7 @@ public:
 		: ThresholdPosition(THRESH_POINTS_ARE_SAME)
 		, ThresholdTangentNormal(THRESH_NORMALS_ARE_SAME)
 		, ThresholdUV(THRESH_UVS_ARE_SAME)
+		, MorphThresholdPosition(THRESH_POINTS_ARE_NEAR)
 	{}
 
 	/** Threshold use to decide if two vertex position are equal. */
@@ -27,6 +28,9 @@ public:
 	
 	/** Threshold use to decide if two UVs are equal. */
 	float ThresholdUV;
+
+	/** Threshold use to decide if two vertex position are equal. */
+	float MorphThresholdPosition;
 };
 
 
