Reduce the size of the morph target source data, avoid storing the PointToRawMap array.

diff --git a/Engine/Source/Editor/UnrealEd/Private/Fbx/FbxSkeletalMeshImport.cpp b/Engine/Source/Editor/UnrealEd/Private/Fbx/FbxSkeletalMeshImport.cpp
--- a/Engine/Source/Editor/UnrealEd/Private/Fbx/FbxSkeletalMeshImport.cpp
+++ b/Engine/Source/Editor/UnrealEd/Private/Fbx/FbxSkeletalMeshImport.cpp
@@ -4321,6 +4321,8 @@ void UnFbx::FFbxImporter::ImportMorphTargetsInternal( TArray<FbxNode*>& SkelMesh
 		BaseImportData.MorphTargetNames.Add(ShapeName);
 		TSet<uint32>& ModifiedPoints = BaseImportData.MorphTargetModifiedPoints.AddDefaulted_GetRef();
 		GatherPointsForMorphTarget(&ShapeImportData, SkelMeshNodeArray, &ShapeArray, ModifiedPoints);
+		//We do not need this data anymore empty it so we reduce the size of what we save into memory
+		ShapeImportData.PointToRawMap.Empty();
 		BaseImportData.MorphTargets.Add(ShapeImportData);
 		check(BaseImportData.MorphTargetNames.Num() == BaseImportData.MorphTargets.Num() && BaseImportData.MorphTargetNames.Num() == BaseImportData.MorphTargetModifiedPoints.Num());
 	}
diff --git a/Engine/Source/Runtime/Engine/Private/Rendering/SkeletalMeshLODImporterData.cpp b/Engine/Source/Runtime/Engine/Private/Rendering/SkeletalMeshLODImporterData.cpp
--- a/Engine/Source/Runtime/Engine/Private/Rendering/SkeletalMeshLODImporterData.cpp
+++ b/Engine/Source/Runtime/Engine/Private/Rendering/SkeletalMeshLODImporterData.cpp
@@ -15,7 +15,9 @@ DEFINE_LOG_CATEGORY_STATIC(LogSkeletalMeshLODImporterData, Log, All);
 
 void FSkeletalMeshImportData::CopyDataNeedByMorphTargetImport(FSkeletalMeshImportData& Other) const
 {
+	//The points array is the only data we need to compute the morph target in the skeletalmesh build
 	Other.Points = Points;
+	//PointToRawMap should not be save when saving morph target data, we only need it temporary to gather the point from the fbx shape
 	Other.PointToRawMap = PointToRawMap;
 	Other.bDiffPose = bDiffPose;
 	Other.bUseT0AsRefPose = bUseT0AsRefPose;
