Support for overriding EngineShowFlags in EditorViewportClient. External code can provide a TFunction that is given the opportunity to edit the FEngineShowFlags instance used in FEditorViewportClient::Draw. This is only a per-frame change, the .EngineShowFlags member is not modified. DisableOverrideEngineShowFlags() must be called to remove the override. This allows (for example) an EdMode to temporarily disable rendering features like TAA or Motion Blur which might be undesirable in content-creation / visualization contexts.

diff --git a/Engine/Source/Editor/UnrealEd/Private/EditorViewportClient.cpp b/Engine/Source/Editor/UnrealEd/Private/EditorViewportClient.cpp
--- a/Engine/Source/Editor/UnrealEd/Private/EditorViewportClient.cpp
+++ b/Engine/Source/Editor/UnrealEd/Private/EditorViewportClient.cpp
@@ -3657,16 +3657,22 @@ void FEditorViewportClient::Draw(FViewport* InViewport, FCanvas* Canvas)
 	Canvas->SetScaledToRenderTarget(bStereoRendering);
 	Canvas->SetStereoRendering(bStereoRendering);
 
+	FEngineShowFlags UseEngineShowFlags = EngineShowFlags;
+	if (OverrideShowFlagsFunc)
+	{
+		OverrideShowFlagsFunc(UseEngineShowFlags);
+	}
+
 	// Setup a FSceneViewFamily/FSceneView for the viewport.
 	FSceneViewFamilyContext ViewFamily(FSceneViewFamily::ConstructionValues(
 		Canvas->GetRenderTarget(),
 		GetScene(),
-		EngineShowFlags)
+		UseEngineShowFlags)
 		.SetWorldTimes( TimeSeconds, DeltaTimeSeconds, RealTimeSeconds )
 		.SetRealtimeUpdate( IsRealtime() && FSlateThrottleManager::Get().IsAllowingExpensiveTasks() )
 		.SetViewModeParam( ViewModeParam, ViewModeParamName ) );
 
-	ViewFamily.EngineShowFlags = EngineShowFlags;
+	ViewFamily.EngineShowFlags = UseEngineShowFlags;
 
 	ViewFamily.bIsHDR = Viewport->IsHDRViewport();
 
@@ -3674,7 +3680,7 @@ void FEditorViewportClient::Draw(FViewport* InViewport, FCanvas* Canvas)
 
 	if( ModeTools->GetActiveMode( FBuiltinEditorModes::EM_InterpEdit ) == 0 || !AllowsCinematicControl() )
 	{
-		if( !EngineShowFlags.Game )
+		if( !UseEngineShowFlags.Game )
 		{
 			// in the editor, disable camera motion blur and other rendering features that rely on the former frame
 			// unless the view port is Matinee controlled
@@ -5876,6 +5882,16 @@ FMatrix FEditorViewportClient::CalcViewRotationMatrix(const FRotator& InViewRota
 	}
 }
 
+void FEditorViewportClient::EnableOverrideEngineShowFlags(TUniqueFunction<void(FEngineShowFlags&)> OverrideFunc)
+{
+	OverrideShowFlagsFunc = MoveTemp(OverrideFunc);
+}
+
+void FEditorViewportClient::DisableOverrideEngineShowFlags()
+{
+	OverrideShowFlagsFunc = nullptr;
+}
+
 ////////////////
 
 bool FEditorViewportStats::bInitialized(false);
diff --git a/Engine/Source/Editor/UnrealEd/Public/EditorViewportClient.h b/Engine/Source/Editor/UnrealEd/Public/EditorViewportClient.h
--- a/Engine/Source/Editor/UnrealEd/Public/EditorViewportClient.h
+++ b/Engine/Source/Editor/UnrealEd/Public/EditorViewportClient.h
@@ -1632,6 +1632,26 @@ protected:
 	FVector DefaultOrbitZoom;
 	FVector DefaultOrbitLookAt;
 
+public:
+	/**
+	 * Enable customization of the EngineShowFlags for rendering. After calling this function,
+	 * the provided OverrideFunc will be passed a copy of .EngineShowFlags in ::Draw() just before rendering setup.
+	 * Changes made to the ShowFlags will be used for that frame but .EngineShowFlags will not be modified.
+	 * @param OverrideFunc custom override function that will be called every frame until override is disabled.
+	 */
+	void EnableOverrideEngineShowFlags(TUniqueFunction<void(FEngineShowFlags&)> OverrideFunc);
+
+	/** Disable EngineShowFlags override if enabled */
+	void DisableOverrideEngineShowFlags();
+
+	/** @return true if Override EngineShowFlags are currently enabled */
+	bool IsEngineShowFlagsOverrideEnabled() const { return !! OverrideShowFlagsFunc; }
+
+protected:
+	
+	/** Custom override function that will be called every ::Draw() until override is disabled */
+	TUniqueFunction<void(FEngineShowFlags&)> OverrideShowFlagsFunc;
+
 protected:
 	// Used for the display of the current preview light after it has been adjusted
 	FVector2D MovingPreviewLightSavedScreenPos;
