Fixed incorrectly reported non-determinism in DiffAssetRegistries.

If an asset's dependencies change, the asset itself can be unaffected, but assets which reference it can change. Therefore, we shouldn't stop propagating dependency change flags just because an asset doesn't have the GuidChange or HashChange flags. A common example is when material instance A references material instance B, which in turn references material M. If M changes, it's possible that B stays the same, but A can still change; A was incorrectly reported as non-deterministic in this case.

diff --git a/Engine/Source/Editor/UnrealEd/Private/Commandlets/DiffAssetRegistriesCommandlet.cpp b/Engine/Source/Editor/UnrealEd/Private/Commandlets/DiffAssetRegistriesCommandlet.cpp
--- a/Engine/Source/Editor/UnrealEd/Private/Commandlets/DiffAssetRegistriesCommandlet.cpp
+++ b/Engine/Source/Editor/UnrealEd/Private/Commandlets/DiffAssetRegistriesCommandlet.cpp
@@ -1310,9 +1310,16 @@ void UDiffAssetRegistriesCommandlet::DiffAssetRegistries(const FString& OldPath,
 			TArray<FAssetIdentifier> Referencers;
 			
 			// grab the hash/guid change flags, shift up to the dependency ones
-			int32 NewFlags = (AssetPathFlags.FindOrAdd(Package) & 0x0C) << 2;
+			int32 PackageFlags = AssetPathFlags.FindOrAdd(Package);
+			int32 NewFlags = (PackageFlags & 0x0C) << 2;
+
+			// If we have a dependency chain like C -> B -> A, and A changes, this does not
+			// necessarily cause B's binary representation to change, but it can still impact C.
+			// We must propagate the dependency change flags too, otherwise C will be marked as
+			// non-deterministic when this happens.
+			NewFlags |= PackageFlags & (EAssetFlags::DepGuidChange | EAssetFlags::DepHashChange);
 			
-			// don't bother touching anything if this asset didn't change
+			// don't bother touching anything if this asset or its dependencies didn't change
 			if (NewFlags)
 			{
 				NewState.GetReferencers(Package, Referencers, EAssetRegistryDependencyType::Hard);
