Removed logging added to track UE-20486

diff --git a/Engine/Source/Editor/UnrealEd/Private/PlayLevel.cpp b/Engine/Source/Editor/UnrealEd/Private/PlayLevel.cpp
--- a/Engine/Source/Editor/UnrealEd/Private/PlayLevel.cpp
+++ b/Engine/Source/Editor/UnrealEd/Private/PlayLevel.cpp
@@ -176,33 +176,6 @@ void UEditorEngine::EndPlayMap()
 
 	FlushAsyncLoading();
 
-	// Monitoring when PIE corrupts references between the World and the PIE generated World for UE-20486
-	{
-		TArray<ULevel*> Levels = EditorWorld->GetLevels();
-
-		for (ULevel* Level : Levels)
-		{
-			TArray<UBlueprint*> LevelBlueprints = Level->GetLevelBlueprints();
-
-			if (LevelBlueprints.Num() > 0)
-			{
-				UBlueprint* LevelScriptBlueprint = LevelBlueprints[0];
-				if (LevelScriptBlueprint && LevelScriptBlueprint->GeneratedClass && LevelScriptBlueprint->GeneratedClass->ClassGeneratedBy)
-				{
-					UE_LOG(LogBlueprintUserMessages, Log, TEXT("Early EndPlayMap Detection: Level '%s' has LevelScriptBlueprint '%s' with GeneratedClass '%s' with ClassGeneratedBy '%s'"), *Level->GetPathName(), *LevelScriptBlueprint->GetPathName(), *LevelScriptBlueprint->GeneratedClass->GetPathName(), *LevelScriptBlueprint->GeneratedClass->ClassGeneratedBy->GetPathName());
-				}
-				else if (LevelScriptBlueprint && LevelScriptBlueprint->GeneratedClass)
-				{
-					UE_LOG(LogBlueprintUserMessages, Log, TEXT("Early EndPlayMap Detection: Level '%s' has LevelScriptBlueprint '%s' with GeneratedClass '%s'"), *Level->GetPathName(), *LevelScriptBlueprint->GetPathName(), *LevelScriptBlueprint->GeneratedClass->GetPathName());
-				}
-				else if (LevelScriptBlueprint)
-				{
-					UE_LOG(LogBlueprintUserMessages, Log, TEXT("Early EndPlayMap Detection: Level '%s' has LevelScriptBlueprint '%s'"), *Level->GetPathName(), *LevelScriptBlueprint->GetPathName());
-				}
-			}
-		}
-	}
-
 	if (GEngine->XRSystem.IsValid() && !bIsSimulatingInEditor)
 	{
 		GEngine->XRSystem->OnEndPlay(*GEngine->GetWorldContextFromWorld(PlayWorld));
@@ -371,34 +344,7 @@ void UEditorEngine::EndPlayMap()
 	FNavigationSystem::OnPIEEnd(*EditorWorld);
 
 	FGameDelegates::Get().GetEndPlayMapDelegate().Broadcast();
-
-	// Monitoring when PIE corrupts references between the World and the PIE generated World for UE-20486
-	{
-		TArray<ULevel*> Levels = EditorWorld->GetLevels();
-
-		for (ULevel* Level : Levels)
-		{
-			TArray<UBlueprint*> LevelBlueprints = Level->GetLevelBlueprints();
-
-			if (LevelBlueprints.Num() > 0)
-			{
-				UBlueprint* LevelScriptBlueprint = LevelBlueprints[0];
-				if (LevelScriptBlueprint && LevelScriptBlueprint->GeneratedClass && LevelScriptBlueprint->GeneratedClass->ClassGeneratedBy)
-				{
-					UE_LOG(LogBlueprintUserMessages, Log, TEXT("Late EndPlayMap Detection: Level '%s' has LevelScriptBlueprint '%s' with GeneratedClass '%s' with ClassGeneratedBy '%s'"), *Level->GetPathName(), *LevelScriptBlueprint->GetPathName(), *LevelScriptBlueprint->GeneratedClass->GetPathName(), *LevelScriptBlueprint->GeneratedClass->ClassGeneratedBy->GetPathName());
-				}
-				else if (LevelScriptBlueprint && LevelScriptBlueprint->GeneratedClass)
-				{
-					UE_LOG(LogBlueprintUserMessages, Log, TEXT("Late EndPlayMap Detection: Level '%s' has LevelScriptBlueprint '%s' with GeneratedClass '%s'"), *Level->GetPathName(), *LevelScriptBlueprint->GetPathName(), *LevelScriptBlueprint->GeneratedClass->GetPathName());
-				}
-				else if (LevelScriptBlueprint)
-				{
-					UE_LOG(LogBlueprintUserMessages, Log, TEXT("Late EndPlayMap Detection: Level '%s' has LevelScriptBlueprint '%s'"), *Level->GetPathName(), *LevelScriptBlueprint->GetPathName());
-				}
-			}
-		}
-	}
-
+	
 	// find objects like Textures in the playworld levels that won't get garbage collected as they are marked RF_Standalone
 	for (FObjectIterator It; It; ++It)
 	{
