Restore full search for dependent blueprints when not compiling

diff --git a/Engine/Plugins/Tests/EditorTests/Source/EditorTests/Private/UnrealEd/BlueprintAutomationTests.cpp b/Engine/Plugins/Tests/EditorTests/Source/EditorTests/Private/UnrealEd/BlueprintAutomationTests.cpp
--- a/Engine/Plugins/Tests/EditorTests/Source/EditorTests/Private/UnrealEd/BlueprintAutomationTests.cpp
+++ b/Engine/Plugins/Tests/EditorTests/Source/EditorTests/Private/UnrealEd/BlueprintAutomationTests.cpp
@@ -902,7 +902,7 @@ bool FBlueprintCompileOnLoadTest::RunTest(const FString& BlueprintAssetPath)
 	TSet<TWeakObjectPtr<UBlueprint>> BlueprintDependencies;
 	{
 		TArray<UBlueprint*> DependentBlueprints;
-		FBlueprintEditorUtils::GetDependentBlueprints(InitialBlueprint, DependentBlueprints);
+		FBlueprintEditorUtils::FindDependentBlueprints(InitialBlueprint, DependentBlueprints);
 		for (auto BP : DependentBlueprints)
 		{
 			BlueprintDependencies.Add(BP);
diff --git a/Engine/Source/Editor/UnrealEd/Private/Kismet2/BlueprintEditorUtils.cpp b/Engine/Source/Editor/UnrealEd/Private/Kismet2/BlueprintEditorUtils.cpp
--- a/Engine/Source/Editor/UnrealEd/Private/Kismet2/BlueprintEditorUtils.cpp
+++ b/Engine/Source/Editor/UnrealEd/Private/Kismet2/BlueprintEditorUtils.cpp
@@ -1,6 +1,8 @@
 // Copyright 1998-2019 Epic Games, Inc. All Rights Reserved.
 
 #include "Kismet2/BlueprintEditorUtils.h"
+#include "Algo/Copy.h"
+#include "Algo/Transform.h"
 #include "BlueprintCompilationManager.h"
 #include "UObject/Interface.h"
 #include "Engine/BlueprintGeneratedClass.h"
@@ -2955,6 +2957,68 @@ void FBlueprintEditorUtils::GetDependentBlueprints(UBlueprint* Blueprint, TArray
 	}
 }
 
+void FBlueprintEditorUtils::FindDependentBlueprints(UBlueprint* Blueprint, TArray<UBlueprint*>& DependentBlueprints)
+{	
+	if(Blueprint == nullptr)
+	{
+		return;
+	}
+
+	TArray<UObject*> AllBlueprints;
+	bool const bIncludeDerivedClasses = true;
+	GetObjectsOfClass(UBlueprint::StaticClass(), AllBlueprints, bIncludeDerivedClasses );
+
+	// Sanitize, add correct type.. can't find a UObject* helper to do this, and 
+	// the previous version of htis code checked IsPendingKill():
+	TArray<UBlueprint*> AllBlueprintSafe;
+	Algo::TransformIf(AllBlueprints, AllBlueprintSafe, 
+		[](UObject* Obj)->bool { return Obj && !Obj->IsPendingKill(); }, 
+		[](UObject* Obj)->UBlueprint* { return static_cast<UBlueprint*>(Obj); } 
+	);
+
+	// Update *all* dependencies:
+	for (UBlueprint* TestBP : AllBlueprintSafe)
+	{
+		EnsureCachedDependenciesUpToDate(TestBP);
+	}
+
+	// Gather macro blueprints that we're dependent on:
+	TSet<UBlueprint*> DepSet;
+	Algo::CopyIf(AllBlueprintSafe, DepSet,
+		[Blueprint](UBlueprint* TestBP) -> bool 
+		{ 
+			if(TestBP->CachedDependencies.Contains(Blueprint))
+			{
+				if (TestBP->BlueprintType == BPTYPE_MacroLibrary)
+				{
+					return true;
+				}
+			}
+			return false;
+		}
+	);
+	DepSet.Add(Blueprint);
+	
+	// Find all blueprints that have a cached dep on Blueprint *or* one of the 
+	// macros that is dependent on Blueprint:
+	Algo::CopyIf(AllBlueprintSafe, DependentBlueprints,
+		[DepSet](UBlueprint* TestBP) -> bool
+		{
+			// check for interesection between CachedDependencies and DepSet:
+			for(TWeakObjectPtr<UBlueprint> ObjWeak : TestBP->CachedDependencies)
+			{
+				if(DepSet.Contains(ObjWeak.Get()))
+				{
+					return true;
+				}
+			}
+			return false;
+		}
+	);
+
+	// Remove ourself from the list of dependents:
+	DependentBlueprints.RemoveSwap(Blueprint);
+}
 
 bool FBlueprintEditorUtils::IsGraphIntermediate(const UEdGraph* Graph)
 {
@@ -4562,7 +4626,7 @@ void FBlueprintEditorUtils::RenameComponentMemberVariable(UBlueprint* Blueprint,
 		{
 			const FComponentKey Key(Node);
 			TArray<UBlueprint*> Dependents;
-			GetDependentBlueprints(Blueprint, Dependents);
+			FindDependentBlueprints(Blueprint, Dependents);
 			for (UBlueprint* DepBP : Dependents)
 			{
 				UInheritableComponentHandler* InheritableComponentHandler = DepBP ? DepBP->GetInheritableComponentHandler(false) : nullptr;
@@ -5391,7 +5455,7 @@ void FBlueprintEditorUtils::ReplaceVariableReferences(UBlueprint* Blueprint, con
 	FBlueprintEditorUtils::RenameVariableReferences(Blueprint, Blueprint->GeneratedClass, OldName, NewName);
 
 	TArray<UBlueprint*> Dependents;
-	GetDependentBlueprints(Blueprint, Dependents);
+	FindDependentBlueprints(Blueprint, Dependents);
 
 	for (UBlueprint* DependentBp : Dependents)
 	{
diff --git a/Engine/Source/Editor/UnrealEd/Public/Kismet2/BlueprintEditorUtils.h b/Engine/Source/Editor/UnrealEd/Public/Kismet2/BlueprintEditorUtils.h
--- a/Engine/Source/Editor/UnrealEd/Public/Kismet2/BlueprintEditorUtils.h
+++ b/Engine/Source/Editor/UnrealEd/Public/Kismet2/BlueprintEditorUtils.h
@@ -523,9 +523,12 @@ public:
 	/** Gather all bps that Blueprint depends on */
 	static void GatherDependencies(const UBlueprint* Blueprint, TSet<TWeakObjectPtr<UBlueprint>>& OutDependencies, TSet<TWeakObjectPtr<UStruct>>& OutUDSDependencies);
 
-	/** Returns a list of loaded Blueprints that are dependent on the given Blueprint. */
+	/** Returns cached a list of loaded Blueprints that are dependent on the given Blueprint. */
 	static void GetDependentBlueprints(UBlueprint* Blueprint, TArray<UBlueprint*>& DependentBlueprints);
 
+	/** Searches the reference graph to find blueprints that are dependent on this BP */
+	static void FindDependentBlueprints(UBlueprint* Blueprint, TArray<UBlueprint*>& DependentBlueprints);
+
 	/** Ensures, that CachedDependencies in BP are up to date */
 	static void EnsureCachedDependenciesUpToDate(UBlueprint* Blueprint);
 
diff --git a/Engine/Source/Runtime/Engine/Private/Blueprint.cpp b/Engine/Source/Runtime/Engine/Private/Blueprint.cpp
--- a/Engine/Source/Runtime/Engine/Private/Blueprint.cpp
+++ b/Engine/Source/Runtime/Engine/Private/Blueprint.cpp
@@ -571,7 +571,7 @@ bool UBlueprint::Rename( const TCHAR* InName, UObject* NewOuter, ERenameFlags Fl
 	{
 		// Gather all blueprints that currently depend on this one.
 		TArray<UBlueprint*> Dependents;
-		FBlueprintEditorUtils::GetDependentBlueprints(this, Dependents);
+		FBlueprintEditorUtils::FindDependentBlueprints(this, Dependents);
 
 		FKismetEditorUtilities::CompileBlueprint(this);
 
