Update force deleting to also replace references on Actors with no owning World (previously these actors were just ignored). This ensures in the correct results of objects created in actor components such as Child Actor Component and in Sequencer.

diff --git a/Engine/Source/Editor/UnrealEd/Private/ObjectTools.cpp b/Engine/Source/Editor/UnrealEd/Private/ObjectTools.cpp
--- a/Engine/Source/Editor/UnrealEd/Private/ObjectTools.cpp
+++ b/Engine/Source/Editor/UnrealEd/Private/ObjectTools.cpp
@@ -2630,6 +2630,9 @@ namespace ObjectTools
 
 		bool bSelectionChanged = false;
 
+		TArray<UObject*> ObjectsToReplace;
+		ObjectsToReplace.Reserve(ObjectsToDelete.Num());
+
 		// Destroy all Components
 		if (ComponentsToDelete.Num() > 0)
 		{
@@ -2685,6 +2688,12 @@ namespace ObjectTools
 					{
 						CurActor->GetWorld()->EditorDestroyActor( CurActor, false );
 					}
+					// Ensure that we replace any generated actors who don't have worlds that are left such as the template
+					// from Child Actor Components
+					else
+					{
+						ObjectsToReplace.Add(CurActor);
+					}
 
 					bNeedsGarbageCollection = true;
 				}
@@ -2707,9 +2716,6 @@ namespace ObjectTools
 		{
 			int32 ReplaceableObjectsNum = 0;
 			{
-				TArray<UObject*> ObjectsToReplace;
-				ObjectsToReplace.Reserve(ObjectsToDelete.Num());
-
 				for(TWeakObjectPtr<UObject>& Object : ObjectsToDelete)
 				{
 					if(Object.IsValid())
@@ -2767,7 +2773,7 @@ namespace ObjectTools
 					TArray<UObject*> UDStructToReplace;
 					for (int32 Iter = 0; Iter < ObjectsToReplace.Num(); )
 					{
-						if (auto UDStruct = Cast<UUserDefinedStruct>(ObjectsToReplace[Iter]))
+						if (UUserDefinedStruct* UDStruct = Cast<UUserDefinedStruct>(ObjectsToReplace[Iter]))
 						{
 							ObjectsToReplace.RemoveAtSwap(Iter);
 							UDStructToReplace.Add(UDStruct);
@@ -2788,7 +2794,7 @@ namespace ObjectTools
 
 				{
 					FForceReplaceInfo ReplaceInfo;
-					ForceReplaceReferences(NULL, ObjectsToReplace, ReplaceInfo, false);
+					ForceReplaceReferences(nullptr, ObjectsToReplace, ReplaceInfo, false);
 					ReplaceableObjectsNum += ReplaceInfo.ReplaceableObjects.Num();
 				}
 			}
@@ -2807,7 +2813,7 @@ namespace ObjectTools
 			for(auto It = ObjectsToDelete.CreateIterator(); It; ++It)
 			{
 				UObject* CurObject = It->Get();
-				if ( !ensure(CurObject != NULL) )
+				if ( !ensure(CurObject != nullptr) )
 				{
 					continue;
 				}
