Fix CIS static analysis
#rnx

diff --git a/Engine/Source/Editor/UnrealEd/Private/Fbx/FbxSkeletalMeshImport.cpp b/Engine/Source/Editor/UnrealEd/Private/Fbx/FbxSkeletalMeshImport.cpp
--- a/Engine/Source/Editor/UnrealEd/Private/Fbx/FbxSkeletalMeshImport.cpp
+++ b/Engine/Source/Editor/UnrealEd/Private/Fbx/FbxSkeletalMeshImport.cpp
@@ -3718,14 +3718,20 @@ bool UnFbx::FFbxImporter::FillSkelMeshImporterFromFbx( FSkeletalMeshImportData&
 	//
 	// clean up
 	//
-	if (UniqueUVCount > 0)
+	if (LayerElementUV)
 	{
 		delete[] LayerElementUV;
+	}
+	if (UVReferenceMode)
+	{
 		delete[] UVReferenceMode;
+	}
+	if (UVMappingMode)
+	{
 		delete[] UVMappingMode;
 	}
 	
-	return true; //-V773
+	return true;
 }
 
 void UnFbx::FFbxImporter::InsertNewLODToBaseSkeletalMesh(USkeletalMesh* InSkeletalMesh, USkeletalMesh* BaseSkeletalMesh, int32 DesiredLOD, UFbxSkeletalMeshImportData* TemplateImportData)
