TBA: Expose "view as text" and "roundtrip" functions to the content browser context menu in the editor

diff --git a/Engine/Source/Editor/ContentBrowser/Private/AssetContextMenu.cpp b/Engine/Source/Editor/ContentBrowser/Private/AssetContextMenu.cpp
--- a/Engine/Source/Editor/ContentBrowser/Private/AssetContextMenu.cpp
+++ b/Engine/Source/Editor/ContentBrowser/Private/AssetContextMenu.cpp
@@ -76,6 +76,9 @@
 #include "PackageHelperFunctions.h"
 #include "EngineUtils.h"
 
+#include "Commandlets/TextAssetCommandlet.h"
+#include "Misc/FileHelper.h"
+
 #define LOCTEXT_NAMESPACE "ContentBrowser"
 
 FAssetContextMenu::FAssetContextMenu(const TWeakPtr<SAssetView>& InAssetView)
@@ -645,6 +648,21 @@ void FAssetContextMenu::MakeAssetActionsSubMenu(FMenuBuilder& MenuBuilder)
 				FSlateIcon(),
 				FUIAction(FExecuteAction::CreateSP(this, &FAssetContextMenu::ExportSelectedAssetsToText))
 			);
+
+			MenuBuilder.AddMenuEntry(
+				LOCTEXT("ViewSelectedAssetAsText", "View as text"),
+				LOCTEXT("ViewSelectedAssetAsTextTooltip", "Opens a window showing the selected asset in text format"),
+				FSlateIcon(),
+				FUIAction(FExecuteAction::CreateSP(this, &FAssetContextMenu::ViewSelectedAssetAsText),
+					FCanExecuteAction::CreateSP(this, &FAssetContextMenu::CanViewSelectedAssetAsText))
+			);
+
+			MenuBuilder.AddMenuEntry(
+				LOCTEXT("TextFormatRountrip", "Run Text Asset Roundtrip"),
+				LOCTEXT("TextFormatRountripTooltip", "Save the select asset backwards or forwards between text and binary formats and check for determinism"),
+				FSlateIcon(),
+				FUIAction(FExecuteAction::CreateSP(this, &FAssetContextMenu::DoTextFormatRoundtrip))
+			);
 		}
 		MenuBuilder.EndSection();
 	}
@@ -678,9 +696,55 @@ void FAssetContextMenu::ExportSelectedAssetsToText()
 	}
 }
 
-bool FAssetContextMenu::CanExportSelectedAssetsToText() const
+void FAssetContextMenu::ViewSelectedAssetAsText()
 {
-	return true;
+	if (SelectedAssets.Num() == 1)
+	{
+		UPackage* Package = SelectedAssets[0].GetPackage();
+		FString TargetFilename = FPaths::CreateTempFilename(*FPaths::ProjectSavedDir(), nullptr, *FPackageName::GetTextAssetPackageExtension());
+		if (SavePackageHelper(Package, TargetFilename))
+		{
+			FString TextFormat;
+			if (FFileHelper::LoadFileToString(TextFormat, *TargetFilename))
+			{
+				SGenericDialogWidget::OpenDialog(LOCTEXT("TextAssetViewerTitle", "Viewing AS Text Asset..."), SNew(STextBlock).Text(FText::FromString(TextFormat)));
+			}
+			IFileManager::Get().Delete(*TargetFilename);
+		}
+	}
+}
+
+bool FAssetContextMenu::CanViewSelectedAssetAsText() const
+{
+	return SelectedAssets.Num() == 1;
+}
+
+void FAssetContextMenu::DoTextFormatRoundtrip()
+{
+	UTextAssetCommandlet::FProcessingArgs Args;
+	Args.NumSaveIterations = 1;
+	Args.bIncludeEngineContent = true;
+	Args.bVerifyJson = true;
+	Args.CSVFilename = TEXT("");
+	Args.ProcessingMode = UTextAssetCommandlet::EProcessingMode::RoundTrip;
+	Args.bFilenameIsFilter = false;
+
+	for (const FAssetData& Asset : SelectedAssets)
+	{
+		UPackage* Package = Asset.GetPackage();
+		Args.Filename = FPackageName::LongPackageNameToFilename(Package->GetPathName());
+		if (!UTextAssetCommandlet::DoTextAssetProcessing(Args))
+		{
+			FNotificationInfo Info(LOCTEXT("RountripTextAssetFailed", "Roundtripping of selected asset(s) failed"));
+			Info.ExpireDuration = 3.0f;
+			FSlateNotificationManager::Get().AddNotification(Info);
+			return;
+		}
+	}
+
+	FNotificationInfo Info(LOCTEXT("RoundtripTextAssetsSuccessfully", "Roundtripped selected asset(s) successfully"));
+	Info.ExpireDuration = 3.0f;
+	FSlateNotificationManager::Get().AddNotification(Info);
 }
 
 bool FAssetContextMenu::CanExecuteAssetActions() const
diff --git a/Engine/Source/Editor/ContentBrowser/Private/AssetContextMenu.h b/Engine/Source/Editor/ContentBrowser/Private/AssetContextMenu.h
--- a/Engine/Source/Editor/ContentBrowser/Private/AssetContextMenu.h
+++ b/Engine/Source/Editor/ContentBrowser/Private/AssetContextMenu.h
@@ -288,8 +288,12 @@ private:
 	/** Handler to export the selected asset(s) to experimental text format */
 	void ExportSelectedAssetsToText();
 
-	/** Handler to check if we can export the selected asset(s) to experimental text format */
-	bool CanExportSelectedAssetsToText() const;
+	/** Handler to export the selected asset(s) to experimental text format */
+	void ViewSelectedAssetAsText();
+	bool CanViewSelectedAssetAsText() const;
+
+	/** Run the rountrip test on this asset */
+	void DoTextFormatRoundtrip();
 
 	/** Handler to check to see if a sync to asset tree command is allowed */
 	bool CanExecuteSyncToAssetTree() const;
diff --git a/Engine/Source/Editor/UnrealEd/Classes/Commandlets/TextAssetCommandlet.h b/Engine/Source/Editor/UnrealEd/Classes/Commandlets/TextAssetCommandlet.h
--- a/Engine/Source/Editor/UnrealEd/Classes/Commandlets/TextAssetCommandlet.h
+++ b/Engine/Source/Editor/UnrealEd/Classes/Commandlets/TextAssetCommandlet.h
@@ -18,9 +18,33 @@ class UTextAssetCommandlet
 
 public:
 
+	enum class EProcessingMode
+	{
+		ResaveText,
+		ResaveBinary,
+		RoundTrip,
+		LoadBinary,
+		LoadText,
+		FindMismatchedSerializers,
+	};
+
+	struct FProcessingArgs
+	{
+		EProcessingMode ProcessingMode = EProcessingMode::ResaveText;
+		int32 NumSaveIterations = 1;
+		bool bIncludeEngineContent = false;
+		bool bFilenameIsFilter = true;
+		FString Filename;
+		FString CSVFilename;
+		FString OutputPath;
+		bool bVerifyJson = true;
+	};
+
 	//~ Begin UCommandlet Interface
 
 	virtual int32 Main(const FString& CmdLineParams) override;
 	
 	//~ End UCommandlet Interface
+	static UNREALED_API bool DoTextAssetProcessing(const FString& InCommandLine);
+	static UNREALED_API bool DoTextAssetProcessing(const FProcessingArgs& InArgs);
 };
\ No newline at end of file
diff --git a/Engine/Source/Editor/UnrealEd/Private/Commandlets/TextAssetCommandlet.cpp b/Engine/Source/Editor/UnrealEd/Private/Commandlets/TextAssetCommandlet.cpp
--- a/Engine/Source/Editor/UnrealEd/Private/Commandlets/TextAssetCommandlet.cpp
+++ b/Engine/Source/Editor/UnrealEd/Private/Commandlets/TextAssetCommandlet.cpp
@@ -50,19 +50,17 @@ void FindMismatchedSerializers()
 	}
 }
 
-int32 UTextAssetCommandlet::Main(const FString& CmdLineParams)
+namespace 
 {
-	TRACE_CPUPROFILER_EVENT_SCOPE_TEXT(TEXT("UTextAssetCommandlet::Main"));
-
-	TArray<FString> Blacklist;
-
 	static const FString BackupExtension = TEXT("textassetbackup");
 	static const FString BackupRoundtripExtension = TEXT("textassetbackup_roundtrip");
 	static const FString BackupExtension_WithDot = TEXT(".") + BackupExtension;
 	static const FString BackupRoundtripExtension_WithDot = TEXT(".") + BackupRoundtripExtension;
+}
 
+static void RepairDamagedFiles()
+{
 	// Repair any damage caused by a failed run of this commandlet
-	TArray<FString> IntermediateFilenames;
 	struct FVisitor : public IPlatformFile::FDirectoryVisitor
 	{
 		virtual bool Visit(const TCHAR* FilenameOrDirectory, bool bIsDirectory) override
@@ -88,67 +86,95 @@ int32 UTextAssetCommandlet::Main(const FString& CmdLineParams)
 
 	IFileManager::Get().IterateDirectoryRecursively(*FPaths::ProjectContentDir(), RepairVisitor);
 	IFileManager::Get().IterateDirectoryRecursively(*FPaths::EngineContentDir(), RepairVisitor);
+}
+
+bool UTextAssetCommandlet::DoTextAssetProcessing(const FString& InCommandLine)
+{
+	FProcessingArgs Args;
 
 	FString ModeString = TEXT("ResaveText");
 	FString IterationsString = TEXT("1");
-	FString FilenameFilterString, OutputPathString, CSVFilename;
 
-	FParse::Value(*CmdLineParams, TEXT("csv="), CSVFilename);
-	FParse::Value(*CmdLineParams, TEXT("mode="), ModeString);
-	FParse::Value(*CmdLineParams, TEXT("filter="), FilenameFilterString);
-	FParse::Value(*CmdLineParams, TEXT("outputpath="), OutputPathString);
-	bool bVerifyJson = !FParse::Param(*CmdLineParams, TEXT("noverifyjson"));
+	FString Filename, FilenameFilter;
 
-	enum class EMode
-	{
-		ResaveText,
-		ResaveBinary,
-		RoundTrip,
-		LoadBinary,
-		LoadText,
-		FindMismatchedSerializers,
-	};
-
-	TMap<FString, EMode> Modes;
-	Modes.Add(TEXT("ResaveText"), EMode::ResaveText);
-	Modes.Add(TEXT("ResaveBinary"), EMode::ResaveBinary);
-	Modes.Add(TEXT("RoundTrip"), EMode::RoundTrip);
-	Modes.Add(TEXT("LoadText"), EMode::LoadText);
-	Modes.Add(TEXT("LoadBinary"), EMode::LoadBinary);
-	Modes.Add(TEXT("FindMismatchedSerializers"), EMode::FindMismatchedSerializers);
+	FParse::Value(*InCommandLine, TEXT("mode="), ModeString);
+	FParse::Value(*InCommandLine, TEXT("filename="), Filename);
+	FParse::Value(*InCommandLine, TEXT("filter="), FilenameFilter);
+	FParse::Value(*InCommandLine, TEXT("csv="), Args.CSVFilename);
+	FParse::Value(*InCommandLine, TEXT("outputpath="), Args.OutputPath);
 
-	check(Modes.Contains(ModeString));
-	EMode Mode = Modes[ModeString];
+	if (Filename.Len() > 0 && FilenameFilter.Len() > 0)
+	{
+		UE_LOG(LogTextAsset, Error, TEXT("Cannot specify a filename and a filter at the same time when processing text assets"));
+		return false;
+	}
 
-	if (Mode == EMode::FindMismatchedSerializers)
+	if (Filename.Len() > 0)
 	{
-		FindMismatchedSerializers();
-		return 0;
+		Args.Filename = Filename;
+		Args.bFilenameIsFilter = false;
+	}
+	else if (FilenameFilter.Len() > 0)
+	{
+		Args.Filename = FilenameFilter;
+		Args.bFilenameIsFilter = true;
 	}
+	else
+	{
+		Args.bFilenameIsFilter = true;	// do everything
+	}
+
+	Args.bVerifyJson = !FParse::Param(*InCommandLine, TEXT("noverifyjson"));
+
+	TMap<FString, EProcessingMode> Modes;
+	Modes.Add(TEXT("ResaveText"), EProcessingMode::ResaveText);
+	Modes.Add(TEXT("ResaveBinary"), EProcessingMode::ResaveBinary);
+	Modes.Add(TEXT("RoundTrip"), EProcessingMode::RoundTrip);
+	Modes.Add(TEXT("LoadText"), EProcessingMode::LoadText);
+	Modes.Add(TEXT("LoadBinary"), EProcessingMode::LoadBinary);
+	Modes.Add(TEXT("FindMismatchedSerializers"), EProcessingMode::FindMismatchedSerializers);
+
+	check(Modes.Contains(ModeString));
+	Args.ProcessingMode = Modes[ModeString];
+
+	FParse::Value(*InCommandLine, TEXT("iterations="), Args.NumSaveIterations);
+
+	Args.bIncludeEngineContent = FParse::Param(*InCommandLine, TEXT("includeenginecontent"));
+
+	return DoTextAssetProcessing(Args);
+}
+
+bool UTextAssetCommandlet::DoTextAssetProcessing(const FProcessingArgs& InArgs)
+{
+	TRACE_CPUPROFILER_EVENT_SCOPE_TEXT(TEXT("UTextAssetCommandlet::Main"));
 
-	TArray<UObject*> Objects;	
+	RepairDamagedFiles();
 
-	int32 NumSaveIterations = 1;
-	FParse::Value(*CmdLineParams, TEXT("iterations="), NumSaveIterations);
+	TArray<FString> Blacklist;
 
-	bool bIncludeEngineContent = FParse::Param(*CmdLineParams, TEXT("includeenginecontent"));
+	if (InArgs.ProcessingMode == EProcessingMode::FindMismatchedSerializers)
+	{
+		FindMismatchedSerializers();
+		return true;
+	}
 
+	TArray<UObject*> Objects;		
 	TArray<FString> InputAssetFilenames;
 
 	FString ProjectContentDir = *FPaths::ProjectContentDir();
 	FString EngineContentDir = *FPaths::EngineContentDir();
 	const FString Wildcard = TEXT("*");
 
-	switch (Mode)
+	switch (InArgs.ProcessingMode)
 	{
-	case EMode::ResaveBinary:
-	case EMode::ResaveText:
-	case EMode::RoundTrip:
+	case EProcessingMode::ResaveBinary:
+	case EProcessingMode::ResaveText:
+	case EProcessingMode::RoundTrip:
 	{
 		IFileManager::Get().FindFilesRecursive(InputAssetFilenames, *ProjectContentDir, *(Wildcard + FPackageName::GetAssetPackageExtension()), true, false, true);
 		IFileManager::Get().FindFilesRecursive(InputAssetFilenames, *ProjectContentDir, *(Wildcard + FPackageName::GetMapPackageExtension()), true, false, false);
 
-		if (bIncludeEngineContent)
+		if (InArgs.bIncludeEngineContent)
 		{
 			IFileManager::Get().FindFilesRecursive(InputAssetFilenames, *EngineContentDir, *(Wildcard + FPackageName::GetAssetPackageExtension()), true, false, false);
 			IFileManager::Get().FindFilesRecursive(InputAssetFilenames, *EngineContentDir, *(Wildcard + FPackageName::GetMapPackageExtension()), true, false, false);
@@ -157,14 +183,14 @@ int32 UTextAssetCommandlet::Main(const FString& CmdLineParams)
 		break;
 	}
 
-	case EMode::LoadText:
+	case EProcessingMode::LoadText:
 	{
 		IFileManager::Get().FindFilesRecursive(InputAssetFilenames, *ProjectContentDir, *(Wildcard + FPackageName::GetTextAssetPackageExtension()), true, false, true);
 		//IFileManager::Get().FindFilesRecursive(InputAssetFilenames, *BasePath, *(Wildcard + FPackageName::GetTextMapPackageExtension()), true, false, false);
 		break;
 	}
 
-	case EMode::LoadBinary:
+	case EProcessingMode::LoadBinary:
 	{
 		IFileManager::Get().FindFilesRecursive(InputAssetFilenames, *ProjectContentDir, *(Wildcard + FPackageName::GetAssetPackageExtension()), true, false, true);
 		//IFileManager::Get().FindFilesRecursive(InputAssetFilenames, *BasePath, *(Wildcard + FPackageName::GetTextMapPackageExtension()), true, false, false);
@@ -172,13 +198,33 @@ int32 UTextAssetCommandlet::Main(const FString& CmdLineParams)
 	}
 	}
 
+	FString FilenameFilter = InArgs.Filename;
+	if (!InArgs.bFilenameIsFilter)
+	{
+		FString PotentialFilenames[] = {
+			InArgs.Filename + FPackageName::GetAssetPackageExtension(),
+			InArgs.Filename + FPackageName::GetMapPackageExtension(),
+			InArgs.Filename + FPackageName::GetTextAssetPackageExtension(),
+			InArgs.Filename + FPackageName::GetTextMapPackageExtension()
+		};
+
+		for (const FString& Filename : PotentialFilenames)
+		{
+			if (FPaths::FileExists(Filename))
+			{
+				FilenameFilter = Filename;
+				break;
+			}
+		}
+	}
+
 	TArray<TTuple<FString, FString>> FilesToProcess;
 
 	for (const FString& InputAssetFilename : InputAssetFilenames)
 	{
 		bool bIgnore = false;
 
-		if (FilenameFilterString.Len() > 0 && !InputAssetFilename.Contains(FilenameFilterString))
+		if (FilenameFilter.Len() > 0 && !InputAssetFilename.Contains(FilenameFilter))
 		{
 			bIgnore = true;
 		}
@@ -202,15 +248,15 @@ int32 UTextAssetCommandlet::Main(const FString& CmdLineParams)
 		bool bShouldProcess = true;
 		FString DestinationFilename = InputAssetFilename;
 
-		switch (Mode)
+		switch (InArgs.ProcessingMode)
 		{
-		case EMode::ResaveBinary:
+		case EProcessingMode::ResaveBinary:
 		{
 			DestinationFilename = InputAssetFilename + TEXT(".tmp");
 			break;
 		}
 
-		case EMode::ResaveText:
+		case EProcessingMode::ResaveText:
 		{
 			if (InputAssetFilename.EndsWith(FPackageName::GetAssetPackageExtension())) DestinationFilename = FPaths::ChangeExtension(InputAssetFilename, FPackageName::GetTextAssetPackageExtension());;
 			if (InputAssetFilename.EndsWith(FPackageName::GetMapPackageExtension())) DestinationFilename = FPaths::ChangeExtension(InputAssetFilename, FPackageName::GetTextMapPackageExtension());;
@@ -218,8 +264,8 @@ int32 UTextAssetCommandlet::Main(const FString& CmdLineParams)
 			break;
 		}
 
-		case EMode::LoadText:
-		case EMode::LoadBinary:
+		case EProcessingMode::LoadText:
+		case EProcessingMode::LoadBinary:
 		{
 			break;
 		}
@@ -232,19 +278,15 @@ int32 UTextAssetCommandlet::Main(const FString& CmdLineParams)
 	}
 
 	const FString FailedDiffsPath = FPaths::ProjectSavedDir() / TEXT("FailedDiffs");
-	static const bool bKeepFailedDiffs = FParse::Param(FCommandLine::Get(), TEXT("keepfaileddiffs"));
-	if (bKeepFailedDiffs)
-	{
-		IFileManager::Get().DeleteDirectory(*FailedDiffsPath, false, true);
-	}
+	IFileManager::Get().DeleteDirectory(*FailedDiffsPath, false, true);
 
 	float TotalPackageLoadTime = 0.0;
 	float TotalPackageSaveTime = 0.0;
 
 	FArchive* CSVWriter = nullptr;
-	if (CSVFilename.Len() > 0)
+	if (InArgs.CSVFilename.Len() > 0)
 	{
-		CSVWriter = IFileManager::Get().CreateFileWriter(*CSVFilename);
+		CSVWriter = IFileManager::Get().CreateFileWriter(*InArgs.CSVFilename);
 		if (CSVWriter != nullptr)
 		{
 			FString CSVLine = FString::Printf(TEXT("Total Time,Num Files,AvgFileTime,MinFileTime,MaxFileTime,TotalLoadTime\n"));
@@ -252,12 +294,12 @@ int32 UTextAssetCommandlet::Main(const FString& CmdLineParams)
 		}
 	}
 
-	for (int32 Iteration = 0; Iteration < NumSaveIterations; ++Iteration)
+	for (int32 Iteration = 0; Iteration < InArgs.NumSaveIterations; ++Iteration)
 	{
-		if (NumSaveIterations > 1)
+		if (InArgs.NumSaveIterations > 1)
 		{
 			UE_LOG(LogTextAsset, Display, TEXT("-----------------------------------------------------"));
-			UE_LOG(LogTextAsset, Display, TEXT("Iteration %i/%i"), Iteration + 1, NumSaveIterations);
+			UE_LOG(LogTextAsset, Display, TEXT("Iteration %i/%i"), Iteration + 1, InArgs.NumSaveIterations);
 		}
 
 		double MaxTime = FLT_MIN;
@@ -273,6 +315,7 @@ int32 UTextAssetCommandlet::Main(const FString& CmdLineParams)
 		TArray<FString> PhaseSuccess;
 		TArray<TArray<FString>> PhaseFails;
 		PhaseFails.AddDefaulted(3);
+		TArray<FString> IntermediateFilenames;
 
 		for (const TTuple<FString, FString>& FileToProcess : FilesToProcess)
 		{
@@ -286,9 +329,9 @@ int32 UTextAssetCommandlet::Main(const FString& CmdLineParams)
 
 			double StartTime = FPlatformTime::Seconds();
 
-			switch (Mode)
+			switch (InArgs.ProcessingMode)
 			{
-			case EMode::RoundTrip:
+			case EProcessingMode::RoundTrip:
 			{
 				const FString WorkingFilenames[2] = { SourceFilename, FPaths::ChangeExtension(SourceFilename, FPackageName::GetTextAssetPackageExtension()) };
 				
@@ -391,17 +434,14 @@ int32 UTextAssetCommandlet::Main(const FString& CmdLineParams)
 						FSHAHash& Hash = PhaseHashes[PhaseHashes.AddDefaulted()];
 						HashFile(*WorkingFilenames[Bucket], Hash);
 
-						if (bKeepFailedDiffs)
-						{ 
-							FString TargetPath = WorkingFilenames[Bucket];
-							FPaths::MakePathRelativeTo(TargetPath, *FPaths::ProjectContentDir());
-							TargetPath = FailedDiffsPath / TargetPath;
+						FString TargetPath = WorkingFilenames[Bucket];
+						FPaths::MakePathRelativeTo(TargetPath, *FPaths::ProjectContentDir());
+						TargetPath = FailedDiffsPath / TargetPath;
 
-							FString IntermediateFilename = FString::Printf(TEXT("%s_Phase%i_%03i%s"), *FPaths::ChangeExtension(TargetPath, TEXT("")), Phase, i + 1, *FPaths::GetExtension(WorkingFilenames[Bucket], true));
-							IFileManager::Get().Copy(*IntermediateFilename, *WorkingFilenames[Bucket]);
+						FString IntermediateFilename = FString::Printf(TEXT("%s_Phase%i_%03i%s"), *FPaths::ChangeExtension(TargetPath, TEXT("")), Phase, i + 1, *FPaths::GetExtension(WorkingFilenames[Bucket], true));
+						IFileManager::Get().Copy(*IntermediateFilename, *WorkingFilenames[Bucket]);
 
-							DiffFilenames.Add(IntermediateFilename);
-						}
+						DiffFilenames.Add(IntermediateFilename);
 					}
 
 					UE_LOG(LogTextAsset, Display, TEXT("Phase %i (%s) Results"), Phase + 1, PhaseNames[Phase]);
@@ -491,8 +531,8 @@ int32 UTextAssetCommandlet::Main(const FString& CmdLineParams)
 				break;
 			}
 
-			case EMode::ResaveBinary:
-			case EMode::ResaveText:
+			case EProcessingMode::ResaveBinary:
+			case EProcessingMode::ResaveText:
 			{
 				UPackage* Package = nullptr;
 
@@ -522,7 +562,7 @@ int32 UTextAssetCommandlet::Main(const FString& CmdLineParams)
 
 				if (bSaveSuccessful)
 				{
-					if (bVerifyJson && Mode == EMode::ResaveText)
+					if (InArgs.bVerifyJson && InArgs.ProcessingMode == EProcessingMode::ResaveText)
 					{
 						FArchive* File = IFileManager::Get().CreateFileReader(*DestinationFilename);
 						TSharedPtr< FJsonObject > RootObject;
@@ -531,11 +571,11 @@ int32 UTextAssetCommandlet::Main(const FString& CmdLineParams)
 						delete File;
 					}
 
-					if (OutputPathString.Len() > 0)
+					if (InArgs.OutputPath.Len() > 0)
 					{
 						FString CopyFilename = DestinationFilename;
 						FPaths::MakePathRelativeTo(CopyFilename, *FPaths::RootDir());
-						CopyFilename = OutputPathString / CopyFilename;
+						CopyFilename = InArgs.OutputPath / CopyFilename;
 						CopyFilename.RemoveFromEnd(TEXT(".tmp"));
 						IFileManager::Get().MakeDirectory(*FPaths::GetPath(CopyFilename));
 						IFileManager::Get().Move(*CopyFilename, *DestinationFilename);
@@ -545,7 +585,7 @@ int32 UTextAssetCommandlet::Main(const FString& CmdLineParams)
 				break;
 			}
 
-			case EMode::LoadText:
+			case EProcessingMode::LoadText:
 			{
 				UPackage* Package = nullptr;
 				UE_LOG(LogTextAsset, Display, TEXT("Loading Text Asset '%s'"), *SourceFilename);
@@ -565,7 +605,7 @@ int32 UTextAssetCommandlet::Main(const FString& CmdLineParams)
 				break;
 			}
 
-			case EMode::LoadBinary:
+			case EProcessingMode::LoadBinary:
 			{
 				UPackage* Package = nullptr;
 				UE_LOG(LogTextAsset, Display, TEXT("Loading Binary Asset '%s'"), *SourceFilename);
@@ -589,7 +629,7 @@ int32 UTextAssetCommandlet::Main(const FString& CmdLineParams)
 			double EndTime = FPlatformTime::Seconds();
 			double Time = EndTime - StartTime;
 
-			if (Mode == EMode::LoadBinary || Mode == EMode::LoadText)
+			if (InArgs.ProcessingMode == EProcessingMode::LoadBinary || InArgs.ProcessingMode == EProcessingMode::LoadText)
 			{
 				if (ThisPackageLoadTime > MaxTime)
 				{
@@ -622,7 +662,7 @@ int32 UTextAssetCommandlet::Main(const FString& CmdLineParams)
 			NumFiles++;
 		}
 
-		if (Mode == EMode::RoundTrip)
+		if (InArgs.ProcessingMode == EProcessingMode::RoundTrip)
 		{
 			UE_LOG(LogTextAsset, Display, TEXT("\t-----------------------------------------------------"));
 			UE_LOG(LogTextAsset, Display, TEXT("\tRoundTrip Results"));
@@ -649,7 +689,7 @@ int32 UTextAssetCommandlet::Main(const FString& CmdLineParams)
 
 		double AvgFileTime, MinFileTime, MaxFileTime;
 
-		if (Mode == EMode::LoadBinary || Mode == EMode::LoadText)
+		if (InArgs.ProcessingMode == EProcessingMode::LoadBinary || InArgs.ProcessingMode == EProcessingMode::LoadText)
 		{
 			AvgFileTime = IterationPackageLoadTime;
 		}
@@ -675,7 +715,7 @@ int32 UTextAssetCommandlet::Main(const FString& CmdLineParams)
 			CSVWriter->Serialize(TCHAR_TO_ANSI(*CSVLine), CSVLine.Len());
 		}
 
-		if (Mode != EMode::LoadText && Mode != EMode::ResaveText)
+		if (InArgs.ProcessingMode != EProcessingMode::LoadText && InArgs.ProcessingMode != EProcessingMode::ResaveText)
 		{
 			UE_LOG(LogTextAsset, Display, TEXT("\tTotal Package Save Time:  \t%.2fs"), IterationPackageSaveTime);
 		}
@@ -686,19 +726,40 @@ int32 UTextAssetCommandlet::Main(const FString& CmdLineParams)
 	UE_LOG(LogTextAsset, Display, TEXT("-----------------------------------------------------"));
 	UE_LOG(LogTextAsset, Display, TEXT("Text Asset Commandlet Completed!"));
 	UE_LOG(LogTextAsset, Display, TEXT("\tTotal Files Processed:  \t%i"), FilesToProcess.Num());
-	UE_LOG(LogTextAsset, Display, TEXT("\tAvg Iteration Package Load Time:  \t%.2fs"), TotalPackageLoadTime / (float)NumSaveIterations);
+	UE_LOG(LogTextAsset, Display, TEXT("\tAvg Iteration Package Load Time:  \t%.2fs"), TotalPackageLoadTime / (float)InArgs.NumSaveIterations);
 
 	if (CSVWriter != nullptr)
 	{
 		delete CSVWriter;
 	}
 
-	if (Mode != EMode::LoadText)
+	if (InArgs.ProcessingMode != EProcessingMode::LoadText)
 	{
-		UE_LOG(LogTextAsset, Display, TEXT("\tAvg Iteration Save Time:  \t%.2fs"), TotalPackageSaveTime / (float)NumSaveIterations);
+		UE_LOG(LogTextAsset, Display, TEXT("\tAvg Iteration Save Time:  \t%.2fs"), TotalPackageSaveTime / (float)InArgs.NumSaveIterations);
 	}
 	
 	UE_LOG(LogTextAsset, Display, TEXT("-----------------------------------------------------"));
 	
-	return 0;
-}
\ No newline at end of file
+	return true;
+}
+
+int32 UTextAssetCommandlet::Main(const FString& CmdLineParams)
+{
+	return DoTextAssetProcessing(CmdLineParams) ? 0 : 1;
+}
+
+static void TextAssetToolCVarCommand(const TArray<FString>& Args)
+{
+	FString JoinedArgs;
+	for (const FString& Arg : Args) 
+	{
+		JoinedArgs += Arg;
+		JoinedArgs += TEXT(" ");
+	}
+	UTextAssetCommandlet::DoTextAssetProcessing(JoinedArgs);
+}
+
+static FAutoConsoleCommand CVar_TextAssetTool(
+	TEXT("TextAssetTool"),
+	TEXT("--"),
+	FConsoleCommandWithArgsDelegate::CreateStatic(TextAssetToolCVarCommand));
\ No newline at end of file
