Remove obsolete cooker save setting bConsiderCompressedPackageFileLengthRequirements

diff --git a/Engine/Config/BaseEditor.ini b/Engine/Config/BaseEditor.ini
--- a/Engine/Config/BaseEditor.ini
+++ b/Engine/Config/BaseEditor.ini
@@ -353,7 +353,6 @@ BatchCount=100
 MaxMemoryAllowance=16384
 
 [CookSettings]
-bConsiderCompressedPackageFileLengthRequirements=true
 PackagesPerGC=0
 IdleTimeToGC=20
 MaxMemoryAllowance=16384
diff --git a/Engine/Source/Editor/UnrealEd/Classes/CookOnTheSide/CookOnTheFlyServer.h b/Engine/Source/Editor/UnrealEd/Classes/CookOnTheSide/CookOnTheFlyServer.h
--- a/Engine/Source/Editor/UnrealEd/Classes/CookOnTheSide/CookOnTheFlyServer.h
+++ b/Engine/Source/Editor/UnrealEd/Classes/CookOnTheSide/CookOnTheFlyServer.h
@@ -851,9 +851,6 @@ private:
 		return (CookFlags & InCookFlags) != ECookInitializationFlags::None;
 	}
 
-	/** If true, the maximum file length of a package being saved will be reduced by 32 to compensate for compressed package intermediate files */
-	bool ShouldConsiderCompressedPackageFileLengthRequirements() const;
-
 	/**
 	*	Cook (save) the given package
 	*
diff --git a/Engine/Source/Editor/UnrealEd/Private/CookOnTheFlyServer.cpp b/Engine/Source/Editor/UnrealEd/Private/CookOnTheFlyServer.cpp
--- a/Engine/Source/Editor/UnrealEd/Private/CookOnTheFlyServer.cpp
+++ b/Engine/Source/Editor/UnrealEd/Private/CookOnTheFlyServer.cpp
@@ -3622,13 +3622,6 @@ void UCookOnTheFlyServer::SaveCookedPackage(UPackage* Package, uint32 SaveFlags,
 	return SaveCookedPackage( Package, SaveFlags, TargetPlatformNames, SavePackageResults);
 }
 
-bool UCookOnTheFlyServer::ShouldConsiderCompressedPackageFileLengthRequirements() const
-{
-	bool bConsiderCompressedPackageFileLengthRequirements = true;
-	GConfig->GetBool(TEXT("CookSettings"), TEXT("bConsiderCompressedPackageFileLengthRequirements"), bConsiderCompressedPackageFileLengthRequirements, GEditorIni);
-	return bConsiderCompressedPackageFileLengthRequirements;
-}
-
 bool UCookOnTheFlyServer::MakePackageFullyLoaded(UPackage* Package) const
 {
 	if (Package->IsFullyLoaded())
@@ -4000,15 +3993,11 @@ void UCookOnTheFlyServer::SaveCookedPackage(UPackage* Package, uint32 SaveFlags,
 						World->PersistentLevel->HandleLegacyMapBuildData();
 					}
 
-					// need to subtract 32 because the SavePackage code creates temporary files with longer file names then the one we provide
-					// projects may ignore this restriction if desired
-					static bool bConsiderCompressedPackageFileLengthRequirements = ShouldConsiderCompressedPackageFileLengthRequirements();
-					const int32 CompressedPackageFileLengthRequirement = bConsiderCompressedPackageFileLengthRequirements ? 32 : 0;
 					const FString FullFilename = FPaths::ConvertRelativePathToFull(PlatFilename);
-					if (FullFilename.Len() >= (FPlatformMisc::GetMaxPathLength() - CompressedPackageFileLengthRequirement))
+					if (FullFilename.Len() >= FPlatformMisc::GetMaxPathLength())
 					{
-						LogCookerMessage(FString::Printf(TEXT("Couldn't save package, filename is too long: %s"), *PlatFilename), EMessageSeverity::Error);
-						UE_LOG(LogCook, Error, TEXT("Couldn't save package, filename is too long :%s"), *PlatFilename);
+						LogCookerMessage(FString::Printf(TEXT("Couldn't save package, filename is too long (%d >= %d): %s"), FullFilename.Len(), FPlatformMisc::GetMaxPathLength(), *PlatFilename), EMessageSeverity::Error);
+						UE_LOG(LogCook, Error, TEXT("Couldn't save package, filename is too long (%d >= %d): %s"), FullFilename.Len(), FPlatformMisc::GetMaxPathLength(), *PlatFilename);
 						Result = ESavePackageResult::Error;
 					}
 					else
