TBA: Minor bugfix for textasset roundtrip test mode - empty the diff filename array after each pass so as not to mark failed early pass results for deletion if the last pass succeeds

#rnx

diff --git a/Engine/Source/Editor/UnrealEd/Private/Commandlets/TextAssetCommandlet.cpp b/Engine/Source/Editor/UnrealEd/Private/Commandlets/TextAssetCommandlet.cpp
--- a/Engine/Source/Editor/UnrealEd/Private/Commandlets/TextAssetCommandlet.cpp
+++ b/Engine/Source/Editor/UnrealEd/Private/Commandlets/TextAssetCommandlet.cpp
@@ -481,6 +481,8 @@ bool UTextAssetCommandlet::DoTextAssetProcessing(const FProcessingArgs& InArgs)
 							IFileManager::Get().Delete(*DiffFilename, false, false, true);
 						}
 					}
+
+					DiffFilenames.Empty();
 				}
 
 				static const bool bDisableCleanup = FParse::Param(FCommandLine::Get(), TEXT("disablecleanup"));
