TBA: Roundtrip test improvements.
 - When logging the phase failures at the end of execution, only don't include phase0 fails in the list for phase1 and phase2, because results are meaningless in those cases
 - When phase0 fails, don't keep failed diffs
 - Try and clean up empty directories in the failed diffs folder (may not fully work)

diff --git a/Engine/Source/Editor/UnrealEd/Private/Commandlets/TextAssetCommandlet.cpp b/Engine/Source/Editor/UnrealEd/Private/Commandlets/TextAssetCommandlet.cpp
--- a/Engine/Source/Editor/UnrealEd/Private/Commandlets/TextAssetCommandlet.cpp
+++ b/Engine/Source/Editor/UnrealEd/Private/Commandlets/TextAssetCommandlet.cpp
@@ -473,16 +473,6 @@ bool UTextAssetCommandlet::DoTextAssetProcessing(const FProcessingArgs& InArgs)
 					{
 						IFileManager::Get().Delete(*WorkingFilenames[1], false, false, true);
 					}
-
-					if (bTotalSuccess)
-					{
-						for (const FString& DiffFilename : DiffFilenames)
-						{
-							IFileManager::Get().Delete(*DiffFilename, false, false, true);
-						}
-					}
-
-					DiffFilenames.Empty();
 				}
 
 				static const bool bDisableCleanup = FParse::Param(FCommandLine::Get(), TEXT("disablecleanup"));
@@ -500,6 +490,7 @@ bool UTextAssetCommandlet::DoTextAssetProcessing(const FProcessingArgs& InArgs)
 					}
 				}
 
+				bool bDeleteDiffs = true;
 				if (!bPhasesMatched[0])
 				{
 					UE_LOG(LogTextAsset, Display, TEXT("-----------------------------------------------------------------------------------------"));
@@ -509,6 +500,27 @@ bool UTextAssetCommandlet::DoTextAssetProcessing(const FProcessingArgs& InArgs)
 				{
 					UE_LOG(LogTextAsset, Display, TEXT("-----------------------------------------------------------------------------------------"));
 					UE_LOG(LogTextAsset, Error, TEXT("Binary determinism tests succeeded, but text and/or alternating tests failed for asset '%s'"), *SourceLongPackageName);
+					bDeleteDiffs = false;
+				}
+
+				if (bDeleteDiffs)
+				{
+					FString Directory;
+					for (const FString& DiffFilename : DiffFilenames)
+					{
+						IFileManager::Get().Delete(*DiffFilename, false, false, true);
+
+						if (Directory.Len() == 0)
+						{
+							Directory = FPaths::GetPath(DiffFilename) + TEXT("/");
+						}
+					}
+					TArray<FString> Found;
+					IFileManager::Get().FindFiles(Found, *Directory, true, true);
+					if (Found.Num() == 0)
+					{
+						IFileManager::Get().DeleteDirectory(*Directory, true, false);
+					}
 				}
 
 				bool bSuccess = true;
@@ -675,14 +687,17 @@ bool UTextAssetCommandlet::DoTextAssetProcessing(const FProcessingArgs& InArgs)
 			UE_LOG(LogTextAsset, Display, TEXT("\tPhase 2 Fails: %i (Mixed Package Determinism Fails)"), PhaseFails[2].Num());
 			UE_LOG(LogTextAsset, Display, TEXT("\t-----------------------------------------------------"));
 
-			for (int32 PhaseIndex = 0; PhaseIndex < PhaseFails.Num(); ++PhaseIndex)
+			for (int32 PhaseIndex = 1; PhaseIndex < PhaseFails.Num(); ++PhaseIndex)
 			{
 				if (PhaseFails[PhaseIndex].Num() > 0)
 				{
 					UE_LOG(LogTextAsset, Display, TEXT("\tPhase %i Fails:"), PhaseIndex);
 					for (const FString& PhaseFail : PhaseFails[PhaseIndex])
 					{
-						UE_LOG(LogTextAsset, Display, TEXT("\t\t%s"), *PhaseFail);
+						if (!PhaseFails[0].Contains(PhaseFail))
+						{
+							UE_LOG(LogTextAsset, Display, TEXT("\t\t%s"), *PhaseFail);
+						}
 					}
 					UE_LOG(LogTextAsset, Display, TEXT("\t-----------------------------------------------------"));
 				}
