Fix for crash in cook-on-the-fly mode due to not taking into consideration whether we were in package store mode or not

diff --git a/Engine/Source/Editor/UnrealEd/Private/CookOnTheFlyServer.cpp b/Engine/Source/Editor/UnrealEd/Private/CookOnTheFlyServer.cpp
--- a/Engine/Source/Editor/UnrealEd/Private/CookOnTheFlyServer.cpp
+++ b/Engine/Source/Editor/UnrealEd/Private/CookOnTheFlyServer.cpp
@@ -4029,8 +4029,12 @@ void UCookOnTheFlyServer::SaveCookedPackage(UPackage* Package, uint32 SaveFlags,
 						}
 						else
 						{
-							FSavePackageContext* SavePackageContext = SavePackageContexts[PlatformIndex];
-							Result = GEditor->Save(Package, World, FlagsToCook, *PlatFilename, GError, NULL, bSwap, false, SaveFlags, Target, FDateTime::MinValue(), false, /*DiffMap*/nullptr, SavePackageContext);
+							FSavePackageContext* const SavePackageContext = IsUsingPackageStore() ? SavePackageContexts[PlatformIndex] : nullptr;
+
+							Result = GEditor->Save(	Package, World, FlagsToCook, *PlatFilename, 
+													GError, nullptr, bSwap, false, SaveFlags, Target, 
+													FDateTime::MinValue(), false, /*DiffMap*/ nullptr, 
+													SavePackageContext);
 						}
 						GIsCookerLoadingPackage = false;
 						{
