TBA: Minor changes to the textasset commandlet
 - Added trace markers for profiling purposes
 - When roundtripping, generate the intermediate files to a separate location and then copy them into the failed diffs folder if the roundtrip fails. Stops loads of empty directories being created in the faileddiffs folder

diff --git a/Engine/Source/Editor/UnrealEd/Private/Commandlets/TextAssetCommandlet.cpp b/Engine/Source/Editor/UnrealEd/Private/Commandlets/TextAssetCommandlet.cpp
--- a/Engine/Source/Editor/UnrealEd/Private/Commandlets/TextAssetCommandlet.cpp
+++ b/Engine/Source/Editor/UnrealEd/Private/Commandlets/TextAssetCommandlet.cpp
@@ -277,6 +277,7 @@ bool UTextAssetCommandlet::DoTextAssetProcessing(const FProcessingArgs& InArgs)
 		}
 	}
 
+	const FString TempFailedDiffsPath = FPaths::ProjectSavedDir() / TEXT(".roundtrip");
 	const FString FailedDiffsPath = FPaths::ProjectSavedDir() / TEXT("FailedDiffs");
 	IFileManager::Get().DeleteDirectory(*FailedDiffsPath, false, true);
 
@@ -333,6 +334,9 @@ bool UTextAssetCommandlet::DoTextAssetProcessing(const FProcessingArgs& InArgs)
 			{
 			case EProcessingMode::RoundTrip:
 			{
+				UE_LOG(LogTextAsset, Display, TEXT("Starting roundtrip test for '%s' [%d/%d]"), *SourceLongPackageName, NumFiles + 1, FilesToProcess.Num());
+				UE_LOG(LogTextAsset, Display, TEXT("-----------------------------------------------------------------------------------------"));
+
 				const FString WorkingFilenames[2] = { SourceFilename, FPaths::ChangeExtension(SourceFilename, FPackageName::GetTextAssetPackageExtension()) };
 				
 
@@ -362,18 +366,29 @@ bool UTextAssetCommandlet::DoTextAssetProcessing(const FProcessingArgs& InArgs)
 				static const int32 NumPhases = 3;
 				static const int32 NumTests = 3;
 
-				static const TCHAR* PhaseNames[NumPhases] = { TEXT("Binary Only"), TEXT("Text Only"), TEXT("Alternating Binary/Text") };
+				static const TCHAR* PhaseNames[] = { TEXT("Binary Only"), TEXT("Text Only"), TEXT("Alternating Binary/Text") };
 
+				static const uint16 PhaseEventTypes[3] = {
+						FCpuProfilerTrace::OutputEventType(TEXT("BinaryOnly"), CpuProfilerGroup_Default),
+						FCpuProfilerTrace::OutputEventType(TEXT("TextOnly"), CpuProfilerGroup_Default),
+						FCpuProfilerTrace::OutputEventType(TEXT("Alternating"), CpuProfilerGroup_Default)
+				};
+
+				static const uint16 TestEventTypes[6] = {
+						FCpuProfilerTrace::OutputEventType(TEXT("Test1"), CpuProfilerGroup_Default),
+						FCpuProfilerTrace::OutputEventType(TEXT("Test2"), CpuProfilerGroup_Default),
+						FCpuProfilerTrace::OutputEventType(TEXT("Test3"), CpuProfilerGroup_Default),
+						FCpuProfilerTrace::OutputEventType(TEXT("Test4"), CpuProfilerGroup_Default),
+						FCpuProfilerTrace::OutputEventType(TEXT("Test5"), CpuProfilerGroup_Default),
+						FCpuProfilerTrace::OutputEventType(TEXT("Test6"), CpuProfilerGroup_Default)
+				};
 
 				TArray<TArray<FSHAHash>> Hashes;
 
 				CollectGarbage(RF_NoFlags, true);
 
-				UE_LOG(LogTextAsset, Display, TEXT("Starting roundtrip test for '%s' [%d/%d]"), *SourceLongPackageName, NumFiles + 1, FilesToProcess.Num());
-				UE_LOG(LogTextAsset, Display, TEXT("-----------------------------------------------------------------------------------------"));
-
 				bool bPhasesMatched[NumPhases] = { true, true, true };
-				TArray<FString> DiffFilenames;
+				TArray<TPair<FString,FString>> DiffFilenames;
 
 				for (int32 Phase = 0; Phase < NumPhases; ++Phase)
 				{
@@ -381,11 +396,15 @@ bool UTextAssetCommandlet::DoTextAssetProcessing(const FProcessingArgs& InArgs)
 					IFileManager::Get().Copy(*SourceFilename, *BaseBinaryPackageBackup, true);
 
 					TArray<FSHAHash> PhaseHashes = Hashes[Hashes.AddDefaulted()];
+					
+					FCpuProfilerTrace::OutputBeginEvent(PhaseEventTypes[Phase]);
 
 					for (int32 i = 0; i < ((Phase == 2) ? NumTests * 2 : NumTests); ++i)
 					{
 						int32 Bucket;
 
+						FCpuProfilerTrace::OutputBeginEvent(TestEventTypes[i]);
+
 						switch (Phase)
 						{
 						case 0: // binary only
@@ -426,22 +445,46 @@ bool UTextAssetCommandlet::DoTextAssetProcessing(const FProcessingArgs& InArgs)
 						}
 						};
 
-						UPackage* Package = LoadPackage(nullptr, *SourceLongPackageName, LOAD_None);
-						SavePackageHelper(Package, *WorkingFilenames[Bucket], RF_Standalone, GWarn, nullptr, SAVE_KeepGUID);
-						ResetLoaders(Package);
-						CollectGarbage(RF_NoFlags, true);
+						UPackage* Package = nullptr;
+						
+						{
+							TRACE_CPUPROFILER_EVENT_SCOPE_TEXT(TEXT("LoadPackage"));
+							Package = LoadPackage(nullptr, *SourceLongPackageName, LOAD_None);
+						}
+						
+						{
+							TRACE_CPUPROFILER_EVENT_SCOPE_TEXT(TEXT("SavePackage")); 
+							SavePackageHelper(Package, *WorkingFilenames[Bucket], RF_Standalone, GWarn, nullptr, SAVE_KeepGUID);
+						}
+						
+						{
+							TRACE_CPUPROFILER_EVENT_SCOPE_TEXT(TEXT("ResetLoaders")); 
+							ResetLoaders(Package);
+						}
 
-						FSHAHash& Hash = PhaseHashes[PhaseHashes.AddDefaulted()];
-						HashFile(*WorkingFilenames[Bucket], Hash);
+						{
+							TRACE_CPUPROFILER_EVENT_SCOPE_TEXT(TEXT("CollectGarbage"));
+							CollectGarbage(RF_NoFlags, true);
+						}
+
+						{
+							TRACE_CPUPROFILER_EVENT_SCOPE_TEXT(TEXT("RoundtripTestCleanup"));
+							FSHAHash& Hash = PhaseHashes[PhaseHashes.AddDefaulted()];
+							HashFile(*WorkingFilenames[Bucket], Hash);
 
-						FString TargetPath = WorkingFilenames[Bucket];
-						FPaths::MakePathRelativeTo(TargetPath, *FPaths::ProjectContentDir());
-						TargetPath = FailedDiffsPath / TargetPath;
+							FString TargetPath = WorkingFilenames[Bucket];
+							FPaths::MakePathRelativeTo(TargetPath, *FPaths::ProjectContentDir());
+							FString IntermediateTargetPath = TempFailedDiffsPath / TargetPath;
+							FString FinalTargetPath = FailedDiffsPath / TargetPath;
 
-						FString IntermediateFilename = FString::Printf(TEXT("%s_Phase%i_%03i%s"), *FPaths::ChangeExtension(TargetPath, TEXT("")), Phase, i + 1, *FPaths::GetExtension(WorkingFilenames[Bucket], true));
-						IFileManager::Get().Copy(*IntermediateFilename, *WorkingFilenames[Bucket]);
+							FString IntermediateFilename = FString::Printf(TEXT("%s_Phase%i_%03i%s"), *FPaths::ChangeExtension(IntermediateTargetPath, TEXT("")), Phase, i + 1, *FPaths::GetExtension(WorkingFilenames[Bucket], true));
+							FString FinalFilename = FString::Printf(TEXT("%s_Phase%i_%03i%s"), *FPaths::ChangeExtension(FinalTargetPath, TEXT("")), Phase, i + 1, *FPaths::GetExtension(WorkingFilenames[Bucket], true));
+							IFileManager::Get().Copy(*IntermediateFilename, *WorkingFilenames[Bucket]);
 
-						DiffFilenames.Add(IntermediateFilename);
+							DiffFilenames.Add(TPair<FString, FString>(IntermediateFilename, FinalFilename));
+						}
+
+						FCpuProfilerTrace::OutputEndEvent();
 					}
 
 					UE_LOG(LogTextAsset, Display, TEXT("Phase %i (%s) Results"), Phase + 1, PhaseNames[Phase]);
@@ -473,6 +516,19 @@ bool UTextAssetCommandlet::DoTextAssetProcessing(const FProcessingArgs& InArgs)
 					{
 						IFileManager::Get().Delete(*WorkingFilenames[1], false, false, true);
 					}
+
+					if (!bTotalSuccess)
+					{
+						for (const TPair<FString, FString>& DiffPair : DiffFilenames)
+						{
+							IFileManager::Get().MakeDirectory(*FPaths::GetPath(DiffPair.Value));
+							IFileManager::Get().Move(*DiffPair.Value, *DiffPair.Key);
+						}
+					}
+
+					DiffFilenames.Empty();
+					IFileManager::Get().DeleteDirectory(*TempFailedDiffsPath, false, true);
+					FCpuProfilerTrace::OutputEndEvent();
 				}
 
 				static const bool bDisableCleanup = FParse::Param(FCommandLine::Get(), TEXT("disablecleanup"));
@@ -490,7 +546,6 @@ bool UTextAssetCommandlet::DoTextAssetProcessing(const FProcessingArgs& InArgs)
 					}
 				}
 
-				bool bDeleteDiffs = true;
 				if (!bPhasesMatched[0])
 				{
 					UE_LOG(LogTextAsset, Display, TEXT("-----------------------------------------------------------------------------------------"));
@@ -500,30 +555,9 @@ bool UTextAssetCommandlet::DoTextAssetProcessing(const FProcessingArgs& InArgs)
 				{
 					UE_LOG(LogTextAsset, Display, TEXT("-----------------------------------------------------------------------------------------"));
 					UE_LOG(LogTextAsset, Error, TEXT("Binary determinism tests succeeded, but text and/or alternating tests failed for asset '%s'"), *SourceLongPackageName);
-					bDeleteDiffs = false;
-				}
-
-				if (bDeleteDiffs)
-				{
-					FString Directory;
-					for (const FString& DiffFilename : DiffFilenames)
-					{
-						IFileManager::Get().Delete(*DiffFilename, false, false, true);
-
-						if (Directory.Len() == 0)
-						{
-							Directory = FPaths::GetPath(DiffFilename) + TEXT("/");
-						}
-					}
-					TArray<FString> Found;
-					IFileManager::Get().FindFiles(Found, *Directory, true, true);
-					if (Found.Num() == 0)
-					{
-						IFileManager::Get().DeleteDirectory(*Directory, true, false);
-					}
 				}
 
-				bool bSuccess = true;
+					bool bSuccess = true;
 				for (int32 PhaseIndex = 0; PhaseIndex < NumPhases; ++PhaseIndex)
 				{
 					if (!bPhasesMatched[PhaseIndex])
@@ -551,10 +585,12 @@ bool UTextAssetCommandlet::DoTextAssetProcessing(const FProcessingArgs& InArgs)
 				UPackage* Package = nullptr;
 
 				UE_LOG(LogTextAsset, Display, TEXT("Resaving asset %s"), *SourceFilename);
+				TRACE_CPUPROFILER_EVENT_SCOPE_TEXT(TEXT("UTextAssetCommandlet::Resave"));
 
 				double Timer = 0.0;
 				{
 					SCOPE_SECONDS_COUNTER(Timer);
+					TRACE_CPUPROFILER_EVENT_SCOPE_TEXT(TEXT("UTextAssetCommandlet::LoadPackage"));
 					Package = LoadPackage(nullptr, *SourceFilename, 0);
 				}
 				IterationPackageLoadTime += Timer;
@@ -566,7 +602,7 @@ bool UTextAssetCommandlet::DoTextAssetProcessing(const FProcessingArgs& InArgs)
 				{
 					{
 						SCOPE_SECONDS_COUNTER(Timer);
-
+						TRACE_CPUPROFILER_EVENT_SCOPE_TEXT(TEXT("UTextAssetCommandlet::SavePackage"));
 						IFileManager::Get().Delete(*DestinationFilename, false, true, true);
 						bSaveSuccessful = SavePackageHelper(Package, *DestinationFilename, RF_Standalone, GWarn, nullptr, SAVE_KeepGUID);
 					}
@@ -578,6 +614,7 @@ bool UTextAssetCommandlet::DoTextAssetProcessing(const FProcessingArgs& InArgs)
 				{
 					if (InArgs.bVerifyJson && InArgs.ProcessingMode == EProcessingMode::ResaveText)
 					{
+						TRACE_CPUPROFILER_EVENT_SCOPE_TEXT(TEXT("UTextAssetCommandlet::VerifyJson"));
 						FArchive* File = IFileManager::Get().CreateFileReader(*DestinationFilename);
 						TSharedPtr< FJsonObject > RootObject;
 						TSharedRef< TJsonReader<char> > Reader = TJsonReaderFactory<char>::Create(File);
@@ -587,6 +624,7 @@ bool UTextAssetCommandlet::DoTextAssetProcessing(const FProcessingArgs& InArgs)
 
 					if (InArgs.OutputPath.Len() > 0)
 					{
+						TRACE_CPUPROFILER_EVENT_SCOPE_TEXT(TEXT("UTextAssetCommandlet::CopyToExternalOutput"));
 						FString CopyFilename = DestinationFilename;
 						FPaths::MakePathRelativeTo(CopyFilename, *FPaths::RootDir());
 						CopyFilename = InArgs.OutputPath / CopyFilename;
