Fixed narrowing error on Mac.

#rnx

diff --git a/Engine/Source/Editor/UnrealEd/Private/Commandlets/TextAssetCommandlet.cpp b/Engine/Source/Editor/UnrealEd/Private/Commandlets/TextAssetCommandlet.cpp
--- a/Engine/Source/Editor/UnrealEd/Private/Commandlets/TextAssetCommandlet.cpp
+++ b/Engine/Source/Editor/UnrealEd/Private/Commandlets/TextAssetCommandlet.cpp
@@ -368,13 +368,13 @@ bool UTextAssetCommandlet::DoTextAssetProcessing(const FProcessingArgs& InArgs)
 
 				static const TCHAR* PhaseNames[] = { TEXT("Binary Only"), TEXT("Text Only"), TEXT("Alternating Binary/Text") };
 
-				static const uint16 PhaseEventTypes[3] = {
+				static const uint32 PhaseEventTypes[3] = {
 						FCpuProfilerTrace::OutputEventType(TEXT("BinaryOnly"), CpuProfilerGroup_Default),
 						FCpuProfilerTrace::OutputEventType(TEXT("TextOnly"), CpuProfilerGroup_Default),
 						FCpuProfilerTrace::OutputEventType(TEXT("Alternating"), CpuProfilerGroup_Default)
 				};
 
-				static const uint16 TestEventTypes[6] = {
+				static const uint32 TestEventTypes[6] = {
 						FCpuProfilerTrace::OutputEventType(TEXT("Test1"), CpuProfilerGroup_Default),
 						FCpuProfilerTrace::OutputEventType(TEXT("Test2"), CpuProfilerGroup_Default),
 						FCpuProfilerTrace::OutputEventType(TEXT("Test3"), CpuProfilerGroup_Default),
@@ -817,4 +817,4 @@ static void TextAssetToolCVarCommand(const TArray<FString>& Args)
 static FAutoConsoleCommand CVar_TextAssetTool(
 	TEXT("TextAssetTool"),
 	TEXT("--"),
-	FConsoleCommandWithArgsDelegate::CreateStatic(TextAssetToolCVarCommand));
\ No newline at end of file
+	FConsoleCommandWithArgsDelegate::CreateStatic(TextAssetToolCVarCommand));
