TBA: Minor tweaks to how text asset commandlet mode is defined and parsed from the commandline, making it easier to add new modes

diff --git a/Engine/Source/Editor/ContentBrowser/Private/AssetContextMenu.cpp b/Engine/Source/Editor/ContentBrowser/Private/AssetContextMenu.cpp
--- a/Engine/Source/Editor/ContentBrowser/Private/AssetContextMenu.cpp
+++ b/Engine/Source/Editor/ContentBrowser/Private/AssetContextMenu.cpp
@@ -871,7 +871,7 @@ void FAssetContextMenu::DoTextFormatRoundtrip()
 	Args.bIncludeEngineContent = true;
 	Args.bVerifyJson = true;
 	Args.CSVFilename = TEXT("");
-	Args.ProcessingMode = UTextAssetCommandlet::EProcessingMode::RoundTrip;
+	Args.ProcessingMode = ETextAssetCommandletMode::RoundTrip;
 	Args.bFilenameIsFilter = false;
 
 	for (const FAssetData& Asset : SelectedAssets)
diff --git a/Engine/Source/Editor/UnrealEd/Classes/Commandlets/TextAssetCommandlet.h b/Engine/Source/Editor/UnrealEd/Classes/Commandlets/TextAssetCommandlet.h
--- a/Engine/Source/Editor/UnrealEd/Classes/Commandlets/TextAssetCommandlet.h
+++ b/Engine/Source/Editor/UnrealEd/Classes/Commandlets/TextAssetCommandlet.h
@@ -10,6 +10,17 @@
 
 UNREALED_API DECLARE_LOG_CATEGORY_EXTERN(LogTextAsset, Log, All);
 
+UENUM()
+enum class ETextAssetCommandletMode
+{
+	ResaveText,
+	ResaveBinary,
+	RoundTrip,
+	LoadBinary,
+	LoadText,
+	FindMismatchedSerializers
+};
+
 UCLASS(config=Editor)
 class UTextAssetCommandlet
 	: public UCommandlet
@@ -18,19 +29,9 @@ class UTextAssetCommandlet
 
 public:
 
-	enum class EProcessingMode
-	{
-		ResaveText,
-		ResaveBinary,
-		RoundTrip,
-		LoadBinary,
-		LoadText,
-		FindMismatchedSerializers,
-	};
-
 	struct FProcessingArgs
 	{
-		EProcessingMode ProcessingMode = EProcessingMode::ResaveText;
+		ETextAssetCommandletMode ProcessingMode = ETextAssetCommandletMode::ResaveText;
 		int32 NumSaveIterations = 1;
 		bool bIncludeEngineContent = false;
 		bool bFilenameIsFilter = true;
diff --git a/Engine/Source/Editor/UnrealEd/Private/Commandlets/TextAssetCommandlet.cpp b/Engine/Source/Editor/UnrealEd/Private/Commandlets/TextAssetCommandlet.cpp
--- a/Engine/Source/Editor/UnrealEd/Private/Commandlets/TextAssetCommandlet.cpp
+++ b/Engine/Source/Editor/UnrealEd/Private/Commandlets/TextAssetCommandlet.cpp
@@ -125,17 +125,7 @@ bool UTextAssetCommandlet::DoTextAssetProcessing(const FString& InCommandLine)
 	}
 
 	Args.bVerifyJson = !FParse::Param(*InCommandLine, TEXT("noverifyjson"));
-
-	TMap<FString, EProcessingMode> Modes;
-	Modes.Add(TEXT("ResaveText"), EProcessingMode::ResaveText);
-	Modes.Add(TEXT("ResaveBinary"), EProcessingMode::ResaveBinary);
-	Modes.Add(TEXT("RoundTrip"), EProcessingMode::RoundTrip);
-	Modes.Add(TEXT("LoadText"), EProcessingMode::LoadText);
-	Modes.Add(TEXT("LoadBinary"), EProcessingMode::LoadBinary);
-	Modes.Add(TEXT("FindMismatchedSerializers"), EProcessingMode::FindMismatchedSerializers);
-
-	check(Modes.Contains(ModeString));
-	Args.ProcessingMode = Modes[ModeString];
+	Args.ProcessingMode = (ETextAssetCommandletMode)StaticEnum<ETextAssetCommandletMode>()->GetValueByNameString(ModeString);
 
 	FParse::Value(*InCommandLine, TEXT("iterations="), Args.NumSaveIterations);
 
@@ -152,7 +142,7 @@ bool UTextAssetCommandlet::DoTextAssetProcessing(const FProcessingArgs& InArgs)
 
 	TArray<FString> Blacklist;
 
-	if (InArgs.ProcessingMode == EProcessingMode::FindMismatchedSerializers)
+	if (InArgs.ProcessingMode == ETextAssetCommandletMode::FindMismatchedSerializers)
 	{
 		FindMismatchedSerializers();
 		return true;
@@ -167,9 +157,9 @@ bool UTextAssetCommandlet::DoTextAssetProcessing(const FProcessingArgs& InArgs)
 
 	switch (InArgs.ProcessingMode)
 	{
-	case EProcessingMode::ResaveBinary:
-	case EProcessingMode::ResaveText:
-	case EProcessingMode::RoundTrip:
+	case ETextAssetCommandletMode::ResaveBinary:
+	case ETextAssetCommandletMode::ResaveText:
+	case ETextAssetCommandletMode::RoundTrip:
 	{
 		IFileManager::Get().FindFilesRecursive(InputAssetFilenames, *ProjectContentDir, *(Wildcard + FPackageName::GetAssetPackageExtension()), true, false, true);
 		IFileManager::Get().FindFilesRecursive(InputAssetFilenames, *ProjectContentDir, *(Wildcard + FPackageName::GetMapPackageExtension()), true, false, false);
@@ -183,14 +173,14 @@ bool UTextAssetCommandlet::DoTextAssetProcessing(const FProcessingArgs& InArgs)
 		break;
 	}
 
-	case EProcessingMode::LoadText:
+	case ETextAssetCommandletMode::LoadText:
 	{
 		IFileManager::Get().FindFilesRecursive(InputAssetFilenames, *ProjectContentDir, *(Wildcard + FPackageName::GetTextAssetPackageExtension()), true, false, true);
 		//IFileManager::Get().FindFilesRecursive(InputAssetFilenames, *BasePath, *(Wildcard + FPackageName::GetTextMapPackageExtension()), true, false, false);
 		break;
 	}
 
-	case EProcessingMode::LoadBinary:
+	case ETextAssetCommandletMode::LoadBinary:
 	{
 		IFileManager::Get().FindFilesRecursive(InputAssetFilenames, *ProjectContentDir, *(Wildcard + FPackageName::GetAssetPackageExtension()), true, false, true);
 		//IFileManager::Get().FindFilesRecursive(InputAssetFilenames, *BasePath, *(Wildcard + FPackageName::GetTextMapPackageExtension()), true, false, false);
@@ -250,13 +240,13 @@ bool UTextAssetCommandlet::DoTextAssetProcessing(const FProcessingArgs& InArgs)
 
 		switch (InArgs.ProcessingMode)
 		{
-		case EProcessingMode::ResaveBinary:
+		case ETextAssetCommandletMode::ResaveBinary:
 		{
 			DestinationFilename = InputAssetFilename + TEXT(".tmp");
 			break;
 		}
 
-		case EProcessingMode::ResaveText:
+		case ETextAssetCommandletMode::ResaveText:
 		{
 			if (InputAssetFilename.EndsWith(FPackageName::GetAssetPackageExtension())) DestinationFilename = FPaths::ChangeExtension(InputAssetFilename, FPackageName::GetTextAssetPackageExtension());;
 			if (InputAssetFilename.EndsWith(FPackageName::GetMapPackageExtension())) DestinationFilename = FPaths::ChangeExtension(InputAssetFilename, FPackageName::GetTextMapPackageExtension());;
@@ -264,8 +254,8 @@ bool UTextAssetCommandlet::DoTextAssetProcessing(const FProcessingArgs& InArgs)
 			break;
 		}
 
-		case EProcessingMode::LoadText:
-		case EProcessingMode::LoadBinary:
+		case ETextAssetCommandletMode::LoadText:
+		case ETextAssetCommandletMode::LoadBinary:
 		{
 			break;
 		}
@@ -332,7 +322,7 @@ bool UTextAssetCommandlet::DoTextAssetProcessing(const FProcessingArgs& InArgs)
 
 			switch (InArgs.ProcessingMode)
 			{
-			case EProcessingMode::RoundTrip:
+			case ETextAssetCommandletMode::RoundTrip:
 			{
 				UE_LOG(LogTextAsset, Display, TEXT("Starting roundtrip test for '%s' [%d/%d]"), *SourceLongPackageName, NumFiles + 1, FilesToProcess.Num());
 				UE_LOG(LogTextAsset, Display, TEXT("-----------------------------------------------------------------------------------------"));
@@ -579,8 +569,8 @@ bool UTextAssetCommandlet::DoTextAssetProcessing(const FProcessingArgs& InArgs)
 				break;
 			}
 
-			case EProcessingMode::ResaveBinary:
-			case EProcessingMode::ResaveText:
+			case ETextAssetCommandletMode::ResaveBinary:
+			case ETextAssetCommandletMode::ResaveText:
 			{
 				UPackage* Package = nullptr;
 
@@ -612,7 +602,7 @@ bool UTextAssetCommandlet::DoTextAssetProcessing(const FProcessingArgs& InArgs)
 
 				if (bSaveSuccessful)
 				{
-					if (InArgs.bVerifyJson && InArgs.ProcessingMode == EProcessingMode::ResaveText)
+					if (InArgs.bVerifyJson && InArgs.ProcessingMode == ETextAssetCommandletMode::ResaveText)
 					{
 						TRACE_CPUPROFILER_EVENT_SCOPE_TEXT(TEXT("UTextAssetCommandlet::VerifyJson"));
 						FArchive* File = IFileManager::Get().CreateFileReader(*DestinationFilename);
@@ -637,7 +627,7 @@ bool UTextAssetCommandlet::DoTextAssetProcessing(const FProcessingArgs& InArgs)
 				break;
 			}
 
-			case EProcessingMode::LoadText:
+			case ETextAssetCommandletMode::LoadText:
 			{
 				UPackage* Package = nullptr;
 				UE_LOG(LogTextAsset, Display, TEXT("Loading Text Asset '%s'"), *SourceFilename);
@@ -657,7 +647,7 @@ bool UTextAssetCommandlet::DoTextAssetProcessing(const FProcessingArgs& InArgs)
 				break;
 			}
 
-			case EProcessingMode::LoadBinary:
+			case ETextAssetCommandletMode::LoadBinary:
 			{
 				UPackage* Package = nullptr;
 				UE_LOG(LogTextAsset, Display, TEXT("Loading Binary Asset '%s'"), *SourceFilename);
@@ -681,7 +671,7 @@ bool UTextAssetCommandlet::DoTextAssetProcessing(const FProcessingArgs& InArgs)
 			double EndTime = FPlatformTime::Seconds();
 			double Time = EndTime - StartTime;
 
-			if (InArgs.ProcessingMode == EProcessingMode::LoadBinary || InArgs.ProcessingMode == EProcessingMode::LoadText)
+			if (InArgs.ProcessingMode == ETextAssetCommandletMode::LoadBinary || InArgs.ProcessingMode == ETextAssetCommandletMode::LoadText)
 			{
 				if (ThisPackageLoadTime > MaxTime)
 				{
@@ -714,7 +704,7 @@ bool UTextAssetCommandlet::DoTextAssetProcessing(const FProcessingArgs& InArgs)
 			NumFiles++;
 		}
 
-		if (InArgs.ProcessingMode == EProcessingMode::RoundTrip)
+		if (InArgs.ProcessingMode == ETextAssetCommandletMode::RoundTrip)
 		{
 			UE_LOG(LogTextAsset, Display, TEXT("\t-----------------------------------------------------"));
 			UE_LOG(LogTextAsset, Display, TEXT("\tRoundTrip Results"));
@@ -744,7 +734,7 @@ bool UTextAssetCommandlet::DoTextAssetProcessing(const FProcessingArgs& InArgs)
 
 		double AvgFileTime, MinFileTime, MaxFileTime;
 
-		if (InArgs.ProcessingMode == EProcessingMode::LoadBinary || InArgs.ProcessingMode == EProcessingMode::LoadText)
+		if (InArgs.ProcessingMode == ETextAssetCommandletMode::LoadBinary || InArgs.ProcessingMode == ETextAssetCommandletMode::LoadText)
 		{
 			AvgFileTime = IterationPackageLoadTime;
 		}
@@ -770,7 +760,7 @@ bool UTextAssetCommandlet::DoTextAssetProcessing(const FProcessingArgs& InArgs)
 			CSVWriter->Serialize(TCHAR_TO_ANSI(*CSVLine), CSVLine.Len());
 		}
 
-		if (InArgs.ProcessingMode != EProcessingMode::LoadText && InArgs.ProcessingMode != EProcessingMode::ResaveText)
+		if (InArgs.ProcessingMode != ETextAssetCommandletMode::LoadText && InArgs.ProcessingMode != ETextAssetCommandletMode::ResaveText)
 		{
 			UE_LOG(LogTextAsset, Display, TEXT("\tTotal Package Save Time:  \t%.2fs"), IterationPackageSaveTime);
 		}
@@ -788,7 +778,7 @@ bool UTextAssetCommandlet::DoTextAssetProcessing(const FProcessingArgs& InArgs)
 		delete CSVWriter;
 	}
 
-	if (InArgs.ProcessingMode != EProcessingMode::LoadText)
+	if (InArgs.ProcessingMode != ETextAssetCommandletMode::LoadText)
 	{
 		UE_LOG(LogTextAsset, Display, TEXT("\tAvg Iteration Save Time:  \t%.2fs"), TotalPackageSaveTime / (float)InArgs.NumSaveIterations);
 	}
