PR #6371: Fix iterative cooking on Windows when the project is outside of the root directory but on the same drive (Contributed by [REDACTED_PERSON])

diff --git a/Engine/Source/Editor/UnrealEd/Private/CookOnTheFlyServer.cpp b/Engine/Source/Editor/UnrealEd/Private/CookOnTheFlyServer.cpp
--- a/Engine/Source/Editor/UnrealEd/Private/CookOnTheFlyServer.cpp
+++ b/Engine/Source/Editor/UnrealEd/Private/CookOnTheFlyServer.cpp
@@ -5072,6 +5072,9 @@ FName UCookOnTheFlyServer::ConvertCookedPathToUncookedPath(
 		BuildUncookedPath(FullCookedFilename, SandboxRootDir, RelativeRootDir);
 	}
 
+	// Convert to a standard filename as required by FPackageNameCache where this path is used.
+	FPaths::MakeStandardFilename(OutUncookedPath);
+
 	return FName(*OutUncookedPath);
 }
 
