Fix a crash when trying to launch a map from the editor.

GetCookingTargetPlatforms does not always return the same platform count as we give to InitializePackageStore and thus was causing out of bounds errors. We might as well just iterate over the array instead.

#rnx

diff --git a/Engine/Source/Editor/UnrealEd/Private/CookOnTheFlyServer.cpp b/Engine/Source/Editor/UnrealEd/Private/CookOnTheFlyServer.cpp
--- a/Engine/Source/Editor/UnrealEd/Private/CookOnTheFlyServer.cpp
+++ b/Engine/Source/Editor/UnrealEd/Private/CookOnTheFlyServer.cpp
@@ -6740,12 +6740,11 @@ void UCookOnTheFlyServer::FinalizePackageStore()
 	SCOPE_TIMER(FinalizePackageStore);
 
 	UE_LOG(LogCook, Display, TEXT("Saving BulkData manifest(s)..."));
-	const TArray<ITargetPlatform*>& TargetPlatforms = GetCookingTargetPlatforms();
-	for (int32 PlatformIndex = 0; PlatformIndex < TargetPlatforms.Num(); ++PlatformIndex)
+	for (FSavePackageContext* PackageContext : SavePackageContexts)
 	{
-		if (SavePackageContexts[PlatformIndex] != nullptr && SavePackageContexts[PlatformIndex]->BulkDataManifest != nullptr)
+		if (PackageContext != nullptr && PackageContext->BulkDataManifest != nullptr)
 		{
-			SavePackageContexts[PlatformIndex]->BulkDataManifest->Save();
+			PackageContext->BulkDataManifest->Save();
 		}
 	}
 	UE_LOG(LogCook, Display, TEXT("Done saving BulkData manifest(s)"));
