Fix a potential out of array bounds access when cooking on the fly.

When cooking on the fly SavePackageContexts will not be initialized so we need to check for this.

#rnx

diff --git a/Engine/Source/Editor/UnrealEd/Private/CookOnTheFlyServer.cpp b/Engine/Source/Editor/UnrealEd/Private/CookOnTheFlyServer.cpp
--- a/Engine/Source/Editor/UnrealEd/Private/CookOnTheFlyServer.cpp
+++ b/Engine/Source/Editor/UnrealEd/Private/CookOnTheFlyServer.cpp
@@ -4031,7 +4031,7 @@ void UCookOnTheFlyServer::SaveCookedPackage(UPackage* Package, uint32 SaveFlags,
 						}
 						else
 						{
-							FSavePackageContext* const SavePackageContext = SavePackageContexts[PlatformIndex];
+							FSavePackageContext* const SavePackageContext = SavePackageContexts.Num() > 0 ? SavePackageContexts[PlatformIndex] : nullptr;
 
 							Result = GEditor->Save(	Package, World, FlagsToCook, *PlatFilename, 
 													GError, nullptr, bSwap, false, SaveFlags, Target, 
