Optimize a slow collect references operation in the Force Delete command in the editor.

diff --git a/Engine/Source/Editor/UnrealEd/Private/ObjectTools.cpp b/Engine/Source/Editor/UnrealEd/Private/ObjectTools.cpp
--- a/Engine/Source/Editor/UnrealEd/Private/ObjectTools.cpp
+++ b/Engine/Source/Editor/UnrealEd/Private/ObjectTools.cpp
@@ -2403,23 +2403,64 @@ namespace ObjectTools
 		return true;
 	}
 
-	static void RecursiveRetrieveReferencers(UObject* Object, TSet<FWeakObjectPtr>& ReferencingObjects)
+	/**
+	 * Inspects all objects in memory and returns the set of all objects that transitively refer to the given InInterestSet
+	 * Objects in the original InInterestSet are included in the output ReferencingObjects set
+	 * Inner Objects that only have a path to the InterestSet through their outers are excluded.
+	 */
+	static void RecursiveRetrieveReferencers(const TArray<UObject*>& InInterestSet, TSet<FWeakObjectPtr>& OutReferencingObjects)
 	{
-		TArray<FReferencerInformation> ExternalReferencers;
-		Object->RetrieveReferencers(nullptr /* internal refs */, &ExternalReferencers);
+		const int32 ExpectedArraySize = 100;
+		const int32 ExpectedReferencesPerObject = 5;
+		TArray<UObject*> InterestSetAdditions(InInterestSet, FMath::Max(0,ExpectedArraySize - InInterestSet.Num()));
+
+		TMap<UObject*, int32> References;
+		TArray<UObject*> InterestSet;
+		InterestSet.Reserve(InterestSetAdditions.Max()*2);
+		References.Reserve(ExpectedReferencesPerObject);
+
+		// It would be faster to run a single TObjectIterator+Serialize loop and capture the complete graph of object references, and then do operations
+		// on the resultant graph, but that would require memory equal to sizeof(pointer)*num objects*(average references per object+3) to hold the graph.
+		// The extra cost of the current solution is that the TObjectIterator will be executed a number of times equal to 
+		// the length of the maximum (minimum reference chain length) from any object to the original interest set
+		// TODO: Worth the memory cost?
+		while (InterestSetAdditions.Num() > 0)
+		{
+			InterestSet.Append(InterestSetAdditions);
+			Algo::Sort(InterestSet, TLess<UObject*>());
+			InterestSetAdditions.Reset();
+
+			for (FObjectIterator It; It; ++It)
+			{
+				UObject* Object = *It;
+				if (Algo::BinarySearch(InterestSet, Object, TLess<UObject*>()) != INDEX_NONE)
+				{
+					continue;
+				}
 
-		for (const FReferencerInformation& RefInfo : ExternalReferencers)
-		{
-			if (RefInfo.Referencer != nullptr)
-			{
-				bool bAlreadyIn = false;
-				ReferencingObjects.Add(RefInfo.Referencer, &bAlreadyIn);
-				if (!bAlreadyIn)
+				const bool bAlsoFindWeakReferences = false;
+				FFindReferencersArchive ArFind(Object, InterestSet, bAlsoFindWeakReferences);
+				ArFind.GetReferenceCounts(References);
+				if (References.Num() > 0)
 				{
-					RecursiveRetrieveReferencers(RefInfo.Referencer, ReferencingObjects);
+					// Ignore internal references; only add the searched object if it refers to a member of the interest set but is not inside that member
+					for (const TPair<UObject*,int32>& kvpair : References)
+					{
+						if (!Object->IsIn(kvpair.Key))
+						{
+							InterestSetAdditions.Add(Object);
+							break;
+						}
+					}
+					References.Reset();
 				}
 			}
 		}
+
+		for (UObject* Referencer : InterestSet)
+		{
+			OutReferencingObjects.Add(Referencer);
+		}
 	}
 
 	int32 ForceDeleteObjects(const TArray< UObject* >& InObjectsToDelete, bool ShowConfirmation)
@@ -2446,12 +2487,7 @@ namespace ObjectTools
 
 		// Recursively find all references to objects being deleted
 		TSet<FWeakObjectPtr> ReferencingObjects;
-		for (UObject* ToDelete : InObjectsToDelete)
-		{
-			ReferencingObjects.Add(ToDelete);
-
-			RecursiveRetrieveReferencers(ToDelete, ReferencingObjects);
-		}
+		RecursiveRetrieveReferencers(InObjectsToDelete, ReferencingObjects);
 
 		// Attempt to close all editors referencing any of the deleted objects
 		bool bClosedAllEditors = true;
