Fix bug in Async MD5 code.  Some packages won't be saved to disk  (eg: if nativization converts the blueprint contents so there's nothing left to store on disk) and thus won't produce a TFuture for the CookedHash in the FSavePackageResultStruct.  If this is the case, we shouldn't try and setup a continuation to update the CookedHash in the asset registry for that package.

diff --git a/Engine/Source/Editor/UnrealEd/Classes/CookOnTheSide/CookOnTheFlyServer.h b/Engine/Source/Editor/UnrealEd/Classes/CookOnTheSide/CookOnTheFlyServer.h
--- a/Engine/Source/Editor/UnrealEd/Classes/CookOnTheSide/CookOnTheFlyServer.h
+++ b/Engine/Source/Editor/UnrealEd/Classes/CookOnTheSide/CookOnTheFlyServer.h
@@ -660,6 +660,15 @@ private:
 	 */
 	void SaveCookedPackages(UPackage* PackageToSave, const TArray<FName>& TargetPlatformNames, const TArray<const ITargetPlatform*>& TargetPlatformsToCache, struct FCookerTimer& Timer, uint32& CookedPackageCount, uint32& Result);
 
+	/**
+	 * Attempts to update the metadata for a package in an asset registry generator
+	 *
+	 * @param Generator The asset registry generator to update
+	 * @param PackageName The name of the package to update info on
+	 * @param SavePackageResult The metadata to associate with the given package name
+	 */
+	void UpdateAssetRegistryPackageData(FAssetRegistryGenerator* Generator, const FName& PackageName, FSavePackageResultStruct& SavePackageResult);
+
 	/** Perform any special processing for freshly loaded packages 
 	 */
 	void ProcessUnsolicitedPackages();
diff --git a/Engine/Source/Editor/UnrealEd/Private/CookOnTheFlyServer.cpp b/Engine/Source/Editor/UnrealEd/Private/CookOnTheFlyServer.cpp
--- a/Engine/Source/Editor/UnrealEd/Private/CookOnTheFlyServer.cpp
+++ b/Engine/Source/Editor/UnrealEd/Private/CookOnTheFlyServer.cpp
@@ -2963,19 +2963,8 @@ void UCookOnTheFlyServer::SaveCookedPackages(
 						if (CookByTheBookOptions)
 						{
 							FAssetRegistryGenerator* Generator = RegistryGenerators.FindRef(SaveTargetPlatformNames[iResultIndex]);
-							if (Generator)
-							{
-								FAssetPackageData* PackageData = Generator->GetAssetPackageData(Package->GetFName());
-								PackageData->DiskSize = SavePackageResult.TotalFileSize;
-								SavePackageResult.CookedHash.Next([PackageData](const FMD5Hash& CookedHash)
-								{
-									// Store the cooked hash in the Asset Registry when it is done computing in another thread.
-									// NOTE: For this to work, we rely on:
-									// 1) UPackage::WaitForAsyncFileWrites to have been called before any use of the CookedHash - it is called in CookByTheBookFinished before the registry does any work with the registries
-									// 2) PackageData must continue to be a valid pointer - the asset registry allocates the FAssetPackageData individually and doesn't relocate or delete them until pruning, which happens after WaitForAsyncFileWrites
-									PackageData->CookedHash = CookedHash;
-								});
-							}
+							UpdateAssetRegistryPackageData(Generator, Package->GetFName(), SavePackageResult);
+
 						}
 
 					}
@@ -3050,6 +3039,27 @@ void UCookOnTheFlyServer::SaveCookedPackages(
 	}
 }
 
+void UCookOnTheFlyServer::UpdateAssetRegistryPackageData(FAssetRegistryGenerator* Generator, const FName& PackageName, FSavePackageResultStruct& SavePackageResult)
+{
+	if (!Generator)
+		return;
+
+	FAssetPackageData* PackageData = Generator->GetAssetPackageData(PackageName);
+	PackageData->DiskSize = SavePackageResult.TotalFileSize;
+	// If there is no hash (e.g.: when SavePackageResult == ESavePackageResult::ReplaceCompletely), don't attempt to setup a continuation to update
+	// the AssetRegistry entry with it later.  Just leave the asset registry entry with a default constructed FMD5Hash which is marked as invalid.
+	if (SavePackageResult.CookedHash.IsValid())
+	{
+		SavePackageResult.CookedHash.Next([PackageData](const FMD5Hash& CookedHash)
+		{
+			// Store the cooked hash in the Asset Registry when it is done computing in another thread.
+			// NOTE: For this to work, we rely on:
+			// 1) UPackage::WaitForAsyncFileWrites to have been called before any use of the CookedHash - it is called in CookByTheBookFinished before the registry does any work with the registries
+			// 2) PackageData must continue to be a valid pointer - the asset registry allocates the FAssetPackageData individually and doesn't relocate or delete them until pruning, which happens after WaitForAsyncFileWrites
+			PackageData->CookedHash = CookedHash;
+		});
+	}
+}
 
 void UCookOnTheFlyServer::PostLoadPackageFixup(UPackage* Package)
 {
@@ -7998,19 +8008,7 @@ uint32 UCookOnTheFlyServer::FullLoadAndSave(uint32& CookedPackageCount)
 						if (bSucceededSavePackage)
 						{
 							FAssetRegistryGenerator* Generator = RegistryGenerators.FindRef(FName(*Target->PlatformName()));
-							if (Generator)
-							{
-								FAssetPackageData* PackageData = Generator->GetAssetPackageData(Package->GetFName());
-								PackageData->DiskSize = SaveResult.TotalFileSize;
-								SaveResult.CookedHash.Next([PackageData](const FMD5Hash& CookedHash)
-								{
-									// Store the cooked hash in the Asset Registry when it is done computing in another thread.
-									// NOTE: For this to work, we rely on:
-									// 1) UPackage::WaitForAsyncFileWrites to have been called before any use of the CookedHash - it is called in CookByTheBookFinished before the registry does any work with the registries
-									// 2) PackageData must continue to be a valid pointer - the asset registry allocates the FAssetPackageData individually and doesn't relocate or delete them until pruning, which happens after WaitForAsyncFileWrites
-									PackageData->CookedHash = CookedHash;
-								});
-							}
+							UpdateAssetRegistryPackageData(Generator, Package->GetFName(), SaveResult);
 
 							FPlatformAtomics::InterlockedIncrement(&ParallelSavedPackages);
 						}
