Guarded usage of FCpuProfilerTrace with #if CPUPROFILERTRACE_ENABLED, in TextAssetCommandlet.cpp.

diff --git a/Engine/Source/Editor/UnrealEd/Private/Commandlets/TextAssetCommandlet.cpp b/Engine/Source/Editor/UnrealEd/Private/Commandlets/TextAssetCommandlet.cpp
--- a/Engine/Source/Editor/UnrealEd/Private/Commandlets/TextAssetCommandlet.cpp
+++ b/Engine/Source/Editor/UnrealEd/Private/Commandlets/TextAssetCommandlet.cpp
@@ -358,6 +358,7 @@ bool UTextAssetCommandlet::DoTextAssetProcessing(const FProcessingArgs& InArgs)
 
 				static const TCHAR* PhaseNames[] = { TEXT("Binary Only"), TEXT("Text Only"), TEXT("Alternating Binary/Text") };
 
+				#if CPUPROFILERTRACE_ENABLED
 				static const uint32 PhaseEventTypes[3] = {
 						FCpuProfilerTrace::OutputEventType(TEXT("BinaryOnly"), CpuProfilerGroup_Default),
 						FCpuProfilerTrace::OutputEventType(TEXT("TextOnly"), CpuProfilerGroup_Default),
@@ -372,6 +373,7 @@ bool UTextAssetCommandlet::DoTextAssetProcessing(const FProcessingArgs& InArgs)
 						FCpuProfilerTrace::OutputEventType(TEXT("Test5"), CpuProfilerGroup_Default),
 						FCpuProfilerTrace::OutputEventType(TEXT("Test6"), CpuProfilerGroup_Default)
 				};
+				#endif // CPUPROFILERTRACE_ENABLED
 
 				TArray<TArray<FSHAHash>> Hashes;
 
@@ -387,13 +389,17 @@ bool UTextAssetCommandlet::DoTextAssetProcessing(const FProcessingArgs& InArgs)
 
 					TArray<FSHAHash> PhaseHashes = Hashes[Hashes.AddDefaulted()];
 					
+					#if CPUPROFILERTRACE_ENABLED
 					FCpuProfilerTrace::OutputBeginEvent(PhaseEventTypes[Phase]);
+					#endif
 
 					for (int32 i = 0; i < ((Phase == 2) ? NumTests * 2 : NumTests); ++i)
 					{
 						int32 Bucket;
 
+						#if CPUPROFILERTRACE_ENABLED
 						FCpuProfilerTrace::OutputBeginEvent(TestEventTypes[i]);
+						#endif
 
 						switch (Phase)
 						{
@@ -474,7 +480,9 @@ bool UTextAssetCommandlet::DoTextAssetProcessing(const FProcessingArgs& InArgs)
 							DiffFilenames.Add(TPair<FString, FString>(IntermediateFilename, FinalFilename));
 						}
 
+						#if CPUPROFILERTRACE_ENABLED
 						FCpuProfilerTrace::OutputEndEvent();
+						#endif
 					}
 
 					UE_LOG(LogTextAsset, Display, TEXT("Phase %i (%s) Results"), Phase + 1, PhaseNames[Phase]);
@@ -518,7 +526,9 @@ bool UTextAssetCommandlet::DoTextAssetProcessing(const FProcessingArgs& InArgs)
 
 					DiffFilenames.Empty();
 					IFileManager::Get().DeleteDirectory(*TempFailedDiffsPath, false, true);
+					#if CPUPROFILERTRACE_ENABLED
 					FCpuProfilerTrace::OutputEndEvent();
+					#endif
 				}
 
 				static const bool bDisableCleanup = FParse::Param(FCommandLine::Get(), TEXT("disablecleanup"));
@@ -807,4 +817,4 @@ static void TextAssetToolCVarCommand(const TArray<FString>& Args)
 static FAutoConsoleCommand CVar_TextAssetTool(
 	TEXT("TextAssetTool"),
 	TEXT("--"),
-	FConsoleCommandWithArgsDelegate::CreateStatic(TextAssetToolCVarCommand));
+	FConsoleCommandWithArgsDelegate::CreateStatic(TextAssetToolCVarCommand));
\ No newline at end of file
