Delete the Sandbox directory asynchronously while the rest of the cook goes ahead, to save time when the sandbox directory has many files to delete.
#rn minor CookingOptimization

diff --git a/Engine/Source/Editor/UnrealEd/Classes/CookOnTheSide/CookOnTheFlyServer.h b/Engine/Source/Editor/UnrealEd/Classes/CookOnTheSide/CookOnTheFlyServer.h
--- a/Engine/Source/Editor/UnrealEd/Classes/CookOnTheSide/CookOnTheFlyServer.h
+++ b/Engine/Source/Editor/UnrealEd/Classes/CookOnTheSide/CookOnTheFlyServer.h
@@ -16,6 +16,7 @@
 
 
 class FAssetRegistryGenerator;
+class FAsyncIODelete;
 class ITargetPlatform;
 struct FPropertyChangedEvent;
 class IPlugin;
@@ -96,12 +97,15 @@ enum class ECookTickFlags : uint8
 };
 ENUM_CLASS_FLAGS(ECookTickFlags);
 
+DECLARE_LOG_CATEGORY_EXTERN(LogCook, Log, All);
+
 UCLASS()
 class UNREALED_API UCookOnTheFlyServer : public UObject, public FTickableEditorObject, public FExec
 {
 	GENERATED_BODY()
 
 		UCookOnTheFlyServer(const FObjectInitializer& ObjectInitializer = FObjectInitializer::Get());
+		UCookOnTheFlyServer(FVTableHelper& Helper); // Declare the FVTableHelper constructor manually so that we can forward-declare-only TUniquePtrs in the header without getting compile error in generated cpp
 
 private:
 	/** Current cook mode the cook on the fly server is running in */
@@ -142,6 +146,7 @@ private:
 
 	ECookInitializationFlags CookFlags = ECookInitializationFlags::None;
 	TUniquePtr<class FSandboxPlatformFile> SandboxFile;
+	TUniquePtr<FAsyncIODelete> AsyncIODelete; // Helper for deleting the old cook directory asynchronously
 	bool bIsInitializingSandbox = false; // stop recursion into callbacks when we are initializing sandbox
 	mutable bool bIgnoreMarkupPackageAlreadyLoaded = false; // avoid marking up packages as already loaded (want to put this around some functionality as we want to load packages fully some times)
 	bool bIsSavingPackage = false; // used to stop recursive mark package dirty functions
@@ -396,7 +401,7 @@ public:
 	* 
 	* @param name of the platform to clear the cooked packages for
 	*/
-	void ClearPlatformCookedData( const FName& PlatformName );
+	void ClearPlatformCookedData( const FString& PlatformName );
 
 
 	/**
@@ -540,10 +545,10 @@ private:
 	*/
 	void AddFileToCook( TArray<FName>& InOutFilesToCook, const FString &InFilename ) const;
 
-    /**
-     * Invokes the necessary FShaderCodeLibrary functions to start cooking the shader code library.
-     */
-    void InitShaderCodeLibrary(void);
+	/**
+	* Invokes the necessary FShaderCodeLibrary functions to start cooking the shader code library.
+	*/
+	void InitShaderCodeLibrary(void);
     
 	/**
 	* Invokes the necessary FShaderCodeLibrary functions to open a named code library.
@@ -560,10 +565,10 @@ private:
 	*/
 	void ProcessShaderCodeLibraries(const FString& LibraryName);
 
-    /**
-     * Invokes the necessary FShaderCodeLibrary functions to clean out all the temporary files.
-     */
-    void CleanShaderCodeLibraries();
+	/**
+	* Invokes the necessary FShaderCodeLibrary functions to clean out all the temporary files.
+	*/
+	void CleanShaderCodeLibraries();
 
 	/**
 	* Call back from the TickCookOnTheSide when a cook by the book finishes (when started form StartCookByTheBook)
@@ -826,6 +831,15 @@ private:
 	*/
 	FString GetSandboxDirectory( const FString& PlatformName ) const;
 
+	/* Delete the sandbox directory (asynchronously) for the given platform in preparation for a clean cook */
+	void DeleteSandboxDirectory(const FString& PlatformName);
+
+	/* Create the delete-old-cooked-directory helper. PlatformName tells it where to create its temp directory. AsyncDeleteDirectory is DeleteSandboxDirectory internal use only and should otherwise be set to nullptr. */
+	FAsyncIODelete& GetAsyncIODelete(const FString& PlatformName, const FString* AsyncDeleteDirectory = nullptr);
+
+	/** Get the directory that should be used for AsyncDeletes based on the given platform. SandboxDirectory is DeleteSandboxDirectory internal use only and should otherwise be set to nullptr. */
+	FString GetAsyncDeleteDirectory(const FString& PlatformName, const FString* SandboxDirectory = nullptr) const;
+
 	bool IsCookingDLC() const;
 
 	/**
diff --git a/Engine/Source/Editor/UnrealEd/Private/CookOnTheFlyServer.cpp b/Engine/Source/Editor/UnrealEd/Private/CookOnTheFlyServer.cpp
--- a/Engine/Source/Editor/UnrealEd/Private/CookOnTheFlyServer.cpp
+++ b/Engine/Source/Editor/UnrealEd/Private/CookOnTheFlyServer.cpp
@@ -67,6 +67,7 @@
 #include "PlatformInfo.h"
 #include "Serialization/ArchiveStackTrace.h"
 #include "DistanceFieldAtlas.h"
+#include "Cooker/AsyncIODelete.h"
 
 #include "AssetRegistryModule.h"
 #include "AssetRegistryState.h"
@@ -104,7 +105,7 @@
 #define LOCTEXT_NAMESPACE "Cooker"
 #define REMAPPED_PLUGGINS TEXT("RemappedPlugins")
 
-DEFINE_LOG_CATEGORY_STATIC(LogCook, Log, All);
+DEFINE_LOG_CATEGORY(LogCook);
 
 #define DEBUG_COOKONTHEFLY 0
 #define OUTPUT_TIMING 1
@@ -1446,6 +1447,8 @@ UCookOnTheFlyServer::UCookOnTheFlyServer(const FObjectInitializer& ObjectInitial
 {
 }
 
+UCookOnTheFlyServer::UCookOnTheFlyServer(FVTableHelper& Helper) :Super(Helper) {}
+
 UCookOnTheFlyServer::~UCookOnTheFlyServer()
 {
 	FCoreDelegates::OnFConfigCreated.RemoveAll(this);
@@ -5124,6 +5127,56 @@ void UCookOnTheFlyServer::GetAllCookedFiles(TMap<FName, FName>& UncookedPathToCo
 	}
 }
 
+void UCookOnTheFlyServer::DeleteSandboxDirectory(const FString& PlatformName)
+{
+	FString SandboxDirectory = GetSandboxDirectory(PlatformName);
+	FPaths::NormalizeDirectoryName(SandboxDirectory);
+	FString AsyncDeleteDirectory = GetAsyncDeleteDirectory(PlatformName, &SandboxDirectory);
+
+	FAsyncIODelete& LocalAsyncIODelete = GetAsyncIODelete(PlatformName, &AsyncDeleteDirectory);
+	LocalAsyncIODelete.DeleteDirectory(SandboxDirectory);
+
+	// Part of Deleting the sandbox includes deleting the old AsyncDelete directory for the sandbox, in case a previous cooker crashed before cleaning it up.
+	// The AsyncDelete directory is associated with a sandbox but is necessarily outside of it since it is used to delete the sandbox.
+	// Note that for the Platform we used to create the AsyncIODelete, this Delete will fail because AsyncIODelete refuses to delete its own temproot; this is okay because it will delete the temproot on exit.
+	LocalAsyncIODelete.DeleteDirectory(AsyncDeleteDirectory);
+}
+
+FAsyncIODelete& UCookOnTheFlyServer::GetAsyncIODelete(const FString& PlatformName, const FString* AsyncDeleteDirectory)
+{
+	FAsyncIODelete* AsyncIODeletePtr = AsyncIODelete.Get();
+	if (!AsyncIODeletePtr)
+	{
+		FString Buffer;
+		if (!AsyncDeleteDirectory)
+		{
+			Buffer = GetAsyncDeleteDirectory(PlatformName);
+			AsyncDeleteDirectory = &Buffer;
+		}
+		AsyncIODelete = MakeUnique<FAsyncIODelete>(*AsyncDeleteDirectory);
+		AsyncIODeletePtr = AsyncIODelete.Get();
+	}
+	// If we have already created the AsyncIODelete, we ignore the input PlatformName and use the existing AsyncIODelete initialized from whatever platform we used before
+	// The PlatformName is used only to construct a directory that we can be sure no other process is using (because a sandbox can only be cooked by one process at a time)
+	return *AsyncIODeletePtr;
+}
+
+FString UCookOnTheFlyServer::GetAsyncDeleteDirectory(const FString& PlatformName, const FString* SandboxDirectory) const
+{
+	// The TempRoot we will delete into is a sibling of the the Platform-specific sandbox directory, with name [PlatformDir]AsyncDelete
+	// Note that two UnrealEd-Cmd processes cooking to the same sandbox at the same time will therefore cause an error since FAsyncIODelete doesn't handle multiple processes sharing TempRoots.
+	// That simultaneous-cook behavior is also not supported in other assumptions throughout the cooker.
+	FString Buffer;
+	if (!SandboxDirectory)
+	{
+		// Avoid recalculating SandboxDirectory if caller supplied it, but if they didn't, calculate it here
+		Buffer = GetSandboxDirectory(PlatformName);
+		FPaths::NormalizeDirectoryName(Buffer);
+		SandboxDirectory = &Buffer;
+	}
+	return (*SandboxDirectory) + TEXT("AsyncDelete");
+}
+
 void UCookOnTheFlyServer::PopulateCookedPackagesFromDisk(const TArray<ITargetPlatform*>& Platforms)
 {
 	TRACE_CPUPROFILER_EVENT_SCOPE(UCookOnTheFlyServer::PopulateCookedPackagesFromDisk);
@@ -5176,15 +5229,14 @@ void UCookOnTheFlyServer::PopulateCookedPackagesFromDisk(const TArray<ITargetPla
 			if ( (CurrentIterativeCookedBuild >= CurrentLocalCookedBuild) && 
 				(CurrentIterativeCookedBuild != FDateTime::MinValue()) )
 			{
-				// clean the sandbox 
-				ClearPlatformCookedData(FName(*Target->PlatformName()));
-				FString SandboxDirectory = GetSandboxDirectory(Target->PlatformName());
-				IFileManager::Get().DeleteDirectory(*SandboxDirectory, false, true);
+				// clean the sandbox
+				const FString PlatformName = Target->PlatformName();
+				ClearPlatformCookedData(PlatformName);
 
 				// SaveCurrentIniSettings(Target); // use this if we don't care about ini safty.
 				// copy the ini settings from the shared cooked build. 
-				const FString SharedCookedIniFile = FPaths::Combine(*FPaths::ProjectSavedDir(), TEXT("SharedIterativeBuild"), *Target->PlatformName(), TEXT("Metadata"), TEXT("CookedIniVersion.txt"));
-				const FString SandboxCookedIniFile = ConvertToFullSandboxPath(*(FPaths::ProjectDir() / TEXT("Metadata") / TEXT("CookedIniVersion.txt")), true).Replace(TEXT("[Platform]"), *Target->PlatformName());
+				const FString SharedCookedIniFile = FPaths::Combine(*FPaths::ProjectSavedDir(), TEXT("SharedIterativeBuild"), *PlatformName, TEXT("Metadata"), TEXT("CookedIniVersion.txt"));
+				const FString SandboxCookedIniFile = ConvertToFullSandboxPath(*(FPaths::ProjectDir() / TEXT("Metadata") / TEXT("CookedIniVersion.txt")), true).Replace(TEXT("[Platform]"), *PlatformName);
 
 				IFileManager::Get().Copy(*SandboxCookedIniFile, *SharedCookedIniFile);
 
@@ -5405,41 +5457,25 @@ void UCookOnTheFlyServer::CleanSandbox(const bool bIterative)
 		SCOPE_SECONDS_COUNTER(SandboxCleanTime);
 		SCOPE_TIMER(CleanSandboxTime);
 #endif
-		if (bIterative == false)
+
+		for (const ITargetPlatform* Target : Platforms)
 		{
-			// for now we are going to wipe the cooked directory
-			for (int32 Index = 0; Index < Platforms.Num(); Index++)
+			bool bShouldClearCookedContent = false;
+			const bool bIsIniSettingsOutOfDate = IniSettingsOutOfDate(Target); // needs to be executed for side effects even if non-iterative
+			if (!bIterative)
 			{
-				ITargetPlatform* Target = Platforms[Index];
-				UE_LOG(LogCook, Display, TEXT("Cooked content cleared for platform %s"), *Target->PlatformName());
-
-				FString SandboxDirectory = GetSandboxDirectory(Target->PlatformName()); // GetOutputDirectory(Target->PlatformName());
-				IFileManager::Get().DeleteDirectory(*SandboxDirectory, false, true);
-
-				ClearPlatformCookedData(FName(*Target->PlatformName()));
-
-				IniSettingsOutOfDate(Target);
-				SaveCurrentIniSettings(Target);
+				// for now we always wipe the cooked directory on full cooks
+				UE_LOG(LogCook, Display, TEXT("Clearing all cooked content for platform %s"), *Target->PlatformName());
+				bShouldClearCookedContent = true;
 			}
-
-		}
-		else
-		{
-			for (const ITargetPlatform* Target : Platforms)
+			else
 			{
-				bool bIniSettingsOutOfDate = IniSettingsOutOfDate(Target);
-				if (bIniSettingsOutOfDate)
+				if (bIsIniSettingsOutOfDate)
 				{
 					if (!IsCookFlagSet(ECookInitializationFlags::IgnoreIniSettingsOutOfDate))
 					{
 						UE_LOG(LogCook, Display, TEXT("Cook invalidated for platform %s ini settings don't match from last cook, clearing all cooked content"), *Target->PlatformName());
-
-						ClearPlatformCookedData(FName(*Target->PlatformName()));
-
-						FString SandboxDirectory = GetSandboxDirectory(Target->PlatformName());
-						IFileManager::Get().DeleteDirectory(*SandboxDirectory, false, true);
-
-						SaveCurrentIniSettings(Target);
+						bShouldClearCookedContent = true;
 					}
 					else
 					{
@@ -5448,6 +5484,15 @@ void UCookOnTheFlyServer::CleanSandbox(const bool bIterative)
 				}
 			}
 
+			if (bShouldClearCookedContent)
+			{
+				ClearPlatformCookedData(Target->PlatformName());
+				SaveCurrentIniSettings(Target);
+			}
+		}
+
+		if (bIterative)
+		{
 			// This is fast in asset registry iterate, so just reconstruct
 			PackageTracker->CookedPackages.Empty();
 			PopulateCookedPackagesFromDisk(Platforms);
@@ -6669,16 +6714,19 @@ void UCookOnTheFlyServer::ClearAllCookedData()
 	PackageTracker->CookedPackages.Empty(); // set of files which have been cooked when needing to recook a file the entry will need to be removed from here
 }
 
-void UCookOnTheFlyServer::ClearPlatformCookedData(const FName& PlatformName)
+void UCookOnTheFlyServer::ClearPlatformCookedData(const FString& PlatformNameString)
 {
 	// if we are going to clear the cooked packages it is conceivable that we will recook the packages which we just cooked 
 	// that means it's also conceivable that we will recook the same package which currently has an outstanding async write request
 	UPackage::WaitForAsyncFileWrites();
 
+	FName PlatformName(*PlatformNameString);
 	PackageTracker->CookedPackages.RemoveAllFilesForPlatform(PlatformName);
 
 	TArray<FName> PackageNames;
 	PackageTracker->UnsolicitedCookedPackages.GetPackagesForPlatformAndRemove(PlatformName, PackageNames);
+
+	DeleteSandboxDirectory(PlatformNameString);
 }
 
 void UCookOnTheFlyServer::ClearCachedCookedPlatformDataForPlatform( const FName& PlatformName )
diff --git a/Engine/Source/Editor/UnrealEd/Private/Cooker/AsyncIODelete.cpp b/Engine/Source/Editor/UnrealEd/Private/Cooker/AsyncIODelete.cpp
new file mode 100644
--- /dev/null
+++ b/Engine/Source/Editor/UnrealEd/Private/Cooker/AsyncIODelete.cpp
@@ -0,0 +1,301 @@
+// Copyright 1998-2019 Epic Games, Inc. All Rights Reserved.
+
+#include "AsyncIODelete.h"
+#include "CoreMinimal.h"
+
+#include "Async/Async.h"
+#include "Containers/UnrealString.h"
+#include "HAL/Event.h"
+#include "HAL/FileManager.h"
+#include "Math/NumericLimits.h"
+#include "Misc/StringBuilder.h"
+#include "Misc/Paths.h"
+#include "Templates/UnrealTemplate.h"
+
+#if WITH_ASYNCIODELETE_DEBUG
+TArray<FString> FAsyncIODelete::AllTempRoots;
+#endif
+
+DECLARE_LOG_CATEGORY_EXTERN(LogCook, Log, All);
+
+FAsyncIODelete::FAsyncIODelete(const FStringView& InOwnedTempRoot)
+{
+	SetTempRoot(InOwnedTempRoot);
+}
+
+FAsyncIODelete::~FAsyncIODelete()
+{
+	SetTempRoot(FStringView());
+}
+
+void FAsyncIODelete::SetTempRoot(const FStringView& InOwnedTempRoot)
+{
+	Teardown();
+
+#if WITH_ASYNCIODELETE_DEBUG
+	if (!TempRoot.IsEmpty())
+	{
+		RemoveTempRoot(*TempRoot);
+	}
+#endif
+
+	TempRoot = InOwnedTempRoot;
+
+#if WITH_ASYNCIODELETE_DEBUG
+	if (!TempRoot.IsEmpty())
+	{
+		AddTempRoot(*TempRoot);
+	}
+#endif
+}
+
+void FAsyncIODelete::SetDeletesPaused(bool bInPaused)
+{
+	bPaused = bInPaused;
+	if (!bPaused)
+	{
+		IFileManager& FileManager = IFileManager::Get();
+		for (const FString& DeletePath : PausedDeletes)
+		{
+			const bool IsDirectory = FileManager.DirectoryExists(*DeletePath);
+			const bool IsFile = !IsDirectory && FileManager.FileExists(*DeletePath);
+			if (!IsDirectory && !IsFile)
+			{
+				continue;
+			}
+			CreateDeleteTask(DeletePath, IsDirectory ? EPathType::Directory : EPathType::File);
+		}
+		PausedDeletes.Empty();
+	}
+}
+
+bool FAsyncIODelete::Setup()
+{
+	if (bInitialized)
+	{
+		return true;
+	}
+
+	if (TempRoot.IsEmpty())
+	{
+		checkf(false, TEXT("DeleteDirectory called without having first set a TempRoot"));
+		return false;
+	}
+
+	// Delete the TempRoot directory to clear the results from any previous process using the same TempRoot that did not shut down cleanly
+	if (!ensure(DeleteTempRootDirectory()))
+	{
+		return false;
+	}
+
+	// Create the empty directory to work in
+	if (!ensure(IFileManager::Get().MakeDirectory(*TempRoot, true)))
+	{
+		return false;
+	}
+
+	// Allocate the task event
+	check(TasksComplete == nullptr);
+	TasksComplete = FPlatformProcess::GetSynchEventFromPool(true /* IsManualReset */);
+	check(ActiveTaskCount == 0);
+	TasksComplete->Trigger(); // We have 0 tasks so the event should be in the Triggered state
+
+	// Assert that all other teardown-transient variables were cleared by the constructor or by the previous teardown
+	// TempRoot and bPaused are preserved across setup/teardown and may have any value
+	check(PausedDeletes.Num() == 0);
+	check(DeleteCounter == 0);
+
+	// We are now setup and ready to create DeleteTasks
+	bInitialized = true;
+
+	return true;
+}
+
+void FAsyncIODelete::Teardown()
+{
+	if (!bInitialized)
+	{
+		return;
+	}
+
+	// Clear task variables
+	WaitForAllTasks();
+	check(ActiveTaskCount == 0 && TasksComplete != nullptr && TasksComplete->Wait(0));
+	FPlatformProcess::ReturnSynchEventToPool(TasksComplete);
+	TasksComplete = nullptr;
+
+	// Remove the temp directory from disk
+	if (!DeleteTempRootDirectory())
+	{
+		// This will leave directories (and potentially files, if we were paused or if any of the asyncdeletes failed) on disk, so it is bad for users, but is not fatal for our operations.
+		UE_LOG(LogCook, Warning, TEXT("Could not delete asyncdelete directory '%s'."), *TempRoot);
+	}
+
+	// Clear delete variables; we don't need to run the tasks for the remaining pauseddeletes because synchronously deleting the temp directory above did the work they were going to do
+	PausedDeletes.Empty();
+	DeleteCounter = 0;
+
+	// We are now torn down and ready for a new setup
+	bInitialized = false;
+}
+
+void FAsyncIODelete::WaitForAllTasks()
+{
+	if (!bInitialized)
+	{
+		return;
+	}
+
+	TasksComplete->Wait();
+	check(ActiveTaskCount == 0);
+}
+
+bool FAsyncIODelete::Delete(const FStringView& PathToDelete, EPathType ExpectedType)
+{
+	IFileManager& FileManager = IFileManager::Get();
+	TStringBuilder<128> PathToDeleteBuffer;
+	PathToDeleteBuffer << PathToDelete;
+	const TCHAR* PathToDeleteSZ = PathToDeleteBuffer.ToString();
+
+	const bool IsDirectory = FileManager.DirectoryExists(PathToDeleteSZ);
+	const bool IsFile = !IsDirectory && FileManager.FileExists(PathToDeleteSZ);
+	if (!IsDirectory && !IsFile)
+	{
+		return true;
+	}
+	if (ExpectedType == EPathType::Directory && !IsDirectory)
+	{
+		checkf(false, TEXT("DeleteDirectory called on \"%.*s\" which is not a directory."), PathToDelete.Len(), PathToDelete.GetData());
+		return false;
+	}
+	if (ExpectedType == EPathType::File && !IsFile)
+	{
+		checkf(false, TEXT("DeleteFile called on \"%.*s\" which is not a file."), PathToDelete.Len(), PathToDelete.GetData());
+		return false;
+	}
+	// Prevent the user from trying to delete our temproot or anything inside it
+	FString PathToDeleteStr(PathToDelete);
+	if (FPaths::IsUnderDirectory(PathToDeleteStr, TempRoot) || FPaths::IsUnderDirectory(TempRoot, PathToDeleteStr))
+	{
+		return false;
+	}
+
+	if (DeleteCounter == UINT32_MAX)
+	{
+		Teardown();
+	}
+	if (!Setup())
+	{
+		// Setup failed; we are not able to provide asynchronous deletes; fall back to synchronous
+		return SynchronousDelete(PathToDeleteSZ, ExpectedType);
+	}
+
+	const FString TempPath = FPaths::Combine(TempRoot, FString::Printf(TEXT("Delete%u"), DeleteCounter));
+	DeleteCounter++;
+
+	const bool bReplace = true;
+	const bool bEvenIfReadOnly = true;
+	const bool bMoveAttributes = false;
+	const bool bDoNotRetryOnError = true;
+	if (!IFileManager::Get().Move(*TempPath, PathToDeleteSZ, bReplace, bEvenIfReadOnly, bMoveAttributes, bDoNotRetryOnError)) // IFileManager::Move works on either files or directories
+	{
+		// The move failed; try a synchronous delete as backup
+		UE_LOG(LogCook, Warning, TEXT("Failed to move path '%.*s' for async delete; falling back to synchronous delete."), PathToDelete.Len(), PathToDelete.GetData());
+		return SynchronousDelete(PathToDeleteSZ, ExpectedType);
+	}
+
+	if (bPaused)
+	{
+		PausedDeletes.Add(TempPath);
+	}
+	else
+	{
+		CreateDeleteTask(TempPath, ExpectedType);
+	}
+	return true;
+}
+
+void FAsyncIODelete::CreateDeleteTask(const FStringView& InDeletePath, EPathType PathType)
+{
+	{
+		FScopeLock Lock(&CriticalSection);
+		TasksComplete->Reset();
+		ActiveTaskCount++;
+	}
+
+	AsyncThread(
+		[this, DeletePath = FString(InDeletePath), PathType]() { ensure(SynchronousDelete(*DeletePath, PathType)); },
+		0, TPri_Normal,
+		[this]() { OnTaskComplete(); });
+}
+
+void FAsyncIODelete::OnTaskComplete()
+{
+	FScopeLock Lock(&CriticalSection);
+	check(ActiveTaskCount > 0);
+	ActiveTaskCount--;
+	if (ActiveTaskCount == 0)
+	{
+		TasksComplete->Trigger();
+	}
+}
+
+bool FAsyncIODelete::SynchronousDelete(const TCHAR* InDeletePath, EPathType PathType)
+{
+	const bool bRequireExists = false;
+	if (PathType == EPathType::Directory)
+	{
+		const bool bTree = true;
+		return IFileManager::Get().DeleteDirectory(InDeletePath, bRequireExists, bTree);
+	}
+	else
+	{
+		const bool bEvenIfReadOnly = true;
+		return IFileManager::Get().Delete(InDeletePath, bRequireExists, bEvenIfReadOnly);
+	}
+}
+
+bool FAsyncIODelete::DeleteTempRootDirectory()
+{
+	IFileManager& FileManager = IFileManager::Get();
+	if (!FileManager.DirectoryExists(*TempRoot))
+	{
+		return true;
+	}
+
+	// Since we sometimes will be creating the directory again immediately, we need to take precautions against the delayed delete of directories that
+	// occurs on Windows platforms; creating a new file/directory in one that was just deleted can fail.  So we need to move-delete our TempRoot
+	// in addition to move-delete our clients' directories.  Since we don't have a TempRoot to move-delete into, we create a unique sibling directory name.
+	FString UniqueDirectory = FPaths::CreateTempFilename(*FPaths::GetPath(TempRoot), TEXT("DeleteTemp"), TEXT(""));
+
+	const bool bReplace = false;
+	const bool bEvenIfReadOnly = true;
+	const TCHAR* DirectoryToDelete = *UniqueDirectory;
+	if (!FileManager.Move(DirectoryToDelete, *TempRoot, bReplace, bEvenIfReadOnly))
+	{
+		// Move failed; fallback to inplace delete
+		DirectoryToDelete = *TempRoot;
+	}
+
+	const bool bRequireExists = false;
+	const bool bTree = true;
+	return FileManager.DeleteDirectory(DirectoryToDelete, bRequireExists, bTree);
+}
+
+#if WITH_ASYNCIODELETE_DEBUG
+void FAsyncIODelete::AddTempRoot(const FStringView& InTempRoot)
+{
+	FString TempRoot(InTempRoot);
+	for (FString& Existing : AllTempRoots)
+	{
+		checkf(!FPaths::IsUnderDirectory(Existing, TempRoot), TEXT("New FAsyncIODelete has TempRoot \"%s\" that is a subdirectory of existing TempRoot \"%s\"."), *TempRoot, *Existing);
+		checkf(!FPaths::IsUnderDirectory(TempRoot, Existing), TEXT("New FAsyncIODelete has TempRoot \"%s\" that is a parent directory of existing TempRoot \"%s\"."), *TempRoot, *Existing);
+	}
+	AllTempRoots.Add(MoveTemp(TempRoot));
+}
+
+void FAsyncIODelete::RemoveTempRoot(const FStringView& InTempRoot)
+{
+	AllTempRoots.Remove(FString(InTempRoot));
+}
+#endif
diff --git a/Engine/Source/Editor/UnrealEd/Private/Cooker/AsyncIODelete.h b/Engine/Source/Editor/UnrealEd/Private/Cooker/AsyncIODelete.h
new file mode 100644
--- /dev/null
+++ b/Engine/Source/Editor/UnrealEd/Private/Cooker/AsyncIODelete.h
@@ -0,0 +1,129 @@
+// Copyright 1998-2019 Epic Games, Inc. All Rights Reserved.
+
+#pragma once
+
+#include "CoreTypes.h"
+
+#include "Async/Future.h"
+#include "Containers/Array.h"
+#include "Containers/StringView.h"
+#include "Misc/AssertionMacros.h"
+
+#ifndef WITH_ASYNCIODELETE_DEBUG
+	#if DO_CHECK
+		#define WITH_ASYNCIODELETE_DEBUG 1 // When enabled the class uses extra memory and does some tracking
+	#else
+		#define WITH_ASYNCIODELETE_DEBUG 0
+	#endif
+#endif
+
+class FEvent;
+
+/**
+ * Helper class to asynchronously delete paths (files or directories).
+ * The deleted paths are moved to a unique temporary location immediately, and then asynchronously deleted later.
+ * Note that two FAsyncIODeletes (in the same process or in multiple processes) at the same with the same InOwnedTempRoot 
+ * is forbidden and will cause behavior errors.  Caller must ensure the machine-global uniqueness of InOwnedTempRoot.
+ *
+ * FAsyncIODelete's public interface is not threadsafe; all calls to the public interface must be run from the same thread.
+ */
+class FAsyncIODelete
+{
+public:
+	/**
+	 * Quick constructor to create flags; does not take any IO action until Setup or a Delete function is used
+	 *
+	 * @param InOwnedTempRoot - The FAsyncIODelete takes ownership of this directory, and deletes it when done.
+	 *							It is a temporary directory used to hold the deleted directories from elsewhere.  Must not be shared with (or a parent or child of) any other FAsyncDelete's InTempRoot.
+	 *						    The FAsyncIODelete will not be able to delete directories until a valid TempRoot has been set, either in the constructor or in SetTempRoot
+	 */
+	explicit FAsyncIODelete(const FStringView& InOwnedTempRoot=FStringView());
+
+	/* Synchronously waits for all requested deletions to complete */
+	~FAsyncIODelete();
+
+	/**
+	 * Set the TempRoot directory used to hold the deleted directories from elsewhere.  The FAsyncIODelete takes ownership of this directory, and deletes it when destructed or when it is changed.
+	 * If the FAsyncIODelete has already been used to delete files, this Set will block until all deletes are finished.
+	 */
+	void SetTempRoot(const FStringView& InOwnedTempRoot);
+	FStringView GetTempRoot() const { return TempRoot; }
+
+	/**
+	 * Set whether new background deletes are paused.  If paused, paths will be moved immediately but will not be deleted from the temporary location until unpaused (or at the FAsyncIODelete's destruction).
+	 * Setting deletes paused does not affect deletes already in progress.
+	 */
+	void SetDeletesPaused(bool bInPaused);
+	bool GetDeletesPaused() const { return bPaused; }
+
+	/**
+	 * Prepare the directory on disk.  Called OnDemand from the other interface functions.  Can also be called if desired by client code.
+	 * May take a long time to run if cleanup from a crashed previous process is required.
+	 * @return False if the TempRoot directory could not be constructed; calls to DeleteDirectory will fail in this case
+	 */
+	bool Setup();
+
+	/** Synchronously wait for all tasks to complete, and remove the temproot. */
+	void Teardown();
+
+	/**
+	 * Asynchronously delete the given directory (and all subdirs).  The directory is moved and no longer exists before returning, but the move location will be deleted from disk later.
+	 * The delete will fail if the given path is not a directory.
+	 * The delete will fail and return false if the given path is a parent directory, child directory, or is equal to TempRoot.
+	 * The delete can also fail due to IO Errors - insufficient permissions, device failure.
+	 * If the given path is on a different volume than temp root, or the move fails for other reasons, the delete will proceed synchronously but will succeed.
+	 * If the given path does not exist, DeleteDirectory will return true.
+	 *
+	 * @param PathToDelete - If DeleteDirectory returns true, this path has been removed from disk.  
+	 *
+	 * @return Whether the requested path is no longer present on disk.
+	 */
+	bool DeleteDirectory(const FStringView& PathToDelete)
+	{
+		return Delete(PathToDelete, EPathType::Directory);
+	}
+
+	/** Asynchronously delete a file.  Behavior is otherwise the same as DeleteDirectory. */
+	bool DeleteFile(const FStringView& PathToDelete)
+	{
+		return Delete(PathToDelete, EPathType::File);
+	}
+
+	/** Synchronously wait for all deletes to complete.  */
+	void WaitForAllTasks();
+
+private:
+	/** Update the ActiveTaskCount and events for a task completing */
+	void OnTaskComplete();
+
+	enum class EPathType
+	{
+		File,
+		Directory
+	};
+	/** Asynchronously delete a file or directory.  We handle both in the same function, but we want the interface to be explicit. */
+	bool Delete(const FStringView& PathToDelete, EPathType ExpectedType);
+
+	/** Create and store the task to delete the given Directory or File */
+	void CreateDeleteTask(const FStringView& InDeletePath, EPathType PathType);
+
+	/** Delete the given path synchronously; called from a task or in error fallback cases from the public thread */
+	bool SynchronousDelete(const TCHAR* InDeletePath, EPathType PathType);
+
+	bool DeleteTempRootDirectory();
+
+	FString	TempRoot;
+	TArray<FString> PausedDeletes;
+	FCriticalSection CriticalSection; // We use a CriticalSection instead of a TAtomic ActiveTaskCount so that we can atomically { trigger TasksComplete if ActiveTaskCount == 0 }
+	FEvent* TasksComplete = nullptr;
+	uint32 ActiveTaskCount = 0;
+	uint32 DeleteCounter = 0;
+	bool bInitialized = false;
+	bool bPaused = false;
+
+#if WITH_ASYNCIODELETE_DEBUG
+	static TArray<FString> AllTempRoots;
+	static void AddTempRoot(const FStringView& InTempRoot);
+	static void RemoveTempRoot(const FStringView& InTempRoot);
+#endif
+};
\ No newline at end of file
diff --git a/Engine/Source/Editor/UnrealEd/Private/Cooker/AsyncIODeleteTest.cpp b/Engine/Source/Editor/UnrealEd/Private/Cooker/AsyncIODeleteTest.cpp
new file mode 100644
--- /dev/null
+++ b/Engine/Source/Editor/UnrealEd/Private/Cooker/AsyncIODeleteTest.cpp
@@ -0,0 +1,130 @@
+// Copyright 1998-2019 Epic Games, Inc. All Rights Reserved.
+
+#include "CoreMinimal.h"
+
+#include "AsyncIODelete.h"
+#include "HAL/FileManager.h"
+#include "Misc/AutomationTest.h"
+#include "Misc/FileHelper.h"
+#include "Misc/Paths.h"
+#include "Misc/ScopeExit.h"
+
+IMPLEMENT_SIMPLE_AUTOMATION_TEST(FAsyncIODeleteTest, "System.Core.Misc.AsyncIODelete", EAutomationTestFlags::ApplicationContextMask | EAutomationTestFlags::EngineFilter);
+bool FAsyncIODeleteTest::RunTest(const FString& Parameters)
+{
+	FString TestRoot = FPaths::CreateTempFilename(FPlatformProcess::UserTempDir(), TEXT("AsyncIODelete"), TEXT(""));
+	IFileManager& FileManager = IFileManager::Get();
+	const bool bApplyToTreeTrue = true;
+	ON_SCOPE_EXIT
+	{
+		const bool bRequireExists = false;
+		FileManager.DeleteDirectory(*TestRoot, bRequireExists, bApplyToTreeTrue);
+	};
+
+	FString TempRoot = FPaths::Combine(TestRoot, TEXT("TempRoot"));
+	FString TempRoot2 = FPaths::Combine(TestRoot, TEXT("TempRoot2"));
+	FString TempRoot3 = FPaths::Combine(TestRoot, TEXT("TempRoot3"));
+	FString TempRoot4 = FPaths::Combine(TestRoot, TEXT("TempRoot4"));
+	FString TestFile1 = FPaths::Combine(TestRoot, TEXT("TestFile1"));
+	FString TestDir1 = FPaths::Combine(TestRoot, TEXT("TestDir1"));
+	int32 NumFiles = 0;
+	int32 NumDirs = 0;
+	const TCHAR* TestText = TEXT("Test");
+	auto CreateTestPathsToDelete = [=, &TestFile1, &TestDir1, &FileManager]()
+	{
+		FFileHelper::SaveStringToFile(TestText, *TestFile1);
+		FileManager.MakeDirectory(*TestDir1, bApplyToTreeTrue);
+	};
+	IPlatformFile::FDirectoryVisitorFunc CountFilesAndDirs = [&NumFiles, &NumDirs](const TCHAR* VisitFilename, bool VisitIsDir)
+	{
+		NumFiles += VisitIsDir ? 1 : 0;
+		NumDirs += VisitIsDir ? 0 : 1;
+		return true;
+	};
+	auto TestTempRootCountsEqual = [&](const FString& RootDir, int32 ExpectedFileCount, int32 ExpectedDirCount, const TCHAR* TestDescription)
+	{
+		NumFiles = 0;
+		NumDirs = 0;
+		FileManager.IterateDirectory(*RootDir, CountFilesAndDirs);
+		TestTrue(TestDescription, NumFiles == ExpectedFileCount && NumDirs == ExpectedDirCount);
+	};
+	auto TestRequestedPathsDeleted = [&](const TCHAR* TestDescription)
+	{
+		TestTrue(TestDescription, !FileManager.FileExists(*TestFile1) && !FileManager.DirectoryExists(*TestDir1));
+	};
+
+	FileManager.MakeDirectory(*TestRoot);
+
+	{
+		FAsyncIODelete AsyncIODelete(TempRoot);
+
+		// Waiting for tasks to complete when none have been launched should succeed
+		AsyncIODelete.WaitForAllTasks();
+
+		// Moving file and directory from the source location should be finished by the time DeleteFile/DeleteDirectory returns
+		CreateTestPathsToDelete();
+		AsyncIODelete.DeleteFile(TestFile1);
+		AsyncIODelete.DeleteDirectory(TestDir1);
+		TestRequestedPathsDeleted(TEXT("AsyncIODelete::Delete should have moved the deleted paths before returning."));
+
+		// Deleting the temporary files/directories should be finished before WaitForAllTasks returns
+		AsyncIODelete.WaitForAllTasks();
+		TestTempRootCountsEqual(TempRoot, 0, 0, TEXT("AsyncIODelete should have deleted the moved paths before WaitForAllTasks returned."));
+
+		// Two FAsyncIODelete constructed at once should be legal, as long as they have different TempRoots
+		FAsyncIODelete AsyncIODelete2(TempRoot2);
+
+		// Use the pause feature to verify that the paths are indeed moved into the TempRoot
+		AsyncIODelete2.SetDeletesPaused(true);
+		CreateTestPathsToDelete();
+		AsyncIODelete2.DeleteFile(TestFile1);
+		AsyncIODelete2.DeleteDirectory(TestDir1);
+		TestRequestedPathsDeleted(TEXT("AsyncIODelete::Delete should have moved the deleted paths before returning even when paused."));
+		AsyncIODelete2.WaitForAllTasks();
+		TestTempRootCountsEqual(TempRoot2, 1, 1, TEXT("AsyncIODelete should not have deleted the moved paths because it is paused."));
+		AsyncIODelete2.SetDeletesPaused(false);
+		AsyncIODelete2.WaitForAllTasks();
+		TestTempRootCountsEqual(TempRoot2, 0, 0, TEXT("AsyncIODelete should have deleted the moved paths after unpausing."));
+
+		// Verify Teardown() deletes the TempRoot and Setup() creates it
+		AsyncIODelete2.Teardown();
+		TestTrue(TEXT("AsyncIODelete::Teardown should have deleted its TempRoot."), !FileManager.DirectoryExists(*TempRoot2));
+		AsyncIODelete2.Setup();
+		TestTrue(TEXT("AsyncIODelete::Setup should have created its TempRoot."), FileManager.DirectoryExists(*TempRoot2));
+
+		// Manual setup works as long as you call SetTempRoot before Setup
+		FAsyncIODelete AsyncIODelete3;
+		AsyncIODelete3.SetTempRoot(TempRoot3);
+		AsyncIODelete3.Setup();
+		TestTrue(TEXT("Setup should have created the TempRoot."), FileManager.DirectoryExists(*TempRoot3));
+
+		// Check that even after Setup, waiting for tasks to complete when none have been launched should succeed
+		AsyncIODelete.WaitForAllTasks();
+
+		// Changing TempRoot and then deleting a file works
+		CreateTestPathsToDelete();
+		AsyncIODelete3.DeleteFile(TestFile1);
+		AsyncIODelete3.DeleteDirectory(TestDir1);
+		AsyncIODelete3.SetTempRoot(TempRoot4);
+		TestTrue(TEXT("SetTempRoot should have deleted the old TempRoot."), !FileManager.DirectoryExists(*TempRoot3));
+		CreateTestPathsToDelete();
+		AsyncIODelete3.DeleteFile(TestFile1);
+		AsyncIODelete3.DeleteDirectory(TestDir1);
+		TestRequestedPathsDeleted(TEXT("AsyncIODelete::Delete should have worked after changing the TempRoot."));
+		TestTrue(TEXT("Delete should have created the new TempRoot after SetTempRoot."), FileManager.DirectoryExists(*TempRoot4));
+		AsyncIODelete3.WaitForAllTasks();
+		TestTempRootCountsEqual(TempRoot4, 0, 0, TEXT("AsyncIODelete::Delete should have created the tasks to delete moved files after changing the TempRoot."));
+
+		// Attempting to delete a parent directory of the temproot, the temproot itself, or a child inside of it fails
+		FString SubDirInTempRoot4 = FPaths::Combine(TempRoot4, TEXT("SubDir"));
+		FileManager.MakeDirectory(*SubDirInTempRoot4, bApplyToTreeTrue); // Note it's illegal to add files into TempRoot, but we're not currently checking for it and we're not colliding with the DeleteN paths AsyncIODelete uses, so this breaking of the rule will not cause problems
+		TestFalse(TEXT("AsyncIODelete should refuse to delete a parent of its TempRoot."), AsyncIODelete3.DeleteDirectory(TestRoot));
+		TestFalse(TEXT("AsyncIODelete should refuse to delete its TempRoot."), AsyncIODelete3.DeleteDirectory(TempRoot4));
+		TestFalse(TEXT("AsyncIODelete should refuse to delete a child of its TempRoot."), AsyncIODelete3.DeleteDirectory(SubDirInTempRoot4));
+	}
+
+	TestTrue(TEXT("AsyncIODelete destructor should have deleted its TempRoot."),
+		!FileManager.DirectoryExists(*TempRoot) && !FileManager.DirectoryExists(*TempRoot2) && !FileManager.DirectoryExists(*TempRoot3) && !FileManager.DirectoryExists(*TempRoot4));
+
+	return true;
+}
\ No newline at end of file
