AsyncIODelete Non-unity compile fix
#rnx

diff --git a/Engine/Source/Editor/UnrealEd/Private/Cooker/AsyncIODelete.cpp b/Engine/Source/Editor/UnrealEd/Private/Cooker/AsyncIODelete.cpp
--- a/Engine/Source/Editor/UnrealEd/Private/Cooker/AsyncIODelete.cpp
+++ b/Engine/Source/Editor/UnrealEd/Private/Cooker/AsyncIODelete.cpp
@@ -5,6 +5,7 @@
 
 #include "Async/Async.h"
 #include "Containers/UnrealString.h"
+#include "CookOnTheSide/CookOnTheFlyServer.h" // needed for DECLARE_LOG_CATEGORY_EXTERN(LogCook,...)
 #include "HAL/Event.h"
 #include "HAL/FileManager.h"
 #include "Math/NumericLimits.h"
@@ -16,8 +17,6 @@
 TArray<FString> FAsyncIODelete::AllTempRoots;
 #endif
 
-DECLARE_LOG_CATEGORY_EXTERN(LogCook, Log, All);
-
 FAsyncIODelete::FAsyncIODelete(const FStringView& InOwnedTempRoot)
 {
 	SetTempRoot(InOwnedTempRoot);
