Remove the ensure from FAsyncDelete's asynchronous call to DeleteDirectory.  It is failing and bring down the cook for some users.  Need to add lower level logging to diagnose why, but for now we can just log it and ignore it; it will be cleaned up on the next cook.
#rnx

diff --git a/Engine/Source/Editor/UnrealEd/Private/Cooker/AsyncIODelete.cpp b/Engine/Source/Editor/UnrealEd/Private/Cooker/AsyncIODelete.cpp
--- a/Engine/Source/Editor/UnrealEd/Private/Cooker/AsyncIODelete.cpp
+++ b/Engine/Source/Editor/UnrealEd/Private/Cooker/AsyncIODelete.cpp
@@ -223,7 +223,7 @@ void FAsyncIODelete::CreateDeleteTask(const FStringView& InDeletePath, EPathType
 	}
 
 	AsyncThread(
-		[this, DeletePath = FString(InDeletePath), PathType]() { ensure(SynchronousDelete(*DeletePath, PathType)); },
+		[this, DeletePath = FString(InDeletePath), PathType]() { SynchronousDelete(*DeletePath, PathType); },
 		0, TPri_Normal,
 		[this]() { OnTaskComplete(); });
 }
@@ -241,17 +241,24 @@ void FAsyncIODelete::OnTaskComplete()
 
 bool FAsyncIODelete::SynchronousDelete(const TCHAR* InDeletePath, EPathType PathType)
 {
+	bool Result;
 	const bool bRequireExists = false;
 	if (PathType == EPathType::Directory)
 	{
 		const bool bTree = true;
-		return IFileManager::Get().DeleteDirectory(InDeletePath, bRequireExists, bTree);
+		Result = IFileManager::Get().DeleteDirectory(InDeletePath, bRequireExists, bTree);
 	}
 	else
 	{
 		const bool bEvenIfReadOnly = true;
 		return IFileManager::Get().Delete(InDeletePath, bRequireExists, bEvenIfReadOnly);
 	}
+
+	if (!Result)
+	{
+		UE_LOG(LogCook, Warning, TEXT("Could not delete asyncdelete %s '%s'."), PathType == EPathType::Directory ? TEXT("directory") : TEXT("file"), InDeletePath);
+	}
+	return Result;
 }
 
 bool FAsyncIODelete::DeleteTempRootDirectory()
