Fix static analysis compile warning; an oversight in the previous checkin.
#rnx

diff --git a/Engine/Source/Editor/UnrealEd/Private/Cooker/AsyncIODelete.cpp b/Engine/Source/Editor/UnrealEd/Private/Cooker/AsyncIODelete.cpp
--- a/Engine/Source/Editor/UnrealEd/Private/Cooker/AsyncIODelete.cpp
+++ b/Engine/Source/Editor/UnrealEd/Private/Cooker/AsyncIODelete.cpp
@@ -251,7 +251,7 @@ bool FAsyncIODelete::SynchronousDelete(const TCHAR* InDeletePath, EPathType Path
 	else
 	{
 		const bool bEvenIfReadOnly = true;
-		return IFileManager::Get().Delete(InDeletePath, bRequireExists, bEvenIfReadOnly);
+		Result = IFileManager::Get().Delete(InDeletePath, bRequireExists, bEvenIfReadOnly);
 	}
 
 	if (!Result)
