Fix assets reporting incorrect warnings about duplicate bulkdata if cooked more than once in the same editor session.

- The problem was caused by UCookOnTheFlyServer::SavePackageContexts never being cleaned up, so on subsequent cooks we added new contexts to the back of the TArray but would use the original contexts at the front of the array, which is where we started getting duplicate warnings from.
- Iterative cooking was preventing our local testing from noticing the issue. This will only show up if an asset is changed, the cooked data directory is delete, or is invalidated in some way.
- Add UCookOnTheFlyServer::ClearPackageStoreContexts to delete the contents of SavePackageContexts, called when ever a cook finishes or the UCookOnTheFlyServer object itself is deleted.

#rnx

diff --git a/Engine/Source/Editor/UnrealEd/Classes/CookOnTheSide/CookOnTheFlyServer.h b/Engine/Source/Editor/UnrealEd/Classes/CookOnTheSide/CookOnTheFlyServer.h
--- a/Engine/Source/Editor/UnrealEd/Classes/CookOnTheSide/CookOnTheFlyServer.h
+++ b/Engine/Source/Editor/UnrealEd/Classes/CookOnTheSide/CookOnTheFlyServer.h
@@ -708,6 +708,11 @@ private:
 	*/
 	void FinalizePackageStore();
 
+	/**
+	* Empties SavePackageContexts and deletes the contents
+	*/
+	void ClearPackageStoreContexts();
+
 	/**
 	* Initialize all target platforms
 	*/
diff --git a/Engine/Source/Editor/UnrealEd/Private/CookOnTheFlyServer.cpp b/Engine/Source/Editor/UnrealEd/Private/CookOnTheFlyServer.cpp
--- a/Engine/Source/Editor/UnrealEd/Private/CookOnTheFlyServer.cpp
+++ b/Engine/Source/Editor/UnrealEd/Private/CookOnTheFlyServer.cpp
@@ -1451,6 +1451,8 @@ UCookOnTheFlyServer::UCookOnTheFlyServer(FVTableHelper& Helper) :Super(Helper) {
 
 UCookOnTheFlyServer::~UCookOnTheFlyServer()
 {
+	ClearPackageStoreContexts();
+
 	FCoreDelegates::OnFConfigCreated.RemoveAll(this);
 	FCoreDelegates::OnFConfigDeleted.RemoveAll(this);
 
@@ -6831,6 +6833,18 @@ void UCookOnTheFlyServer::FinalizePackageStore()
 		}
 	}
 	UE_LOG(LogCook, Display, TEXT("Done saving BulkData manifest(s)"));
+
+	ClearPackageStoreContexts();
+}
+
+void UCookOnTheFlyServer::ClearPackageStoreContexts()
+{
+	for (FSavePackageContext* Context : SavePackageContexts)
+	{
+		delete Context;
+	}
+
+	SavePackageContexts.Empty();
 }
 
 void UCookOnTheFlyServer::InitializeTargetPlatforms()
diff --git a/Engine/Source/Runtime/CoreUObject/Public/UObject/SavePackage.h b/Engine/Source/Runtime/CoreUObject/Public/UObject/SavePackage.h
--- a/Engine/Source/Runtime/CoreUObject/Public/UObject/SavePackage.h
+++ b/Engine/Source/Runtime/CoreUObject/Public/UObject/SavePackage.h
@@ -125,6 +125,12 @@ public:
 	{
 	}
 
+	~FSavePackageContext()
+	{
+		delete PackageStoreWriter;
+		delete BulkDataManifest;
+	}
+
 	FPackageStoreWriter* PackageStoreWriter;
 	FPackageStoreBulkDataManifest* BulkDataManifest;
 };
