FileManagerGeneric: Add preserving of windows' GetLastError when DeleteDirectory fails.
AsyncIODelete: add logging for the asyncdelete error code, and fix static analysis compile error.
#rnx

diff --git a/Engine/Source/Editor/UnrealEd/Private/Cooker/AsyncIODelete.cpp b/Engine/Source/Editor/UnrealEd/Private/Cooker/AsyncIODelete.cpp
--- a/Engine/Source/Editor/UnrealEd/Private/Cooker/AsyncIODelete.cpp
+++ b/Engine/Source/Editor/UnrealEd/Private/Cooker/AsyncIODelete.cpp
@@ -8,6 +8,7 @@
 #include "CookOnTheSide/CookOnTheFlyServer.h" // needed for DECLARE_LOG_CATEGORY_EXTERN(LogCook,...)
 #include "HAL/Event.h"
 #include "HAL/FileManager.h"
+#include "HAL/PlatformMisc.h"
 #include "Math/NumericLimits.h"
 #include "Misc/StringBuilder.h"
 #include "Misc/Paths.h"
@@ -82,14 +83,17 @@ bool FAsyncIODelete::Setup()
 	}
 
 	// Delete the TempRoot directory to clear the results from any previous process using the same TempRoot that did not shut down cleanly
-	if (!ensure(DeleteTempRootDirectory()))
+	uint32 ErrorCode;
+	if (!DeleteTempRootDirectory(ErrorCode))
 	{
+		UE_LOG(LogCook, Fatal, TEXT("Could not clear asyncdelete root directory '%s'.  LastError: %i."), *TempRoot, ErrorCode);
 		return false;
 	}
 
 	// Create the empty directory to work in
-	if (!ensure(IFileManager::Get().MakeDirectory(*TempRoot, true)))
+	if (!IFileManager::Get().MakeDirectory(*TempRoot, true))
 	{
+		UE_LOG(LogCook, Fatal, TEXT("Could not create asyncdelete root directory '%s'.  LastError: %i."), *TempRoot, FPlatformMisc::GetLastError());
 		return false;
 	}
 
@@ -124,10 +128,11 @@ void FAsyncIODelete::Teardown()
 	TasksComplete = nullptr;
 
 	// Remove the temp directory from disk
-	if (!DeleteTempRootDirectory())
+	uint32 ErrorCode;
+	if (!DeleteTempRootDirectory(ErrorCode))
 	{
 		// This will leave directories (and potentially files, if we were paused or if any of the asyncdeletes failed) on disk, so it is bad for users, but is not fatal for our operations.
-		UE_LOG(LogCook, Warning, TEXT("Could not delete asyncdelete directory '%s'."), *TempRoot);
+		UE_LOG(LogCook, Warning, TEXT("Could not delete asyncdelete root directory '%s'.  LastError: %i."), *TempRoot, ErrorCode);
 	}
 
 	// Clear delete variables; we don't need to run the tasks for the remaining pauseddeletes because synchronously deleting the temp directory above did the work they were going to do
@@ -199,7 +204,7 @@ bool FAsyncIODelete::Delete(const FStringView& PathToDelete, EPathType ExpectedT
 	if (!IFileManager::Get().Move(*TempPath, PathToDeleteSZ, bReplace, bEvenIfReadOnly, bMoveAttributes, bDoNotRetryOnError)) // IFileManager::Move works on either files or directories
 	{
 		// The move failed; try a synchronous delete as backup
-		UE_LOG(LogCook, Warning, TEXT("Failed to move path '%.*s' for async delete; falling back to synchronous delete."), PathToDelete.Len(), PathToDelete.GetData());
+		UE_LOG(LogCook, Warning, TEXT("Failed to move path '%.*s' for async delete (LastError == %i); falling back to synchronous delete."), PathToDelete.Len(), PathToDelete.GetData(), FPlatformMisc::GetLastError());
 		return SynchronousDelete(PathToDeleteSZ, ExpectedType);
 	}
 
@@ -256,13 +261,14 @@ bool FAsyncIODelete::SynchronousDelete(const TCHAR* InDeletePath, EPathType Path
 
 	if (!Result)
 	{
-		UE_LOG(LogCook, Warning, TEXT("Could not delete asyncdelete %s '%s'."), PathType == EPathType::Directory ? TEXT("directory") : TEXT("file"), InDeletePath);
+		UE_LOG(LogCook, Warning, TEXT("Failed to asyncdelete %s '%s'. LastError == %i."), PathType == EPathType::Directory ? TEXT("directory") : TEXT("file"), InDeletePath, FPlatformMisc::GetLastError());
 	}
 	return Result;
 }
 
-bool FAsyncIODelete::DeleteTempRootDirectory()
+bool FAsyncIODelete::DeleteTempRootDirectory(uint32& OutErrorCode)
 {
+	OutErrorCode = 0;
 	IFileManager& FileManager = IFileManager::Get();
 	if (!FileManager.DirectoryExists(*TempRoot))
 	{
@@ -277,7 +283,8 @@ bool FAsyncIODelete::DeleteTempRootDirectory()
 	const bool bReplace = false;
 	const bool bEvenIfReadOnly = true;
 	const TCHAR* DirectoryToDelete = *UniqueDirectory;
-	if (!FileManager.Move(DirectoryToDelete, *TempRoot, bReplace, bEvenIfReadOnly))
+	const bool bMoveSucceeded = FileManager.Move(DirectoryToDelete, *TempRoot, bReplace, bEvenIfReadOnly);
+	if (!bMoveSucceeded)
 	{
 		// Move failed; fallback to inplace delete
 		DirectoryToDelete = *TempRoot;
@@ -285,7 +292,17 @@ bool FAsyncIODelete::DeleteTempRootDirectory()
 
 	const bool bRequireExists = false;
 	const bool bTree = true;
-	return FileManager.DeleteDirectory(DirectoryToDelete, bRequireExists, bTree);
+	const bool bDeleteSucceeded = FileManager.DeleteDirectory(DirectoryToDelete, bRequireExists, bTree);
+	if (!bDeleteSucceeded)
+	{
+		OutErrorCode = FPlatformMisc::GetLastError();
+		if (bMoveSucceeded && !bDeleteSucceeded)
+		{
+			// Try to move the directory back so that we can try again to delete it next time.
+			FileManager.Move(*TempRoot, DirectoryToDelete, bReplace, bEvenIfReadOnly);
+		}
+	}
+	return bDeleteSucceeded;
 }
 
 #if WITH_ASYNCIODELETE_DEBUG
diff --git a/Engine/Source/Editor/UnrealEd/Private/Cooker/AsyncIODelete.h b/Engine/Source/Editor/UnrealEd/Private/Cooker/AsyncIODelete.h
--- a/Engine/Source/Editor/UnrealEd/Private/Cooker/AsyncIODelete.h
+++ b/Engine/Source/Editor/UnrealEd/Private/Cooker/AsyncIODelete.h
@@ -110,7 +110,7 @@ private:
 	/** Delete the given path synchronously; called from a task or in error fallback cases from the public thread */
 	bool SynchronousDelete(const TCHAR* InDeletePath, EPathType PathType);
 
-	bool DeleteTempRootDirectory();
+	bool DeleteTempRootDirectory(uint32& OutErrorCode);
 
 	FString	TempRoot;
 	TArray<FString> PausedDeletes;
diff --git a/Engine/Source/Runtime/Core/Private/GenericPlatform/GenericPlatformFile.cpp b/Engine/Source/Runtime/Core/Private/GenericPlatform/GenericPlatformFile.cpp
--- a/Engine/Source/Runtime/Core/Private/GenericPlatform/GenericPlatformFile.cpp
+++ b/Engine/Source/Runtime/Core/Private/GenericPlatform/GenericPlatformFile.cpp
@@ -3,6 +3,7 @@
 #include "GenericPlatform/GenericPlatformFile.h"
 #include "HAL/FileManager.h"
 #include "Misc/Paths.h"
+#include "HAL/PlatformMisc.h"
 #include "HAL/ThreadSafeCounter.h"
 #include "Stats/Stats.h"
 #include "Async/AsyncWork.h"
@@ -642,6 +643,8 @@ bool IPlatformFile::DeleteDirectoryRecursively(const TCHAR* Directory)
 	{
 	public:
 		IPlatformFile&		PlatformFile;
+		uint32 FirstError = 0;
+
 		FRecurse(IPlatformFile&	InPlatformFile)
 			: PlatformFile(InPlatformFile)
 		{
@@ -651,7 +654,10 @@ bool IPlatformFile::DeleteDirectoryRecursively(const TCHAR* Directory)
 			if (bIsDirectory)
 			{
 				PlatformFile.IterateDirectory(FilenameOrDirectory, *this);
-				PlatformFile.DeleteDirectory(FilenameOrDirectory);
+				if (!PlatformFile.DeleteDirectory(FilenameOrDirectory) && FirstError == 0)
+				{
+					FirstError = FPlatformMisc::GetLastError();
+				}
 			}
 			else
 			{
@@ -660,14 +666,22 @@ bool IPlatformFile::DeleteDirectoryRecursively(const TCHAR* Directory)
 
 				// File delete failed -- unset readonly flag and try again
 				PlatformFile.SetReadOnly(FilenameOrDirectory, false);
-				PlatformFile.DeleteFile(FilenameOrDirectory);
+				if (!PlatformFile.DeleteFile(FilenameOrDirectory) && FirstError == 0)
+				{
+					FirstError = FPlatformMisc::GetLastError();
+				}
 			}
 			return true; // continue searching
 		}
 	};
 	FRecurse Recurse(*this);
 	Recurse.Visit(Directory, true);
-	return !DirectoryExists(Directory);
+	const bool bSucceeded = !DirectoryExists(Directory);
+	if (!bSucceeded)
+	{
+		FPlatformMisc::SetLastError(Recurse.FirstError);
+	}
+	return bSucceeded;
 }
 
 
@@ -813,11 +827,18 @@ static bool InternalCreateDirectoryTree(IPlatformFile& Ipf, const FString& Direc
 			if (!InternalCreateDirectoryTree(Ipf, Directory.Left(SeparatorIndex)))
 				return false;
 
-			return Ipf.CreateDirectory(*Directory) || Ipf.DirectoryExists(*Directory);
+			if (Ipf.CreateDirectory(*Directory))
+				return true;
 		}
 	}
 
-	return Ipf.DirectoryExists(*Directory);
+	uint32 ErrorCode = FPlatformMisc::GetLastError();
+	const bool bExists = Ipf.DirectoryExists(*Directory);
+	if (!bExists)
+	{
+		FPlatformMisc::SetLastError(ErrorCode);
+	}
+	return bExists;
 }
 
 bool IPlatformFile::CreateDirectoryTree(const TCHAR* Directory)
diff --git a/Engine/Source/Runtime/Core/Private/HAL/FileManagerGeneric.cpp b/Engine/Source/Runtime/Core/Private/HAL/FileManagerGeneric.cpp
--- a/Engine/Source/Runtime/Core/Private/HAL/FileManagerGeneric.cpp
+++ b/Engine/Source/Runtime/Core/Private/HAL/FileManagerGeneric.cpp
@@ -378,11 +378,26 @@ bool FFileManagerGeneric::MakeDirectory( const TCHAR* Path, bool Tree )
 
 bool FFileManagerGeneric::DeleteDirectory( const TCHAR* Path, bool RequireExists, bool Tree )
 {
-	if( Tree )
+	bool bSucceeded;
+	if (Tree)
 	{
-		return GetLowLevel().DeleteDirectoryRecursively( Path ) || ( !RequireExists && !GetLowLevel().DirectoryExists( Path ) );
+		bSucceeded = GetLowLevel().DeleteDirectoryRecursively(Path);
+	}
+	else
+	{
+		bSucceeded = GetLowLevel().DeleteDirectory(Path);
+	}
+
+	if (!bSucceeded && !RequireExists)
+	{
+		uint32 ErrorFromDelete = FPlatformMisc::GetLastError();
+		bSucceeded = !GetLowLevel().DirectoryExists(Path);
+		if (!bSucceeded)
+		{
+			FPlatformMisc::SetLastError(ErrorFromDelete);
+		}
 	}
-	return GetLowLevel().DeleteDirectory( Path ) || ( !RequireExists && !GetLowLevel().DirectoryExists( Path ) );
+	return bSucceeded;
 }
 
 FFileStatData FFileManagerGeneric::GetStatData(const TCHAR* FilenameOrDirectory)
diff --git a/Engine/Source/Runtime/Core/Private/Windows/WindowsPlatformFile.cpp b/Engine/Source/Runtime/Core/Private/Windows/WindowsPlatformFile.cpp
--- a/Engine/Source/Runtime/Core/Private/Windows/WindowsPlatformFile.cpp
+++ b/Engine/Source/Runtime/Core/Private/Windows/WindowsPlatformFile.cpp
@@ -1051,7 +1051,13 @@ public:
 	virtual bool DeleteDirectory(const TCHAR* Directory) override
 	{
 		RemoveDirectoryW(*NormalizeDirectory(Directory));
-		return !DirectoryExists(Directory);
+		uint32 LastError = GetLastError();
+		const bool bSucceeded = !DirectoryExists(Directory);
+		if (!bSucceeded)
+		{
+			SetLastError(LastError);
+		}
+		return bSucceeded;
 	}
 	virtual FFileStatData GetStatData(const TCHAR* FilenameOrDirectory) override
 	{
diff --git a/Engine/Source/Runtime/Core/Private/Windows/WindowsPlatformMisc.cpp b/Engine/Source/Runtime/Core/Private/Windows/WindowsPlatformMisc.cpp
--- a/Engine/Source/Runtime/Core/Private/Windows/WindowsPlatformMisc.cpp
+++ b/Engine/Source/Runtime/Core/Private/Windows/WindowsPlatformMisc.cpp
@@ -1928,6 +1928,11 @@ uint32 FWindowsPlatformMisc::GetLastError()
 	return (uint32)::GetLastError();
 }
 
+void FWindowsPlatformMisc::SetLastError(uint32 ErrorCode)
+{
+	::SetLastError((DWORD)ErrorCode);
+}
+
 bool FWindowsPlatformMisc::CoInitialize()
 {
 	HRESULT hr = ::CoInitialize(NULL);
diff --git a/Engine/Source/Runtime/Core/Public/GenericPlatform/GenericPlatformMisc.h b/Engine/Source/Runtime/Core/Public/GenericPlatform/GenericPlatformMisc.h
--- a/Engine/Source/Runtime/Core/Public/GenericPlatform/GenericPlatformMisc.h
+++ b/Engine/Source/Runtime/Core/Public/GenericPlatform/GenericPlatformMisc.h
@@ -600,6 +600,8 @@ struct CORE_API FGenericPlatformMisc
 		return 0;
 	}
 
+	static void SetLastError(uint32 ErrorCode) {}
+
 	static void RaiseException( uint32 ExceptionCode );
 
 public:
diff --git a/Engine/Source/Runtime/Core/Public/Windows/WindowsPlatformMisc.h b/Engine/Source/Runtime/Core/Public/Windows/WindowsPlatformMisc.h
--- a/Engine/Source/Runtime/Core/Public/Windows/WindowsPlatformMisc.h
+++ b/Engine/Source/Runtime/Core/Public/Windows/WindowsPlatformMisc.h
@@ -85,6 +85,7 @@ struct CORE_API FWindowsPlatformMisc
 	static FString GetDefaultLocale();
 
 	static uint32 GetLastError();
+	static void SetLastError(uint32 ErrorCode);
 	static void RaiseException( uint32 ExceptionCode );
 	static bool SetStoredValue(const FString& InStoreId, const FString& InSectionName, const FString& InKeyName, const FString& InValue);
 	static bool GetStoredValue(const FString& InStoreId, const FString& InSectionName, const FString& InKeyName, FString& OutValue);
