[UE-85892] Fix assert due to broken contract in FAssetRegistryGenerator::GenerateStreamingInstallManifest that FinalChunkManifests never has nulls in it.  Thas wasfailing when cooking with -cooksinglepackage.
#rnx

diff --git a/Engine/Source/Editor/UnrealEd/Private/Commandlets/AssetRegistryGenerator.cpp b/Engine/Source/Editor/UnrealEd/Private/Commandlets/AssetRegistryGenerator.cpp
--- a/Engine/Source/Editor/UnrealEd/Private/Commandlets/AssetRegistryGenerator.cpp
+++ b/Engine/Source/Editor/UnrealEd/Private/Commandlets/AssetRegistryGenerator.cpp
@@ -339,7 +339,10 @@ bool FAssetRegistryGenerator::GenerateStreamingInstallManifest(int64 InExtraFlav
 				int32 PakchunkIndex = GetPakchunkIndex(ChunkID);
 				if (PakchunkIndex >= FinalChunkManifests.Num())
 				{
-					FinalChunkManifests.AddDefaulted(PakchunkIndex - FinalChunkManifests.Num() + 1);
+					// Extend the array until it is large enough to hold the requested index, filling it in with nulls on all the newly added indices.
+					// Note that this will temporarily break our contract that FinalChunkManifests does not contain null pointers; we fix up the contract
+					// by replacing any remaining null pointers in the loop over FinalChunkManifests at the bottom of this function.
+					FinalChunkManifests.AddZeroed(PakchunkIndex - FinalChunkManifests.Num() + 1);
 				}
 				checkf(PakchunkIndex < FinalChunkManifests.Num(), TEXT("Chunk %i out of range. %i manifests available"), PakchunkIndex, FinalChunkManifests.Num() - 1);
 				checkf(FinalChunkManifests[PakchunkIndex] == nullptr || FinalChunkManifests[PakchunkIndex]->Num() == 0, TEXT("Manifest already exists for chunk %i"), PakchunkIndex);
@@ -364,6 +367,7 @@ bool FAssetRegistryGenerator::GenerateStreamingInstallManifest(int64 InExtraFlav
 	}
 
 	// generate per-chunk pak list files
+	bool bSucceeded = true;
 	for (int32 PakchunkIndex = 0; PakchunkIndex < FinalChunkManifests.Num(); ++PakchunkIndex)
 	{
 		// Serialize chunk layers whether chunk is empty or not
@@ -373,10 +377,11 @@ bool FAssetRegistryGenerator::GenerateStreamingInstallManifest(int64 InExtraFlav
 		FString LayerString = FString::Printf(TEXT("%d\r\n"), TargetLayer);
 		ChunkLayerFile->Serialize(TCHAR_TO_ANSI(*LayerString), LayerString.Len());
 
-		// Is this chunk empty?
+		// Is this index a null placeholder that we added in the loop over EncryptedNonUFSFileGroups and then never filled in?  If so, 
+		// fill it in with an empty FChunkPackageSet
 		if (!FinalChunkManifests[PakchunkIndex])
 		{
-			continue;
+			FinalChunkManifests[PakchunkIndex] = new FChunkPackageSet;
 		}
 
 		int32 FilenameIndex = 0;
@@ -395,7 +400,8 @@ bool FAssetRegistryGenerator::GenerateStreamingInstallManifest(int64 InExtraFlav
 			if (!PakListFile)
 			{
 				UE_LOG(LogAssetRegistryGenerator, Error, TEXT("Failed to open output paklist file %s"), *PakListFilename);
-				return false;
+				bSucceeded = false;
+				break;
 			}
 
 			FString PakChunkOptions;
@@ -470,7 +476,7 @@ bool FAssetRegistryGenerator::GenerateStreamingInstallManifest(int64 InExtraFlav
 	ChunkLayerFile->Close();
 	PakChunkListFile->Close();
 
-	return true;
+	return bSucceeded;
 }
 
 void FAssetRegistryGenerator::GenerateChunkManifestForPackage(const FName& PackageFName, const FString& PackagePathName, const FString& SandboxFilename, const FString& LastLoadedMapName, FSandboxPlatformFile* InSandboxFile)
