AsyncIODelete directory name reduced in length to avoid creating filenames over MAX_PATH.
#rnx

diff --git a/Engine/Source/Editor/UnrealEd/Private/CookOnTheFlyServer.cpp b/Engine/Source/Editor/UnrealEd/Private/CookOnTheFlyServer.cpp
--- a/Engine/Source/Editor/UnrealEd/Private/CookOnTheFlyServer.cpp
+++ b/Engine/Source/Editor/UnrealEd/Private/CookOnTheFlyServer.cpp
@@ -5142,6 +5142,9 @@ void UCookOnTheFlyServer::DeleteSandboxDirectory(const FString& PlatformName)
 	// The AsyncDelete directory is associated with a sandbox but is necessarily outside of it since it is used to delete the sandbox.
 	// Note that for the Platform we used to create the AsyncIODelete, this Delete will fail because AsyncIODelete refuses to delete its own temproot; this is okay because it will delete the temproot on exit.
 	LocalAsyncIODelete.DeleteDirectory(AsyncDeleteDirectory);
+
+	// UE_DEPRECATED(4.25, "Delete the old location for AsyncDeleteDirectory until all users have cooked at least once")
+	LocalAsyncIODelete.DeleteDirectory(SandboxDirectory + TEXT("AsyncDelete"));
 }
 
 FAsyncIODelete& UCookOnTheFlyServer::GetAsyncIODelete(const FString& PlatformName, const FString* AsyncDeleteDirectory)
@@ -5176,7 +5179,7 @@ FString UCookOnTheFlyServer::GetAsyncDeleteDirectory(const FString& PlatformName
 		FPaths::NormalizeDirectoryName(Buffer);
 		SandboxDirectory = &Buffer;
 	}
-	return (*SandboxDirectory) + TEXT("AsyncDelete");
+	return (*SandboxDirectory) + TEXT("_Del");
 }
 
 void UCookOnTheFlyServer::PopulateCookedPackagesFromDisk(const TArray<ITargetPlatform*>& Platforms)
diff --git a/Engine/Source/Editor/UnrealEd/Private/Cooker/AsyncIODelete.cpp b/Engine/Source/Editor/UnrealEd/Private/Cooker/AsyncIODelete.cpp
--- a/Engine/Source/Editor/UnrealEd/Private/Cooker/AsyncIODelete.cpp
+++ b/Engine/Source/Editor/UnrealEd/Private/Cooker/AsyncIODelete.cpp
@@ -194,7 +194,7 @@ bool FAsyncIODelete::Delete(const FStringView& PathToDelete, EPathType ExpectedT
 		return SynchronousDelete(PathToDeleteSZ, ExpectedType);
 	}
 
-	const FString TempPath = FPaths::Combine(TempRoot, FString::Printf(TEXT("Delete%u"), DeleteCounter));
+	const FString TempPath = FPaths::Combine(TempRoot, FString::Printf(TEXT("%u"), DeleteCounter));
 	DeleteCounter++;
 
 	const bool bReplace = true;
