Cooking - FullLoadAndSave needs to skip child Packages in the same way that one-at-a-time cooking does.  This prevents the error when cooking EngineTest:
Asset 'TranslucencyLightingModes_BuiltData.uasset' cannot be saved as it has only been partially loaded.
#rnx

diff --git a/Engine/Source/Editor/UnrealEd/Private/CookOnTheFlyServer.cpp b/Engine/Source/Editor/UnrealEd/Private/CookOnTheFlyServer.cpp
--- a/Engine/Source/Editor/UnrealEd/Private/CookOnTheFlyServer.cpp
+++ b/Engine/Source/Editor/UnrealEd/Private/CookOnTheFlyServer.cpp
@@ -7754,6 +7754,11 @@ uint32 UCookOnTheFlyServer::FullLoadAndSave(uint32& CookedPackageCount)
 					continue;
 				}
 
+				if (Package->GetOuter() != nullptr)
+				{
+					continue;
+				}
+
 				PackagesToSave.Add(Package);
 
 				{
