[UE-85892] Fix return of the null manifest pointer error in singlepackagecook after robomerge of Jamie's latest changes to the function.
#rnx

diff --git a/Engine/Source/Editor/UnrealEd/Private/Commandlets/AssetRegistryGenerator.cpp b/Engine/Source/Editor/UnrealEd/Private/Commandlets/AssetRegistryGenerator.cpp
--- a/Engine/Source/Editor/UnrealEd/Private/Commandlets/AssetRegistryGenerator.cpp
+++ b/Engine/Source/Editor/UnrealEd/Private/Commandlets/AssetRegistryGenerator.cpp
@@ -385,6 +385,7 @@ bool FAssetRegistryGenerator::GenerateStreamingInstallManifest(int64 InExtraFlav
 		if (!Manifest)
 		{
 			FinalChunkManifests[PakchunkIndex] = new FChunkPackageSet;
+			Manifest = FinalChunkManifests[PakchunkIndex];
 		}
 
 		int32 FilenameIndex = 0;
