Add features to PkgInfoCommandlet to support comparing two results of ResavePackagesCommandlet run on all packages.
Added -dumptopath to support dumping all packages in a directory to a separate .txt per package.
Added more Hide flags to hide values that are expected to change per process/per save.

Fix DiffAssetsCommandlet commandline parsing.  ProcessCommandLine is now using ::CommandLineToArgvW and recombining whitespace tokens to surround them with quotes. This turns DiffCmd="foo" into "DiffCmd=foo", which breaks FParse::Value(CmdLine, "DiffCmd"). Iterate over tokens and split on = to work around this.

Deprecate DiffPackagesCommandlet since it no longer runs and DiffAssetsCommandlet replaces it.
#rnx

diff --git a/Engine/Source/Editor/UnrealEd/Classes/Commandlets/DiffPackagesCommandlet.h b/Engine/Source/Editor/UnrealEd/Classes/Commandlets/DiffPackagesCommandlet.h
--- a/Engine/Source/Editor/UnrealEd/Classes/Commandlets/DiffPackagesCommandlet.h
+++ b/Engine/Source/Editor/UnrealEd/Classes/Commandlets/DiffPackagesCommandlet.h
@@ -32,8 +32,8 @@ enum EObjectDiff
 	OD_Invalid,
 };
 
-UCLASS()
-class UDiffPackagesCommandlet : public UCommandlet
+UCLASS(deprecated)
+class UE_DEPRECATED(4.25, "No longer maintained; use DiffAssetCommandlet instead") UDEPRECATED_DiffPackagesCommandlet : public UCommandlet
 {
 	GENERATED_UCLASS_BODY()
 
@@ -77,14 +77,6 @@ class UDiffPackagesCommandlet : public UCommandlet
 
 	EObjectDiff DiffObjects(UObject* ObjA, UObject* ObjB, UObject* ObjAncestor, struct FObjectComparison& PropDiffs);
 
-	/**
-	 * Copies the raw property values for the natively serialized properties of the specified object into the output var.
-	 *
-	 * @param	Object	the object to load values for
-	 * @param	out_NativePropertyData	receives the raw bytes corresponding to Object's natively serialized property values.
-	 */
-	static void LoadNativePropertyData( UObject* Object, TArray<uint8>& out_NativePropertyData );
-
 	/**
 	 * Compares the natively serialized property values for the specified objects by comparing the non-script serialized portion of each object's data as it
 	 * is on disk.  If a different is detected, gives each object the chance to generate a textual representation of its natively serialized property values
diff --git a/Engine/Source/Editor/UnrealEd/Private/Commandlets/DiffAssetsCommandlet.cpp b/Engine/Source/Editor/UnrealEd/Private/Commandlets/DiffAssetsCommandlet.cpp
--- a/Engine/Source/Editor/UnrealEd/Private/Commandlets/DiffAssetsCommandlet.cpp
+++ b/Engine/Source/Editor/UnrealEd/Private/Commandlets/DiffAssetsCommandlet.cpp
@@ -33,9 +33,28 @@ bool UDiffAssetsCommandlet::ExportFilesToTextAndDiff(const FString& InParams)
 	TArray<FString> Switches;
 	ParseCommandLine(*Params, Tokens, Switches);
 
-	const FString AssetPackageExtension = FPackageName::GetAssetPackageExtension();
+	const FString EqualStr(TEXT("="));
+	const FString DiffCmdKey(TEXT("DiffCmd"));
+	TArray<FString> PositionalTokens;
 	FString DiffCmd;
-	if (!FParse::Value(*Params, TEXT("DiffCmd="), DiffCmd) || Tokens.Num() < 2)
+	for (FString& Token : Tokens)
+	{
+		FString Key, Value;
+		if (Token.Split(EqualStr, &Key, &Value))
+		{
+			if (DiffCmd.IsEmpty() && Key.Equals(DiffCmdKey, ESearchCase::IgnoreCase))
+			{
+				DiffCmd = Value;
+			}
+		}
+		else
+		{
+			PositionalTokens.Add(MoveTemp(Token));
+		}
+	}
+
+	const FString AssetPackageExtension = FPackageName::GetAssetPackageExtension();
+	if (DiffCmd.IsEmpty() || PositionalTokens.Num() < 2)
 	{
 		UE_LOG(LogDiffAssetsCommandlet, Warning, TEXT("Usage: UDiffAssets File1%s File2%s DiffCmd=\"C:/Program Files/Araxis/Araxis Merge/AraxisP4Diff.exe {1} {2}\""), *AssetPackageExtension, *AssetPackageExtension);
 		return false;
@@ -46,7 +65,7 @@ bool UDiffAssetsCommandlet::ExportFilesToTextAndDiff(const FString& InParams)
 		return false;
 	}
 
-	return ExportFilesToTextAndDiff(Tokens[0], Tokens[1], DiffCmd);
+	return ExportFilesToTextAndDiff(PositionalTokens[0], PositionalTokens[1], DiffCmd);
 }
 
 bool UDiffAssetsCommandlet::CopyFileToTempLocation(FString& InOutFilename)
diff --git a/Engine/Source/Editor/UnrealEd/Private/Commandlets/DiffPackagesCommandlet.cpp b/Engine/Source/Editor/UnrealEd/Private/Commandlets/DiffPackagesCommandlet.cpp
--- a/Engine/Source/Editor/UnrealEd/Private/Commandlets/DiffPackagesCommandlet.cpp
+++ b/Engine/Source/Editor/UnrealEd/Private/Commandlets/DiffPackagesCommandlet.cpp
@@ -30,7 +30,7 @@ DEFINE_LOG_CATEGORY_STATIC(LogDiffPackagesCommandlet, Log, All);
 // whether to skip levels when building the initial diff sets (for debugging)
 #define OPTIMIZE_LEVEL_DIFFS 1
 
-UDiffPackagesCommandlet::UDiffPackagesCommandlet(const FObjectInitializer& ObjectInitializer)
+UDEPRECATED_DiffPackagesCommandlet::UDEPRECATED_DiffPackagesCommandlet(const FObjectInitializer& ObjectInitializer)
 	: Super(ObjectInitializer)
 {
 	LogToConsole = false;
@@ -167,6 +167,11 @@ FNativePropertyData::FNativePropertyData( UObject* InObject )
 	SetObject(InObject);
 }
 
+namespace DiffPackagesCommandletImpl
+{
+	void LoadNativePropertyData(UObject* Object, TArray<uint8>& out_NativePropertyData);
+}
+
 /**
  * Changes the UObject associated with this native property data container and re-initializes the
  * PropertyData and PropertyText members
@@ -179,7 +184,7 @@ void FNativePropertyData::SetObject( UObject* NewObject )
 
 	if ( Object != NULL )
 	{
-		UDiffPackagesCommandlet::LoadNativePropertyData(Object, PropertyData);
+		DiffPackagesCommandletImpl::LoadNativePropertyData(Object, PropertyData);
 		Object->GetNativePropertyValues(PropertyText, PPF_SimpleObjectText);
 	}
 }
@@ -393,7 +398,7 @@ void FixupObjectReferences(UPackage* Package, TMap<UObject*, UObject*>& ObjectRe
 }
 
 
-bool UDiffPackagesCommandlet::GenerateObjectComparison( UObject* RootObject, FObjectComparison& out_Comparison, TArray<FObjectComparison>* ObjectsToIgnore/*=NULL*/ )
+bool UDEPRECATED_DiffPackagesCommandlet::GenerateObjectComparison( UObject* RootObject, FObjectComparison& out_Comparison, TArray<FObjectComparison>* ObjectsToIgnore/*=NULL*/ )
 {
 	check(RootObject);
 
@@ -449,10 +454,16 @@ bool UDiffPackagesCommandlet::GenerateObjectComparison( UObject* RootObject, FOb
 }
 
 
-bool UDiffPackagesCommandlet::Initialize( const TCHAR* Parms )
+bool UDEPRECATED_DiffPackagesCommandlet::Initialize( const TCHAR* Parms )
 {
 	bool bResult = false;
 
+	SET_WARN_COLOR(COLOR_RED);
+	UE_LOG(LogDiffPackagesCommandlet, Warning, TEXT("#####################################################"));
+	UE_LOG(LogDiffPackagesCommandlet, Warning, TEXT("# DiffPackages is deprecated, use DiffAssets instead."));
+	UE_LOG(LogDiffPackagesCommandlet, Warning, TEXT("#####################################################"));
+	CLEAR_WARN_COLOR();
+
 	// parse the command line into tokens and switches
 	TArray<FString> Tokens, Switches;
 	ParseCommandLine(Parms, Tokens, Switches);
@@ -695,7 +706,7 @@ bool UDiffPackagesCommandlet::Initialize( const TCHAR* Parms )
 	return bResult;
 }
 
-int32 UDiffPackagesCommandlet::Main(const FString& Params)
+int32 UDEPRECATED_DiffPackagesCommandlet::Main(const FString& Params)
 {
 	if ( !Initialize(*Params) )
 	{
@@ -906,7 +917,7 @@ void AppendComparisonResultText(FString& ExistingResultText, const FString& NewR
 	ExistingResultText += NewResultText + LINE_TERMINATOR;
 }
 
-bool UDiffPackagesCommandlet::ProcessDiff(FObjectComparison& Diff)
+bool UDEPRECATED_DiffPackagesCommandlet::ProcessDiff(FObjectComparison& Diff)
 {
 	// always diff the root objects against each other
 	Diff.OverallDiffType = DiffObjects(
@@ -971,7 +982,7 @@ bool UDiffPackagesCommandlet::ProcessDiff(FObjectComparison& Diff)
 	return Diff.OverallDiffType != OD_None;
 }
 
-EObjectDiff UDiffPackagesCommandlet::DiffObjects(UObject* ObjA, UObject* ObjB, UObject* ObjAncestor, FObjectComparison& Diff)
+EObjectDiff UDEPRECATED_DiffPackagesCommandlet::DiffObjects(UObject* ObjA, UObject* ObjB, UObject* ObjAncestor, FObjectComparison& Diff)
 {
 	// if all objects are NULL, there's no difference :)
 	if (!ObjA && !ObjB && !ObjAncestor)
@@ -1309,7 +1320,10 @@ EObjectDiff UDiffPackagesCommandlet::DiffObjects(UObject* ObjA, UObject* ObjB, U
 
 
 
-void UDiffPackagesCommandlet::LoadNativePropertyData( UObject* Object, TArray<uint8>& out_NativePropertyData )
+namespace DiffPackagesCommandletImpl
+{
+
+void LoadNativePropertyData( UObject* Object, TArray<uint8>& out_NativePropertyData )
 {
 	// first, validate our input parameters
 	check(Object);
@@ -1350,8 +1364,10 @@ void UDiffPackagesCommandlet::LoadNativePropertyData( UObject* Object, TArray<ui
 	}
 }
 
+}
+
 
-EObjectDiff UDiffPackagesCommandlet::CompareNativePropertyValues( UObject* ObjA, UObject* ObjB, UObject* ObjAncestor, FObjectComparison& PropertyValueComparisons )
+EObjectDiff UDEPRECATED_DiffPackagesCommandlet::CompareNativePropertyValues( UObject* ObjA, UObject* ObjB, UObject* ObjAncestor, FObjectComparison& PropertyValueComparisons )
 {
 	FNativePropertyData PropertyDataA(ObjA), PropertyDataB(ObjB), PropertyDataAncestor(ObjAncestor);
 
diff --git a/Engine/Source/Editor/UnrealEd/Private/Commandlets/PackageUtilities.cpp b/Engine/Source/Editor/UnrealEd/Private/Commandlets/PackageUtilities.cpp
--- a/Engine/Source/Editor/UnrealEd/Private/Commandlets/PackageUtilities.cpp
+++ b/Engine/Source/Editor/UnrealEd/Private/Commandlets/PackageUtilities.cpp
@@ -960,7 +960,10 @@ void FPkgInfoReporter_Log::GeneratePackageReport( FLinkerLoad* InLinker /*=nullp
 	Out.Logf(ELogVerbosity::Display, TEXT("\t  Custom Versions:\n%s"), *Linker->Summary.GetCustomVersionContainer().ToString("\t\t"));
 	
 
-	Out.Logf(ELogVerbosity::Display, TEXT("\t             Guid: %s"), *Linker->Summary.Guid.ToString() );
+	if (!IsHideSaveUnstable())
+	{
+		Out.Logf(ELogVerbosity::Display, TEXT("\t             Guid: %s"), *Linker->Summary.Guid.ToString());
+	}
 	Out.Logf(ELogVerbosity::Display, TEXT("\t   PersistentGuid: %s"), *Linker->Summary.PersistentGuid.ToString() );
 	if (Linker->Summary.OwnerPersistentGuid.IsValid())
 	{
@@ -981,7 +984,14 @@ void FPkgInfoReporter_Log::GeneratePackageReport( FLinkerLoad* InLinker /*=nullp
 		for( int32 i = 0; i < Linker->NameMap.Num(); ++i )
 		{
 			FName name = FName::CreateFromDisplayId(Linker->NameMap[ i ], 0);
-			Out.Logf(ELogVerbosity::Display, TEXT("\t%d: Name '%s' Comparison Index %d Display Index %d [Internal: %s, %d]"), i, *name.ToString(), name.GetComparisonIndex().ToUnstableInt(), name.GetDisplayIndex().ToUnstableInt(), *name.GetPlainNameString(), name.GetNumber() );
+			if (IsHideProcessUnstable())
+			{
+				Out.Logf(ELogVerbosity::Display, TEXT("\t%d: Name '%s' [Internal: %s, %d]"), i, *name.ToString(), *name.GetPlainNameString(), name.GetNumber());
+			}
+			else
+			{
+				Out.Logf(ELogVerbosity::Display, TEXT("\t%d: Name '%s' Comparison Index %d Display Index %d [Internal: %s, %d]"), i, *name.ToString(), name.GetComparisonIndex().ToUnstableInt(), name.GetDisplayIndex().ToUnstableInt(), *name.GetPlainNameString(), name.GetNumber());
+			}
 		}
 	}
 
@@ -1200,12 +1210,12 @@ void FPkgInfoReporter_Log::GeneratePackageReport( FLinkerLoad* InLinker /*=nullp
 				Out.Logf(ELogVerbosity::Display, TEXT("\t\t      Pkg Guid: %s"), *Export.PackageGuid.ToString());
 				Out.Logf(ELogVerbosity::Display, TEXT("\t\t   ObjectFlags: 0x%08X"), (uint32)Export.ObjectFlags );
 				Out.Logf(ELogVerbosity::Display, TEXT("\t\t          Size: %d"), Export.SerialSize );
-				if ( !bHideOffsets )
+				if ( !IsHideOffsets())
 				{
 					Out.Logf(ELogVerbosity::Display, TEXT("\t\t      Offset: %d"), Export.SerialOffset );
 				}
 				Out.Logf(ELogVerbosity::Display, TEXT("\t\t       Object: %s"), Export.Object ? TEXT("VALID") : TEXT("NULL"));
-				if ( !bHideOffsets )
+				if ( !IsHideOffsets() )
 				{
 					Out.Logf(ELogVerbosity::Display, TEXT("\t\t    HashNext: %d"), Export.HashNext );
 				}
@@ -1467,21 +1477,31 @@ int32 UPkgInfoCommandlet::Main( const FString& Params )
 	// Create a file writer to dump the info to
 	FOutputDevice* OutputOverride = GWarn;
 	FString OutputFilename;
+	TUniquePtr<FOutputDeviceFile> OutputBuffer;
 	if (FParse::Value(*Params, TEXT("dumptofile="), OutputFilename))
 	{
-		OutputOverride = new FOutputDeviceFile(*OutputFilename, true);
+		OutputBuffer = MakeUnique<FOutputDeviceFile>(*OutputFilename, true);
+		OutputBuffer->SetSuppressEventTag(true);
+		OutputOverride = OutputBuffer.Get();
 	}
 
-	const bool bHideOffsets = Switches.Contains(TEXT("HideOffsets"));
 
-	FPkgInfoReporter* Reporter = new FPkgInfoReporter_Log(InfoFlags, bHideOffsets);
+	uint32 DisplayFlags = PKGINFODISPLAY_None;
+	DisplayFlags |= Switches.Contains(TEXT("HideUnstable")) ? PKGINFODISPLAY_HideAllUnstable : 0;
+	DisplayFlags |= Switches.Contains(TEXT("HideProcessUnstable")) ? PKGINFODISPLAY_HideProcessUnstable : 0;
+	DisplayFlags |= Switches.Contains(TEXT("HideSaveUnstable")) ? PKGINFODISPLAY_HideSaveUnstable : 0;
+	DisplayFlags |= Switches.Contains(TEXT("HideOffsets")) ? PKGINFODISPLAY_HideOffsets : 0;
 
-	TArray<FString> FilesInPath;
+	FPkgInfoReporter* Reporter = new FPkgInfoReporter_Log(InfoFlags, (EPackageInfoDisplayFlags)DisplayFlags);
 
+	TArray<FString> FilesInPath;
 	FString PathWithPackages;
+	FString RelPathSibling;
 	if (FParse::Value(*Params, TEXT("AllPackagesIn="), PathWithPackages))
 	{
 		FPackageName::FindPackagesInDirectory(FilesInPath, *PathWithPackages);
+		RelPathSibling = FPaths::ConvertRelativePathToFull(PathWithPackages);
+		RelPathSibling = FPaths::Combine(RelPathSibling, TEXT("Placeholder"));
 	}
 	else if( Switches.Contains(TEXT("AllPackages")) )
 	{
@@ -1538,9 +1558,41 @@ int32 UPkgInfoCommandlet::Main( const FString& Params )
 		}
 	}
 
+	FString OutputPath;
+	if (FParse::Value(*Params, TEXT("dumptopath="), OutputPath))
+	{
+		if (!OutputFilename.IsEmpty())
+		{
+			UE_LOG(LogPackageUtilities, Warning, TEXT("-dumptopath is not supported with -dumptofile, ignoring -dumptopath."));
+			OutputPath.Empty();
+		}
+		else if (RelPathSibling.IsEmpty())
+		{
+			UE_LOG(LogPackageUtilities, Warning, TEXT("-dumptopath is only supported with -AllPackagesIn, ignoring -dumptopath."));
+			OutputPath.Empty();
+		}
+		else
+		{
+			OutputPath = FPaths::ConvertRelativePathToFull(OutputPath);
+		}
+	}
+
 	for( int32 FileIndex = 0; FileIndex < FilesInPath.Num(); FileIndex++ )
 	{
 		FString Filename = FPaths::ConvertRelativePathToFull(FilesInPath[FileIndex]);
+		FString PackageOutputFilename;
+
+		if (!OutputPath.IsEmpty())
+		{
+			PackageOutputFilename = Filename;
+			if (!FPaths::MakePathRelativeTo(PackageOutputFilename, *RelPathSibling))
+			{
+				UE_LOG(LogPackageUtilities, Error, TEXT("Package filename '%s' is not a child path of root content path '%s', unable to create Outputfile, skipping the file."),
+					*Filename, *FPaths::GetPath(RelPathSibling));
+				continue;
+			}
+			PackageOutputFilename = FPaths::Combine(OutputPath, PackageOutputFilename) + TEXT(".txt");
+		}
 
 		{
 			// reset the loaders for the packages we want to load so that we don't find the wrong version of the file
@@ -1609,6 +1661,14 @@ int32 UPkgInfoCommandlet::Main( const FString& Params )
 				bDumpProperties = false;
 			}
 		}
+
+		if (!PackageOutputFilename.IsEmpty())
+		{
+			OutputBuffer = MakeUnique<FOutputDeviceFile>(*PackageOutputFilename, true);
+			OutputBuffer->SetSuppressEventTag(true);
+			OutputOverride = OutputBuffer.Get();
+		}
+
 		{
 			// Turn off log categories etc as it makes diffing hard
 			TGuardValue<ELogTimes::Type> GuardPrintLogTimes(GPrintLogTimes, ELogTimes::None);
@@ -1653,10 +1713,6 @@ int32 UPkgInfoCommandlet::Main( const FString& Params )
 		CollectGarbage(RF_NoFlags);
 	}
 
-	if (OutputOverride != GWarn)
-	{
-		delete OutputOverride;		
-	}
 	delete Reporter;
 	Reporter = NULL;
 	return 0;
@@ -2732,4 +2788,4 @@ int32 UReplaceActorCommandlet::Main(const FString& Params)
 	}
 	GWorld = UWorld::CreateWorld(EWorldType::Editor, false );
 	return 0;
-}
+}
\ No newline at end of file
diff --git a/Engine/Source/Editor/UnrealEd/Public/PackageUtilityWorkers.h b/Engine/Source/Editor/UnrealEd/Public/PackageUtilityWorkers.h
--- a/Engine/Source/Editor/UnrealEd/Public/PackageUtilityWorkers.h
+++ b/Engine/Source/Editor/UnrealEd/Public/PackageUtilityWorkers.h
@@ -28,6 +28,16 @@ enum EPackageInfoFlags
 	PKGINFO_All			= PKGINFO_Names|PKGINFO_Imports|PKGINFO_Exports|PKGINFO_Depends|PKGINFO_Paths|PKGINFO_Thumbs|PKGINFO_Lazy|PKGINFO_AssetRegistry|PKGINFO_Text,
 };
 
+enum EPackageInfoDisplayFlags
+{
+	PKGINFODISPLAY_None					= 0x00,
+	PKGINFODISPLAY_HideOffsets			= 0x01, // If present, FObjectExport::SerialOffset will not be included in the output; useful when generating a report for comparison against another version of the same package.
+	PKGINFODISPLAY_HideProcessUnstable	= 0x02, // If present, values that are unstable across different process invocations will be not included, such as FName.GetComparisonIndex().ToUnstableInt()
+	PKGINFODISPLAY_HideSaveUnstable		= 0x04, // If present, values that are recalculated during save such as the Package guid will not be included
+
+	PKGINFODISPLAY_HideAllUnstable		= PKGINFODISPLAY_HideOffsets | PKGINFODISPLAY_HideProcessUnstable | PKGINFODISPLAY_HideSaveUnstable,
+	PKGINFODISPLAY_All					= PKGINFODISPLAY_HideAllUnstable
+};
 /**
  * Base for classes which generate output for the PkgInfo commandlet
  */
@@ -35,13 +45,16 @@ struct FPkgInfoReporter
 {
 	/** Constructors */
 	FPkgInfoReporter() 
-	: InfoFlags(PKGINFO_None), bHideOffsets(false), Linker(NULL), PackageCount(0)
+		: FPkgInfoReporter(PKGINFO_None, (EPackageInfoDisplayFlags)PKGINFODISPLAY_None, nullptr)
 	{}
-	FPkgInfoReporter( uint32 InInfoFlags, bool bInHideOffsets, FLinkerLoad* InLinker=NULL )
-	: InfoFlags(InInfoFlags), bHideOffsets(bInHideOffsets), Linker(InLinker), PackageCount(0)
+	UE_DEPRECATED(4.25, "Use The EPackageInfoDisplayFlags constructor instead") FPkgInfoReporter(uint32 InInfoFlags, bool bInHideOffsets, FLinkerLoad* InLinker = NULL)
+		: FPkgInfoReporter(InInfoFlags, bInHideOffsets ? PKGINFODISPLAY_HideOffsets : PKGINFODISPLAY_None, InLinker)
+	{}
+	FPkgInfoReporter(uint32 InInfoFlags, EPackageInfoDisplayFlags InDisplayFlags, FLinkerLoad* InLinker = NULL)
+		: InfoFlags(InInfoFlags), DisplayFlags(InDisplayFlags), Linker(InLinker), PackageCount(0)
 	{}
 	FPkgInfoReporter( const FPkgInfoReporter& Other )
-	: InfoFlags(Other.InfoFlags), bHideOffsets(Other.bHideOffsets), Linker(Other.Linker), PackageCount(Other.PackageCount)
+	: InfoFlags(Other.InfoFlags), DisplayFlags(Other.DisplayFlags), Linker(Other.Linker), PackageCount(Other.PackageCount)
 	{}
 
 	/** Destructor */
@@ -64,15 +77,14 @@ struct FPkgInfoReporter
 
 protected:
 	/**
-	 * A bitmask of PKGINFO_ flags that determine the information that is included in the report.
+	 * A bitmask of PKGINFO_ flags that determine the categories of information included in the report.
 	 */
 	uint32 InfoFlags;
 
-	/**
-	 * Determines whether FObjectExport::SerialOffset will be included in the output; useful when generating
-	 * a report for comparison against another version of the same package.
+	/*
+	 * A bitmask of EPackageInfoDisplayFlags that determine the display of information in the report.
 	 */
-	bool bHideOffsets;
+	uint32 DisplayFlags;
 
 	/**
 	 * The linker of the package to generate the report for
@@ -83,16 +95,16 @@ protected:
 	 * The number of packages evaluated by this reporter so far.  Must be incremented by child classes.
 	 */
 	int32 PackageCount;
+
+	bool IsHideOffsets() const { return (DisplayFlags & PKGINFODISPLAY_HideOffsets) != 0; }
+	bool IsHideProcessUnstable() const { return (DisplayFlags & PKGINFODISPLAY_HideProcessUnstable) != 0; }
+	bool IsHideSaveUnstable() const { return (DisplayFlags & PKGINFODISPLAY_HideSaveUnstable) != 0; }
 };
 
 struct FPkgInfoReporter_Log : public FPkgInfoReporter
 {
-	FPkgInfoReporter_Log( uint32 InInfoFlags, bool bInHideOffsets, FLinkerLoad* InLinker=NULL )
-	: FPkgInfoReporter(InInfoFlags, bInHideOffsets, InLinker)
-	{}
-	FPkgInfoReporter_Log( const FPkgInfoReporter_Log& Other )
-	: FPkgInfoReporter( Other )
-	{}
+	using FPkgInfoReporter::FPkgInfoReporter;
+
 	/**
 	 * Writes information about the linker to the log.
 	 *
