Add SessionPlatforms to record the platforms being cooked in the current CookByTheBookSession or CookOnTheFly request. Replace use of TargetPlatformManager.GetTargetPlatforms with SessionPlatforms. Remove the up-front registration of all possible platforms that occurred when creating a CookOnTheFLy server.  This fixes the performance issue of unnecessarily generating AssetRegistries for unused platforms when cooking in the editor.

Keep track of data about platforms that have been requested in previous sessions, to maintain the prevention of duplicate work that was previously prevented by doing the work for all possible platforms at startup.  This requirement lead to the creation of FPlatformData.

Performance: Get rid of passing TargetPlatform FNames around and looking up the platform from TPM; just pass the TargetPlatforms around.

Maintain or improve performance and maintainability in the code being touched:
	Performance: Use TInlineAllocator where possible, pass around TArrayViews to support that.
	Maintainability: Remove no-longer-used functions and parameters.
	Memory: Move per-platform RegistryGenerator into a single map of Platform to FPlatformData, along with the new flags per platform.
	Performance/Maintainability: Merge CleanSandbox into InitializeSandbox to avoid the need to activate many-to-one work twice in the two functions
	Maintainability: Change the check for bFullLoadAndSave at the top of TickCookOnTheSide to include CookByTheBookOptions->bRunning.  This was already always true; I added it to make it clear the FullLoadAndSave does not happen unless StartCookByTheBook was called earlier.

#rn Minor Cooking

diff --git a/Engine/Source/Editor/UnrealEd/Classes/Commandlets/CookCommandlet.h b/Engine/Source/Editor/UnrealEd/Classes/Commandlets/CookCommandlet.h
--- a/Engine/Source/Editor/UnrealEd/Classes/Commandlets/CookCommandlet.h
+++ b/Engine/Source/Editor/UnrealEd/Classes/Commandlets/CookCommandlet.h
@@ -64,10 +64,11 @@ class UCookCommandlet
 	 * @param  BindAnyPort					Whether to bind on any port or the default port.
 	 * @param  Timeout						Length of time to wait for connections before attempting to close
 	 * @param  bForceClose					Whether or not the server should always shutdown after a timeout or after a user disconnects
+	 * @param  TargetPlatforms				The list of platforms that should be initialized at startup.  Other platforms will be initialized when first requested
 	 *
 	 * @return true on success, false otherwise.
 	 */
-	bool CookOnTheFly( FGuid InstanceId, int32 Timeout = 180, bool bForceClose = false );
+	bool CookOnTheFly( FGuid InstanceId, int32 Timeout = 180, bool bForceClose = false, const TArray<ITargetPlatform*>& TargetPlatforms=TArray<ITargetPlatform*>() );
 
 	/** Cooks specified list of files */
 	bool CookByTheBook(const TArray<ITargetPlatform*>& Platforms, TArray<FString>& FilesInPath);
diff --git a/Engine/Source/Editor/UnrealEd/Classes/CookOnTheSide/CookOnTheFlyServer.h b/Engine/Source/Editor/UnrealEd/Classes/CookOnTheSide/CookOnTheFlyServer.h
--- a/Engine/Source/Editor/UnrealEd/Classes/CookOnTheSide/CookOnTheFlyServer.h
+++ b/Engine/Source/Editor/UnrealEd/Classes/CookOnTheSide/CookOnTheFlyServer.h
@@ -3,12 +3,15 @@
 #pragma once
 
 #include "CoreMinimal.h"
+
+#include "Containers/ArrayView.h"
 #include "Misc/EnumClassFlags.h"
 #include "UObject/ObjectMacros.h"
 #include "UObject/UObjectGlobals.h"
 #include "UObject/Object.h"
 #include "UObject/WeakObjectPtr.h"
 #include "UObject/Package.h"
+#include "Templates/UniquePtr.h"
 #include "TickableEditorObject.h"
 #include "IPlatformFileSandboxWrapper.h"
 #include "INetworkFileSystemModule.h"
@@ -107,15 +110,20 @@ class UNREALED_API UCookOnTheFlyServer : public UObject, public FTickableEditorO
 		UCookOnTheFlyServer(const FObjectInitializer& ObjectInitializer = FObjectInitializer::Get());
 		UCookOnTheFlyServer(FVTableHelper& Helper); // Declare the FVTableHelper constructor manually so that we can forward-declare-only TUniquePtrs in the header without getting compile error in generated cpp
 
+public:
+	struct FCookByTheBookOptions;
+	struct FPlatformData;
+	struct FPlatformManager;
+
 private:
 	/** Current cook mode the cook on the fly server is running in */
 	ECookMode::Type CurrentCookMode = ECookMode::CookOnTheFly;
 	/** Directory to output to instead of the default should be empty in the case of DLC cooking */ 
 	FString OutputDirectoryOverride;
 
-	struct FCookByTheBookOptions;
-
 	FCookByTheBookOptions* CookByTheBookOptions = nullptr;
+	TUniquePtr<FPlatformManager> PlatformManager;
+	FCriticalSection RequestLock;
 
 	//////////////////////////////////////////////////////////////////////////
 	// Cook on the fly options
@@ -216,9 +224,6 @@ private:
 	/** Cached copy of asset registry */
 	IAssetRegistry* AssetRegistry = nullptr;
 
-	/** Map of platform name to asset registry generators, which hold the state of asset registry data for a platform */
-	TMap<FName, FAssetRegistryGenerator*> RegistryGenerators;
-
 	/** Map of platform name to scl.csv files we saved out. */
 	TMap<FName, TArray<FString>> OutSCLCSVPaths;
 
@@ -249,17 +254,18 @@ private:
 	*/
 	void OnTargetPlatformChangedSupportedFormats(const ITargetPlatform* TargetPlatform);
 
-	/**
-	 * Returns the current set of cooking targetplatforms
-	 *  mostly used for cook on the fly or in situations where the cooker can't figure out what the target platform is
-	 * 
-	 * @return Array of target platforms which are can be used
-	 */
-	const TArray<ITargetPlatform*>& GetCookingTargetPlatforms() const;
+	/* In CookOnTheFly, requests can add a new platform; this function conditionally initializes the requested platform.  Returns false if the given TargetPlatform is not available, else true */
+	bool AddCookOnTheFlyPlatform(ITargetPlatform* TargetPlatform);
+	/* Version of AddCookOnTheFlyPlatform that takes the Platform name instead of an ITargetPlatform*.  Returns the Platform if found */
+	ITargetPlatform* AddCookOnTheFlyPlatform(const FString& PlatformName);
+	/* Internal helper for AddCookOnTheFlyPlatform.  Initializing Platforms must be done on the tickloop thread; Platform data is read only on other threads */
+	void AddCookOnTheFlyPlatformFromGameThread(ITargetPlatform* TargetPlatform);
 
-	/** Cached cooking target platforms from the targetmanager, these are used when we don't know what platforms we should be targeting */
-	mutable TArray<ITargetPlatform*> CookingTargetPlatforms;
+	/* Constructs the PackageNameCache and PackageTracker */
+	void ConstructPackageTracker();
 
+	/* Update polled fields used by CookOnTheFly's network request handlers */
+	void TickNetwork();
 public:
 
 	enum ECookOnTheSideResult
@@ -308,11 +314,12 @@ public:
 	/**
 	* Cook on the side, cooks while also running the editor...
 	*
-	* @param  BindAnyPort					Whether to bind on any port or the default port.
+	* @param BindAnyPort					Whether to bind on any port or the default port.
+	* @param TargetPlatforms				If nonempty, cooking will be prepared (generate AssetRegistry, etc) for each platform in the array
 	*
 	* @return true on success, false otherwise.
 	*/
-	bool StartNetworkFileServer( bool BindAnyPort );
+	bool StartNetworkFileServer( bool BindAnyPort, const TArray<ITargetPlatform*>& TargetPlatforms = TArray<ITargetPlatform*>() );
 
 	/**
 	* Broadcast our the fileserver presence on the network
@@ -365,6 +372,9 @@ public:
 	/**
 	* Get any packages which are in memory, these were probably required to be loaded because of the current package we are cooking, so we should probably cook them also
 	*/
+	TArray<UPackage*> GetUnsolicitedPackages(const TArray<const ITargetPlatform*>& TargetPlatforms) const;
+
+	UE_DEPRECATED(4.25, "Use the version that takes TArray<const ITargetPlatform*> instead")
 	TArray<UPackage*> GetUnsolicitedPackages(const TArray<FName>& TargetPlatformNames) const;
 
 	/**
@@ -390,18 +400,24 @@ public:
 
 	/**
 	* Clear any cached cooked platform data for a platform
-	*  call ClearCachedCookedPlatformData on all Uobjects
-	* @param PlatformName platform to clear all the cached data for
+	*  call ClearCachedCookedPlatformData on all UObjects
+	* @param TargetPlatform platform to clear all the cached data for
 	*/
-	void ClearCachedCookedPlatformDataForPlatform( const FName& PlatformName );
+	void ClearCachedCookedPlatformDataForPlatform(const ITargetPlatform* TargetPlatform);
+
+	UE_DEPRECATED(4.25, "Use version that takes const ITargetPlatform* instead")
+	void ClearCachedCookedPlatformDataForPlatform(const FName& PlatformName);
 
 
 	/**
 	* Clear all the previously cooked data for the platform passed in 
 	* 
-	* @param name of the platform to clear the cooked packages for
+	* @param TargetPlatform the platform to clear the cooked packages for
 	*/
-	void ClearPlatformCookedData( const FString& PlatformName );
+	void ClearPlatformCookedData(const ITargetPlatform* TargetPlatform);
+
+	UE_DEPRECATED(4.25, "Use version that takes const ITargetPlatform* instead")
+	void ClearPlatformCookedData(const FString& PlatformName);
 
 
 	/**
@@ -410,7 +426,10 @@ public:
 	* 
 	* @return return true if shaders were recompiled
 	*/
-	bool RecompileChangedShaders(const TArray<FName>& TargetPlatforms);
+	bool RecompileChangedShaders(const TArray<const ITargetPlatform*>& TargetPlatforms);
+
+	UE_DEPRECATED(4.25, "Use version that takes const ITargetPlatform* instead")
+	bool RecompileChangedShaders(const TArray<FName>& TargetPlatformNames);
 
 
 	/**
@@ -485,10 +504,13 @@ public:
 	* RequestPackage to be cooked
 	*
 	* @param StandardPackageFName name of the package in standard format as returned by FPaths::MakeStandardFilename
-	* @param TargetPlatforms name of the targetplatforms we want this package cooked for
+	* @param TargetPlatforms The TargetPlatforms we want this package cooked for
 	* @param bForceFrontOfQueue should we put this package in the front of the cook queue (next to be processed) or at the end
 	*/
-	bool RequestPackage(const FName& StandardPackageFName, const TArray<FName>& TargetPlatforms, const bool bForceFrontOfQueue);
+	bool RequestPackage(const FName& StandardPackageFName, const TArray<const ITargetPlatform*>& TargetPlatforms, const bool bForceFrontOfQueue);
+
+	UE_DEPRECATED(4.25, "Use Version that takes TArray<const ITargetPlatform*> instead")
+	bool RequestPackage(const FName& StandardPackageFName, const TArray<FName>& TargetPlatformNames, const bool bForceFrontOfQueue);
 
 	/**
 	* RequestPackage to be cooked
@@ -538,7 +560,8 @@ private:
 	*/
 	void CollectFilesToCook(TArray<FName>& FilesInPath, 
 		const TArray<FString>& CookMaps, const TArray<FString>& CookDirectories, 
-		const TArray<FString>& IniMapSections, ECookByTheBookOptions FilesToCookFlags);
+		const TArray<FString>& IniMapSections, ECookByTheBookOptions FilesToCookFlags,
+		const TArrayView<const ITargetPlatform* const>& TargetPlatforms);
 
 	/**
 	* AddFileToCook add file to cook list 
@@ -588,17 +611,17 @@ private:
 	* BuildMapDependencyGraph
 	* builds a map of dependencies from maps
 	* 
-	* @param PlatformName name of the platform we want to build a map for
+	* @param TargetPlatform the platform we want to build a map for
 	*/
-	void BuildMapDependencyGraph(const FName& PlatformName);
+	void BuildMapDependencyGraph(const ITargetPlatform* TargetPlatform);
 
 	/**
 	* WriteMapDependencyGraph
 	* write a previously built map dependency graph out to the sandbox directory for a platform
 	*
-	* @param PlatformName name of the platform we want to save out the dependency graph for
+	* @param TargetPlatform the platform we want to save out the dependency graph for
 	*/
-	void WriteMapDependencyGraph(const FName& PlatformName);
+	void WriteMapDependencyGraph(const ITargetPlatform* TargetPlatform);
 
 	//////////////////////////////////////////////////////////////////////////
 	// cook on the fly specific functions
@@ -611,7 +634,7 @@ private:
 	 */
 	bool HandleNetworkFileServerNewConnection( const FString& VersionInfo, const FString& PlatformName );
 
-	void GetCookOnTheFlyUnsolicitedFiles(const FName& PlatformName, TArray<FString>& UnsolicitedFiles, const FString& Filename);
+	void GetCookOnTheFlyUnsolicitedFiles(const ITargetPlatform* TargetPlatform, TArray<FString>& UnsolicitedFiles, const FString& Filename);
 
 	/**
 	* Cook requests for a package from network
@@ -636,8 +659,6 @@ private:
 	 */
 	FString HandleNetworkGetSandboxPath();
 
-	void GetCookOnTheFlyUnsolicitedFiles( const FName& PlatformName, TArray<FString>& UnsolicitedFiles );
-
 	/**
 	 * HandleNetworkGetPrecookedList 
 	 * this is used specifically for cook on the fly with shared cooked builds
@@ -656,14 +677,13 @@ private:
 	 *  internal function should not be used externally Call Tick / RequestPackage to initiate
 	 * 
 	 * @param PackageToSave main cooked package to save
-	 * @param TargetPlatformNames list of target platforms names that we want to save this package for
-	 * @param TargetPlatformsToCache list of target platforms that we want to cache uobjects for, might not be the same as the list of packages to save
+	 * @param InTargetPlatforms list of target platforms names that we want to save this package for and cache UObjects for
 	 * @param Timer FCookerTimer struct which defines the timeslicing behavior 
 	 * @param FirstUnsolicitedPackage first package which was not actually requested, unsolicited packages are prioritized differently, they are saved the same way
 	 * @param Result (in+out) used to modify the result of the operation and add any relevant flags
 	 * @return returns true if we saved all the packages false if we bailed early for any reason
 	 */
-	void SaveCookedPackages(UPackage* PackageToSave, const TArray<FName>& TargetPlatformNames, const TArray<const ITargetPlatform*>& TargetPlatformsToCache, struct FCookerTimer& Timer, uint32& CookedPackageCount, uint32& Result);
+	void SaveCookedPackages(UPackage* PackageToSave, const TArray<const ITargetPlatform*>& InTargetPlatforms, struct FCookerTimer& Timer, uint32& CookedPackageCount, uint32& Result);
 
 	/**
 	 * Attempts to update the metadata for a package in an asset registry generator
@@ -683,9 +703,10 @@ private:
 	 *  this is the same as a normal load but also ensures that the sublevels are loaded if they are streaming sublevels
 	 *
 	 * @param BuildFilename long package name of the package to load 
-	 * @return UPackage of the package loaded, null if the file didn't exist or other failure
+	 * @param OutPackage UPackage of the package loaded, non-null on success, may be non-null on failure if the UPackage existed but had another failure
+	 * @return Whether LoadPackage was completely successful and the package can be cooked
 	 */
-	UPackage* LoadPackageForCooking(const FString& BuildFilename);
+	bool LoadPackageForCooking(const FString& BuildFilename, UPackage*& OutPackage);
 
 	/**
 	* Makes sure a package is fully loaded before we save it out
@@ -694,14 +715,14 @@ private:
 	bool MakePackageFullyLoaded(UPackage* Package) const;
 
 	/**
-	* Initialize the sandbox 
+	* Initialize the sandbox for @param TargetPlatforms
 	*/
-	void InitializeSandbox();
+	void InitializeSandbox(const TArrayView<const ITargetPlatform* const>& TargetPlatforms);
 
 	/**
 	* Initialize the package store
 	*/
-	void InitializePackageStore(const TArray<FName>& TargetPlatformNames);
+	void InitializePackageStore(const TArrayView<const ITargetPlatform* const>& TargetPlatforms);
 
 	/**
 	* Finalize the package store
@@ -714,16 +735,16 @@ private:
 	void ClearPackageStoreContexts();
 
 	/**
-	* Initialize all target platforms
+	* Initialize platforms in @param NewTargetPlatforms
 	*/
-	void InitializeTargetPlatforms();
+	void InitializeTargetPlatforms(const TArrayView<ITargetPlatform* const>& NewTargetPlatforms);
 
 	/**
 	* Some content plugins does not support all target platforms.
 	* Build up a map of unsupported packages per platform that can be checked before saving.
 	*/
 	void DiscoverPlatformSpecificNeverCookPackages(
-		const TArray<FName>& TargetPlatformNames, const TArray<FString>& UBTPlatformStrings);
+		const TArrayView<const ITargetPlatform* const>& TargetPlatforms, const TArray<FString>& UBTPlatformStrings);
 
 	/**
 	* Clean up the sandbox
@@ -783,7 +804,7 @@ private:
 	 * @param TargetPlatforms target platforms to cache for
 	 * @return false if time slice was reached, true if all objects have had BeginCacheForCookedPlatformData called
 	 */
-	bool BeginPackageCacheForCookedPlatformData(UPackage* Package, const TArray<const ITargetPlatform*>& TargetPlatforms, struct FCookerTimer& Timer) const;
+	bool BeginPackageCacheForCookedPlatformData(UPackage* Package, const TArrayView<const ITargetPlatform* const>& TargetPlatforms, struct FCookerTimer& Timer) const;
 	
 	/**
 	 * Returns true when all objects in package have all their cooked platform data loaded
@@ -793,7 +814,7 @@ private:
 	 * @param TargetPlatforms target platforms to cache for
 	 * @return false if time slice was reached, true if all return true for IsCachedCookedPlatformDataLoaded 
 	 */
-	bool FinishPackageCacheForCookedPlatformData(UPackage* Package, const TArray<const ITargetPlatform*>& TargetPlatforms, struct FCookerTimer& Timer) const;
+	bool FinishPackageCacheForCookedPlatformData(UPackage* Package, const TArrayView<const ITargetPlatform* const>& TargetPlatforms, struct FCookerTimer& Timer) const;
 
 	/**
 	* GetCurrentIniVersionStrings gets the current ini version strings for compare against previous cook
@@ -899,22 +920,11 @@ private:
 	*	@param	Package				The package to cook/save
 	*	@param	SaveFlags			The flags to pass to the SavePackage function
 	*	@param	bOutWasUpToDate		Upon return, if true then the cooked package was cached (up to date)
+	*	@param  TargetPlatforms		Only cook for target platforms which are included in this array
 	*
 	*	@return	ESavePackageResult::Success if packages was cooked
 	*/
-	void SaveCookedPackage(UPackage* Package, uint32 SaveFlags, TArray<FSavePackageResultStruct>& SavePackageResults);
-	/**
-	*	Cook (save) the given package
-	*
-	*	@param	Package				The package to cook/save
-	*	@param	SaveFlags			The flags to pass to the SavePackage function
-	*	@param	bOutWasUpToDate		Upon return, if true then the cooked package was cached (up to date)
-	*	@param  TargetPlatformNames Only cook for target platforms which are included in this array (if empty cook for all target platforms specified on commandline options)
-	*									TargetPlatformNames is in and out value returns the platforms which the SaveCookedPackage function saved for
-	*
-	*	@return	ESavePackageResult::Success if packages was cooked
-	*/
-	void SaveCookedPackage(UPackage* Package, uint32 SaveFlags, TArray<FName> &TargetPlatformNames, TArray<FSavePackageResultStruct>& SavePackageResults);
+	void SaveCookedPackage(UPackage* Package, uint32 SaveFlags, const TArrayView<const ITargetPlatform* const>& TargetPlatforms, TArray<FSavePackageResultStruct>& SavePackageResults);
 
 
 	/**
@@ -922,7 +932,7 @@ private:
 	*  
 	*  @param	Platforms		List of platforms to make global shader maps for
 	*/
-	void SaveGlobalShaderMapFiles(const TArray<ITargetPlatform*>& Platforms);
+	void SaveGlobalShaderMapFiles(const TArrayView<const ITargetPlatform* const>& Platforms);
 
 
 	/** Create sandbox file in directory using current settings supplied */
@@ -930,16 +940,13 @@ private:
 	/** Gets the output directory respecting any command line overrides */
 	FString GetOutputDirectoryOverride() const;
 
-	/** Cleans sandbox folders for all target platforms */
-	void CleanSandbox( const bool bIterative );
-
 	/**
 	* Populate the cooked packages list from the on disk content using time stamps and dependencies to figure out if they are ok
 	* delete any local content which is out of date
 	* 
 	* @param Platforms to process
 	*/
-	void PopulateCookedPackagesFromDisk( const TArray<ITargetPlatform*>& Platforms );
+	void PopulateCookedPackagesFromDisk(const TArrayView<const ITargetPlatform* const>& Platforms);
 
 	/**
 	* Searches the disk for all the cooked files in the sandbox path provided
@@ -950,9 +957,12 @@ private:
 	*/
 	void GetAllCookedFiles(TMap<FName, FName>& UncookedPathToCookedPath, const FString& SandboxRootDir);
 
-	/** Generates asset registry */
+	/** Loads the platform-independent asset registry for use by the cooker */
 	void GenerateAssetRegistry();
 
+	/** Construct or refresh-for-filechanges the platform-specific asset registry for the given platforms */
+	void RefreshPlatformAssetRegistries(const TArrayView<const ITargetPlatform* const>& TargetPlatforms);
+
 	/** Generates long package names for all files to be cooked */
 	void GenerateLongPackageNames(TArray<FName>& FilesInPath);
 
diff --git a/Engine/Source/Editor/UnrealEd/Private/Commandlets/CookCommandlet.cpp b/Engine/Source/Editor/UnrealEd/Private/Commandlets/CookCommandlet.cpp
--- a/Engine/Source/Editor/UnrealEd/Private/Commandlets/CookCommandlet.cpp
+++ b/Engine/Source/Editor/UnrealEd/Private/Commandlets/CookCommandlet.cpp
@@ -324,7 +324,7 @@ UCookCommandlet::UCookCommandlet( const FObjectInitializer& ObjectInitializer )
 	LogToConsole = false;
 }
 
-bool UCookCommandlet::CookOnTheFly( FGuid InstanceId, int32 Timeout, bool bForceClose )
+bool UCookCommandlet::CookOnTheFly( FGuid InstanceId, int32 Timeout, bool bForceClose, const TArray<ITargetPlatform*>& TargetPlatforms)
 {
 	UCookOnTheFlyServer *CookOnTheFlyServer = NewObject<UCookOnTheFlyServer>();
 
@@ -356,7 +356,7 @@ bool UCookCommandlet::CookOnTheFly( FGuid InstanceId, int32 Timeout, bool bForce
 
 	bool BindAnyPort = InstanceId.IsValid();
 
-	if ( CookOnTheFlyServer->StartNetworkFileServer(BindAnyPort) == false )
+	if ( CookOnTheFlyServer->StartNetworkFileServer(BindAnyPort, TargetPlatforms) == false )
 	{
 		return false;
 	}
@@ -549,8 +549,20 @@ int32 UCookCommandlet::Main(const FString& CmdLineParams)
 
 	COOK_STAT(DetailedCookStats::CookProject = FApp::GetProjectName());
 
+
+	ITargetPlatformManagerModule& TPM = GetTargetPlatformManagerRef();
 	if ( bCookOnTheFly )
 	{
+		// In cook on the fly, if the user did not provide a targetplatform on the commandline, then we do not intialize any platforms up front; we wait for the first connection.
+		// TPM.GetActiveTargetPlatforms defaults to the currently running platform (e.g. Windows, with editor) in the no-target case, so we need to only call GetActiveTargetPlatforms
+		// if targetplatform was on the commandline
+		FString Unused;
+		TArray<ITargetPlatform*> TargetPlatforms;
+		if (FParse::Value(FCommandLine::Get(), TEXT("TARGETPLATFORM="), Unused))
+		{
+			TargetPlatforms = TPM.GetActiveTargetPlatforms();
+		}
+
 		// parse instance identifier
 		FString InstanceIdString;
 		bool bForceClose = Switches.Contains(TEXT("FORCECLOSE"));
@@ -570,12 +582,10 @@ int32 UCookCommandlet::Main(const FString& CmdLineParams)
 			Timeout = 180;
 		}
 
-		CookOnTheFly( InstanceId, Timeout, bForceClose);
+		CookOnTheFly( InstanceId, Timeout, bForceClose, TargetPlatforms);
 	}
 	else
 	{
-		
-		ITargetPlatformManagerModule& TPM = GetTargetPlatformManagerRef();
 		const TArray<ITargetPlatform*>& Platforms = TPM.GetActiveTargetPlatforms();
 
 		TArray<FString> FilesInPath;
diff --git a/Engine/Source/Editor/UnrealEd/Private/CookOnTheFlyServer.cpp b/Engine/Source/Editor/UnrealEd/Private/CookOnTheFlyServer.cpp
--- a/Engine/Source/Editor/UnrealEd/Private/CookOnTheFlyServer.cpp
+++ b/Engine/Source/Editor/UnrealEd/Private/CookOnTheFlyServer.cpp
@@ -22,6 +22,7 @@
 #include "Serialization/CustomVersion.h"
 #include "Misc/App.h"
 #include "Misc/CoreDelegates.h"
+#include "Misc/ScopeExit.h"
 #include "Misc/ScopeLock.h"
 #include "Modules/ModuleManager.h"
 #include "Serialization/ArchiveUObject.h"
@@ -405,6 +406,7 @@ DECLARE_STATS_GROUP(TEXT("Cooking"), STATGROUP_Cooking, STATCAT_Advanced);
 DECLARE_CYCLE_STAT(TEXT("Precache Derived data for platform"), STAT_TickPrecacheCooking, STATGROUP_Cooking);
 DECLARE_CYCLE_STAT(TEXT("Tick cooking"), STAT_TickCooker, STATGROUP_Cooking);
 
+constexpr uint32 ExpectedMaxNumPlatforms = 32;
 
 
 /* helper structs functions
@@ -495,11 +497,11 @@ const FString& GetAssetRegistryPath()
 /**
  * Return the release asset registry filename for the release version supplied
  */
-FString GetReleaseVersionAssetRegistryPath(const FString& ReleaseVersion, const FName& PlatformName )
+FString GetReleaseVersionAssetRegistryPath(const FString& ReleaseVersion, const FString& PlatformName )
 {
 	// cache the part of the path which is static because getting the ProjectDir is really slow and also string manipulation
 	const static FString ProjectDirectory = FPaths::ProjectDir() / FString(TEXT("Releases"));
-	return  ProjectDirectory / ReleaseVersion / PlatformName.ToString();
+	return  ProjectDirectory / ReleaseVersion / PlatformName;
 }
 
 const FString& GetAssetRegistryFilename()
@@ -536,12 +538,243 @@ void LogCookerMessage( const FString& MessageText, EMessageSeverity::Type Severi
 
 //////////////////////////////////////////////////////////////////////////
 
+/** A BaseKeyFuncs for Maps and Sets with a quicker hash function for pointers than TDefaultMapKeyFuncs */
+template<typename KeyType, typename ValueType, bool bInAllowDuplicateKeys>
+struct TFastPointerKeyFuncs : public TDefaultMapKeyFuncs<KeyType, ValueType, bInAllowDuplicateKeys>
+{
+	using typename TDefaultMapKeyFuncs<KeyType, ValueType, bInAllowDuplicateKeys>::KeyInitType;
+	static FORCEINLINE uint32 GetKeyHash(KeyInitType Key)
+	{
+#if PLATFORM_64BITS
+		static_assert(sizeof(UPTRINT) == sizeof(uint64), "Expected pointer size to be 64 bits");
+		// Ignoring the lower 4 bits since they are likely zero anyway.
+		const uint64 ImportantBits = reinterpret_cast<uint64>(Key) >> 4;
+		return GetTypeHash(ImportantBits);
+#else
+		static_assert(sizeof(UPTRINT) == sizeof(uint32), "Expected pointer size to be 32 bits");
+		return static_cast<uint32>(reinterpret_cast<uint32>(Key));
+#endif
+	}
+};
+
+/** A TMap which uses TFastPointerKeyFuncs instead of TDefaultMapKeyFuncs */
+template<typename KeyType, typename ValueType, typename SetAllocator = FDefaultSetAllocator>
+class TFastPointerMap : public TMap<KeyType, ValueType, FDefaultSetAllocator, TFastPointerKeyFuncs<KeyType, ValueType, false>>
+{};
+
+//////////////////////////////////////////////////////////////////////////
+
+/*
+ * Struct to hold data about each platform we have encountered in the cooker.
+ * Fields in this struct persist across multiple CookByTheBook sessions.
+ * Fields on this struct are used on multiple threads; see variable comments for thread synchronization rules.
+ */
+struct UCookOnTheFlyServer::FPlatformData
+{
+	/* Name of the Platform, a cache of FName(ITargetPlatform->PlatformName()) */
+	FName PlatformName;
+
+	/* Pointer to the platform-specific RegistryGenerator for this platform.  If already constructed we can take a faster refresh path on future sessions.
+	 * Read/Write on TickCookOnTheSide thread only
+	 */
+	TUniquePtr<FAssetRegistryGenerator> RegistryGenerator;
+
+	/*
+	 * Whether InitializeSandbox has been called for this platform.  If we have already initialized the sandbox we can take a faster refresh path on future sessions.
+	 * Threadsafe due to write-once.  Written only once after construction on the game thread.
+	 */
+	bool bIsSandboxInitialized = false;
+
+	/*
+	 * The last FPlatformTime::GetSeconds() at which this platform was requested in a CookOnTheFly request.
+	 * If equal to 0, the platform was not requested in a CookOnTheFly since the last clear.
+	 * Written only when SessionLock critical section is held.
+	 */
+	double LastReferenceTime = 0.0;
+
+	/*
+	 * The count of how many active CookOnTheFly requests are currently using the Platform.
+	 * Read/Write only when the SessionLock critical section is held.
+	 */
+	uint32 ReferenceCount = 0;
+};
+
+/* Information about the platforms (a) known and (b) active for the current cook session in the UCookOnTheFlyServer. */
+struct UCookOnTheFlyServer::FPlatformManager
+{
+public:
+	FPlatformManager(FCriticalSection& InSessionLock)
+		:SessionLock(InSessionLock)
+	{}
+
+	FCriticalSection& GetSessionLock()
+	{
+		return SessionLock;
+	}
+
+	/*
+	 * Returns the set of TargetPlatforms that is active for the CurrentCookByTheBook session or CookOnTheFly request.
+	 * This function can be called and its return value referenced only from the TickCookOnTheSide thread or when SessionLock is held.
+	 */
+	const TArray<const ITargetPlatform*>& GetSessionPlatforms() const
+	{
+		checkf(bHasSelectedSessionPlatforms, TEXT("Calling GetSessionPlatforms or (any of the top level cook functions that call it) without first calling SelectSessionPlatforms is invalid"));
+		return SessionPlatforms;
+	}
+
+	/*
+	 * Return whether the platforms have been selected for the current session (may be empty if current session is a null session).
+	 * This function can be called from any thread, but is guaranteed thread-accurate only from the TickCookOnTheSide thread or when SessionLock is held.
+	*/
+	bool HasSelectedSessionPlatforms() const
+	{
+		return bHasSelectedSessionPlatforms;
+	}
+
+	/*
+	 * Return whether the given platform is already in the list of platforms for the current session.
+	 * This function can be called only from the TickCookOnTheSide thread or when SessionLock is held.
+	 */
+	bool HasSessionPlatform(const ITargetPlatform* TargetPlatform) const
+	{
+		return SessionPlatforms.Contains(TargetPlatform);
+	}
+
+	/*
+	 * Specify the set of TargetPlatforms to use for the currently-initializing CookByTheBook session or CookOnTheFly request.
+	 * This function can be called only from the TickCookOnTheSide thread.
+	 */
+	void SelectSessionPlatforms(const TArrayView<const ITargetPlatform* const>& TargetPlatforms, FPackageTracker* PackageTracker);
+
+	/*
+	 * Mark that the list of TargetPlatforms for the session has no longer been set; it will be an error to try to read them until SelectSessionPlatforms is called.
+	 * This function can be called only from the TickCookOnTheSide thread.
+	 */
+	void ClearSessionPlatforms()
+	{
+		FScopeLock Lock(&SessionLock);
+
+		SessionPlatforms.Empty();
+		bHasSelectedSessionPlatforms = false;
+	}
+
+	/*
+	 * Add @param TargetPlatform to the session platforms if not already present.
+	 * This function can be called only from the TickCookOnTheSide thread.
+	 *
+	 * @param PackageTracker If the TargetPlatform does not already exist, we notify the PackageTracker that it was added.
+	 */
+	void AddSessionPlatform(const ITargetPlatform* TargetPlatform, FPackageTracker* PackageTracker);
+
+	/*
+	 * Get The PlatformData for the given Platform.
+	 * Guaranteed to return non-null for any Platform in the current list of SessionPlatforms.
+	 * Can be called from any thread.
+	 */
+	UCookOnTheFlyServer::FPlatformData* GetPlatformData(const ITargetPlatform* Platform)
+	{
+		return PlatformDatas.Find(Platform);
+	}
+
+	/*
+	 * Create if not already created the necessary platform-specific data for the given platform.
+	 * This function can be called with new platforms only before multithreading begins (e.g. in StartNetworkFileServer or StartCookByTheBook).
+	 */
+	UCookOnTheFlyServer::FPlatformData& CreatePlatformData(const ITargetPlatform* Platform)
+	{
+		check(Platform != nullptr);
+		FPlatformData& PlatformData = PlatformDatas.FindOrAdd(Platform);
+		if (PlatformData.PlatformName.IsNone())
+		{
+			check(!bPlatformDataFrozen); // It is not legal to add new platforms to this map when running CookOnTheFlyServer because we read this map from threads handling network requests, and mutating the map is not threadsafe when it is being read
+
+			// Newly added, construct
+			PlatformData.PlatformName = FName(Platform->PlatformName());
+			checkf(!PlatformData.PlatformName.IsNone(), TEXT("Invalid ITargetPlatform with an empty name"));
+		}
+		return PlatformData;
+	}
+
+	/* 
+	 * Return whether platform-specific setup steps have been executed for the given platform in the current UCookOnTheFlyServer.
+	 * This function can be called from any thread, but is guaranteed thread-accurate only from the TickCookOnTheSide thread.
+	 */
+	bool IsPlatformInitialized(const ITargetPlatform* Platform) const
+	{
+		const FPlatformData* PlatformData = PlatformDatas.Find(Platform);
+		if (!PlatformData)
+		{
+			return false;
+		}
+		return PlatformData->bIsSandboxInitialized;
+	}
+
+	/* If and only if bFrozen is set to true, it is invalid to call CreatePlatformData with a Platform that does not already exist. */
+	void SetPlatformDataFrozen(bool bFrozen)
+	{
+		bPlatformDataFrozen = bFrozen;
+	}
+
+	/*
+	 * Platforms requested in CookOnTheFly requests are added to the list of SessionPlatforms, and some packages (e.g. unsolicited packages) are cooked against all session packages.
+	 * To have good performance in the case where multiple targetplatforms are being requested over time from the same CookOnTheFly server, we prune platforms from the list of 
+	 * active session platforms if they haven't been requested in a while.
+	 * This function can be called only from the TickCookOnTheSide thread.
+	 * 
+	 * @param PackageTracker If a platform is removed, we remove it from the PackageTracker as well.
+	 */
+	void PruneUnreferencedSessionPlatforms(FPackageTracker& PackageTracker);
+
+	/*
+	 * Increment the counter indicating the current platform is requested in an active CookOnTheFly request.  Add it to the SessionPlatforms if not already present.
+	 * This function can be called only when the SessionLock is held.
+	 * 
+	 * @param PackageTracker - PackageTracker that provides the queue of tick commands to execute AddSessionPlatform on the game thread, and to be notified of the new platform.
+	*/
+	void AddRefCookOnTheFlyPlatform(const ITargetPlatform* TargetPlatform, FPackageTracker& PackageTracker);
+
+	/*
+	 * Decrement the counter indicating the current platform is being used in a CookOnTheFly request.
+	 * This function can be called only when the SessionLock is held.
+	 */
+	void ReleaseCookOnTheFlyPlatform(const ITargetPlatform* TargetPlatform)
+	{
+		check(TargetPlatform != nullptr);
+		FPlatformData* PlatformData = GetPlatformData(TargetPlatform);
+		checkf(PlatformData != nullptr, TEXT("Unrecognized Platform %s"), *TargetPlatform->PlatformName());
+		check(PlatformData->ReferenceCount > 0);
+		--PlatformData->ReferenceCount;
+		PlatformData->LastReferenceTime = FPlatformTime::Seconds();
+	}
+
+private:
+	/* A collection of initialization flags and other data we maintain for each platform we have encountered in any session. */
+	TFastPointerMap<const ITargetPlatform*, UCookOnTheFlyServer::FPlatformData> PlatformDatas;
+
+	/*
+	 * A collection of Platforms that are active for the current CookByTheBook session or CookOnTheFly request.  Used so we can refer to "all platforms" without having to store a list on every FileRequest.
+	 * Writing to the list of active session platforms requires a CriticalSection, because it is read (under critical section) on NetworkFileServer threads.
+	 */
+	TArray<const ITargetPlatform*> SessionPlatforms;
+
+	/* A reference to the critical section used to guard SessionPlatforms. */
+	FCriticalSection& SessionLock;
+
+	/* If PlatformData is frozen, it is invalid to add new PlatformDatas. */
+	bool bPlatformDataFrozen = false;
+
+	/* It is invalid to attempt to cook if session platforms have not been selected. */
+	bool bHasSelectedSessionPlatforms = false;
+};
+
+//////////////////////////////////////////////////////////////////////////
+
 /** cooked file requests which includes platform which file is requested for */
 struct FFilePlatformRequest
 {
 protected:
 	FName			Filename;
-	TArray<FName>	PlatformNames;
+	TArray<const ITargetPlatform*>	Platforms;
 
 public:
 	// yes we have some friends
@@ -549,19 +782,19 @@ public:
 
 	FFilePlatformRequest() = default;
 
-	FFilePlatformRequest(const FName& InFileName, const FName& InPlatformName) : Filename(InFileName)
+	FFilePlatformRequest(const FName& InFileName, const ITargetPlatform* InPlatform) : Filename(InFileName)
 	{
-		PlatformNames.Add(InPlatformName);
+		Platforms.Add(InPlatform);
 	}
 
-	FFilePlatformRequest(const FName& InFilename, const TArray<FName>& InPlatformName) : Filename(InFilename)
-	{
-		PlatformNames = InPlatformName;
-	}
+	FFilePlatformRequest(const FName& InFilename, const TArrayView<const ITargetPlatform* const>& InPlatforms)
+		: Filename(InFilename)
+		, Platforms(InPlatforms.GetData(), InPlatforms.Num())
+	{}
 
-	FFilePlatformRequest(const FName& InFilename, TArray<FName>&& InPlatformName) : Filename(InFilename), PlatformNames(MoveTemp(InPlatformName)) { }
-	FFilePlatformRequest(const FFilePlatformRequest& InFilePlatformRequest) : Filename(InFilePlatformRequest.Filename), PlatformNames(InFilePlatformRequest.PlatformNames) { }
-	FFilePlatformRequest(FFilePlatformRequest&& InFilePlatformRequest) : Filename(MoveTemp(InFilePlatformRequest.Filename)), PlatformNames(MoveTemp(InFilePlatformRequest.PlatformNames)) { }
+	FFilePlatformRequest(const FName& InFilename, TArray<const ITargetPlatform*>&& InPlatforms) : Filename(InFilename), Platforms(MoveTemp(InPlatforms)) { }
+	FFilePlatformRequest(const FFilePlatformRequest& InFilePlatformRequest) : Filename(InFilePlatformRequest.Filename), Platforms(InFilePlatformRequest.Platforms) { }
+	FFilePlatformRequest(FFilePlatformRequest&& InFilePlatformRequest) : Filename(MoveTemp(InFilePlatformRequest.Filename)), Platforms(MoveTemp(InFilePlatformRequest.Platforms)) { }
 
 	void SetFilename(const FString& InFilename)
 	{
@@ -573,25 +806,25 @@ public:
 		return Filename;
 	}
 
-	const TArray<FName>& GetPlatformNames() const
+	const TArray<const ITargetPlatform*>& GetPlatforms() const
 	{
-		return PlatformNames;
+		return Platforms;
 	}
 
-	void RemovePlatform(const FName& Platform)
+	void RemovePlatform(const ITargetPlatform* Platform)
 	{
-		PlatformNames.Remove(Platform);
+		Platforms.Remove(Platform);
 	}
 
-	void AddPlatform(const FName& Platform)
+	void AddPlatform(const ITargetPlatform* Platform)
 	{
-		check(Platform != NAME_None);
-		PlatformNames.Add(Platform);
+		check(Platform != nullptr);
+		Platforms.Add(Platform);
 	}
 
-	bool HasPlatform(const FName& Platform) const
+	bool HasPlatform(const ITargetPlatform* Platform) const
 	{
-		return PlatformNames.Find(Platform) != INDEX_NONE;
+		return Platforms.Find(Platform) != INDEX_NONE;
 	}
 
 	bool IsValid()  const
@@ -602,13 +835,13 @@ public:
 	void Clear()
 	{
 		Filename = TEXT("");
-		PlatformNames.Empty();
+		Platforms.Empty();
 	}
 
 	FFilePlatformRequest& operator=(FFilePlatformRequest&& InFileRequest)
 	{
 		Filename = MoveTemp(InFileRequest.Filename);
-		PlatformNames = MoveTemp(InFileRequest.PlatformNames);
+		Platforms = MoveTemp(InFileRequest.Platforms);
 		return *this;
 	}
 
@@ -616,7 +849,7 @@ public:
 	{
 		if (InFileRequest.Filename == Filename)
 		{
-			if (InFileRequest.PlatformNames == PlatformNames)
+			if (InFileRequest.Platforms == Platforms)
 			{
 				return true;
 			}
@@ -628,9 +861,9 @@ public:
 	{
 		FString Result = FString::Printf(TEXT("%s;"), *Filename.ToString());
 
-		for (const FName& Platform : PlatformNames)
+		for (const ITargetPlatform* Platform : Platforms)
 		{
-			Result += FString::Printf(TEXT("%s,"), *Platform.ToString());
+			Result += FString::Printf(TEXT("%s,"), *Platform->PlatformName());
 		}
 		return Result;
 	}
@@ -640,9 +873,9 @@ uint32 GetTypeHash(const FFilePlatformRequest& Key)
 {
 	uint32 Hash = GetTypeHash(Key.Filename);
 
-	for (const FName& PlatformName : Key.PlatformNames)
+	for (const ITargetPlatform* Platform : Key.Platforms)
 	{
-		Hash += Hash << 2 ^ GetTypeHash(PlatformName);
+		Hash += Hash << 2 ^ GetTypeHash(reinterpret_cast<UPTRINT>(Platform));
 	}
 
 	return Hash;
@@ -651,41 +884,41 @@ uint32 GetTypeHash(const FFilePlatformRequest& Key)
 struct FFilePlatformCookedPackage
 {
 public:
-	FFilePlatformCookedPackage(const FFilePlatformRequest& InFilePlatformRequest, const TArray<bool>& bInSuccededSavePackage)
+	FFilePlatformCookedPackage(const FFilePlatformRequest& InFilePlatformRequest, const TArray<bool>& bInSucceededSavePackage)
 	: Filename(InFilePlatformRequest.GetFilename())
-	, PlatformNames(InFilePlatformRequest.GetPlatformNames())
-	, bSucceededSavePackage(bInSuccededSavePackage)
+	, Platforms(InFilePlatformRequest.GetPlatforms())
+	, bSucceededSavePackage(bInSucceededSavePackage)
 	{ 
-		check(PlatformNames.Num() == bSucceededSavePackage.Num()); 
+		check(Platforms.Num() == bSucceededSavePackage.Num());
 	}
 
-	FFilePlatformCookedPackage(const FName& InFilename, const TArray<FName>& InPlatformName) 
+	FFilePlatformCookedPackage(const FName& InFilename, const TArrayView<const ITargetPlatform* const>& InPlatforms) 
 	: Filename(InFilename)
-	, PlatformNames(InPlatformName)
+	, Platforms(InPlatforms.GetData(), InPlatforms.Num())
 	{
 		// only use this constructor to short hand when packages fail
-		for (int32 I = 0; I < InPlatformName.Num(); ++I)
+		for (int32 I = 0; I < InPlatforms.Num(); ++I)
 		{
 			bSucceededSavePackage.Add(false);
 		}
 
-		check(PlatformNames.Num() == bSucceededSavePackage.Num());
+		check(Platforms.Num() == bSucceededSavePackage.Num());
 	}
 
-	FFilePlatformCookedPackage(const FName& InFilename, const TArray<FName>& InPlatformName, const TArray<bool>& bInSuccededSavePackage)
+	FFilePlatformCookedPackage(const FName& InFilename, const TArrayView<const ITargetPlatform* const>& InPlatforms, const TArray<bool>& bInSucceededSavePackage)
 	: Filename(InFilename)
-	, PlatformNames(InPlatformName)
-	, bSucceededSavePackage(bInSuccededSavePackage) 
+	, Platforms(InPlatforms.GetData(), InPlatforms.Num())
+	, bSucceededSavePackage(bInSucceededSavePackage) 
 	{ 
-		check(PlatformNames.Num() == bSucceededSavePackage.Num()); 
+		check(Platforms.Num() == bSucceededSavePackage.Num()); 
 	}
 
-	FFilePlatformCookedPackage(const FName& InFilename, TArray<FName>&& InPlatformName, TArray<bool>&& bInSuccededSavePackage)
+	FFilePlatformCookedPackage(const FName& InFilename, TArray<const ITargetPlatform*>&& InPlatforms, TArray<bool>&& bInSucceededSavePackage)
 	: Filename(InFilename)
-	, PlatformNames(MoveTemp(InPlatformName))
-	, bSucceededSavePackage(MoveTemp(bInSuccededSavePackage))
+	, Platforms(MoveTemp(InPlatforms))
+	, bSucceededSavePackage(MoveTemp(bInSucceededSavePackage))
 	{ 
-		check(PlatformNames.Num() == bSucceededSavePackage.Num());
+		check(Platforms.Num() == bSucceededSavePackage.Num());
 	}
 
 	FFilePlatformCookedPackage(const FFilePlatformCookedPackage& InFilePlatformRequest) = default;
@@ -701,42 +934,75 @@ public:
 		return Filename;
 	}
 
-	const TArray<FName>& GetPlatformNames() const
+	const TArray<const ITargetPlatform*>& GetPlatforms() const
 	{
-		return PlatformNames;
+		return Platforms;
 	}
 
-	bool HasPlatform(const FName& Platform) const
+	bool HasPlatform(const ITargetPlatform* Platform) const
 	{
-		return PlatformNames.Find(Platform) != INDEX_NONE;
+		return Platforms.Find(Platform) != INDEX_NONE;
 	}
 
-	const void AddPlatform(const FName& Platform, bool bSucceeded)
+	bool HasPlatforms(const TArrayView<const ITargetPlatform* const>& QueryPlatforms, bool bIncludeFailed = true) const
 	{
-		check(PlatformNames.Num() == bSucceededSavePackage.Num());
-		check(Platform != NAME_None);
-		PlatformNames.Add(Platform);
+		if (bIncludeFailed == false)
+		{
+			bool bAllFailed = true;
+
+			for (const ITargetPlatform* Platform : QueryPlatforms)
+			{
+				if (HasSucceededSavePackage(Platform))
+				{
+					bAllFailed = false;
+					break;
+				}
+			}
+
+			if (bAllFailed)
+			{
+				return false;
+			}
+		}
+
+		// make sure all the platforms are completed
+		for (const ITargetPlatform* Platform : QueryPlatforms)
+		{
+			if (!Platforms.Contains(Platform))
+			{
+				return false;
+			}
+		}
+
+		return true;
+	}
+
+	const void AddPlatform(const ITargetPlatform* Platform, bool bSucceeded)
+	{
+		check(Platforms.Num() == bSucceededSavePackage.Num());
+		check(Platform != nullptr);
+		Platforms.Add(Platform);
 		bSucceededSavePackage.Add(bSucceeded);
-		check(PlatformNames.Num() == bSucceededSavePackage.Num());
+		check(Platforms.Num() == bSucceededSavePackage.Num());
 	}
 
-	void RemovePlatform(const FName& Platform)
+	void RemovePlatform(const ITargetPlatform* Platform)
 	{
-		check(PlatformNames.Num() == bSucceededSavePackage.Num());
-		const int32 Index = PlatformNames.IndexOfByKey(Platform);
+		check(Platforms.Num() == bSucceededSavePackage.Num());
+		const int32 Index = Platforms.IndexOfByKey(Platform);
 
 		if (Index != -1)
 		{
-			PlatformNames.RemoveAt(Index);
+			Platforms.RemoveAt(Index);
 			bSucceededSavePackage.RemoveAt(Index);
 		}
 
-		check(PlatformNames.Num() == bSucceededSavePackage.Num());
+		check(Platforms.Num() == bSucceededSavePackage.Num());
 	}
 
-	const bool HasSucceededSavePackage(const FName& PlatformName) const
+	const bool HasSucceededSavePackage(const ITargetPlatform* Platform) const
 	{
-		const int32 Index = PlatformNames.IndexOfByKey(PlatformName);
+		const int32 Index = Platforms.IndexOfByKey(Platform);
 
 		if ((Index != -1) && (Index < bSucceededSavePackage.Num()))
 		{
@@ -747,7 +1013,7 @@ public:
 
 private:
 	FName			Filename;
-	TArray<FName>	PlatformNames;
+	TArray<const ITargetPlatform*>	Platforms;
 	TArray<bool>	bSucceededSavePackage; // one bool for each platform
 };
 
@@ -765,7 +1031,7 @@ public:
 		return FilesProcessed.Num();
 	}
 
-	void Add(const FFilePlatformCookedPackage& InRequest)
+	FFilePlatformCookedPackage& Add(const FFilePlatformCookedPackage& InRequest)
 	{
 		check(InRequest.IsValid());
 
@@ -778,24 +1044,25 @@ public:
 		{
 			check(ExistingRequest->GetFilename() == InRequest.GetFilename());
 
-			for (const FName& Platform : InRequest.GetPlatformNames())
+			for (const ITargetPlatform* Platform : InRequest.GetPlatforms())
 			{
 				const bool bSucceeded = InRequest.HasSucceededSavePackage(Platform);
 				ExistingRequest->AddPlatform(Platform, bSucceeded);
 			}
+			return *ExistingRequest;
 		}
 		else
 		{
-			FilesProcessed.Add(InRequest.GetFilename(), InRequest);
+			return FilesProcessed.Add(InRequest.GetFilename(), InRequest);
 		}
 	}
 
 	bool Exists(const FFilePlatformRequest& Request) const
 	{
-		return Exists(Request.GetFilename(), Request.GetPlatformNames(), /* bIncludeFailed */ true);
+		return Exists(Request.GetFilename(), Request.GetPlatforms(), /* bIncludeFailed */ true);
 	}
 
-	bool Exists(const FName& Filename, const TArray<FName>& PlatformNames, bool bIncludeFailed = true) const
+	bool Exists(const FName& Filename, const TArrayView<const ITargetPlatform* const>& Platforms, bool bIncludeFailed = true) const
 	{
 		FScopeLock ScopeLock(&SynchronizationObject);
 
@@ -806,54 +1073,26 @@ public:
 			return false;
 		}
 
-		if (bIncludeFailed == false)
-		{
-			bool allFailed = true;
-
-			for (const FName& PlatformName : PlatformNames)
-			{
-				if (OurRequest->HasSucceededSavePackage(PlatformName))
-				{
-					allFailed = false;
-					break;
-				}
-			}
-
-			if (allFailed)
-			{
-				return false;
-			}
-		}
-
-		// make sure all the platforms are completed
-		for (const FName& Platform : PlatformNames)
-		{
-			if (!OurRequest->GetPlatformNames().Contains(Platform))
-			{
-				return false;
-			}
-		}
-
-		return true;
+		return OurRequest->HasPlatforms(Platforms, bIncludeFailed);
 	}
 
-	void RemoveAllFilesForPlatform(const FName& PlatformName)
+	void RemoveAllFilesForPlatform(const ITargetPlatform* Platform)
 	{
 		FScopeLock ScopeLock(&SynchronizationObject);
 
 		for (auto& Request : FilesProcessed)
 		{
-			Request.Value.RemovePlatform(PlatformName);
+			Request.Value.RemovePlatform(Platform);
 		}
 	}
 
-	bool GetCookedPlatforms(const FName& Filename, TArray<FName>& OutPlatformList) const
+	bool GetCookedPlatforms(const FName& Filename, TArray<const ITargetPlatform*>& OutPlatformList) const
 	{
 		FScopeLock ScopeLock(&SynchronizationObject);
 
 		if (const FFilePlatformCookedPackage* Request = FilesProcessed.Find(Filename))
 		{
-			OutPlatformList = Request->GetPlatformNames();
+			OutPlatformList = Request->GetPlatforms();
 			return true;
 		}
 
@@ -867,14 +1106,14 @@ public:
 		return FilesProcessed.Remove(Filename);
 	}
 
-	void GetCookedFilesForPlatform(const FName& PlatformName, TArray<FName>& CookedFiles, bool bGetFailedCookedPackages, bool bGetSuccessfulCookedPackages)
+	void GetCookedFilesForPlatform(const ITargetPlatform* Platform, TArray<FName>& CookedFiles, bool bGetFailedCookedPackages, bool bGetSuccessfulCookedPackages)
 	{
 		FScopeLock ScopeLock(&SynchronizationObject);
 		for (const auto& CookedFile : FilesProcessed)
 		{
-			if (CookedFile.Value.HasPlatform(PlatformName))
+			if (CookedFile.Value.HasPlatform(Platform))
 			{
-				const bool bHasSucceededSavePackage = CookedFile.Value.HasSucceededSavePackage(PlatformName);
+				const bool bHasSucceededSavePackage = CookedFile.Value.HasSucceededSavePackage(Platform);
 
 				if (	(bHasSucceededSavePackage && bGetSuccessfulCookedPackages)
 					|| ((bHasSucceededSavePackage == false) && bGetFailedCookedPackages))
@@ -892,12 +1131,15 @@ public:
 	}
 };
 
+/* 
+ * Holds a queue of requests for cooking packages on specified platforms.
+ * Not threadsafe; caller must handle locking a CriticalSection around use when used on multiple threads.
+ */
 struct CookRequestQueue
 {
 	template <class PREDICATE_CLASS>
 	void Sort(const PREDICATE_CLASS& Predicate)
 	{
-		FScopeLock ScopeLock(&SynchronizationObject);
 		Queue.Sort(Predicate);
 	}
 
@@ -908,17 +1150,16 @@ struct CookRequestQueue
 
 	void EnqueueUnique(const FFilePlatformRequest& Request, bool ForceEnqueFront = false)
 	{
-		FScopeLock ScopeLock(&SynchronizationObject);
-		TArray<FName>* Platforms = PlatformList.Find(Request.GetFilename());
+		TArray<const ITargetPlatform*>* Platforms = PlatformList.Find(Request.GetFilename());
 		if (Platforms == NULL)
 		{
-			PlatformList.Add(Request.GetFilename(), Request.GetPlatformNames());
+			PlatformList.Add(Request.GetFilename(), Request.GetPlatforms());
 			Queue.Add(Request.GetFilename());
 		}
 		else
 		{
 			// add the requested platforms to the platform list
-			for (const FName& Platform : Request.GetPlatformNames())
+			for (const ITargetPlatform* Platform : Request.GetPlatforms())
 			{
 				Platforms->AddUnique(Platform);
 			}
@@ -937,12 +1178,11 @@ struct CookRequestQueue
 
 	bool Dequeue(FFilePlatformRequest& OutResult)
 	{
-		FScopeLock ScopeLock(&SynchronizationObject);
 		if (Queue.Num())
 		{
 			FName Filename = Queue[0];
 			Queue.RemoveAt(0);
-			TArray<FName> Platforms = PlatformList.FindChecked(Filename);
+			TArray<const ITargetPlatform*> Platforms = PlatformList.FindChecked(Filename);
 			PlatformList.Remove(Filename);
 
 			OutResult = FFilePlatformRequest(MoveTemp(Filename), MoveTemp(Platforms));
@@ -953,7 +1193,6 @@ struct CookRequestQueue
 
 	void DequeueAllRequests(TArray<FFilePlatformRequest>& RequestArray)
 	{
-		FScopeLock ScopeLock(&SynchronizationObject);
 		if (Queue.Num())
 		{
 			for (const auto& Request : PlatformList)
@@ -965,17 +1204,16 @@ struct CookRequestQueue
 		}
 	}
 
-	bool Exists(const FName& Filename, const TArray<FName>& PlatformNames) const
+	bool Exists(const FName& Filename, const TArray<const ITargetPlatform*>& Platforms) const
 	{
-		FScopeLock ScopeLock(&SynchronizationObject);
-		const TArray<FName>* Platforms = PlatformList.Find(Filename);
-		if (Platforms == NULL)
+		const TArray<const ITargetPlatform*>* ExistingPlatforms = PlatformList.Find(Filename);
+		if (ExistingPlatforms == nullptr)
 		{
 			return false;
 		}
-		for (const FName& PlatformName : PlatformNames)
+		for (const ITargetPlatform* Platform : Platforms)
 		{
-			if (!Platforms->Contains(PlatformName))
+			if (!ExistingPlatforms->Contains(Platform))
 			{
 				return false;
 			}
@@ -985,9 +1223,8 @@ struct CookRequestQueue
 
 	bool Exists(const FName& Filename)
 	{
-		FScopeLock ScopeLock(&SynchronizationObject);
-		const TArray<FName>* Platforms = PlatformList.Find(Filename);
-		if (Platforms == NULL)
+		const TArray<const ITargetPlatform*>* ExistingPlatforms = PlatformList.Find(Filename);
+		if (ExistingPlatforms == nullptr)
 		{
 			return false;
 		}
@@ -996,27 +1233,36 @@ struct CookRequestQueue
 
 	bool HasItems() const
 	{
-		FScopeLock ScopeLock(&SynchronizationObject);
 		return Queue.Num() > 0;
 	}
 
 	int Num() const
 	{
-		FScopeLock ScopeLock(&SynchronizationObject);
 		return Queue.Num();
 	}
 
 	void Empty()
 	{
-		FScopeLock ScopeLock(&SynchronizationObject);
 		Queue.Empty();
 		PlatformList.Empty();
 	}
 
+	void RemovePlatform(const ITargetPlatform* TargetPlatform)
+	{
+		for (auto& kvpair : PlatformList)
+		{
+			TArray<const ITargetPlatform*>& RequestPlatforms = kvpair.Value;
+			RequestPlatforms.Remove(TargetPlatform);
+			if (RequestPlatforms.Num() == 0)
+			{
+				UE_LOG(LogCook, Error, TEXT("RemovePlatform call has left an empty list of platforms requested in CookOnTheSide request."));
+			}
+		}
+	}
+
 private:
 	TArray<FName>				Queue;
-	TMap<FName, TArray<FName>>	PlatformList;
-	mutable FCriticalSection	SynchronizationObject;
+	TMap<FName, TArray<const ITargetPlatform*>>	PlatformList;
 };
 
 struct FThreadSafeUnsolicitedPackagesList
@@ -1027,7 +1273,7 @@ struct FThreadSafeUnsolicitedPackagesList
 		CookedPackages.Add(PlatformRequest);
 	}
 
-	void GetPackagesForPlatformAndRemove(const FName& Platform, TArray<FName> PackageNames)
+	void GetPackagesForPlatformAndRemove(const ITargetPlatform* Platform, TArray<FName> PackageNames)
 	{
 		FScopeLock _(&SyncObject);
 
@@ -1035,12 +1281,12 @@ struct FThreadSafeUnsolicitedPackagesList
 		{
 			FFilePlatformRequest &Request = CookedPackages[I];
 
-			if (Request.GetPlatformNames().Contains(Platform))
+			if (Request.GetPlatforms().Contains(Platform))
 			{
 				// remove the platform
 				Request.RemovePlatform(Platform);
 
-				if (Request.GetPlatformNames().Num() == 0)
+				if (Request.GetPlatforms().Num() == 0)
 				{
 					CookedPackages.RemoveAt(I);
 				}
@@ -1126,19 +1372,6 @@ public:
 };
 
 
-struct FDummyCriticalSection
-{
-public:
-	FORCEINLINE void Lock() { }
-	FORCEINLINE void Unlock() { }
-};
-
-struct FDummyScopeLock
-{
-public:
-	FDummyScopeLock(FDummyCriticalSection *) { }
-};
-
 template<typename Type>
 struct FThreadSafeQueue : public FUnsynchronizedQueue<Type, FCriticalSection, FScopeLock>
 {
@@ -1195,8 +1428,18 @@ public:
 
 struct FPackageTracker : public FUObjectArray::FUObjectCreateListener, public FUObjectArray::FUObjectDeleteListener
 {
-	FPackageTracker(FPackageNameCache* InPackageNameCache)
-	:	PackageNameCache(InPackageNameCache)
+	enum class ERequestType
+	{
+		None,
+		TickCommand,
+		Cook
+	};
+	typedef TFunction<void()> FTickCommand;
+
+	FPackageTracker(FPackageNameCache* InPackageNameCache, FCriticalSection& InRequestLock, UCookOnTheFlyServer::FPlatformManager& InPlatformManager)
+		: RequestLock(InRequestLock)
+		, PackageNameCache(InPackageNameCache)
+		, PlatformManager(InPlatformManager)
 	{
 		for (TObjectIterator<UPackage> It; It; ++It)
 		{
@@ -1205,7 +1448,6 @@ struct FPackageTracker : public FUObjectArray::FUObjectCreateListener, public FU
 			if (Package->GetOuter() == nullptr)
 			{
 				LoadedPackages.Add(Package);
-				FilterLoadedPackage(Package);
 			}
 		}
 
@@ -1213,6 +1455,9 @@ struct FPackageTracker : public FUObjectArray::FUObjectCreateListener, public FU
 
 		GUObjectArray.AddUObjectDeleteListener(this);
 		GUObjectArray.AddUObjectCreateListener(this);
+
+		// We will call FilterLoadedPackage on every package in LoadedPackages the next time GetPackagesPendingSave is called
+		bPackagesPendingSaveDirty = true;
 	}
 
 	~FPackageTracker()
@@ -1237,7 +1482,10 @@ struct FPackageTracker : public FUObjectArray::FUObjectCreateListener, public FU
 				LoadedPackages.Add(Package);
 				NewPackages.Add(Package);
 
-				FilterLoadedPackage(Package);
+				if (PlatformManager.HasSelectedSessionPlatforms())
+				{
+					FilterLoadedPackage(Package);
+				}
 			}
 		}
 	}
@@ -1261,6 +1509,79 @@ struct FPackageTracker : public FUObjectArray::FUObjectCreateListener, public FU
 		GUObjectArray.RemoveUObjectCreateListener(this);
 	}
 
+	void AddTickCommand(const FTickCommand& TickCommand)
+	{
+		FScopeLock ScopeLock(&RequestLock);
+		TickCommands.Add(TickCommand);
+	}
+
+	bool HasRequests()
+	{
+		FScopeLock ScopeLock(&RequestLock);
+		return TickCommands.Num() > 0 || CookRequests.HasItems();
+	}
+
+	void EnqueueUniqueCookRequest(const FFilePlatformRequest& FileRequest, bool bForceFrontOfQueue = false)
+	{
+		FScopeLock ScopeLock(&RequestLock);
+		CookRequests.EnqueueUnique(FileRequest, bForceFrontOfQueue);
+	}
+
+	int32 GetCookRequestsNum()
+	{
+		FScopeLock ScopeLock(&RequestLock);
+		return CookRequests.Num();
+	}
+
+	ERequestType DequeueRequest(TArray<FTickCommand>& OutTickCommands, FFilePlatformRequest& OutToBuild)
+	{
+		FScopeLock ScopeLock(&RequestLock);
+		if (TickCommands.Num() > 0)
+		{
+			OutTickCommands = MoveTemp(TickCommands);
+			TickCommands.Empty();
+			return ERequestType::TickCommand;
+		}
+		else if (CookRequests.Dequeue(OutToBuild))
+		{
+			return ERequestType::Cook;
+		}
+		else
+		{
+			return ERequestType::None;
+		}
+	}
+
+	void EmptyRequests()
+	{
+		FScopeLock ScopeLock(&RequestLock);
+		CookRequests.Empty();
+		TickCommands.Empty();
+	}
+
+	void DequeueAllRequests(TArray<FPackageTracker::FTickCommand>& OutTickCommands, TArray<FFilePlatformRequest>& RequestArray)
+	{
+		FScopeLock ScopeLock(&RequestLock);
+		OutTickCommands = MoveTemp(TickCommands);
+		TickCommands.Empty();
+		CookRequests.DequeueAllRequests(RequestArray);
+	}
+
+	void RemoveSessionPlatform(const ITargetPlatform* TargetPlatform)
+	{
+		FScopeLock Lock(&RequestLock);
+
+		TArray<FName> UnusedPackageNames;
+		UnsolicitedCookedPackages.GetPackagesForPlatformAndRemove(TargetPlatform, UnusedPackageNames);
+
+		// The caller should not be removing platforms if we have an active request referencing that platform, but in case they did, remove the platform
+		// from all pending requests
+		CookRequests.RemovePlatform(TargetPlatform);
+
+		// Keep information about packages cooked for the platform in the CookedPackages set, in case the platform is requested again in the future
+	}
+
+
 	// This is the set of packages which have already had PostLoadFixup called 
 	TSet<UPackage*>			PostLoadFixupPackages;
 
@@ -1275,7 +1596,15 @@ struct FPackageTracker : public FUObjectArray::FUObjectCreateListener, public FU
 	CookedPackageSet		CookedPackages;
 
 	// List of requested files
+private:
 	CookRequestQueue		CookRequests;
+	TArray<FTickCommand> TickCommands;
+	FCriticalSection&		RequestLock;
+public:
+	// These functions exposing the RequestLock and CookRequests are important in a few places for performance, but will likely be removed in the future and should not be widely used
+	FCriticalSection& GetRequestLock() { return RequestLock; }
+	CookRequestQueue& GetThreadUnsafeCookRequests(){ return CookRequests; }
+public:
 
 	FEvent*					CookRequestEvent = nullptr;
 
@@ -1284,15 +1613,11 @@ struct FPackageTracker : public FUObjectArray::FUObjectCreateListener, public FU
 
 	FThreadSafeSet<FName>						NeverCookPackageList;
 	FThreadSafeSet<FName>						UncookedEditorOnlyPackages; // set of packages that have been rejected due to being referenced by editor-only properties
-	TMap<FName, TSet<FName>> 					PlatformSpecificNeverCookPackages;
-
-	// Currently targeted platforms
-	TArray<FName>			AllTargetPlatformNames;
-
-	void					AddTargetPlatform(const FName& TargetPlatform);
+	TFastPointerMap<const ITargetPlatform*, TSet<FName>> 	PlatformSpecificNeverCookPackages;
 
 private:
 	FPackageNameCache*		PackageNameCache;
+	UCookOnTheFlyServer::FPlatformManager& PlatformManager;
 
 	// Set of packages pending save.
 
@@ -1303,16 +1628,16 @@ private:
 
 	bool					bPackagesPendingSaveDirty = true;
 
-	void					DirtyPackagesPendingSave()		{ bPackagesPendingSaveDirty = true; }
 	void					UpdatePackagesPendingSave();
 
 public:
+	void					OnPlatformAdded() { bPackagesPendingSaveDirty = true; }
 	const PendingPackageSet&	GetPackagesPendingSave()			{ UpdatePackagesPendingSave(); return PackagesPendingSave; }
-	void						PackageSaved(UPackage* Package);
 	
 	void						DirtyPackage(const FName& CookedPackageName, UPackage* Package);
 
 	void						OnPackageCooked(const FFilePlatformCookedPackage& CookedPackage, UPackage* Package);
+	void						RemovePendingSavePackage(UPackage* Package);
 };
 
 void FPackageTracker::DirtyPackage(const FName& CookedPackageName, UPackage* Package)
@@ -1325,27 +1650,19 @@ void FPackageTracker::DirtyPackage(const FName& CookedPackageName, UPackage* Pac
 
 void FPackageTracker::OnPackageCooked(const FFilePlatformCookedPackage& CookedPackage, UPackage* Package)
 {
-	CookedPackages.Add(CookedPackage);
+	FFilePlatformCookedPackage& CurrentStatus = CookedPackages.Add(CookedPackage);
 
 	if (Package)
 	{
-		PackageSaved(Package);
-	}
-}
-
-void FPackageTracker::AddTargetPlatform(const FName& TargetPlatform)
-{
-	if (AllTargetPlatformNames.Find(TargetPlatform) == INDEX_NONE)
-	{
-		AllTargetPlatformNames.Add(TargetPlatform);
-
-		DirtyPackagesPendingSave();
+		if (CurrentStatus.HasPlatforms(PlatformManager.GetSessionPlatforms()))
+		{
+			PackagesPendingSave.Remove(Package);
+		}
 	}
 }
 
-void FPackageTracker::PackageSaved(UPackage* Package)
+void FPackageTracker::RemovePendingSavePackage(UPackage* Package)
 {
-	check(Package != nullptr);
 	PackagesPendingSave.Remove(Package);
 }
 
@@ -1360,8 +1677,9 @@ void FPackageTracker::FilterLoadedPackage(UPackage* Package)
 		return;	// if we have name none that means we are in core packages or something...
 	}
 
-	if (CookedPackages.Exists(StandardPackageFName, AllTargetPlatformNames))
+	if (CookedPackages.Exists(StandardPackageFName, PlatformManager.GetSessionPlatforms()))
 	{
+		// All SessionPlatforms have already been cooked for the package, so we don't need to save it again
 		return;
 	}
 
@@ -1374,6 +1692,7 @@ void FPackageTracker::UpdatePackagesPendingSave()
 	{
 		return;
 	}
+	check(PlatformManager.HasSelectedSessionPlatforms()); // It is not valid to call UpdatePackagesPendingSave until after sessions have been selected
 
 	PackagesPendingSave.Empty(PackagesPendingSave.Num());
 
@@ -1385,6 +1704,108 @@ void FPackageTracker::UpdatePackagesPendingSave()
 	bPackagesPendingSaveDirty = false;
 }
 
+void UCookOnTheFlyServer::FPlatformManager::SelectSessionPlatforms(const TArrayView<const ITargetPlatform* const>& TargetPlatforms, FPackageTracker* InPackageTracker)
+{
+	FScopeLock Lock(&SessionLock);
+
+	SessionPlatforms.Empty(TargetPlatforms.Num());
+	SessionPlatforms.Append(TargetPlatforms.GetData(), TargetPlatforms.Num());
+	for (const ITargetPlatform* TargetPlatform : TargetPlatforms)
+	{
+		CreatePlatformData(TargetPlatform);
+	}
+	bHasSelectedSessionPlatforms = true;
+	if (InPackageTracker)
+	{
+		InPackageTracker->OnPlatformAdded();
+	}
+}
+
+void UCookOnTheFlyServer::FPlatformManager::AddSessionPlatform(const ITargetPlatform* TargetPlatform, FPackageTracker* InPackageTracker)
+{
+	FScopeLock Lock(&SessionLock);
+
+	if (!SessionPlatforms.Contains(TargetPlatform))
+	{
+		SessionPlatforms.Add(TargetPlatform);
+		CreatePlatformData(TargetPlatform);
+		bHasSelectedSessionPlatforms = true;
+		if (InPackageTracker)
+		{
+			InPackageTracker->OnPlatformAdded();
+		}
+	}
+}
+
+void UCookOnTheFlyServer::FPlatformManager::PruneUnreferencedSessionPlatforms(FPackageTracker& InPackageTracker)
+{
+	const double SecondsToLive = 5.0 * 60;
+
+	double OldestKeepTime = -1.0e10; // Constructed to something smaller than 0 - SecondsToLive, so we can robustly detect not-yet-initialized
+	TArray<const ITargetPlatform*, TInlineAllocator<1>> RemovePlatforms;
+
+	for (auto& kvpair : PlatformDatas)
+	{
+		UCookOnTheFlyServer::FPlatformData& PlatformData = kvpair.Value;
+		if (PlatformData.LastReferenceTime > 0. && PlatformData.ReferenceCount == 0)
+		{
+			// We have a platform that we need to check for pruning.  Initialize the OldestKeepTime so we can check whether the platform has aged out.
+			if (OldestKeepTime < -SecondsToLive)
+			{
+				const double CurrentTimeSeconds = FPlatformTime::Seconds();
+				OldestKeepTime = CurrentTimeSeconds - SecondsToLive;
+			}
+
+			// Note that this loop is outside of the critical section, for performance.
+			// If we find any candidates for pruning we have to check them again once inside the critical section.
+			if (kvpair.Value.LastReferenceTime < OldestKeepTime)
+			{
+				RemovePlatforms.Add(kvpair.Key);
+			}
+		}
+	}
+
+	if (RemovePlatforms.Num() > 0)
+	{
+		FScopeLock Lock(&SessionLock);
+
+		for (const ITargetPlatform* TargetPlatform : RemovePlatforms)
+		{
+			UCookOnTheFlyServer::FPlatformData* PlatformData = PlatformDatas.Find(TargetPlatform);
+			if (PlatformData->LastReferenceTime > 0. && PlatformData->ReferenceCount == 0 && PlatformData->LastReferenceTime < OldestKeepTime)
+			{
+				// Mark that the platform no longer needs to be inspected for pruning because we have removed it from CookOnTheFly's SessionPlatforms
+				PlatformData->LastReferenceTime = 0.;
+
+				// Remove the SessionPlatform
+				InPackageTracker.RemoveSessionPlatform(TargetPlatform);
+
+				SessionPlatforms.Remove(TargetPlatform);
+				if (SessionPlatforms.Num() == 0)
+				{
+					bHasSelectedSessionPlatforms = false;
+				}
+			}
+		}
+	}
+}
+
+void UCookOnTheFlyServer::FPlatformManager::AddRefCookOnTheFlyPlatform(const ITargetPlatform* TargetPlatform, FPackageTracker& InPackageTracker)
+{
+	check(TargetPlatform != nullptr);
+	FPlatformData* PlatformData = GetPlatformData(TargetPlatform);
+	checkf(PlatformData != nullptr, TEXT("Unrecognized Platform %s"), *TargetPlatform->PlatformName());
+	++PlatformData->ReferenceCount;
+
+	if (!HasSessionPlatform(TargetPlatform))
+	{
+		InPackageTracker.AddTickCommand([this, TargetPlatform, &InPackageTracker]()
+			{
+				AddSessionPlatform(TargetPlatform, &InPackageTracker);
+			});
+	}
+}
+
 //////////////////////////////////////////////////////////////////////////
 // Cook by the book options
 
@@ -1410,7 +1831,7 @@ public:
 	FString							CreateReleaseVersion;
 
 	/** Dependency graph of maps as root objects. */
-	TMap<FName,TMap<FName,TSet<FName>>> MapDependencyGraphs;
+	TFastPointerMap<const ITargetPlatform*,TMap<FName,TSet<FName>>> MapDependencyGraphs;
 
 	/** If a cook is cancelled next cook will need to resume cooking */
 	TArray<FFilePlatformRequest>		PreviousCookRequests;
@@ -1427,7 +1848,6 @@ public:
 	bool							bDisableUnsolicitedPackages = false;
 	bool							bFullLoadAndSave = false;
 	bool							bPackageStore = false;
-	TArray<FName>					TargetPlatformNames;
 	TArray<FName>					StartupPackages;
 
 	/** Mapping from source packages to their localized variants (based on the culture list in FCookByTheBookStartupOptions) */
@@ -1447,6 +1867,7 @@ UCookOnTheFlyServer::UCookOnTheFlyServer(const FObjectInitializer& ObjectInitial
 	bIsSavingPackage(false),
 	AssetRegistry(nullptr)
 {
+	PlatformManager = MakeUnique<FPlatformManager>(RequestLock);
 }
 
 UCookOnTheFlyServer::UCookOnTheFlyServer(FVTableHelper& Helper) :Super(Helper) {}
@@ -1470,7 +1891,7 @@ UCookOnTheFlyServer::~UCookOnTheFlyServer()
 	ClearHierarchyTimers();
 }
 
-// this tick only happens in the editor cook commandlet directly calls tick on the side
+// This tick only happens in the editor.  The cook commandlet directly calls tick on the side.
 void UCookOnTheFlyServer::Tick(float DeltaTime)
 {
 	TRACE_CPUPROFILER_EVENT_SCOPE(UCookOnTheFlyServer::Tick);
@@ -1482,7 +1903,6 @@ void UCookOnTheFlyServer::Tick(float DeltaTime)
 		// if we are in the editor then precache some stuff ;)
 		TArray<const ITargetPlatform*> CacheTargetPlatforms;
 		const ULevelEditorPlaySettings* PlaySettings = GetDefault<ULevelEditorPlaySettings>();
-		ITargetPlatform* TargetPlatform = nullptr;
 		if (PlaySettings && (PlaySettings->LastExecutedLaunchModeType == LaunchMode_OnDevice))
 		{
 			FString DeviceName = PlaySettings->LastExecutedLaunchDevice.Left(PlaySettings->LastExecutedLaunchDevice.Find(TEXT("@")));
@@ -1515,81 +1935,47 @@ TStatId UCookOnTheFlyServer::GetStatId() const
 	RETURN_QUICK_DECLARE_CYCLE_STAT(UCookServer, STATGROUP_Tickables);
 }
 
-const TArray<ITargetPlatform*>& UCookOnTheFlyServer::GetCookingTargetPlatforms() const
+void UCookOnTheFlyServer::ConstructPackageTracker()
 {
-	ITargetPlatformManagerModule& TPM = GetTargetPlatformManagerRef();
+	delete PackageNameCache;
+	delete PackageTracker;
 
-	FString PlatformStr;
-	if ( FParse::Value(FCommandLine::Get(), TEXT("TARGETPLATFORM="), PlatformStr) == false )
+	if (AssetRegistry)
 	{
-		FString ValueName = TEXT("DefaultTargetPlatform");
-
-		if ( IsCookingInEditor() )
-		{
-			ValueName += TEXT("Editor");
-		}
-		if ( IsCookOnTheFlyMode() )
-		{
-			ValueName += TEXT("OnTheFly");
-		}
-
-
-		TArray<FString> TargetPlatformNames;
-		// see if we have specified in an ini file which target platforms we should use
-		if ( GConfig->GetArray(TEXT("CookSettings"), *ValueName, TargetPlatformNames, GEditorIni) )
-		{
-			for ( const FString& TargetPlatformName : TargetPlatformNames )
-			{
-				ITargetPlatform* TargetPlatform = TPM.FindTargetPlatform(TargetPlatformName);
-
-				if ( TargetPlatform )
-				{
-					CookingTargetPlatforms.AddUnique(TargetPlatform);
-				}
-				else
-				{
-					UE_LOG(LogCook, Warning, TEXT("Unable to resolve targetplatform name %s"), *TargetPlatformName );
-				}
-			}
-		}
+		PackageNameCache = new FPackageNameCache(AssetRegistry);
 	}
-
-
-	if ( CookingTargetPlatforms.Num() == 0 )
+	else
 	{
-		const TArray<ITargetPlatform*>& Platforms = TPM.GetCookingTargetPlatforms();
-
-		CookingTargetPlatforms = Platforms;
+		PackageNameCache = new FPackageNameCache();
 	}
 
-	return CookingTargetPlatforms;
+	PackageTracker = new FPackageTracker(PackageNameCache, RequestLock, *PlatformManager.Get());
 }
 
-bool UCookOnTheFlyServer::StartNetworkFileServer( const bool BindAnyPort )
+bool UCookOnTheFlyServer::StartNetworkFileServer(const bool BindAnyPort, const TArray<ITargetPlatform*>& TargetPlatforms)
 {
-	check( IsCookOnTheFlyMode() );
+	check(IsCookOnTheFlyMode());
 	//GetDerivedDataCacheRef().WaitForQuiescence(false);
 
 #if PROFILE_NETWORK
 	NetworkRequestEvent = FPlatformProcess::GetSynchEventFromPool();
 #endif
 
-	GenerateAssetRegistry();
-
-	InitializeSandbox();
-	InitializeTargetPlatforms();
+	// Precreate the map of all possible target platforms so we can access the collection of existing platforms in a threadsafe manner
+	// Each PlatformData in the map will be uninitialized until we call AddCookOnTheFlyPlatform for the platform
+	ITargetPlatformManagerModule& TPM = GetTargetPlatformManagerRef();
+	for (const ITargetPlatform* TargetPlatform : TPM.GetTargetPlatforms())
+	{
+		PlatformManager->CreatePlatformData(TargetPlatform);
+	}
+	PlatformManager->SetPlatformDataFrozen(true);
 
-	const TArray<ITargetPlatform*>& Platforms = GetCookingTargetPlatforms();
+	CreateSandboxFile();
+	GenerateAssetRegistry();
 
-	// When cooking on the fly the full registry is saved at the beginning
-	// in cook by the book asset registry is saved after the cook is finished
-	for (int32 Index = 0; Index < Platforms.Num(); Index++)
+	for (ITargetPlatform* TargetPlatform: TargetPlatforms)
 	{
-		FAssetRegistryGenerator* Generator = RegistryGenerators.FindRef(FName(*Platforms[Index]->PlatformName()));
-		if (Generator)
-		{
-			Generator->SaveAssetRegistry(GetSandboxAssetRegistryFilename(), true);
-		}
+		AddCookOnTheFlyPlatform(TargetPlatform);
 	}
 
 	// start the listening thread
@@ -1610,28 +1996,105 @@ bool UCookOnTheFlyServer::StartNetworkFileServer( const bool BindAnyPort )
 	NetworkFileDelegateContainer.OnFileModifiedCallback = &FileModifiedDelegate;
 
 
-	INetworkFileServer *TcpFileServer = FModuleManager::LoadModuleChecked<INetworkFileSystemModule>("NetworkFileSystem")
-		.CreateNetworkFileServer(true, BindAnyPort ? 0 : -1, NetworkFileDelegateContainer, ENetworkFileServerProtocol::NFSP_Tcp);
-	if ( TcpFileServer )
+	INetworkFileServer *TcpFileServer = FModuleManager::LoadModuleChecked<INetworkFileSystemModule>("NetworkFileSystem")
+		.CreateNetworkFileServer(true, BindAnyPort ? 0 : -1, NetworkFileDelegateContainer, ENetworkFileServerProtocol::NFSP_Tcp);
+	if ( TcpFileServer )
+	{
+		NetworkFileServers.Add(TcpFileServer);
+	}
+
+#if 0 // cookonthefly server via http
+	INetworkFileServer *HttpFileServer = FModuleManager::LoadModuleChecked<INetworkFileSystemModule>("NetworkFileSystem")
+		.CreateNetworkFileServer(true, BindAnyPort ? 0 : -1, NetworkFileDelegateContainer, ENetworkFileServerProtocol::NFSP_Http);
+	if ( HttpFileServer )
+	{
+		NetworkFileServers.Add( HttpFileServer );
+	}
+#endif
+
+	PackageTracker->CookRequestEvent = FPlatformProcess::GetSynchEventFromPool();
+
+	// loop while waiting for requests
+	return true;
+}
+
+ITargetPlatform* UCookOnTheFlyServer::AddCookOnTheFlyPlatform(const FString& PlatformName)
+{
+	ITargetPlatformManagerModule& TPM = GetTargetPlatformManagerRef();
+	ITargetPlatform* TargetPlatform = TPM.FindTargetPlatform(PlatformName);
+	if (TargetPlatform == nullptr)
+	{
+		UE_LOG(LogCook, Warning, TEXT("Target platform %s wasn't found."), *PlatformName);
+		return nullptr;
+	}
+	if (!AddCookOnTheFlyPlatform(TargetPlatform))
+	{
+		return nullptr;
+	}
+	return TargetPlatform;
+}
+
+bool UCookOnTheFlyServer::AddCookOnTheFlyPlatform(ITargetPlatform* TargetPlatform)
+{
+	check(IsCookOnTheFlyMode());
+
+	check(TargetPlatform != nullptr);
+	const FPlatformData* PlatformData = PlatformManager->GetPlatformData(TargetPlatform);
+	if (PlatformData == nullptr)
 	{
-		NetworkFileServers.Add(TcpFileServer);
+		UE_LOG(LogCook, Warning, TEXT("Target platform %s wasn't found in TargetPlatformManager."), *TargetPlatform->PlatformName());
+		return false;
 	}
 
-#if 0 // cookonthefly server via http
-	INetworkFileServer *HttpFileServer = FModuleManager::LoadModuleChecked<INetworkFileSystemModule>("NetworkFileSystem")
-		.CreateNetworkFileServer(true, BindAnyPort ? 0 : -1, NetworkFileDelegateContainer, ENetworkFileServerProtocol::NFSP_Http);
-	if ( HttpFileServer )
+	if (PlatformData->bIsSandboxInitialized)
 	{
-		NetworkFileServers.Add( HttpFileServer );
+		// Platform has already been added by this function or by StartCookByTheBook
+		return true;
 	}
-#endif
 
-	PackageTracker->CookRequestEvent = FPlatformProcess::GetSynchEventFromPool();
+	if (IsInGameThread())
+	{
+		AddCookOnTheFlyPlatformFromGameThread(TargetPlatform);
+	}
+	else
+	{
+		// Registering a new platform is not thread safe; queue the command for TickCookOnTheSide to execute
+		PackageTracker->AddTickCommand([this, TargetPlatform]() { AddCookOnTheFlyPlatformFromGameThread(TargetPlatform); });
+		if (PackageTracker->CookRequestEvent)
+		{
+			PackageTracker->CookRequestEvent->Trigger();
+		}
+	}
 
-	// loop while waiting for requests
 	return true;
 }
 
+void UCookOnTheFlyServer::AddCookOnTheFlyPlatformFromGameThread(ITargetPlatform* TargetPlatform)
+{
+	check(!!(CookFlags & ECookInitializationFlags::GeneratedAssetRegistry)); // GenerateAssetRegistry should have been called in StartNetworkFileServer
+
+	FPlatformData* PlatformData = PlatformManager->GetPlatformData(TargetPlatform);
+	check(PlatformData != nullptr); // should have been checked by the caller
+	if (PlatformData->bIsSandboxInitialized)
+	{
+		return;
+	}
+
+	TArrayView<ITargetPlatform* const> NewTargetPlatforms{ TargetPlatform };
+
+	RefreshPlatformAssetRegistries(NewTargetPlatforms);
+	InitializeSandbox(NewTargetPlatforms);
+	InitializeTargetPlatforms(NewTargetPlatforms);
+
+	// When cooking on the fly the full registry is saved at the beginning
+	// in cook by the book asset registry is saved after the cook is finished
+	FAssetRegistryGenerator* Generator = PlatformData->RegistryGenerator.Get();
+	if (Generator)
+	{
+		Generator->SaveAssetRegistry(GetSandboxAssetRegistryFilename(), true);
+	}
+	check(PlatformData->bIsSandboxInitialized); // This should have been set by InitializeSandbox, and it is what we use to determine whether a platform has been initialized
+}
 
 bool UCookOnTheFlyServer::BroadcastFileserverPresence( const FGuid &InstanceId )
 {
@@ -2051,26 +2514,37 @@ void UCookOnTheFlyServer::WaitForRequests(int TimeoutMs)
 	}
 }
 
-bool UCookOnTheFlyServer::HasCookRequests() const 
+bool UCookOnTheFlyServer::HasCookRequests() const
 { 
-	return PackageTracker->CookRequests.HasItems(); 
+	return PackageTracker->HasRequests();
 }
 
-bool UCookOnTheFlyServer::RequestPackage(const FName& StandardPackageFName, const TArray<FName>& TargetPlatforms, const bool bForceFrontOfQueue)
+bool UCookOnTheFlyServer::RequestPackage(const FName& StandardPackageFName, const TArray<const ITargetPlatform*>& TargetPlatforms, const bool bForceFrontOfQueue)
 {
 	FFilePlatformRequest FileRequest(StandardPackageFName, TargetPlatforms);
-	PackageTracker->CookRequests.EnqueueUnique(FileRequest, bForceFrontOfQueue);
+	PackageTracker->EnqueueUniqueCookRequest(FileRequest, bForceFrontOfQueue);
 	return true;
 }
 
+bool UCookOnTheFlyServer::RequestPackage(const FName& StandardPackageFName, const TArray<FName>& TargetPlatformNames, const bool bForceFrontOfQueue)
+{
+	TArray<const ITargetPlatform*> TargetPlatforms;
+	ITargetPlatformManagerModule& TPM = GetTargetPlatformManagerRef();
+	for (const FName& TargetPlatformName : TargetPlatformNames)
+	{
+		const ITargetPlatform* TargetPlatform = TPM.FindTargetPlatform(TargetPlatformName.ToString());
+		if (TargetPlatform)
+		{
+			TargetPlatforms.Add(TargetPlatform);
+		}
+	}
+	return RequestPackage(StandardPackageFName, TargetPlatforms, bForceFrontOfQueue);
+}
 
 bool UCookOnTheFlyServer::RequestPackage(const FName& StandardPackageFName, const bool bForceFrontOfQueue)
 {
-	check(IsCookByTheBookMode());
-	// need target platforms if we are not in cook by the book mode
-	FFilePlatformRequest FileRequest(StandardPackageFName, CookByTheBookOptions->TargetPlatformNames);
-	PackageTracker->CookRequests.EnqueueUnique(FileRequest, bForceFrontOfQueue);
-	return true;
+	check(IsCookByTheBookMode()); // Invalid to call RequestPackage without a list of TargetPlatforms unless we are in cook by the book mode
+	return RequestPackage(StandardPackageFName, PlatformManager->GetSessionPlatforms(), bForceFrontOfQueue);
 }
 
 // callback just before the garbage collector gets called.
@@ -2095,7 +2569,9 @@ UCookOnTheFlyServer::FReentryData& UCookOnTheFlyServer::GetReentryData(const UPa
 
 uint32 UCookOnTheFlyServer::TickCookOnTheSide(const float TimeSlice, uint32 &CookedPackageCount, ECookTickFlags TickFlags)
 {
-	if (IsCookByTheBookMode() && CookByTheBookOptions->bFullLoadAndSave)
+	TickNetwork();
+
+	if (IsCookByTheBookMode() && CookByTheBookOptions->bRunning && CookByTheBookOptions->bFullLoadAndSave)
 	{
 		uint32 Result = FullLoadAndSave(CookedPackageCount);
 
@@ -2117,34 +2593,8 @@ uint32 UCookOnTheFlyServer::TickCookOnTheSide(const float TimeSlice, uint32 &Coo
 		}
 	}
 
-	// This is all the target platforms which we needed to process requests for this iteration
-	// we use this in the unsolicited packages processing below
-	TArray<FName> AllTargetPlatformNames;
-
 	while (!IsEngineExitRequested() || CurrentCookMode == ECookMode::CookByTheBook)
 	{
-		if (HasCookRequests())
-		{
-			const float CurrentProgressDisplayTime = FPlatformTime::Seconds();
-
-			if (LastCookedPackagesCount != PackageTracker->CookedPackages.Num()
-				|| LastCookRequestsCount != PackageTracker->CookRequests.Num()
-				|| (CurrentProgressDisplayTime - LastProgressDisplayTime) > GCookProgressRepeatTime)
-			{
-				UE_CLOG(!(TickFlags & ECookTickFlags::HideProgressDisplay) && (GCookProgressDisplay & (int32)ECookProgressDisplayMode::RemainingPackages),
-					LogCook,
-					Display,
-					TEXT("Cooked packages %d Packages Remain %d Total %d"),
-					PackageTracker->CookedPackages.Num(),
-					PackageTracker->CookRequests.Num(),
-					PackageTracker->CookedPackages.Num() + PackageTracker->CookRequests.Num());
-
-				LastCookedPackagesCount = PackageTracker->CookedPackages.Num();
-				LastCookRequestsCount = PackageTracker->CookRequests.Num();
-				LastProgressDisplayTime = CurrentProgressDisplayTime;
-			}
-		}
-
 		// if we just cooked a map then don't process anything the rest of this tick
 		if (Result & COSR_RequiresGC)
 		{
@@ -2160,28 +2610,57 @@ uint32 UCookOnTheFlyServer::TickCookOnTheSide(const float TimeSlice, uint32 &Coo
 			}
 		}
 
-		auto GetNextUncookedRequestLambda = [this](FFilePlatformRequest& OutToBuild) -> bool
+		FFilePlatformRequest ToBuild;
+		TArray<FPackageTracker::FTickCommand> TickCommands;
+		FPackageTracker::ERequestType RequestType;
+		while (true)
 		{
-			while (this->PackageTracker->CookRequests.Dequeue(OutToBuild))
+			RequestType = PackageTracker->DequeueRequest(TickCommands, ToBuild);
+			if (RequestType == FPackageTracker::ERequestType::Cook && this->PackageTracker->CookedPackages.Exists(ToBuild))
 			{
-				if (this->PackageTracker->CookedPackages.Exists(OutToBuild))
-				{
 #if DEBUG_COOKONTHEFLY
-					UE_LOG(LogCook, Display, TEXT("Package for platform already cooked %s, discarding request"), *OutToBuild.GetFilename().ToString());
+				UE_LOG(LogCook, Display, TEXT("Package for platform already cooked %s, discarding request"), *OutToBuild.GetFilename().ToString());
 #endif
-					continue;
-				}
-				return true;
+				continue;
 			}
-			return false;
-		};
+			break;
+		}
 
-		FFilePlatformRequest ToBuild;
-		if (!GetNextUncookedRequestLambda(ToBuild))
+		if (RequestType == FPackageTracker::ERequestType::None)
 		{
 			// no more real work to do this tick, break out and do some other stuff
 			break;
 		}
+		else if (RequestType == FPackageTracker::ERequestType::TickCommand)
+		{
+			// An array of TickCommands to process; execute through them all and then continue
+			for (FPackageTracker::FTickCommand& TickCommand : TickCommands)
+			{
+				TickCommand();
+			}
+			continue;
+		}
+		check(RequestType == FPackageTracker::ERequestType::Cook && ToBuild.IsValid());
+
+		const float CurrentProgressDisplayTime = FPlatformTime::Seconds();
+
+		const int32 CookRequestsNum = PackageTracker->GetCookRequestsNum() + 1; // + 1 to include the request we just dequeued
+		if (LastCookedPackagesCount != PackageTracker->CookedPackages.Num()
+			|| LastCookRequestsCount != CookRequestsNum
+			|| (CurrentProgressDisplayTime - LastProgressDisplayTime) > GCookProgressRepeatTime)
+		{
+			UE_CLOG(!(TickFlags & ECookTickFlags::HideProgressDisplay) && (GCookProgressDisplay & (int32)ECookProgressDisplayMode::RemainingPackages),
+				LogCook,
+				Display,
+				TEXT("Cooked packages %d Packages Remain %d Total %d"),
+				PackageTracker->CookedPackages.Num(),
+				CookRequestsNum,
+				PackageTracker->CookedPackages.Num() + CookRequestsNum);
+
+			LastCookedPackagesCount = PackageTracker->CookedPackages.Num();
+			LastCookRequestsCount = CookRequestsNum;
+			LastProgressDisplayTime = CurrentProgressDisplayTime;
+		}
 
 #if PROFILE_NETWORK
 		if (NetworkRequestEvent)
@@ -2204,25 +2683,17 @@ uint32 UCookOnTheFlyServer::TickCookOnTheSide(const float TimeSlice, uint32 &Coo
 		SCOPE_TIMER(TickCookOnTheSide);
 
 		check(ToBuild.IsValid());
-		const TArray<FName>& TargetPlatformNames = ToBuild.GetPlatformNames();
+		const TArray<const ITargetPlatform*>& TargetPlatforms = ToBuild.GetPlatforms();
+		if (TargetPlatforms.Num() == 0)
+		{
+			UE_LOG(LogCook, Error, TEXT("Empty list of platforms requested in CookOnTheSide request."));
+			continue;
+		}
 
 #if OUTPUT_TIMING
 		//FScopeTimer PackageManualTimer( ToBuild.GetFilename().ToString(), false );
 #endif
 
-		for (const FName& PlatformName : TargetPlatformNames)
-		{
-			AllTargetPlatformNames.AddUnique(PlatformName);
-		}
-
-		for (const FName& PlatformName : AllTargetPlatformNames)
-		{
-			if (ToBuild.HasPlatform(PlatformName) == false)
-			{
-				ToBuild.AddPlatform(PlatformName);
-			}
-		}
-
 		const FString BuildFilename = ToBuild.GetFilename().ToString();
 
 		bool bShouldCook = true;
@@ -2247,16 +2718,16 @@ uint32 UCookOnTheFlyServer::TickCookOnTheSide(const float TimeSlice, uint32 &Coo
 			bShouldCook = false;
 		}
 
+		UPackage* LoadedPackage = nullptr;
 		UPackage* PackageForCooking = nullptr;
 
 		if (bShouldCook) // if we should cook the package then cook it otherwise add it to the list of already cooked packages below
 		{
-			UPackage* Package = LoadPackageForCooking(BuildFilename);
-
-			if (Package)
+			bool bLoadFullySuccessful = LoadPackageForCooking(BuildFilename, LoadedPackage);
+			if (bLoadFullySuccessful)
 			{
-				FString Name = Package->GetPathName();
-				FString PackageFilename(PackageNameCache->GetCachedStandardPackageFilename(Package));
+				FString Name = LoadedPackage->GetPathName();
+				FString PackageFilename(PackageNameCache->GetCachedStandardPackageFilename(LoadedPackage));
 				if (PackageFilename != BuildFilename)
 				{
 					// we have saved something which we didn't mean to load 
@@ -2265,12 +2736,12 @@ uint32 UCookOnTheFlyServer::TickCookOnTheSide(const float TimeSlice, uint32 &Coo
 					//	and then also save our new request as processed so we don't do it again
 					UE_LOG(LogCook, Verbose, TEXT("Request for %s received going to save %s"), *BuildFilename, *PackageFilename);
 
-					PackageTracker->OnPackageCooked(FFilePlatformCookedPackage(ToBuild.GetFilename(), TargetPlatformNames), Package);
+					PackageTracker->OnPackageCooked(FFilePlatformCookedPackage(ToBuild.GetFilename(), TargetPlatforms), LoadedPackage);
 
 					ToBuild.SetFilename(PackageFilename);
 				}
 
-				PackageForCooking = Package;
+				PackageForCooking = LoadedPackage;
 			}
 			else
 			{
@@ -2291,9 +2762,9 @@ uint32 UCookOnTheFlyServer::TickCookOnTheSide(const float TimeSlice, uint32 &Coo
 			// did not cook this package 
 #if DO_CHECK
 			// make sure this package doesn't exist
-			for (const FName& TargetPlatformName : ToBuild.GetPlatformNames())
+			for (const ITargetPlatform* TargetPlatform: ToBuild.GetPlatforms())
 			{
-				const FString SandboxFilename = ConvertToFullSandboxPath(ToBuild.GetFilename().ToString(), true, TargetPlatformName.ToString());
+				const FString SandboxFilename = ConvertToFullSandboxPath(ToBuild.GetFilename().ToString(), true, TargetPlatform->PlatformName());
 				if (IFileManager::Get().FileExists(*SandboxFilename))
 				{
 					// if we find the file this means it was cooked on a previous cook, however source package can't be found now. 
@@ -2304,20 +2775,12 @@ uint32 UCookOnTheFlyServer::TickCookOnTheSide(const float TimeSlice, uint32 &Coo
 				}
 			}
 #endif
-			PackageTracker->OnPackageCooked(FFilePlatformCookedPackage(ToBuild.GetFilename(), TargetPlatformNames), nullptr);
+			PackageTracker->OnPackageCooked(FFilePlatformCookedPackage(ToBuild.GetFilename(), TargetPlatforms), LoadedPackage);
 			continue;
 		}
 
 		bool bIsAllDataCached = true;
 
-		ITargetPlatformManagerModule& TPM = GetTargetPlatformManagerRef();
-		TArray<const ITargetPlatform*> TargetPlatforms;
-		for (const FName& TargetPlatformName : AllTargetPlatformNames)
-		{
-			TargetPlatforms.Add(TPM.FindTargetPlatform(TargetPlatformName.ToString()));
-		}
-
-
 		GShaderCompilingManager->ProcessAsyncResults(true, false);
 
 
@@ -2375,27 +2838,29 @@ uint32 UCookOnTheFlyServer::TickCookOnTheSide(const float TimeSlice, uint32 &Coo
 
 		ProcessUnsolicitedPackages();
 
-		// in cook by the book bail out early because shaders and compiled for the primary package we are trying to save. 
-		// note in this case we also put the package at the end of the queue that queue might be reordered if we do partial gc
-		if ((bIsAllDataCached == false) && IsCookByTheBookMode() && !IsRealtimeMode())
-		{
-			/////
-			// TODO-Stefan: need to understand this logic better
-			/////
-
-			// don't load anymore stuff unless we have space and we don't already have enough stuff to save
-			if ((!(Result & COSR_RequiresGC)) &&
-				(HasExceededMaxMemory() == false) &&
-				((Timer.NumPackagesSaved + PackageTracker->GetPackagesPendingSave().Num()) < Timer.MaxNumPackagesToSave)) // if we don't need to GC and also we have memory then load some more packages ;)
-			{
-				GShaderCompilingManager->ProcessAsyncResults(true, false); // we can afford to do work here because we are essentually requing this package for processing later
-				Timer.SavedPackage();  // this is a special case to prevent infinite loop, if we only have one package we might fall through this and could loop forever.  
-				PackageTracker->CookRequests.EnqueueUnique(ToBuild, false);
+		if (!bIsAllDataCached)
+		{
+			// If we are waiting on asynchronous bulk data builds (e.g. shader compile) then if we proceed to SaveCookedPackages we will have to wastefully wait for that asynchronous work to complete.
+			// Instead, when possible, requeue the package, working on other packages before we get back to it
+			if (IsCookByTheBookMode() && // We only allow reordering if we're doing a CookByTheBook; in CookOnTheFly we care urgently about the currently requested package
+				!IsRealtimeMode() && // In RealTimeMode (e.g. running in editor) we don't requeue.  TODO Daniel Lamb: Why do we not requeue in realtime?
+				!(Result & COSR_RequiresGC) && // GC1 Don't requeue if we are about to garbage collect, we would then waste work reloading this package
+				!HasExceededMaxMemory() && // GC2
+				Timer.NumPackagesSaved + PackageTracker->GetPackagesPendingSave().Num() < Timer.MaxNumPackagesToSave // GC3
+				)
+			{
+				// Given the list of constraints we just checked, we know we're not in a rush on any particular package, and we have some async results to process, so take the time to prcoess them now
+				GShaderCompilingManager->ProcessAsyncResults(true, false); 
+				// If all of our requested packages are missing cached data we might keep requeueing them forever.
+				// To prevent an infinite loop, call Timer.SavedPackage so that we eventually fail the Timer.NumPackagesSaved < Timer.MaxNumPackagesToSave check above and break out of the loop.
+				Timer.SavedPackage();  
+				// Requeue the new package and move on the next one in the queue
+				PackageTracker->EnqueueUniqueCookRequest(ToBuild, false);
 				continue;
 			}
 		}
 
-		SaveCookedPackages(PackageForCooking, AllTargetPlatformNames, TargetPlatforms, Timer, /* out */ CookedPackageCount, /* out */ Result);
+		SaveCookedPackages(PackageForCooking, TargetPlatforms, Timer, /* out */ CookedPackageCount, /* out */ Result);
 
 		if (Timer.IsTimeUp())
 		{
@@ -2421,8 +2886,7 @@ uint32 UCookOnTheFlyServer::TickCookOnTheSide(const float TimeSlice, uint32 &Coo
 		CookByTheBookOptions->CookTime += Timer.GetTimeTillNow();
 	}
 
-	if (IsCookByTheBookRunning() &&
-		(HasCookRequests() == false))
+	if (IsCookByTheBookRunning() && !HasCookRequests())
 	{
 		check(IsCookByTheBookMode());
 
@@ -2436,7 +2900,18 @@ uint32 UCookOnTheFlyServer::TickCookOnTheSide(const float TimeSlice, uint32 &Coo
 	return Result;
 }
 
-bool UCookOnTheFlyServer::BeginPackageCacheForCookedPlatformData(UPackage* Package, const TArray<const ITargetPlatform*>& TargetPlatforms, FCookerTimer& Timer) const
+void UCookOnTheFlyServer::TickNetwork()
+{
+	if (!PackageTracker)
+		return;
+
+	if (!IsCookByTheBookMode()) // StartCookByTheBook does not AddRef its session platforms, so it is not safe to prune unreferenced session platforms when supporting CookByTheBook
+	{
+		PlatformManager->PruneUnreferencedSessionPlatforms(*PackageTracker);
+	}
+}
+
+bool UCookOnTheFlyServer::BeginPackageCacheForCookedPlatformData(UPackage* Package, const TArrayView<const ITargetPlatform* const>& TargetPlatforms, FCookerTimer& Timer) const
 {
 	COOK_STAT(FScopedDurationTimer DurationTimer(DetailedCookStats::TickCookOnTheSideBeginPackageCacheForCookedPlatformDataTimeSec));
 
@@ -2509,7 +2984,7 @@ bool UCookOnTheFlyServer::BeginPackageCacheForCookedPlatformData(UPackage* Packa
 
 }
 
-bool UCookOnTheFlyServer::FinishPackageCacheForCookedPlatformData(UPackage* Package, const TArray<const ITargetPlatform*>& TargetPlatforms, FCookerTimer& Timer) const
+bool UCookOnTheFlyServer::FinishPackageCacheForCookedPlatformData(UPackage* Package, const TArrayView<const ITargetPlatform* const>& TargetPlatforms, FCookerTimer& Timer) const
 {
 	COOK_STAT(FScopedDurationTimer DurationTimer(DetailedCookStats::TickCookOnTheSideFinishPackageCacheForCookedPlatformDataTimeSec));
 
@@ -2566,12 +3041,12 @@ bool UCookOnTheFlyServer::FinishPackageCacheForCookedPlatformData(UPackage* Pack
 		}
 	}
 
-	for (UObject* Obj : CurrentReentryData.CachedObjectsInOuter)
+	if (CurrentCookMode == ECookMode::CookByTheBook)
 	{
-		// if this objects data is cached then we can call FinishedCookedPLatformDataCache
-		// we can only safely call this when we are finished caching this object completely.
+		// For each object for which data is cached we can call FinishedCookedPlatformDataCache
+		// we can only safely call this when we are finished caching the object completely.
 		// this doesn't ever happen for cook in editor or cook on the fly mode
-		if (CurrentCookMode == ECookMode::CookByTheBook)
+		for (UObject* Obj : CurrentReentryData.CachedObjectsInOuter)
 		{
 			check(!IsCookingInEditor());
 			// this might be run multiple times for a single object
@@ -2592,16 +3067,14 @@ bool UCookOnTheFlyServer::FinishPackageCacheForCookedPlatformData(UPackage* Pack
 	return true;
 }
 
-UPackage* UCookOnTheFlyServer::LoadPackageForCooking(const FString& BuildFilename)
+bool UCookOnTheFlyServer::LoadPackageForCooking(const FString& BuildFilename, UPackage*& OutPackage)
 {
 	COOK_STAT(FScopedDurationTimer LoadPackagesTimer(DetailedCookStats::TickCookOnTheSideLoadPackagesTimeSec));
-	UPackage *Package = NULL;
+	OutPackage = NULL;
+	FString PackageNameString;
+	if (FPackageName::TryConvertFilenameToLongPackageName(BuildFilename, PackageNameString))
 	{
-		FString PackageName;
-		if (FPackageName::TryConvertFilenameToLongPackageName(BuildFilename, PackageName))
-		{
-			Package = FindObject<UPackage>(ANY_PACKAGE, *PackageName);
-		}
+		OutPackage = FindObject<UPackage>(ANY_PACKAGE, *PackageNameString);
 	}
 
 #if DEBUG_COOKONTHEFLY
@@ -2614,12 +3087,36 @@ UPackage* UCookOnTheFlyServer::LoadPackageForCooking(const FString& BuildFilenam
 		GOutputCookingWarnings = IsCookFlagSet(ECookInitializationFlags::OutputVerboseCookerWarnings);
 	}
 
-	//  if the package is already loaded then try to avoid reloading it :)
-	if ((Package == NULL) || (Package->IsFullyLoaded() == false))
+	bool bSuccess = true;
+	//  if the package is not yet fully loaded then fully load it
+	if (OutPackage == nullptr || !OutPackage->IsFullyLoaded())
 	{
 		GIsCookerLoadingPackage = true;
 		SCOPE_TIMER(LoadPackage);
-		Package = LoadPackage(NULL, *BuildFilename, LOAD_None);
+		UPackage* LoadedPackage = LoadPackage(NULL, *BuildFilename, LOAD_None);
+		if (LoadedPackage)
+		{
+			OutPackage = LoadedPackage;
+		}
+		else
+		{
+			bSuccess = false;
+			if (OutPackage == nullptr)
+			{
+				// FindObject failed and LoadPackage failed, but there are some Packages that can exist in our PackageTracker that are not returned from FindPackage e.g. because they are async loading
+				// If we are reporting an error on a Package, we are going to mark it as already cooked, and so if it is present in PendingSavePackages we need to be sure to remove it
+				// So look for it in all PendingSavePackages if we are reporting an error and didn't find it in FindPackage
+				FName PackageName = FName(PackageNameString);
+				for (UPackage* PendingSavePackage : PackageTracker->GetPackagesPendingSave())
+				{
+					if (PendingSavePackage->GetFName() == PackageName)
+					{
+						OutPackage = PendingSavePackage;
+						break;
+					}
+				}
+			}
+		}
 
 		++this->StatLoadedPackageCount;
 
@@ -2632,7 +3129,7 @@ UPackage* UCookOnTheFlyServer::LoadPackageForCooking(const FString& BuildFilenam
 	}
 #endif
 
-	if (Package == NULL)
+	if (!bSuccess)
 	{
 		if ((!IsCookOnTheFlyMode()) || (!IsCookingInEditor()))
 		{
@@ -2641,7 +3138,7 @@ UPackage* UCookOnTheFlyServer::LoadPackageForCooking(const FString& BuildFilenam
 		}
 	}
 	GOutputCookingWarnings = false;
-	return Package;
+	return bSuccess;
 }
 
 
@@ -2668,8 +3165,7 @@ void UCookOnTheFlyServer::ProcessUnsolicitedPackages()
 
 void UCookOnTheFlyServer::SaveCookedPackages(
 	UPackage*								PackageToSave, 
-	const TArray<FName>&					TargetPlatformNames, 
-	const TArray<const ITargetPlatform*>&	TargetPlatformsToCache, 
+	const TArray<const ITargetPlatform*>&	InTargetPlatforms,
 	FCookerTimer&							Timer,
 	uint32&									CookedPackageCount,
 	uint32&									Result)
@@ -2678,41 +3174,35 @@ void UCookOnTheFlyServer::SaveCookedPackages(
 
 	bool bIsAllDataCached = true;
 
-	const TArray<FName>& AllTargetPlatformNames = TargetPlatformNames;
-
 	// To make transitioning from the old way of managing unsolicited assets
 	// easier, we construct a temporary array here with the PackageToSave as 
-	// the first entry. It would be nice if we could restructure this so we
-	// could avoid the temporary allocation and memory copying, but I'm 
-	// leaving it like this for now to avoid having change too many things
-	// at once
+	// the first entry
 
-	auto PendingSet = PackageTracker->GetPackagesPendingSave();
-	PendingSet.Remove(PackageToSave);
-	TArray<UPackage*> PackagesToSave = PendingSet.Array();
+	auto& PendingSet = PackageTracker->GetPackagesPendingSave();
 
+	constexpr uint32 ExpectedMaxNumUnsolicitedPlatforms = 10;
+	TArray<UPackage*, TInlineAllocator<ExpectedMaxNumUnsolicitedPlatforms>> PackagesToSave;
+	PackagesToSave.Reserve(PendingSet.Num() + 1);
 	int32 FirstUnsolicitedPackage = 0;
-
 	if (PackageToSave)
 	{
-		if (PackagesToSave.Num())
-		{
-			UPackage* First = PackagesToSave[0];
-			PackagesToSave.Add(First);
-			PackagesToSave[0] = PackageToSave;
-		}
-		else
+		PackagesToSave.Add(PackageToSave);
+		FirstUnsolicitedPackage = 1;
+	}
+	for (UPackage* PendingPackage : PendingSet)
+	{
+		if (PendingPackage != PackageToSave)
 		{
-			PackagesToSave.Add(PackageToSave);
+			PackagesToSave.Add(PendingPackage);
 		}
-
-		FirstUnsolicitedPackage = 1;
 	}
 
 	// Loop over array and save as many packages as we can during our
 	// time slice
 
 	const int32 OriginalPackagesToSaveCount = PackagesToSave.Num();
+	TArray<const ITargetPlatform*, TInlineAllocator<ExpectedMaxNumPlatforms>> PlatformsForPackage;
+	TArray<const ITargetPlatform*> AlreadyCookedPlatformsForPackage;
 
 	{
 		SCOPE_TIMER(SavingPackages);
@@ -2737,25 +3227,34 @@ void UCookOnTheFlyServer::SaveCookedPackages(
 				continue;
 			}
 
-			// TODO: optimize this to avoid creating the full array first only to remove entries (multiple dynamic allocations)
-			//
-			// Could rewrite using FTargetPlatformSet
-			
-			TArray<FName> SaveTargetPlatformNames = AllTargetPlatformNames;
-			TArray<FName> CookedTargetPlatforms;
-			if (PackageTracker->CookedPackages.GetCookedPlatforms(PackageFName, CookedTargetPlatforms))
+			// For passed-in packages, consider cooking only the passed-in target platforms
+			// For pending save packages, consider cooking all platforms in the cook-by-the-book session
+			const bool bProcessingUnsolicitedPackages = Package != PackageToSave;
+			const TArray<const ITargetPlatform*>& PossiblePlatformsForPackage{ bProcessingUnsolicitedPackages ? PlatformManager->GetSessionPlatforms() : InTargetPlatforms };
+
+			// Cook only the platforms that have not yet been cooked for the given package
+			AlreadyCookedPlatformsForPackage.Reset();
+			PackageTracker->CookedPackages.GetCookedPlatforms(PackageFName, AlreadyCookedPlatformsForPackage);
+			PlatformsForPackage.Reset();
+			for (const ITargetPlatform* TargetPlatform : PossiblePlatformsForPackage)
 			{
-				for (const FName& CookedPlatform : CookedTargetPlatforms)
+				if (!AlreadyCookedPlatformsForPackage.Contains(TargetPlatform))
 				{
-					SaveTargetPlatformNames.Remove(CookedPlatform);
+					PlatformsForPackage.Add(TargetPlatform);
 				}
 			}
 
-			// we somehow already cooked this package not sure how that can happen because the PackagesToSave list should have already filtered this
-			if (SaveTargetPlatformNames.Num() == 0)
+			if (PlatformsForPackage.Num() == 0)
 			{
-				UE_LOG(LogCook, Warning, TEXT("Already saved this package not sure how this got here!"));
-				// already saved this package
+				// We've already saved all possible platforms for this package; this should not be possible.
+				// It should not be possible for the passed-in packages because they should never have been added to CookRequests if the remaining packages to cook was empty.
+				// It should not be possible for pendingsave packages because they should not have been queued if they had completed all of their possible platforms, and should be removed from pendingsave if they are ever cooked
+				UE_LOG(LogCook, Warning, TEXT("%s package '%s' in SaveCookedPackages has no more platforms left to cook; this should not be possible!"),
+					(bProcessingUnsolicitedPackages ? TEXT("Unsolicited") : TEXT("Passed-in")), *PackageFName.ToString());
+				if (bProcessingUnsolicitedPackages)
+				{
+					PackageTracker->RemovePendingSavePackage(Package);
+				}
 				continue;
 			}
 
@@ -2773,7 +3272,6 @@ void UCookOnTheFlyServer::SaveCookedPackages(
 
 			// if we are cook the fly then save the package which was requested as fast as we can because the client is waiting on it
 
-			const bool bProcessingUnsolicitedPackages = (I >= 1);
 			bool bForceSavePackage = false;
 
 			if (IsCookOnTheFlyMode())
@@ -2794,7 +3292,7 @@ void UCookOnTheFlyServer::SaveCookedPackages(
 					// if we are in realtime mode then don't wait forever for the package to be ready
 					while ((!Timer.IsTimeUp()) && IsRealtimeMode() && (bShouldFinishTick == false))
 					{
-						if (FinishPackageCacheForCookedPlatformData(Package, TargetPlatformsToCache, Timer) == true)
+						if (FinishPackageCacheForCookedPlatformData(Package, PlatformsForPackage, Timer) == true)
 						{
 							bFinishedCachingCookedPlatformData = true;
 							break;
@@ -2829,11 +3327,11 @@ void UCookOnTheFlyServer::SaveCookedPackages(
 			// if we are forcing save the package then it doesn't matter if we call FinishPackageCacheForCookedPlatformData
 			if (!bShouldFinishTick && !bForceSavePackage)
 			{
-				AllObjectsCookedDataCached = FinishPackageCacheForCookedPlatformData(Package, TargetPlatformsToCache, Timer);
+				AllObjectsCookedDataCached = FinishPackageCacheForCookedPlatformData(Package, PlatformsForPackage, Timer);
 				if (AllObjectsCookedDataCached == false)
 				{
 					GShaderCompilingManager->ProcessAsyncResults(true, false);
-					AllObjectsCookedDataCached = FinishPackageCacheForCookedPlatformData(Package, TargetPlatformsToCache, Timer);
+					AllObjectsCookedDataCached = FinishPackageCacheForCookedPlatformData(Package, PlatformsForPackage, Timer);
 				}
 			}
 
@@ -2860,7 +3358,7 @@ void UCookOnTheFlyServer::SaveCookedPackages(
 				for (int32 RemainingIndex = I; RemainingIndex < NumPackagesToRequeue; ++RemainingIndex)
 				{
 					FName StandardFilename = PackageNameCache->GetCachedStandardPackageFileFName(PackagesToSave[RemainingIndex]);
-					PackageTracker->CookRequests.EnqueueUnique(FFilePlatformRequest(StandardFilename, SaveTargetPlatformNames));
+					PackageTracker->EnqueueUniqueCookRequest(FFilePlatformRequest(StandardFilename, PlatformsForPackage));
 				}
 				Result |= COSR_WaitingOnCache;
 
@@ -2878,12 +3376,12 @@ void UCookOnTheFlyServer::SaveCookedPackages(
 				if (NextPackage != Package)
 				{
 					SCOPE_TIMER(PrecachePlatformDataForNextPackage);
-					BeginPackageCacheForCookedPlatformData(NextPackage, TargetPlatformsToCache, Timer);
+					BeginPackageCacheForCookedPlatformData(NextPackage, PlatformsForPackage, Timer);
 				}
 				if (NextNextPackage != NextPackage)
 				{
 					SCOPE_TIMER(PrecachePlatformDataForNextNextPackage);
-					BeginPackageCacheForCookedPlatformData(NextNextPackage, TargetPlatformsToCache, Timer);
+					BeginPackageCacheForCookedPlatformData(NextNextPackage, PlatformsForPackage, Timer);
 				}
 			}
 
@@ -2933,21 +3431,21 @@ void UCookOnTheFlyServer::SaveCookedPackages(
 				GOutputCookingWarnings = IsCookFlagSet(ECookInitializationFlags::OutputVerboseCookerWarnings);
 				try
 				{
-					SaveCookedPackage(Package, SaveFlags, SaveTargetPlatformNames, SavePackageResults);
+					SaveCookedPackage(Package, SaveFlags, PlatformsForPackage, SavePackageResults);
 				}
 				catch (std::exception&)
 				{
 					FString TargetPlatforms;
-					for ( const FName& PlatformFName : SaveTargetPlatformNames )
+					for ( const ITargetPlatform* Platform : PlatformsForPackage)
 					{
-						TargetPlatforms += FString::Printf( TEXT("%s, "), *PlatformFName.ToString());
+						TargetPlatforms += FString::Printf( TEXT("%s, "), *Platform->PlatformName());
 					}
 					UE_LOG(LogCook, Warning, TEXT("Tried to save package %s for target platforms %s but threw an exception"), *Package->GetPathName(), *TargetPlatforms);
 					throw;
 				}
 				
 				GOutputCookingWarnings = false;
-				check(SaveTargetPlatformNames.Num() == SavePackageResults.Num());
+				check(PlatformsForPackage.Num() == SavePackageResults.Num());
 				for (int iResultIndex = 0; iResultIndex < SavePackageResults.Num(); iResultIndex++)
 				{
 					FSavePackageResultStruct& SavePackageResult = SavePackageResults[iResultIndex];
@@ -2969,7 +3467,7 @@ void UCookOnTheFlyServer::SaveCookedPackages(
 						// Update asset registry
 						if (CookByTheBookOptions)
 						{
-							FAssetRegistryGenerator* Generator = RegistryGenerators.FindRef(SaveTargetPlatformNames[iResultIndex]);
+							FAssetRegistryGenerator* Generator = PlatformManager->GetPlatformData(PlatformsForPackage[iResultIndex])->RegistryGenerator.Get();
 							UpdateAssetRegistryPackageData(Generator, Package->GetFName(), SavePackageResult);
 
 						}
@@ -2995,7 +3493,7 @@ void UCookOnTheFlyServer::SaveCookedPackages(
 				}
 			}
 
-				//@todo ResetLoaders outside of this (ie when Package is NULL) causes problems w/ default materials
+			//@todo ResetLoaders outside of this (ie when Package is NULL) causes problems w/ default materials
 			FName StandardFilename = PackageNameCache->GetCachedStandardPackageFileFName(Package);
 
 			// We always want to mark package as processed unless it wasn't saved because it was referenced by editor-only data
@@ -3003,7 +3501,7 @@ void UCookOnTheFlyServer::SaveCookedPackages(
 			if (StandardFilename != NAME_None)
 			{
 				// mark the package as cooked
-				FFilePlatformCookedPackage FileRequest(StandardFilename, SaveTargetPlatformNames, SucceededSavePackage);
+				FFilePlatformCookedPackage FileRequest(StandardFilename, PlatformsForPackage, SucceededSavePackage);
 				bool bWasReferencedOnlyByEditorOnlyData = false;
 				for (const FSavePackageResultStruct& SavePackageResult : SavePackageResults)
 				{
@@ -3022,7 +3520,7 @@ void UCookOnTheFlyServer::SaveCookedPackages(
 						// this is an unsolicited package
 						if (FPaths::FileExists(FileRequest.GetFilename().ToString()) == true)
 						{
-							PackageTracker->UnsolicitedCookedPackages.AddCookedPackage(FFilePlatformRequest(FileRequest.GetFilename(), FileRequest.GetPlatformNames()));
+							PackageTracker->UnsolicitedCookedPackages.AddCookedPackage(FFilePlatformRequest(FileRequest.GetFilename(), FileRequest.GetPlatforms()));
 
 #if DEBUG_COOKONTHEFLY
 							UE_LOG(LogCook, Display, TEXT("UnsolicitedCookedPackages: %s"), *FileRequest.GetFilename().ToString());
@@ -3278,7 +3776,7 @@ bool UCookOnTheFlyServer::HasExceededMaxMemory() const
 	return false;
 }
 
-TArray<UPackage*> UCookOnTheFlyServer::GetUnsolicitedPackages(const TArray<FName>& TargetPlatformNames) const
+TArray<UPackage*> UCookOnTheFlyServer::GetUnsolicitedPackages(const TArray<const ITargetPlatform*>& TargetPlatforms) const
 {
 	SCOPE_TIMER(GeneratePackageNames);
 
@@ -3293,7 +3791,7 @@ TArray<UPackage*> UCookOnTheFlyServer::GetUnsolicitedPackages(const TArray<FName
 		if (StandardPackageFName == NAME_None)
 			continue;	// if we have name none that means we are in core packages or something...
 
-		if (PackageTracker->CookedPackages.Exists(StandardPackageFName, TargetPlatformNames))
+		if (PackageTracker->CookedPackages.Exists(StandardPackageFName, TargetPlatforms))
 			continue;
 
 		PackagesToSave.Add(Package);
@@ -3304,6 +3802,21 @@ TArray<UPackage*> UCookOnTheFlyServer::GetUnsolicitedPackages(const TArray<FName
 	return PackagesToSave;
 }
 
+TArray<UPackage*> UCookOnTheFlyServer::GetUnsolicitedPackages(const TArray<FName>& TargetPlatformNames) const
+{
+	TArray<const ITargetPlatform*> TargetPlatforms;
+	ITargetPlatformManagerModule& TPM = GetTargetPlatformManagerRef();
+	for (const FName& TargetPlatformName : TargetPlatformNames)
+	{
+		const ITargetPlatform* TargetPlatform = TPM.FindTargetPlatform(TargetPlatformName.ToString());
+		if (TargetPlatform)
+		{
+			TargetPlatforms.Add(TargetPlatform);
+		}
+	}
+	return GetUnsolicitedPackages(TargetPlatforms);
+}
+
 void UCookOnTheFlyServer::OnObjectModified( UObject *ObjectMoving )
 {
 	if (IsGarbageCollecting())
@@ -3405,7 +3918,7 @@ void UCookOnTheFlyServer::MarkPackageDirtyForCooker( UPackage *Package )
 		{
 			if ( IsCookByTheBookMode() )
 			{
-				TArray<FName> CookedPlatforms;
+				TArray<const ITargetPlatform*> CookedPlatforms;
 				// if we have already cooked this package and we have made changes then recook ;)
 				if (PackageTracker->CookedPackages.GetCookedPlatforms(PackageFFileName, CookedPlatforms) )
 				{
@@ -3413,7 +3926,7 @@ void UCookOnTheFlyServer::MarkPackageDirtyForCooker( UPackage *Package )
 					{
 						// if this package was previously cooked and we are doing a cook by the book 
 						// we need to recook this package before finishing cook by the book
-						PackageTracker->CookRequests.EnqueueUnique(FFilePlatformRequest(PackageFFileName, CookedPlatforms));
+						PackageTracker->EnqueueUniqueCookRequest(FFilePlatformRequest(PackageFFileName, CookedPlatforms));
 					}
 					else
 					{
@@ -3457,6 +3970,7 @@ void UCookOnTheFlyServer::EndNetworkFileServer()
 		NetworkFileServer = NULL;
 	}
 	NetworkFileServers.Empty();
+	PlatformManager->SetPlatformDataFrozen(false);
 }
 
 uint32 UCookOnTheFlyServer::GetPackagesPerGC() const
@@ -3574,7 +4088,8 @@ void UCookOnTheFlyServer::MarkGCPackagesToKeepForCooker()
 	// then see if the package is needed by any of the packages which are required to be saved
 
 	TMap<FName, int32> PackageDependenciesCount;
-	for (const FName& QueuedPackage : PackageTracker->CookRequests.GetQueue())
+	FScopeLock RequestScopeLock(&PackageTracker->GetRequestLock());
+	for (const FName& QueuedPackage : PackageTracker->GetThreadUnsafeCookRequests().GetQueue())
 	{
 		const FName* PackageName = PackageNameCache->GetCachedPackageFilenameToPackageFName(QueuedPackage);
 		if ( !PackageName )
@@ -3608,7 +4123,7 @@ void UCookOnTheFlyServer::MarkGCPackagesToKeepForCooker()
 	// then sort by the number of dependencies which are referenced by the package
 	// we want to process the packages with the highest dependencies so that they can
 	// be evicted from memory and are likely to be able to be released on next GC pass
-	PackageTracker->CookRequests.Sort([&PackageDependenciesCount, &LoadedPackages](const FName& A, const FName& B)
+	PackageTracker->GetThreadUnsafeCookRequests().Sort([&PackageDependenciesCount, &LoadedPackages](const FName& A, const FName& B)
 	{
 		int32 ADependencies = PackageDependenciesCount.FindChecked(A);
 		int32 BDependencies = PackageDependenciesCount.FindChecked(B);
@@ -3648,12 +4163,6 @@ bool UCookOnTheFlyServer::HasRecompileShaderRequests() const
 	return PackageTracker->RecompileRequests.HasItems();
 }
 
-void UCookOnTheFlyServer::SaveCookedPackage(UPackage* Package, uint32 SaveFlags, TArray<FSavePackageResultStruct>& SavePackageResults)
-{
-	TArray<FName> TargetPlatformNames; 
-	return SaveCookedPackage( Package, SaveFlags, TargetPlatformNames, SavePackageResults);
-}
-
 bool UCookOnTheFlyServer::MakePackageFullyLoaded(UPackage* Package) const
 {
 	if (Package->IsFullyLoaded())
@@ -3848,7 +4357,7 @@ private:
 	}
 };
 
-void UCookOnTheFlyServer::SaveCookedPackage(UPackage* Package, uint32 SaveFlags, TArray<FName> &TargetPlatformNames, TArray<FSavePackageResultStruct>& SavePackageResults)
+void UCookOnTheFlyServer::SaveCookedPackage(UPackage* Package, uint32 SaveFlags, const TArrayView<const ITargetPlatform* const>& TargetPlatforms, TArray<FSavePackageResultStruct>& SavePackageResults)
 {
 	check( SavePackageResults.Num() == 0);
 	check( bIsSavingPackage == false );
@@ -3916,41 +4425,10 @@ void UCookOnTheFlyServer::SaveCookedPackage(UPackage* Package, uint32 SaveFlags,
 
 		ITargetPlatformManagerModule& TPM = GetTargetPlatformManagerRef();
 
-		static TArray<ITargetPlatform*> ActiveStartupPlatforms = TPM.GetCookingTargetPlatforms();
-
-		TArray<ITargetPlatform*> Platforms;
-
-		if ( TargetPlatformNames.Num() )
-		{
-			const TArray<ITargetPlatform*>& TargetPlatforms = TPM.GetTargetPlatforms();
-
-			for (const FName TargetPlatformFName : TargetPlatformNames)
-			{
-				const FString TargetPlatformName = TargetPlatformFName.ToString();
-
-				for (ITargetPlatform *TargetPlatform  : TargetPlatforms)
-				{
-					if ( TargetPlatform->PlatformName() == TargetPlatformName )
-					{
-						Platforms.Add( TargetPlatform );
-					}
-				}
-			}
-		}
-		else
-		{
-			Platforms = ActiveStartupPlatforms;
-
-			for (ITargetPlatform *Platform : Platforms)
-			{
-				TargetPlatformNames.Add(FName(*Platform->PlatformName()));
-			}
-		}
-		
-		for (int32 PlatformIndex = 0; PlatformIndex < Platforms.Num(); ++PlatformIndex)
+		for (int32 PlatformIndex = 0; PlatformIndex < TargetPlatforms.Num(); ++PlatformIndex)
 		{
 			SavePackageResults.Add(FSavePackageResultStruct(ESavePackageResult::Success));
-			ITargetPlatform* Target = Platforms[PlatformIndex];
+			const ITargetPlatform* Target = TargetPlatforms[PlatformIndex];
 			FString PlatFilename = Filename.Replace(TEXT("[Platform]"), *Target->PlatformName());
 
 			FSavePackageResultStruct& Result = SavePackageResults[PlatformIndex];
@@ -3975,7 +4453,7 @@ void UCookOnTheFlyServer::SaveCookedPackage(UPackage* Package, uint32 SaveFlags,
 			// check if this package is unsupported for the target platform (typically plugin content)
 			else 
 			{
-				TSet<FName>* NeverCookPackages = PackageTracker->PlatformSpecificNeverCookPackages.Find(FName(*Target->PlatformName()));
+				TSet<FName>* NeverCookPackages = PackageTracker->PlatformSpecificNeverCookPackages.Find(Target);
 				if (NeverCookPackages && NeverCookPackages->Find(FName(*PackagePathName)))
 				{
 					Result = ESavePackageResult::ContainsEditorOnlyData;
@@ -4094,7 +4572,7 @@ void UCookOnTheFlyServer::SaveCookedPackage(UPackage* Package, uint32 SaveFlags,
 	}
 	else
 	{
-		for (int32 PlatformIndex = 0; PlatformIndex < TargetPlatformNames.Num(); ++PlatformIndex)
+		for (int32 PlatformIndex = 0; PlatformIndex < TargetPlatforms.Num(); ++PlatformIndex)
 		{
 			SavePackageResults.Add(FSavePackageResultStruct(ESavePackageResult::MissingFile));
 		}
@@ -4105,7 +4583,6 @@ void UCookOnTheFlyServer::SaveCookedPackage(UPackage* Package, uint32 SaveFlags,
 
 }
 
-
 void UCookOnTheFlyServer::Initialize( ECookMode::Type DesiredCookMode, ECookInitializationFlags InCookFlags, const FString &InOutputDirectoryOverride )
 {
 	TRACE_CPUPROFILER_EVENT_SCOPE(UCookOnTheFlyServer::Initialize);
@@ -4121,8 +4598,7 @@ void UCookOnTheFlyServer::Initialize( ECookMode::Type DesiredCookMode, ECookInit
 		// For standalone CookByTheBook the PackageNameCache and PackageTracker are initialized 
 		// later to benefit from a fully scanned and initialized AssetRegistry.
 		// For IsCookingInEditor() these objects are required for the FCoreUObjectDelegates below
-		PackageNameCache = new FPackageNameCache();
-		PackageTracker = new FPackageTracker(PackageNameCache);
+		ConstructPackageTracker();
 	}
 
 	if (IsCookingInEditor())
@@ -4382,14 +4858,13 @@ FString UCookOnTheFlyServer::GetOutputDirectoryOverride() const
 		// Output directory needs to contain [Platform] token to be able to cook for multiple targets.
 		if ( IsCookByTheBookMode() )
 		{
-			const TArray<ITargetPlatform*>& TargetPlatforms = GetCookingTargetPlatforms();
-
-			// more then one target platform specified append "[platform]" to output directory so that multiple platforms can be cooked
-			check( TargetPlatforms.Num() == 1 );
+			checkf(PlatformManager->GetSessionPlatforms().Num() == 1,
+				TEXT("If OutputDirectoryOverride is provided when cooking multiple platforms, it must include [Platform] in the text, to be replaced with the name of each of the requested Platforms.") );
 		}
 		else
 		{
-			// cook on the fly we have to platform because we don't know which platforms we are cooking for up front
+			// In cook on the fly mode we always add a [Platform] subdirectory rather than requiring the command-line user to include it in their path it because we assume they 
+			// don't know which platforms they are cooking for up front
 			OutputDirectory = FPaths::Combine(*OutputDirectory, TEXT("[Platform]"));
 		}
 	}
@@ -4752,7 +5227,7 @@ void UCookOnTheFlyServer::ProcessAccessedIniSettings(const FConfigFile* Config,
 	const TArray<ITargetPlatform*>& Platforms = TPM.GetTargetPlatforms();
 	FString PlatformName;
 	bool bFoundPlatformName = false;
-	for (ITargetPlatform* Platform : Platforms )
+	for (const ITargetPlatform* Platform : Platforms )
 	{
 		FString CurrentPlatformName = Platform->IniPlatformName();
 		for ( const auto& SourceIni : Config->SourceIniHierarchy )
@@ -5184,7 +5659,7 @@ FString UCookOnTheFlyServer::GetAsyncDeleteDirectory(const FString& PlatformName
 	return (*SandboxDirectory) + TEXT("_Del");
 }
 
-void UCookOnTheFlyServer::PopulateCookedPackagesFromDisk(const TArray<ITargetPlatform*>& Platforms)
+void UCookOnTheFlyServer::PopulateCookedPackagesFromDisk(const TArrayView<const ITargetPlatform* const>& Platforms)
 {
 	TRACE_CPUPROFILER_EVENT_SCOPE(UCookOnTheFlyServer::PopulateCookedPackagesFromDisk);
 
@@ -5193,9 +5668,9 @@ void UCookOnTheFlyServer::PopulateCookedPackagesFromDisk(const TArray<ITargetPla
 	{
 		TArray<FString> CookedPackagesToDelete;
 
-		ITargetPlatform* Target = Platforms[Index];
+		const ITargetPlatform* Target = Platforms[Index];
+		FPlatformData* PlatformData = PlatformManager->GetPlatformData(Target);
 		FString SandboxPath = GetSandboxDirectory(Target->PlatformName());
-		FName PlatformFName(*Target->PlatformName());
 
 		FString EngineSandboxPath = SandboxFile->ConvertToSandboxPath(*FPaths::EngineDir()) + TEXT("/");
 		EngineSandboxPath.ReplaceInline(TEXT("[Platform]"), *Target->PlatformName());
@@ -5215,7 +5690,8 @@ void UCookOnTheFlyServer::PopulateCookedPackagesFromDisk(const TArray<ITargetPla
 		static bool bFindCulprit = false; //debugging setting if we want to find culprit for iterative cook issues
 
 		// Registry generator already exists
-		FAssetRegistryGenerator* PlatformAssetRegistry = RegistryGenerators.FindRef(PlatformFName);
+		FAssetRegistryGenerator* PlatformAssetRegistry = PlatformData->RegistryGenerator.Get();
+		check(PlatformAssetRegistry);
 
 		// Load the platform cooked asset registry file
 		const FString CookedAssetRegistry = FPaths::ProjectDir() / TEXT("Metadata") / GetDevelopmentAssetRegistryFilename();
@@ -5237,11 +5713,11 @@ void UCookOnTheFlyServer::PopulateCookedPackagesFromDisk(const TArray<ITargetPla
 				(CurrentIterativeCookedBuild != FDateTime::MinValue()) )
 			{
 				// clean the sandbox
-				const FString PlatformName = Target->PlatformName();
-				ClearPlatformCookedData(PlatformName);
+				ClearPlatformCookedData(Target);
 
 				// SaveCurrentIniSettings(Target); // use this if we don't care about ini safty.
 				// copy the ini settings from the shared cooked build. 
+				const FString PlatformName = Target->PlatformName();
 				const FString SharedCookedIniFile = FPaths::Combine(*FPaths::ProjectSavedDir(), TEXT("SharedIterativeBuild"), *PlatformName, TEXT("Metadata"), TEXT("CookedIniVersion.txt"));
 				const FString SandboxCookedIniFile = ConvertToFullSandboxPath(*(FPaths::ProjectDir() / TEXT("Metadata") / TEXT("CookedIniVersion.txt")), true).Replace(TEXT("[Platform]"), *PlatformName);
 
@@ -5350,8 +5826,8 @@ void UCookOnTheFlyServer::PopulateCookedPackagesFromDisk(const TArray<ITargetPla
 				}
 			}
 
-			TArray<FName> PlatformNames;
-			PlatformNames.Add(PlatformFName);
+			TArray<const ITargetPlatform*> PlatformsForPackage;
+			PlatformsForPackage.Add(Target);
 
 			if (bShouldKeep)
 			{
@@ -5361,7 +5837,7 @@ void UCookOnTheFlyServer::PopulateCookedPackagesFromDisk(const TArray<ITargetPla
 
 				if (IdenticalCookedPackages.Contains(SourcePackageName))
 				{
-					PackageTracker->CookedPackages.Add(FFilePlatformCookedPackage(UncookedFilename, MoveTemp(PlatformNames), MoveTemp(Succeeded)));
+					PackageTracker->CookedPackages.Add(FFilePlatformCookedPackage(UncookedFilename, MoveTemp(PlatformsForPackage), MoveTemp(Succeeded)));
 					KeptPackages.Add(SourcePackageName);
 					++NumPackagesKept;
 				}
@@ -5371,7 +5847,7 @@ void UCookOnTheFlyServer::PopulateCookedPackagesFromDisk(const TArray<ITargetPla
 				if (SourcePackageName != NAME_None && IsCookByTheBookMode()) // cook on the fly will requeue this package when it wants it 
 				{
 					// Force cook the modified file
-					PackageTracker->CookRequests.EnqueueUnique(FFilePlatformRequest(UncookedFilename, PlatformNames));
+					PackageTracker->EnqueueUniqueCookRequest(FFilePlatformRequest(UncookedFilename, PlatformsForPackage));
 				}
 				if (CookedFile != NAME_DummyCookedFilename)
 				{
@@ -5395,124 +5871,56 @@ void UCookOnTheFlyServer::PopulateCookedPackagesFromDisk(const TArray<ITargetPla
 		{
 			const FName UncookedFilename = PackageNameCache->GetCachedStandardPackageFileFName(UncookedPackage);
 
-			TArray<FName> PlatformNames;
-			PlatformNames.Add(PlatformFName);
-
-			ensure(PackageTracker->CookedPackages.Exists(UncookedFilename, PlatformNames, false) == false);
-
-			PackageTracker->CookedPackages.Add(FFilePlatformCookedPackage(UncookedFilename, MoveTemp(PlatformNames)));
-			KeptPackages.Add(UncookedPackage);
-			++NumMarkedFailedSaveKept;
-		}
-
-		PlatformAssetRegistry->UpdateKeptPackages(KeptPackages);
-
-		UE_LOG(LogCook, Display, TEXT("Iterative cooking summary for %s, \nConsidered: %d, \nFile Hash missmatch: %d, \nPackages Kept: %d, \nPackages failed save kept: %d, \nMissing Cooked Info(expected 0): %d"),
-			*Target->PlatformName(),
-			NumPackagesConsidered, NumPackagesFileHashMismatch,
-			NumPackagesKept, NumMarkedFailedSaveKept,
-			NumPackagesUnableToFindCookedPackageInfo);
-
-		auto DeletePackageLambda = [&CookedPackagesToDelete](int32 PackageIndex)
-		{
-			const FString& CookedFullPath = CookedPackagesToDelete[PackageIndex];
-			IFileManager::Get().Delete(*CookedFullPath, true, true, true);
-		};
-		ParallelFor(CookedPackagesToDelete.Num(), DeletePackageLambda);
-	}
-}
-
-const FString ExtractPackageNameFromObjectPath( const FString ObjectPath )
-{
-	// get the path 
-	int32 Beginning = ObjectPath.Find(TEXT("'"), ESearchCase::CaseSensitive);
-	if ( Beginning == INDEX_NONE )
-	{
-		return ObjectPath;
-	}
-	int32 End = ObjectPath.Find(TEXT("."), ESearchCase::CaseSensitive, ESearchDir::FromStart, Beginning + 1);
-	if (End == INDEX_NONE )
-	{
-		End = ObjectPath.Find(TEXT("'"), ESearchCase::CaseSensitive, ESearchDir::FromStart, Beginning + 1);
-	}
-	if ( End == INDEX_NONE )
-	{
-		// one more use case is that the path is "Class'Path" example "OrionBoostItemDefinition'/Game/Misc/Boosts/XP_1Win" dunno why but this is actually dumb
-		if ( ObjectPath[Beginning+1] == '/' )
-		{
-			return ObjectPath.Mid(Beginning+1);
-		}
-		return ObjectPath;
-	}
-	return ObjectPath.Mid(Beginning + 1, End - Beginning - 1);
-}
-
-void UCookOnTheFlyServer::CleanSandbox(const bool bIterative)
-{
-	TRACE_CPUPROFILER_EVENT_SCOPE(UCookOnTheFlyServer::CleanSandbox);
-
-	const TArray<ITargetPlatform*>& Platforms = GetCookingTargetPlatforms();
-
-	// before we can delete any cooked files we need to make sure that we have finished writing them
-	UPackage::WaitForAsyncFileWrites();
-
-#if OUTPUT_TIMING
-	double SandboxCleanTime = 0.0;
-#endif
-	{
-#if OUTPUT_TIMING
-		SCOPE_SECONDS_COUNTER(SandboxCleanTime);
-		SCOPE_TIMER(CleanSandboxTime);
-#endif
-
-		for (const ITargetPlatform* Target : Platforms)
-		{
-			bool bShouldClearCookedContent = false;
-			const bool bIsIniSettingsOutOfDate = IniSettingsOutOfDate(Target); // needs to be executed for side effects even if non-iterative
-			if (!bIterative)
-			{
-				// for now we always wipe the cooked directory on full cooks
-				UE_LOG(LogCook, Display, TEXT("Clearing all cooked content for platform %s"), *Target->PlatformName());
-				bShouldClearCookedContent = true;
-			}
-			else
-			{
-				if (bIsIniSettingsOutOfDate)
-				{
-					if (!IsCookFlagSet(ECookInitializationFlags::IgnoreIniSettingsOutOfDate))
-					{
-						UE_LOG(LogCook, Display, TEXT("Cook invalidated for platform %s ini settings don't match from last cook, clearing all cooked content"), *Target->PlatformName());
-						bShouldClearCookedContent = true;
-					}
-					else
-					{
-						UE_LOG(LogCook, Display, TEXT("Inisettings were out of date for platform %s but we are going with it anyway because IgnoreIniSettingsOutOfDate is set"), *Target->PlatformName());
-					}
-				}
-			}
-
-			if (bShouldClearCookedContent)
-			{
-				ClearPlatformCookedData(Target->PlatformName());
-				SaveCurrentIniSettings(Target);
-			}
+			TArray<const ITargetPlatform*> PlatformsForPackage;
+			PlatformsForPackage.Add(Target);
+
+			ensure(PackageTracker->CookedPackages.Exists(UncookedFilename, PlatformsForPackage, false) == false);
+
+			PackageTracker->CookedPackages.Add(FFilePlatformCookedPackage(UncookedFilename, MoveTemp(PlatformsForPackage)));
+			KeptPackages.Add(UncookedPackage);
+			++NumMarkedFailedSaveKept;
 		}
 
-		if (bIterative)
+		PlatformAssetRegistry->UpdateKeptPackages(KeptPackages);
+
+		UE_LOG(LogCook, Display, TEXT("Iterative cooking summary for %s, \nConsidered: %d, \nFile Hash missmatch: %d, \nPackages Kept: %d, \nPackages failed save kept: %d, \nMissing Cooked Info(expected 0): %d"),
+			*Target->PlatformName(),
+			NumPackagesConsidered, NumPackagesFileHashMismatch,
+			NumPackagesKept, NumMarkedFailedSaveKept,
+			NumPackagesUnableToFindCookedPackageInfo);
+
+		auto DeletePackageLambda = [&CookedPackagesToDelete](int32 PackageIndex)
 		{
-			// This is fast in asset registry iterate, so just reconstruct
-			PackageTracker->CookedPackages.Empty();
-			PopulateCookedPackagesFromDisk(Platforms);
-		}
+			const FString& CookedFullPath = CookedPackagesToDelete[PackageIndex];
+			IFileManager::Get().Delete(*CookedFullPath, true, true, true);
+		};
+		ParallelFor(CookedPackagesToDelete.Num(), DeletePackageLambda);
 	}
-#if OUTPUT_TIMING
-	FString PlatformNames;
-	for (const ITargetPlatform* Target : Platforms)
+}
+
+const FString ExtractPackageNameFromObjectPath( const FString ObjectPath )
+{
+	// get the path 
+	int32 Beginning = ObjectPath.Find(TEXT("'"), ESearchCase::CaseSensitive);
+	if ( Beginning == INDEX_NONE )
 	{
-		PlatformNames += Target->PlatformName() + TEXT(" ");
+		return ObjectPath;
 	}
-	UE_LOG(LogCook, Display, TEXT("Sandbox cleanup took %5.3f seconds for platforms %s iterative %s"), SandboxCleanTime, *PlatformNames, bIterative ? TEXT("true") : TEXT("false"));
-#endif
+	int32 End = ObjectPath.Find(TEXT("."), ESearchCase::CaseSensitive, ESearchDir::FromStart, Beginning + 1);
+	if (End == INDEX_NONE )
+	{
+		End = ObjectPath.Find(TEXT("'"), ESearchCase::CaseSensitive, ESearchDir::FromStart, Beginning + 1);
+	}
+	if ( End == INDEX_NONE )
+	{
+		// one more use case is that the path is "Class'Path" example "OrionBoostItemDefinition'/Game/Misc/Boosts/XP_1Win" dunno why but this is actually dumb
+		if ( ObjectPath[Beginning+1] == '/' )
+		{
+			return ObjectPath.Mid(Beginning+1);
+		}
+		return ObjectPath;
+	}
+	return ObjectPath.Mid(Beginning + 1, End - Beginning - 1);
 }
 
 void UCookOnTheFlyServer::GenerateAssetRegistry()
@@ -5524,18 +5932,18 @@ void UCookOnTheFlyServer::GenerateAssetRegistry()
 	AssetRegistry = &AssetRegistryModule.Get();
 
 	// Mark package as dirty for the last ones saved
-	if(PackageNameCache != nullptr)
+	if (PackageNameCache != nullptr)
 	{
 		for (FName AssetFilename : ModifiedAssetFilenames)
 		{
 			const FString AssetPathOnDisk = AssetFilename.ToString();
-			if(FPaths::FileExists(AssetPathOnDisk))
+			if (FPaths::FileExists(AssetPathOnDisk))
 			{
 				const FString PackageName = FPackageName::FilenameToLongPackageName(AssetPathOnDisk);
 				FSoftObjectPath SoftPackage(PackageName);
-				if(UPackage* Package = Cast<UPackage>(SoftPackage.ResolveObject()))
+				if (UPackage* Package = Cast<UPackage>(SoftPackage.ResolveObject()))
 				{
-					MarkPackageDirtyForCooker( Package );
+					MarkPackageDirtyForCooker(Package);
 				}
 			}
 		}
@@ -5543,6 +5951,8 @@ void UCookOnTheFlyServer::GenerateAssetRegistry()
 
 	if (!!(CookFlags & ECookInitializationFlags::GeneratedAssetRegistry))
 	{
+		UE_LOG(LogCook, Display, TEXT("Updating asset registry"));
+
 		// Force a rescan of modified package files
 		TArray<FString> ModifiedPackageFileList;
 
@@ -5552,35 +5962,25 @@ void UCookOnTheFlyServer::GenerateAssetRegistry()
 		}
 
 		AssetRegistry->ScanModifiedAssetFiles(ModifiedPackageFileList);
-
-		// This is cook in the editor on a second pass, so refresh the generators
-		for (TPair<FName, FAssetRegistryGenerator*>& Pair : RegistryGenerators)
-		{
-			Pair.Value->Initialize(CookByTheBookOptions ? CookByTheBookOptions->StartupPackages : TArray<FName>());
-		}
-		return;
 	}
-	CookFlags |= ECookInitializationFlags::GeneratedAssetRegistry;
-
-	ModifiedAssetFilenames.Reset();
-
-	double GenerateAssetRegistryTime = 0.0;
+	else
 	{
-		SCOPE_TIMER(GenerateAssetRegistryTime);
+		CookFlags |= ECookInitializationFlags::GeneratedAssetRegistry;
 		UE_LOG(LogCook, Display, TEXT("Creating asset registry"));
 
+		ModifiedAssetFilenames.Reset();
+
 		// Perform a synchronous search of any .ini based asset paths (note that the per-game delegate may
 		// have already scanned paths on its own)
 		// We want the registry to be fully initialized when generating streaming manifests too.
 
 		// editor will scan asset registry automagically 
-		bool bCanDelayAssetregistryProcessing = IsRealtimeMode(); 
+		bool bCanDelayAssetregistryProcessing = IsRealtimeMode();
 
 		// if we are running in the editor we need the asset registry to be finished loaded before we process any iterative cook requests
-		bCanDelayAssetregistryProcessing &= !IsCookFlagSet(ECookInitializationFlags::Iterative); // 
+		bCanDelayAssetregistryProcessing &= !IsCookFlagSet(ECookInitializationFlags::Iterative);
 
-		
-		if ( !bCanDelayAssetregistryProcessing)
+		if (!bCanDelayAssetregistryProcessing)
 		{
 			TArray<FString> ScanPaths;
 			if (GConfig->GetArray(TEXT("AssetRegistry"), TEXT("PathsToScanForCook"), ScanPaths, GEngineIni) > 0 && !AssetRegistry->IsLoadingAssets())
@@ -5594,22 +5994,25 @@ void UCookOnTheFlyServer::GenerateAssetRegistry()
 			}
 		}
 	}
+}
 
-	const TArray<ITargetPlatform*>& Platforms = GetCookingTargetPlatforms();
+void UCookOnTheFlyServer::RefreshPlatformAssetRegistries(const TArrayView<const ITargetPlatform* const>& TargetPlatforms)
+{
+	TRACE_CPUPROFILER_EVENT_SCOPE(UCookOnTheFlyServer::RefreshPlatformAssetRegistries);
 
-	for (ITargetPlatform* TargetPlatform : Platforms)
+	for (const ITargetPlatform* TargetPlatform : TargetPlatforms)
 	{
 		FName PlatformName = FName(*TargetPlatform->PlatformName());
 
-		// make sure we have a registry generate for all the platforms 
-		FAssetRegistryGenerator* RegistryGenerator = RegistryGenerators.FindRef(PlatformName);
-		if (RegistryGenerator == nullptr)
+		FPlatformData* PlatformData = PlatformManager->GetPlatformData(TargetPlatform);
+		FAssetRegistryGenerator* RegistryGenerator = PlatformData->RegistryGenerator.Get();
+		if (!RegistryGenerator)
 		{
 			RegistryGenerator = new FAssetRegistryGenerator(TargetPlatform);
+			PlatformData->RegistryGenerator = TUniquePtr<FAssetRegistryGenerator>(RegistryGenerator);
 			RegistryGenerator->CleanManifestDirectories();
-			RegistryGenerator->Initialize(CookByTheBookOptions ? CookByTheBookOptions->StartupPackages : TArray<FName>());
-			RegistryGenerators.Add(PlatformName, RegistryGenerator);
 		}
+		RegistryGenerator->Initialize(CookByTheBookOptions ? CookByTheBookOptions->StartupPackages : TArray<FName>());
 	}
 }
 
@@ -5672,9 +6075,10 @@ void UCookOnTheFlyServer::AddFileToCook( TArray<FName>& InOutFilesToCook, const
 	}
 }
 
-void UCookOnTheFlyServer::CollectFilesToCook(TArray<FName>& FilesInPath, const TArray<FString>& CookMaps, const TArray<FString>& InCookDirectories, const TArray<FString> &IniMapSections, ECookByTheBookOptions FilesToCookFlags)
+void UCookOnTheFlyServer::CollectFilesToCook(TArray<FName>& FilesInPath, const TArray<FString>& CookMaps, const TArray<FString>& InCookDirectories,
+	const TArray<FString> &IniMapSections, ECookByTheBookOptions FilesToCookFlags, const TArrayView<const ITargetPlatform* const>& TargetPlatforms)
 {
-	TRACE_CPUPROFILER_EVENT_SCOPE(UCookOnTheFlyServer::GenerateAssetRegistry);
+	TRACE_CPUPROFILER_EVENT_SCOPE(UCookOnTheFlyServer::CollectFilesToCook);
 
 #if OUTPUT_TIMING
 	SCOPE_TIMER(CollectFilesToCook);
@@ -5936,13 +6340,13 @@ void UCookOnTheFlyServer::CollectFilesToCook(TArray<FName>& FilesInPath, const T
 	if (!(FilesToCookFlags & ECookByTheBookOptions::NoDefaultMaps))
 	{
 		// make sure we cook the default maps
-		ITargetPlatformManagerModule& TPM = GetTargetPlatformManagerRef();
-		static const TArray<ITargetPlatform*>& Platforms = TPM.GetTargetPlatforms();
-		for (int32 Index = 0; Index < Platforms.Num(); Index++)
+		// Collect the default maps for all requested platforms.  Our additions are potentially wasteful if different platforms in the requested list have different default maps.
+		// In that case we will wastefully cook maps for platforms that don't require them.
+		for (const ITargetPlatform* TargetPlatform : TargetPlatforms)
 		{
 			// load the platform specific ini to get its DefaultMap
 			FConfigFile PlatformEngineIni;
-			FConfigCacheIni::LoadLocalIniFile(PlatformEngineIni, TEXT("Engine"), true, *Platforms[Index]->IniPlatformName());
+			FConfigCacheIni::LoadLocalIniFile(PlatformEngineIni, TEXT("Engine"), true, *TargetPlatform->IniPlatformName());
 
 			// get the server and game default maps and cook them
 			FString Obj;
@@ -6041,10 +6445,11 @@ void UCookOnTheFlyServer::CollectFilesToCook(TArray<FName>& FilesInPath, const T
 
 	if (CookByTheBookOptions && !(FilesToCookFlags & ECookByTheBookOptions::DisableUnsolicitedPackages))
 	{
-		// Gather initial unsolicited package list, this is needed in iterative mode because it may skip cooking all explicit packages and never hit this code
+		// Gather initial unsolicited package list, this is needed in iterative mode because all explicitly requested packages may have already been cooked
+		// and so the code inside the TIckCookOnTheSide build loop might never run and never get a chance to call GetUnsolicitedPackages
 		UE_LOG(LogCook, Verbose, TEXT("Finding initial unsolicited packages"));
 
-		TArray<UPackage*> UnsolicitedPackages = GetUnsolicitedPackages(CookByTheBookOptions->TargetPlatformNames);
+		TArray<UPackage*> UnsolicitedPackages = GetUnsolicitedPackages(PlatformManager->GetSessionPlatforms());
 
 		for (UPackage* UnsolicitedPackage : UnsolicitedPackages)
 		{
@@ -6059,7 +6464,7 @@ bool UCookOnTheFlyServer::IsCookByTheBookRunning() const
 }
 
 
-void UCookOnTheFlyServer::SaveGlobalShaderMapFiles(const TArray<ITargetPlatform*>& Platforms)
+void UCookOnTheFlyServer::SaveGlobalShaderMapFiles(const TArrayView<const ITargetPlatform* const>& Platforms)
 {
 	// we don't support this behavior
 	check( !IsCookingDLC() );
@@ -6181,16 +6586,12 @@ void UCookOnTheFlyServer::InitShaderCodeLibrary(void)
 {
     const UProjectPackagingSettings* const PackagingSettings = GetDefault<UProjectPackagingSettings>();
     bool const bCacheShaderLibraries = (CurrentCookMode == ECookMode::CookByTheBook);
-    ITargetPlatformManagerModule& TPM = GetTargetPlatformManagerRef();
     if (bCacheShaderLibraries && PackagingSettings->bShareMaterialShaderCode)
     {
         FShaderCodeLibrary::InitForCooking(PackagingSettings->bSharedMaterialNativeLibraries);
         
-        for (const FName& TargetPlatformName : CookByTheBookOptions->TargetPlatformNames)
+        for (const ITargetPlatform* TargetPlatform : PlatformManager->GetSessionPlatforms())
         {
-            FString TargetPlatformNameString = TargetPlatformName.ToString();
-            const ITargetPlatform* TargetPlatform = TPM.FindTargetPlatform(TargetPlatformNameString);
-
 			// Find out if this platform requires stable shader keys, by reading the platform setting file.
 			bool bNeedShaderStableKeys = false;
 			FConfigFile PlatformIniFile;
@@ -6238,17 +6639,13 @@ void UCookOnTheFlyServer::OpenShaderCodeLibrary(FString const& Name)
 
 void UCookOnTheFlyServer::ProcessShaderCodeLibraries(const FString& LibraryName)
 {
-	ITargetPlatformManagerModule& TPM = GetTargetPlatformManagerRef();
-
-	for (const FName& TargetPlatformName : CookByTheBookOptions->TargetPlatformNames)
+	for (const ITargetPlatform* TargetPlatform : PlatformManager->GetSessionPlatforms())
 	{
-		// make sure we have a registry generate for all the platforms 
-		TArray<FString>* SCLCSVPaths = OutSCLCSVPaths.Find(TargetPlatformName);
+		// make sure we have a registry generated for all the platforms 
+		const FString TargetPlatformName = TargetPlatform->PlatformName();
+		TArray<FString>* SCLCSVPaths = OutSCLCSVPaths.Find(FName(TargetPlatformName));
 		if (SCLCSVPaths && SCLCSVPaths->Num())
 		{
-			FString TargetPlatformNameString = TargetPlatformName.ToString();
-			const ITargetPlatform* TargetPlatform = TPM.FindTargetPlatform(TargetPlatformNameString);
-
 			TArray<FName> ShaderFormats;
 			TargetPlatform->GetAllTargetedShaderFormats(ShaderFormats);
 			for (FName ShaderFormat : ShaderFormats)
@@ -6263,11 +6660,11 @@ void UCookOnTheFlyServer::ProcessShaderCodeLibraries(const FString& LibraryName)
 				IFileManager::Get().FindFilesRecursive(ExpandedFiles, *FPaths::GetPath(StablePCPathCompressed), *FPaths::GetCleanFilename(StablePCPathCompressed), true, false, false);
 				if (!ExpandedFiles.Num())
 				{
-					UE_LOG(LogCook, Display, TEXT("---- NOT Running UShaderPipelineCacheToolsCommandlet for platform %s  shader format %s, no files found at %s"), *TargetPlatformName.ToString(), *ShaderFormat.ToString(), *StablePCPath);
+					UE_LOG(LogCook, Display, TEXT("---- NOT Running UShaderPipelineCacheToolsCommandlet for platform %s  shader format %s, no files found at %s"), *TargetPlatformName, *ShaderFormat.ToString(), *StablePCPath);
 				}
 				else
 				{
-					UE_LOG(LogCook, Display, TEXT("---- Running UShaderPipelineCacheToolsCommandlet for platform %s  shader format %s"), *TargetPlatformName.ToString(), *ShaderFormat.ToString());
+					UE_LOG(LogCook, Display, TEXT("---- Running UShaderPipelineCacheToolsCommandlet for platform %s  shader format %s"), *TargetPlatformName, *ShaderFormat.ToString());
 
 					const FString OutFilename = FString::Printf(TEXT("%s_%s.stable.upipelinecache"), *LibraryName, *ShaderFormat.ToString());
 					const FString PCUncookedPath = FPaths::ProjectDir() / TEXT("Content") / TEXT("PipelineCaches") / TargetPlatform->IniPlatformName() / OutFilename;
@@ -6279,7 +6676,7 @@ void UCookOnTheFlyServer::ProcessShaderCodeLibraries(const FString& LibraryName)
 					}
 
 					const FString PCCookedPath = ConvertToFullSandboxPath(*PCUncookedPath, true);
-					const FString PCPath = PCCookedPath.Replace(TEXT("[Platform]"), *TargetPlatform->PlatformName());
+					const FString PCPath = PCCookedPath.Replace(TEXT("[Platform]"), *TargetPlatformName);
 
 
 					FString Args(TEXT("build "));
@@ -6302,7 +6699,7 @@ void UCookOnTheFlyServer::ProcessShaderCodeLibraries(const FString& LibraryName)
 					}
 					if (!NumMatched)
 					{
-						UE_LOG(LogCook, Warning, TEXT("Shader format %s for platform %s had this file %s, but no .scl.csv files."), *ShaderFormat.ToString(), *TargetPlatformNameString, *StablePCPath);
+						UE_LOG(LogCook, Warning, TEXT("Shader format %s for platform %s had this file %s, but no .scl.csv files."), *ShaderFormat.ToString(), *TargetPlatformName, *StablePCPath);
 						for (int32 Index = 0; Index < SCLCSVPaths->Num(); Index++)
 						{
 							UE_LOG(LogCook, Warning, TEXT("    .scl.csv file: %s"), *((*SCLCSVPaths)[Index]));
@@ -6324,7 +6721,7 @@ void UCookOnTheFlyServer::ProcessShaderCodeLibraries(const FString& LibraryName)
 					}
 					else
 					{
-						UE_LOG(LogCook, Display, TEXT("---- Done running UShaderPipelineCacheToolsCommandlet for platform %s"), *TargetPlatformName.ToString());
+						UE_LOG(LogCook, Display, TEXT("---- Done running UShaderPipelineCacheToolsCommandlet for platform %s"), *TargetPlatformName);
 					}
 				}
 			}
@@ -6335,37 +6732,36 @@ void UCookOnTheFlyServer::ProcessShaderCodeLibraries(const FString& LibraryName)
 
 void UCookOnTheFlyServer::SaveShaderCodeLibrary(FString const& Name)
 {
-    const UProjectPackagingSettings* const PackagingSettings = GetDefault<UProjectPackagingSettings>();
-    bool const bCacheShaderLibraries = (CurrentCookMode == ECookMode::CookByTheBook);
-    ITargetPlatformManagerModule& TPM = GetTargetPlatformManagerRef();
-    if (bCacheShaderLibraries && PackagingSettings->bShareMaterialShaderCode)
-    {
-        FString ActualName = GenerateShaderCodeLibraryName(Name, IsCookFlagSet(ECookInitializationFlags::IterateSharedBuild));
-        
-        // Save shader code map - cleaning directories is deliebrately a separate loop here as we open the cache once per shader platform and we don't assume that they can't be shared across target platforms.
-        for (const FName& TargetPlatformName : CookByTheBookOptions->TargetPlatformNames)
-        {
-            FString TargetPlatformNameString = TargetPlatformName.ToString();
-            const ITargetPlatform* TargetPlatform = TPM.FindTargetPlatform(TargetPlatformNameString);
-            
-            FString BasePath = !IsCookingDLC() ? FPaths::ProjectContentDir() : GetContentDirecctoryForDLC();
-            
-            FString ShaderCodeDir = ConvertToFullSandboxPath(*BasePath, true, TargetPlatformNameString);
+	const UProjectPackagingSettings* const PackagingSettings = GetDefault<UProjectPackagingSettings>();
+	bool const bCacheShaderLibraries = (CurrentCookMode == ECookMode::CookByTheBook);
+	if (bCacheShaderLibraries && PackagingSettings->bShareMaterialShaderCode)
+	{
+		FString ActualName = GenerateShaderCodeLibraryName(Name, IsCookFlagSet(ECookInitializationFlags::IterateSharedBuild));
+		
+		// Save shader code map - cleaning directories is deliberately a separate loop here as we open the cache once per shader platform and we don't assume that they can't be shared across target platforms.
+		for (const ITargetPlatform* TargetPlatform : PlatformManager->GetSessionPlatforms())
+		{
+		   
+			FString BasePath = !IsCookingDLC() ? FPaths::ProjectContentDir() : GetContentDirecctoryForDLC();
+			
+			FString ShaderCodeDir = ConvertToFullSandboxPath(*BasePath, true, TargetPlatform->PlatformName());
 
 			const FString RootMetaDataPath = FPaths::ProjectDir() / TEXT("Metadata") / TEXT("PipelineCaches");
 			const FString MetaDataPathSB = ConvertToFullSandboxPath(*RootMetaDataPath, true);
 			const FString MetaDataPath = MetaDataPathSB.Replace(TEXT("[Platform]"), *TargetPlatform->PlatformName());
 
-            TArray<FName> ShaderFormats;
-            TargetPlatform->GetAllTargetedShaderFormats(ShaderFormats);
-            if (ShaderFormats.Num() > 0)
-            {
-                bool bSaved = FShaderCodeLibrary::SaveShaderCodeMaster(ShaderCodeDir, MetaDataPath, ShaderFormats, OutSCLCSVPaths.FindOrAdd(TargetPlatformName));
-                
-                if(!bSaved)
-                {
-                    LogCookerMessage(FString::Printf(TEXT("Shared Material Shader Code Library failed for %s."),*TargetPlatformNameString), EMessageSeverity::Error);
-                }
+			TArray<FName> ShaderFormats;
+			TargetPlatform->GetAllTargetedShaderFormats(ShaderFormats);
+			if (ShaderFormats.Num() > 0)
+			{
+				FString TargetPlatformName = TargetPlatform->PlatformName();
+				TArray<FString>& PlatformSCLCSVPaths = OutSCLCSVPaths.FindOrAdd(FName(TargetPlatformName));
+				bool bSaved = FShaderCodeLibrary::SaveShaderCodeMaster(ShaderCodeDir, MetaDataPath, ShaderFormats, PlatformSCLCSVPaths);
+				
+				if(!bSaved)
+				{
+					LogCookerMessage(FString::Printf(TEXT("Shared Material Shader Code Library failed for %s."),*TargetPlatformName), EMessageSeverity::Error);
+				}
 				else
 				{
 					if (PackagingSettings->bSharedMaterialNativeLibraries)
@@ -6374,43 +6770,40 @@ void UCookOnTheFlyServer::SaveShaderCodeLibrary(FString const& Name)
 						if (!bSaved)
 						{
 							// This is fatal - In this case we should cancel any launch on device operation or package write but we don't want to assert and crash the editor
-							LogCookerMessage(FString::Printf(TEXT("Package Native Shader Library failed for %s."), *TargetPlatformNameString), EMessageSeverity::Error);
+							LogCookerMessage(FString::Printf(TEXT("Package Native Shader Library failed for %s."), *TargetPlatformName), EMessageSeverity::Error);
 						}
 					}
-					for (const FString& Item : OutSCLCSVPaths.FindOrAdd(TargetPlatformName))
+					for (const FString& Item : PlatformSCLCSVPaths)
 					{
-						UE_LOG(LogCook, Display, TEXT("Saved scl.csv %s for platform %s"), *Item, *TargetPlatformName.ToString());
+						UE_LOG(LogCook, Display, TEXT("Saved scl.csv %s for platform %s"), *Item, *TargetPlatformName);
 					}
 
 				}
-            }
-        }
-        
-        FShaderCodeLibrary::CloseLibrary(ActualName);
-    }
+			}
+		}
+		
+		FShaderCodeLibrary::CloseLibrary(ActualName);
+	}
 }
 
 void UCookOnTheFlyServer::CleanShaderCodeLibraries()
 {
-    const UProjectPackagingSettings* const PackagingSettings = GetDefault<UProjectPackagingSettings>();
-    bool const bCacheShaderLibraries = (CurrentCookMode == ECookMode::CookByTheBook);
-    ITargetPlatformManagerModule& TPM = GetTargetPlatformManagerRef();
-    // If not iterative then clean up our temporary files
-    if (bCacheShaderLibraries && PackagingSettings->bShareMaterialShaderCode && !IsCookFlagSet(ECookInitializationFlags::Iterative))
-    {
-        for (const FName& TargetPlatformName : CookByTheBookOptions->TargetPlatformNames)
-        {
-            FString TargetPlatformNameString = TargetPlatformName.ToString();
-            const ITargetPlatform* TargetPlatform = TPM.FindTargetPlatform(TargetPlatformNameString);
-            
-            TArray<FName> ShaderFormats;
-            TargetPlatform->GetAllTargetedShaderFormats(ShaderFormats);
-            if (ShaderFormats.Num() > 0)
-            {
-                FShaderCodeLibrary::CleanDirectories(ShaderFormats);
-            }
-        }
-    }
+	const UProjectPackagingSettings* const PackagingSettings = GetDefault<UProjectPackagingSettings>();
+	bool const bCacheShaderLibraries = (CurrentCookMode == ECookMode::CookByTheBook);
+	ITargetPlatformManagerModule& TPM = GetTargetPlatformManagerRef();
+	// If not iterative then clean up our temporary files
+	if (bCacheShaderLibraries && PackagingSettings->bShareMaterialShaderCode && !IsCookFlagSet(ECookInitializationFlags::Iterative))
+	{
+		for (const ITargetPlatform* TargetPlatform : PlatformManager->GetSessionPlatforms())
+		{
+			TArray<FName> ShaderFormats;
+			TargetPlatform->GetAllTargetedShaderFormats(ShaderFormats);
+			if (ShaderFormats.Num() > 0)
+			{
+				FShaderCodeLibrary::CleanDirectories(ShaderFormats);
+			}
+		}
+	}
 }
 
 void UCookOnTheFlyServer::CookByTheBookFinished()
@@ -6465,28 +6858,28 @@ void UCookOnTheFlyServer::CookByTheBookFinished()
 		
 		{
 			SCOPE_TIMER(SavingCurrentIniSettings)
-			ITargetPlatformManagerModule& TPM = GetTargetPlatformManagerRef();
-			for ( const FName& TargetPlatformName : CookByTheBookOptions->TargetPlatformNames )
+			for (const ITargetPlatform* TargetPlatform : PlatformManager->GetSessionPlatforms() )
 			{
-				const ITargetPlatform* TargetPlatform = TPM.FindTargetPlatform(TargetPlatformName.ToString());
 				SaveCurrentIniSettings(TargetPlatform);
 			}
 		}
 
 		{
 			SCOPE_TIMER(SavingAssetRegistry);
-			for (TPair<FName, FAssetRegistryGenerator*>& Pair : RegistryGenerators)
+			for (const ITargetPlatform* TargetPlatform : PlatformManager->GetSessionPlatforms())
 			{
+				FPlatformData* PlatformData = PlatformManager->GetPlatformData(TargetPlatform);
+				FAssetRegistryGenerator& Generator = *PlatformData->RegistryGenerator.Get();
 				TArray<FName> CookedPackagesFilenames;
 				TArray<FName> IgnorePackageFilenames;
 
-				const FName& PlatformName = Pair.Key;
-				FAssetRegistryGenerator& Generator = *Pair.Value;
+				const FName& PlatformName = FName(*TargetPlatform->PlatformName());
+				FString PlatformNameString = PlatformName.ToString();
 
-				PackageTracker->CookedPackages.GetCookedFilesForPlatform(PlatformName, CookedPackagesFilenames, false, /* include successful */ true);
+				PackageTracker->CookedPackages.GetCookedFilesForPlatform(TargetPlatform, CookedPackagesFilenames, false, /* include successful */ true);
 
 				// ignore any packages which failed to cook
-				PackageTracker->CookedPackages.GetCookedFilesForPlatform(PlatformName, IgnorePackageFilenames, /* include failed */ true, false);
+				PackageTracker->CookedPackages.GetCookedFilesForPlatform(TargetPlatform, IgnorePackageFilenames, /* include failed */ true, false);
 
 				bool bForceNoFilterAssetsFromAssetRegistry = false;
 
@@ -6567,15 +6960,15 @@ void UCookOnTheFlyServer::CookByTheBookFinished()
 				{
 					if (FParse::Param(FCommandLine::Get(), TEXT("fastcook")))
 					{
-						FFileHelper::SaveStringToFile(FString(), *(GetSandboxDirectory(PlatformName.ToString()) / TEXT("fastcook.txt")));
+						FFileHelper::SaveStringToFile(FString(), *(GetSandboxDirectory(PlatformNameString) / TEXT("fastcook.txt")));
 					}
 				}
 				if (IsCreatingReleaseVersion())
 				{
-					const FString VersionedRegistryPath = GetReleaseVersionAssetRegistryPath(CookByTheBookOptions->CreateReleaseVersion, PlatformName);
+					const FString VersionedRegistryPath = GetReleaseVersionAssetRegistryPath(CookByTheBookOptions->CreateReleaseVersion, PlatformNameString);
 					IFileManager::Get().MakeDirectory(*VersionedRegistryPath, true);
 					const FString VersionedRegistryFilename = VersionedRegistryPath / GetAssetRegistryFilename();
-					const FString CookedAssetRegistryFilename = SandboxRegistryFilename.Replace(TEXT("[Platform]"), *PlatformName.ToString());
+					const FString CookedAssetRegistryFilename = SandboxRegistryFilename.Replace(TEXT("[Platform]"), *PlatformNameString);
 					IFileManager::Get().Copy(*VersionedRegistryFilename, *CookedAssetRegistryFilename, true, true);
 
 					// Also copy development registry if it exists
@@ -6603,9 +6996,11 @@ void UCookOnTheFlyServer::CookByTheBookFinished()
 	CookByTheBookOptions->BasedOnReleaseCookedPackages.Empty();
 
 	CookByTheBookOptions->bRunning = false;
+	CookByTheBookOptions->bFullLoadAndSave = false;
 
 	const FPlatformMemoryStats MemStats = FPlatformMemory::GetStats();
 
+	PlatformManager->ClearSessionPlatforms();
 
 	UE_LOG(LogCook, Display, TEXT("Peak Used virtual %uMB Peak Used physical %uMB"), MemStats.PeakUsedVirtual / 1024 / 1024, MemStats.PeakUsedPhysical / 1024 / 1024 );
 
@@ -6615,12 +7010,12 @@ void UCookOnTheFlyServer::CookByTheBookFinished()
 	UE_LOG(LogCook, Display, TEXT("Done!"));
 }
 
-void UCookOnTheFlyServer::BuildMapDependencyGraph(const FName& PlatformName)
+void UCookOnTheFlyServer::BuildMapDependencyGraph(const ITargetPlatform* TargetPlatform)
 {
-	auto& MapDependencyGraph = CookByTheBookOptions->MapDependencyGraphs.FindChecked(PlatformName);
+	auto& MapDependencyGraph = CookByTheBookOptions->MapDependencyGraphs.FindChecked(TargetPlatform);
 
 	TArray<FName> PlatformCookedPackages;
-	PackageTracker->CookedPackages.GetCookedFilesForPlatform(PlatformName, PlatformCookedPackages, /* include failed */ true, /* include successful */ true);
+	PackageTracker->CookedPackages.GetCookedFilesForPlatform(TargetPlatform, PlatformCookedPackages, /* include failed */ true, /* include successful */ true);
 
 	// assign chunks for all the map packages
 	for (const FName& CookedPackage : PlatformCookedPackages)
@@ -6644,9 +7039,9 @@ void UCookOnTheFlyServer::BuildMapDependencyGraph(const FName& PlatformName)
 	}
 }
 
-void UCookOnTheFlyServer::WriteMapDependencyGraph(const FName& PlatformName)
+void UCookOnTheFlyServer::WriteMapDependencyGraph(const ITargetPlatform* TargetPlatform)
 {
-	auto& MapDependencyGraph = CookByTheBookOptions->MapDependencyGraphs.FindChecked(PlatformName);
+	auto& MapDependencyGraph = CookByTheBookOptions->MapDependencyGraphs.FindChecked(TargetPlatform);
 
 	FString MapDependencyGraphFile = FPaths::ProjectDir() / TEXT("MapDependencyGraph.json");
 	// dump dependency graph. 
@@ -6667,7 +7062,7 @@ void UCookOnTheFlyServer::WriteMapDependencyGraph(const FName& PlatformName)
 	DependencyString.RemoveFromEnd(TEXT(",\n"));
 	DependencyString += "\n}";
 
-	FString CookedMapDependencyGraphFilePlatform = ConvertToFullSandboxPath(MapDependencyGraphFile, true).Replace(TEXT("[Platform]"), *PlatformName.ToString());
+	FString CookedMapDependencyGraphFilePlatform = ConvertToFullSandboxPath(MapDependencyGraphFile, true).Replace(TEXT("[Platform]"), *TargetPlatform->PlatformName());
 	FFileHelper::SaveStringToFile(DependencyString, *CookedMapDependencyGraphFilePlatform, FFileHelper::EEncodingOptions::ForceUnicode);
 }
 
@@ -6687,8 +7082,14 @@ void UCookOnTheFlyServer::CancelCookByTheBook()
 		check(CookByTheBookOptions);
 		check( IsInGameThread() );
 
-		// save the cook requests 
-		PackageTracker->CookRequests.DequeueAllRequests(CookByTheBookOptions->PreviousCookRequests);
+		// save the cook requests and immediately execute any pending TickCommands
+		TArray<FPackageTracker::FTickCommand> TickCommands;
+		PackageTracker->DequeueAllRequests(TickCommands, CookByTheBookOptions->PreviousCookRequests);
+		for (FPackageTracker::FTickCommand& TickCommand : TickCommands)
+		{
+			TickCommand();
+		}
+
 		CookByTheBookOptions->bRunning = false;
 
 		SandboxFile = nullptr;
@@ -6706,7 +7107,7 @@ void UCookOnTheFlyServer::StopAndClearCookedData()
 	}
 
 	PackageTracker->RecompileRequests.Empty();
-	PackageTracker->CookRequests.Empty();
+	PackageTracker->EmptyRequests();
 	PackageTracker->UnsolicitedCookedPackages.Empty();
 	PackageTracker->CookedPackages.Empty(); // set of files which have been cooked when needing to recook a file the entry will need to be removed from here
 }
@@ -6721,26 +7122,33 @@ void UCookOnTheFlyServer::ClearAllCookedData()
 	PackageTracker->CookedPackages.Empty(); // set of files which have been cooked when needing to recook a file the entry will need to be removed from here
 }
 
-void UCookOnTheFlyServer::ClearPlatformCookedData(const FString& PlatformNameString)
+void UCookOnTheFlyServer::ClearPlatformCookedData(const ITargetPlatform* TargetPlatform)
 {
+	if (!TargetPlatform)
+	{
+		return;
+	}
+
 	// if we are going to clear the cooked packages it is conceivable that we will recook the packages which we just cooked 
 	// that means it's also conceivable that we will recook the same package which currently has an outstanding async write request
 	UPackage::WaitForAsyncFileWrites();
 
-	FName PlatformName(*PlatformNameString);
-	PackageTracker->CookedPackages.RemoveAllFilesForPlatform(PlatformName);
+	PackageTracker->CookedPackages.RemoveAllFilesForPlatform(TargetPlatform);
 
 	TArray<FName> PackageNames;
-	PackageTracker->UnsolicitedCookedPackages.GetPackagesForPlatformAndRemove(PlatformName, PackageNames);
+	PackageTracker->UnsolicitedCookedPackages.GetPackagesForPlatformAndRemove(TargetPlatform, PackageNames);
 
-	DeleteSandboxDirectory(PlatformNameString);
+	DeleteSandboxDirectory(TargetPlatform->PlatformName());
 }
 
-void UCookOnTheFlyServer::ClearCachedCookedPlatformDataForPlatform( const FName& PlatformName )
+void UCookOnTheFlyServer::ClearPlatformCookedData(const FString& PlatformName)
 {
-	ITargetPlatformManagerModule& TPM = GetTargetPlatformManagerRef();
-	ITargetPlatform* TargetPlatform = TPM.FindTargetPlatform(PlatformName.ToString());
-	if ( TargetPlatform )
+	ClearPlatformCookedData(GetTargetPlatformManagerRef().FindTargetPlatform(PlatformName));
+}
+
+void UCookOnTheFlyServer::ClearCachedCookedPlatformDataForPlatform( const ITargetPlatform* TargetPlatform )
+{
+	if (TargetPlatform)
 	{
 		for ( TObjectIterator<UObject> It; It; ++It )
 		{
@@ -6749,6 +7157,12 @@ void UCookOnTheFlyServer::ClearCachedCookedPlatformDataForPlatform( const FName&
 	}
 }
 
+void UCookOnTheFlyServer::ClearCachedCookedPlatformDataForPlatform(const FName& PlatformName)
+{
+	ITargetPlatformManagerModule& TPM = GetTargetPlatformManagerRef();
+	const ITargetPlatform* TargetPlatform = TPM.FindTargetPlatform(PlatformName.ToString());
+	return ClearCachedCookedPlatformDataForPlatform(TargetPlatform);
+}
 
 void UCookOnTheFlyServer::OnTargetPlatformChangedSupportedFormats(const ITargetPlatform* TargetPlatform)
 {
@@ -6775,31 +7189,119 @@ void UCookOnTheFlyServer::CreateSandboxFile()
 	SandboxFile->Initialize(&FPlatformFileManager::Get().GetPlatformFile(), *FString::Printf(TEXT("-sandbox=\"%s\""), *OutputDirectory));
 }
 
-void UCookOnTheFlyServer::InitializeSandbox()
+void UCookOnTheFlyServer::InitializeSandbox(const TArrayView<const ITargetPlatform* const>& TargetPlatforms)
 {
-	if ( SandboxFile == nullptr )
+	TRACE_CPUPROFILER_EVENT_SCOPE(UCookOnTheFlyServer::CleanSandbox);
+#if OUTPUT_TIMING
+	double CleanSandboxTime = 0.0;
+#endif
 	{
-		const TArray<ITargetPlatform*>& TargetPlatforms = GetCookingTargetPlatforms();
+#if OUTPUT_TIMING
+		SCOPE_SECONDS_COUNTER(CleanSandboxTime);
+		SCOPE_TIMER(CleanSandbox);
+#endif
+
+		if (SandboxFile == nullptr)
+		{
+			CreateSandboxFile();
+		}
 
-		CreateSandboxFile();
+		// before we can delete any cooked files we need to make sure that we have finished writing them
+		UPackage::WaitForAsyncFileWrites();
 
-		// When looking for deterministic cooking differences in cooked packages, don't delete the packages on disk
-		if (!FParse::Param(FCommandLine::Get(), TEXT("DIFFONLY")))
+		// Skip markup of packages for reload during sandbox initialization.  TODO: Is this no longer required because it is no longer possible to load packages in this function?
+		bIsInitializingSandbox = true;
+		ON_SCOPE_EXIT
 		{
-			bIsInitializingSandbox = true;
-			CleanSandbox(IsCookFlagSet(ECookInitializationFlags::Iterative));
 			bIsInitializingSandbox = false;
+		};
+
+		TArray<const ITargetPlatform*, TInlineAllocator<ExpectedMaxNumPlatforms>> RefreshPlatforms;
+		const bool bIsDiffOnly = FParse::Param(FCommandLine::Get(), TEXT("DIFFONLY"));
+		const bool bIsIterativeCook = IsCookFlagSet(ECookInitializationFlags::Iterative);
+
+		for (const ITargetPlatform* Target : TargetPlatforms)
+		{
+			FPlatformData* PlatformData = PlatformManager->GetPlatformData(Target);
+			const bool bIsIniSettingsOutOfDate = IniSettingsOutOfDate(Target); // needs to be executed for side effects even if non-iterative
+
+			bool bShouldClearCookedContent = false;
+			if (PlatformData->bIsSandboxInitialized)
+			{
+				// We have constructed the sandbox in an earlier cook in this process (e.g. in the editor) and should not clear it again
+				bShouldClearCookedContent = false;
+			}
+			else if (bIsDiffOnly)
+			{
+				// When looking for deterministic cooking differences in cooked packages, don't delete the packages on disk
+				bShouldClearCookedContent = false;
+			}
+			else if (bIsIterativeCook)
+			{
+				if (!bIsIniSettingsOutOfDate)
+				{
+					bShouldClearCookedContent = false;
+				}
+				else
+				{
+					if (!IsCookFlagSet(ECookInitializationFlags::IgnoreIniSettingsOutOfDate))
+					{
+						UE_LOG(LogCook, Display, TEXT("Cook invalidated for platform %s ini settings don't match from last cook, clearing all cooked content"), *Target->PlatformName());
+						bShouldClearCookedContent = true;
+					}
+					else
+					{
+						UE_LOG(LogCook, Display, TEXT("Inisettings were out of date for platform %s but we are going with it anyway because IgnoreIniSettingsOutOfDate is set"), *Target->PlatformName());
+						bShouldClearCookedContent = false;
+					}
+				}
+			}
+			else
+			{
+				// In non-iterative cooks we will be replacing every cooked package and so should wipe the cooked directory
+				UE_LOG(LogCook, Display, TEXT("Clearing all cooked content for platform %s"), *Target->PlatformName());
+				bShouldClearCookedContent = true;
+			}
+
+			if (bShouldClearCookedContent)
+			{
+				ClearPlatformCookedData(Target);
+				SaveCurrentIniSettings(Target);
+			}
+			else
+			{
+				RefreshPlatforms.Add(Target);
+			}
+
+			PlatformData->bIsSandboxInitialized = true;
+		}
 
+		if (RefreshPlatforms.Num() != 0)
+		{
+			for (const ITargetPlatform* Target : RefreshPlatforms)
+			{
+				PackageTracker->CookedPackages.RemoveAllFilesForPlatform(Target);
+			}
+			// The asset registry makes populating cooked packages from disk fast, and populating is a performance benefit
+			// Don't populate however if we are looking for deterministic cooking differences; start from an empty list of cooked packages
+			if (!bIsDiffOnly) 
+			{
+				PopulateCookedPackagesFromDisk(RefreshPlatforms);
+			}
 		}
 	}
-	else
+
+#if OUTPUT_TIMING
+	FString PlatformNames;
+	for (const ITargetPlatform* Target : TargetPlatforms)
 	{
-		// This is an in-editor cook, do an iterative clean
-		CleanSandbox(true);
+		PlatformNames += Target->PlatformName() + TEXT(" ");
 	}
+	UE_LOG(LogCook, Display, TEXT("Sandbox cleanup took %5.3f seconds for platforms %s"), CleanSandboxTime, *PlatformNames);
+#endif
 }
 
-void UCookOnTheFlyServer::InitializePackageStore(const TArray<FName>& TargetPlatformNames)
+void UCookOnTheFlyServer::InitializePackageStore(const TArrayView<const ITargetPlatform* const>& TargetPlatforms)
 {
 	const FString RootPath = FPaths::RootDir();
 	const FString RootPathSandbox = ConvertToFullSandboxPath(*RootPath, true);
@@ -6807,11 +7309,11 @@ void UCookOnTheFlyServer::InitializePackageStore(const TArray<FName>& TargetPlat
 	const FString ProjectPath = FPaths::ProjectDir();
 	const FString ProjectPathSandbox = ConvertToFullSandboxPath(*ProjectPath, true);
 
-	SavePackageContexts.Reserve(TargetPlatformNames.Num());
+	SavePackageContexts.Reserve(TargetPlatforms.Num());
 
-	for (const FName PlatformName : TargetPlatformNames)
+	for (const ITargetPlatform* TargetPlatform: TargetPlatforms)
 	{
-		const FString PlatformString = PlatformName.ToString();
+		const FString PlatformString = TargetPlatform->PlatformName();
 
 		const FString ResolvedRootPath = RootPathSandbox.Replace(TEXT("[Platform]"), *PlatformString);
 		const FString ResolvedProjectPath = ProjectPathSandbox.Replace(TEXT("[Platform]"), *PlatformString);
@@ -6852,25 +7354,19 @@ void UCookOnTheFlyServer::ClearPackageStoreContexts()
 	SavePackageContexts.Empty();
 }
 
-void UCookOnTheFlyServer::InitializeTargetPlatforms()
+void UCookOnTheFlyServer::InitializeTargetPlatforms(const TArrayView<ITargetPlatform* const>& NewTargetPlatforms)
 {
-	const TArray<ITargetPlatform*>& TargetPlatforms = GetCookingTargetPlatforms();
-
 	//allow each platform to update its internals before cooking
-	for (int32 TargetPlatformIndex = 0; TargetPlatformIndex < TargetPlatforms.Num(); ++TargetPlatformIndex)
+	for (ITargetPlatform* TargetPlatform : NewTargetPlatforms)
 	{
-		ITargetPlatform* TargetPlatform = TargetPlatforms[TargetPlatformIndex];
-		if (TargetPlatform)
-		{
-			TargetPlatform->RefreshSettings();
-		}
+		TargetPlatform->RefreshSettings();
 	}
 }
 
 void UCookOnTheFlyServer::DiscoverPlatformSpecificNeverCookPackages(
-	const TArray<FName>& TargetPlatformNames, const TArray<FString>& UBTPlatformStrings)
+	const TArrayView<const ITargetPlatform* const>& TargetPlatforms, const TArray<FString>& UBTPlatformStrings)
 {
-	TArray<FName> PluginUnsupportedTargetPlatforms;
+	TArray<const ITargetPlatform*> PluginUnsupportedTargetPlatforms;
 	TArray<FAssetData> PluginAssets;
 	FARFilter PluginARFilter;
 	FString PluginPackagePath;
@@ -6888,11 +7384,11 @@ void UCookOnTheFlyServer::DiscoverPlatformSpecificNeverCookPackages(
 
 		// find any unsupported target platforms for this plugin
 		PluginUnsupportedTargetPlatforms.Reset();
-		for (int32 I = 0, Count = TargetPlatformNames.Num(); I < Count; ++I)
+		for (int32 I = 0, Count = TargetPlatforms.Num(); I < Count; ++I)
 		{
 			if (!Descriptor.SupportedTargetPlatforms.Contains(UBTPlatformStrings[I]))
 			{
-				PluginUnsupportedTargetPlatforms.Add(TargetPlatformNames[I]);
+				PluginUnsupportedTargetPlatforms.Add(TargetPlatforms[I]);
 			}
 		}
 
@@ -6912,9 +7408,9 @@ void UCookOnTheFlyServer::DiscoverPlatformSpecificNeverCookPackages(
 			PluginAssets.Reset();
 			AssetRegistry->GetAssets(PluginARFilter, PluginAssets);
 
-			for (FName& PlatformName : PluginUnsupportedTargetPlatforms)
+			for (const ITargetPlatform* TargetPlatform: PluginUnsupportedTargetPlatforms)
 			{
-				TSet<FName>& NeverCookPackages = PackageTracker->PlatformSpecificNeverCookPackages.FindOrAdd(PlatformName);
+				TSet<FName>& NeverCookPackages = PackageTracker->PlatformSpecificNeverCookPackages.FindOrAdd(TargetPlatform);
 				for (const FAssetData& Asset : PluginAssets)
 				{
 					NeverCookPackages.Add(Asset.PackageName);
@@ -6935,7 +7431,6 @@ void UCookOnTheFlyServer::StartCookByTheBook( const FCookByTheBookStartupOptions
 {
 	SCOPE_TIMER(StartCookByTheBookTime);
 
-	const TArray<ITargetPlatform*>& TargetPlatforms = CookByTheBookStartupOptions.TargetPlatforms;
 	const TArray<FString>& CookMaps = CookByTheBookStartupOptions.CookMaps;
 	const TArray<FString>& CookDirectories = CookByTheBookStartupOptions.CookDirectories;
 	const TArray<FString>& IniMapSections = CookByTheBookStartupOptions.IniMapSections;
@@ -6964,12 +7459,27 @@ void UCookOnTheFlyServer::StartCookByTheBook( const FCookByTheBookStartupOptions
 	CookByTheBookOptions->bErrorOnEngineContentUse = CookByTheBookStartupOptions.bErrorOnEngineContentUse;
 
 	GenerateAssetRegistry();
+
+	// SelectSessionPlatforms does not check for uniqueness and non-null, and we rely on those properties for performance, so ensure it here before calling SelectSessionPlatforms
+	TArray<ITargetPlatform*> TargetPlatforms;
+	TargetPlatforms.Reserve(CookByTheBookStartupOptions.TargetPlatforms.Num());
+	for (ITargetPlatform* TargetPlatform : CookByTheBookStartupOptions.TargetPlatforms)
+	{
+		if (TargetPlatform)
+		{
+			TargetPlatforms.AddUnique(TargetPlatform);
+		}
+	}
+	PlatformManager->SelectSessionPlatforms(TargetPlatforms, PackageTracker);
+	check(PlatformManager->GetSessionPlatforms().Num() == TargetPlatforms.Num());
+
+	RefreshPlatformAssetRegistries(TargetPlatforms);
+
 	if (CurrentCookMode == ECookMode::CookByTheBook)
 	{
 		check(!PackageNameCache);
 		check(!PackageTracker);
-		PackageNameCache = new FPackageNameCache(AssetRegistry);
-		PackageTracker = new FPackageTracker(PackageNameCache);
+		ConstructPackageTracker();
 		FCoreUObjectDelegates::PackageCreatedForLoad.AddUObject(this, &UCookOnTheFlyServer::MaybeMarkPackageAsAlreadyLoaded);
 	}
 
@@ -7034,9 +7544,10 @@ void UCookOnTheFlyServer::StartCookByTheBook( const FCookByTheBookStartupOptions
 
 		if (LocalizationTargetsToChunk.Num() > 0 && AllCulturesToCook.Num() > 0)
 		{
-			for (TPair<FName, FAssetRegistryGenerator*>& Pair : RegistryGenerators)
+			for (const ITargetPlatform* TargetPlatform : TargetPlatforms)
 			{
-				Pair.Value->RegisterChunkDataGenerator(MakeShared<FLocalizationChunkDataGenerator>(MoveTemp(LocalizationTargetsToChunk), MoveTemp(AllCulturesToCook)));
+				FAssetRegistryGenerator* RegistryGenerator = PlatformManager->GetPlatformData(TargetPlatform)->RegistryGenerator.Get();
+				RegistryGenerator->RegisterChunkDataGenerator(MakeShared<FLocalizationChunkDataGenerator>(MoveTemp(LocalizationTargetsToChunk), MoveTemp(AllCulturesToCook)));
 			}
 		}
 	}
@@ -7095,25 +7606,20 @@ void UCookOnTheFlyServer::StartCookByTheBook( const FCookByTheBookStartupOptions
 		}
 
 	}
-	
-	// build persistent list of all target platform names in CookByTheBookOptions->TargetPlatformNames,
-	// also use temp list of UBT platform strings to discover PlatformSpecificNeverCookPackages
+
+	// use temp list of UBT platform strings to discover PlatformSpecificNeverCookPackages
 	{
 		TArray<FString> UBTPlatformStrings;
 		UBTPlatformStrings.Reserve(TargetPlatforms.Num());
-
-		CookByTheBookOptions->TargetPlatformNames.Empty(TargetPlatforms.Num());
 		for (const ITargetPlatform* Platform : TargetPlatforms)
 		{
 			FString UBTPlatformName;
 			Platform->GetPlatformInfo().UBTTargetId.ToString(UBTPlatformName);
 			UBTPlatformStrings.Emplace(MoveTemp(UBTPlatformName));
-			CookByTheBookOptions->TargetPlatformNames.Emplace(*Platform->PlatformName());
 		}
 
-		DiscoverPlatformSpecificNeverCookPackages(CookByTheBookOptions->TargetPlatformNames, UBTPlatformStrings);
+		DiscoverPlatformSpecificNeverCookPackages(TargetPlatforms, UBTPlatformStrings);
 	}
-	const TArray<FName>& TargetPlatformNames = CookByTheBookOptions->TargetPlatformNames;
 
 	if ( CookByTheBookOptions->DlcName != DLCName )
 	{
@@ -7124,10 +7630,10 @@ void UCookOnTheFlyServer::StartCookByTheBook( const FCookByTheBookStartupOptions
 	}
 
 	// This will either delete the sandbox or iteratively clean it
-	InitializeSandbox();
-	InitializeTargetPlatforms();
+	InitializeSandbox(TargetPlatforms);
+	InitializeTargetPlatforms(TargetPlatforms);
 
-	InitializePackageStore(TargetPlatformNames);
+	InitializePackageStore(TargetPlatforms);
 
 	if (CurrentCookMode == ECookMode::CookByTheBook && !IsCookFlagSet(ECookInitializationFlags::Iterative))
 	{
@@ -7150,13 +7656,11 @@ void UCookOnTheFlyServer::StartCookByTheBook( const FCookByTheBookStartupOptions
 	}
 
 	{
-		for (ITargetPlatform* Platform : TargetPlatforms)
+		if (CookByTheBookOptions->bGenerateDependenciesForMaps)
 		{
-			FName PlatformName = FName(*Platform->PlatformName());
-
-			if (CookByTheBookOptions->bGenerateDependenciesForMaps)
+			for (const ITargetPlatform* Platform : TargetPlatforms)
 			{
-				CookByTheBookOptions->MapDependencyGraphs.Add(PlatformName);
+				CookByTheBookOptions->MapDependencyGraphs.Add(Platform);
 			}
 		}
 	}
@@ -7170,16 +7674,18 @@ void UCookOnTheFlyServer::StartCookByTheBook( const FCookByTheBookStartupOptions
 		// if we are cooking dlc we must be based on a release version cook
 		check( !BasedOnReleaseVersion.IsEmpty() );
 
-		for ( const FName& PlatformName : TargetPlatformNames )
+		for ( const ITargetPlatform* TargetPlatform: TargetPlatforms )
 		{
-			FString OriginalSandboxRegistryFilename = GetReleaseVersionAssetRegistryPath(BasedOnReleaseVersion, PlatformName ) / TEXT("Metadata") / GetDevelopmentAssetRegistryFilename();
+			FString PlatformNameString = TargetPlatform->PlatformName();
+			FName PlatformName(*PlatformNameString);
+			FString OriginalSandboxRegistryFilename = GetReleaseVersionAssetRegistryPath(BasedOnReleaseVersion, PlatformNameString ) / TEXT("Metadata") / GetDevelopmentAssetRegistryFilename();
 
 			TArray<FName> PackageList;
 			// if this check fails probably because the asset registry can't be found or read
 			bool bSucceeded = GetAllPackageFilenamesFromAssetRegistry(OriginalSandboxRegistryFilename, PackageList);
 			if (!bSucceeded)
 			{
-				OriginalSandboxRegistryFilename = GetReleaseVersionAssetRegistryPath(BasedOnReleaseVersion, PlatformName) / GetAssetRegistryFilename();
+				OriginalSandboxRegistryFilename = GetReleaseVersionAssetRegistryPath(BasedOnReleaseVersion, PlatformNameString) / GetAssetRegistryFilename();
 				bSucceeded = GetAllPackageFilenamesFromAssetRegistry(OriginalSandboxRegistryFilename, PackageList);
 			}
 
@@ -7191,7 +7697,7 @@ void UCookOnTheFlyServer::StartCookByTheBook( const FCookByTheBookStartupOptions
 				FVanillaPlatformEntry VanillaPlatfromEntry = BuildPlatformHierarchy(PlatformName, EPlatformFilter::CookFlavor);
 				for (const FPlatformInfo* PlatformFlaworInfo : VanillaPlatfromEntry.PlatformFlavors)
 				{
-					OriginalSandboxRegistryFilename = GetReleaseVersionAssetRegistryPath(BasedOnReleaseVersion, PlatformFlaworInfo->PlatformInfoName) / GetAssetRegistryFilename();
+					OriginalSandboxRegistryFilename = GetReleaseVersionAssetRegistryPath(BasedOnReleaseVersion, PlatformFlaworInfo->PlatformInfoName.ToString()) / GetAssetRegistryFilename();
 					bSucceeded = GetAllPackageFilenamesFromAssetRegistry(OriginalSandboxRegistryFilename, PackageList);
 					if (bSucceeded)
 					{
@@ -7203,13 +7709,13 @@ void UCookOnTheFlyServer::StartCookByTheBook( const FCookByTheBookStartupOptions
 
 			if ( bSucceeded )
 			{
-				TArray<FName> PlatformNames;
-				PlatformNames.Add(PlatformName);
+				TArray<const ITargetPlatform*> ResultPlatforms;
+				ResultPlatforms.Add(TargetPlatform);
 				TArray<bool> Succeeded;
 				Succeeded.Add(true);
 				for (const FName& PackageFilename : PackageList)
 				{
-					PackageTracker->CookedPackages.Add( FFilePlatformCookedPackage( PackageFilename, PlatformNames, Succeeded) );
+					PackageTracker->CookedPackages.Add( FFilePlatformCookedPackage( PackageFilename, ResultPlatforms, Succeeded) );
 				}
 			}
 			CookByTheBookOptions->BasedOnReleaseCookedPackages.Add(PlatformName, MoveTemp(PackageList));
@@ -7243,7 +7749,7 @@ void UCookOnTheFlyServer::StartCookByTheBook( const FCookByTheBookStartupOptions
 		GRedirectCollector.ProcessSoftObjectPathPackageList(StartupPackage, false, StartupSoftObjectPackages);
 	}
 
-	CollectFilesToCook(FilesInPath, CookMaps, CookDirectories, IniMapSections, CookOptions);
+	CollectFilesToCook(FilesInPath, CookMaps, CookDirectories, IniMapSections, CookOptions, TargetPlatforms);
 
 	// Add string asset packages after collecting files, to avoid accidentally activating the behavior to cook all maps if none are specified
 	for (FName SoftObjectPackage : StartupSoftObjectPackages)
@@ -7300,7 +7806,7 @@ void UCookOnTheFlyServer::StartCookByTheBook( const FCookByTheBookStartupOptions
 		
 		if (PackageFileFName != NAME_None)
 		{
-			PackageTracker->CookRequests.EnqueueUnique( FFilePlatformRequest( PackageFileFName, TargetPlatformNames ) );
+			PackageTracker->EnqueueUniqueCookRequest( FFilePlatformRequest( PackageFileFName, TargetPlatforms ) );
 		}
 		else if (!FLinkerLoad::IsKnownMissingPackage(FileFName))
 		{
@@ -7320,20 +7826,19 @@ void UCookOnTheFlyServer::StartCookByTheBook( const FCookByTheBookStartupOptions
 			// note that we might erroneously enter here if we are generating a patch instead and we accidentally passed in BasedOnReleaseVersion to the cooker instead of to unrealpak
 			check( !CreateReleaseVersion.IsEmpty() );
 
-			for ( const FName& PlatformName : TargetPlatformNames )
+			for ( const ITargetPlatform* TargetPlatform : TargetPlatforms )
 			{
-				TArray<FName> PlatformArray;
-				PlatformArray.Add( PlatformName );
-
 				// if we are based of a cook and we are creating a new one we need to make sure that at least all the old packages are cooked as well as the new ones
-				FString OriginalAssetRegistryPath = GetReleaseVersionAssetRegistryPath( BasedOnReleaseVersion, PlatformName ) / GetAssetRegistryFilename();
+				FString OriginalAssetRegistryPath = GetReleaseVersionAssetRegistryPath( BasedOnReleaseVersion, TargetPlatform->PlatformName() ) / GetAssetRegistryFilename();
 
 				TArray<FName> PackageFiles;
 				verify( GetAllPackageFilenamesFromAssetRegistry(OriginalAssetRegistryPath, PackageFiles) );
 
+				TArray<const ITargetPlatform*, TInlineAllocator<1>> RequestPlatforms;
+				RequestPlatforms.Add(TargetPlatform);
 				for ( const FName& PackageFilename : PackageFiles )
 				{
-					PackageTracker->CookRequests.EnqueueUnique( FFilePlatformRequest( PackageFilename, PlatformArray) );
+					PackageTracker->EnqueueUniqueCookRequest( FFilePlatformRequest( PackageFilename, RequestPlatforms) );
 				}
 			}
 		}
@@ -7345,18 +7850,28 @@ void UCookOnTheFlyServer::StartCookByTheBook( const FCookByTheBookStartupOptions
 	// if too many packages are being recooked after resume then we may need to figure out a different way to do this
 	for ( const FFilePlatformRequest& PreviousRequest : CookByTheBookOptions->PreviousCookRequests )
 	{
-		PackageTracker->CookRequests.EnqueueUnique( PreviousRequest );
+		PackageTracker->EnqueueUniqueCookRequest( PreviousRequest );
 	}
 	CookByTheBookOptions->PreviousCookRequests.Empty();
 
 }
 
-bool UCookOnTheFlyServer::RecompileChangedShaders(const TArray<FName>& TargetPlatforms)
+bool UCookOnTheFlyServer::RecompileChangedShaders(const TArray<const ITargetPlatform*>& TargetPlatforms)
+{
+	bool bShadersRecompiled = false;
+	for (const ITargetPlatform* TargetPlatform : TargetPlatforms)
+	{
+		bShadersRecompiled |= RecompileChangedShadersForPlatform(TargetPlatform->PlatformName());
+	}
+	return bShadersRecompiled;
+}
+
+bool UCookOnTheFlyServer::RecompileChangedShaders(const TArray<FName>& TargetPlatformNames)
 {
 	bool bShadersRecompiled = false;
-	for (const FName& TargetPlatform : TargetPlatforms)
+	for (const FName TargetPlatformName : TargetPlatformNames)
 	{
-		bShadersRecompiled |= RecompileChangedShadersForPlatform(TargetPlatform.ToString());
+		bShadersRecompiled |= RecompileChangedShadersForPlatform(TargetPlatformName.ToString());
 	}
 	return bShadersRecompiled;
 }
@@ -7392,11 +7907,11 @@ void UCookOnTheFlyServer::MaybeMarkPackageAsAlreadyLoaded(UPackage *Package)
 
 	bool bShouldMarkAsAlreadyProcessed = false;
 
-	TArray<FName> CookedPlatforms;
+	TArray<const ITargetPlatform*> CookedPlatforms;
 	if (PackageTracker->CookedPackages.GetCookedPlatforms(StandardName, CookedPlatforms))
 	{
 		bShouldMarkAsAlreadyProcessed = true;
-		for (const FName& TargetPlatform : CookByTheBookOptions->TargetPlatformNames)
+		for (const ITargetPlatform* TargetPlatform : PlatformManager->GetSessionPlatforms())
 		{
 			if (!CookedPlatforms.Contains(TargetPlatform))
 			{
@@ -7406,10 +7921,10 @@ void UCookOnTheFlyServer::MaybeMarkPackageAsAlreadyLoaded(UPackage *Package)
 		}
 
 		FString Platforms;
-		for (const FName& CookedPlatform : CookedPlatforms)
+		for (const ITargetPlatform* CookedPlatform : CookedPlatforms)
 		{
 			Platforms += TEXT(" ");
-			Platforms += CookedPlatform.ToString();
+			Platforms += CookedPlatform->PlatformName();
 		}
 		if (IsCookFlagSet(ECookInitializationFlags::LogDebugInfo))
 		{
@@ -7448,6 +7963,12 @@ bool UCookOnTheFlyServer::HandleNetworkFileServerNewConnection(const FString& Ve
 
 	const FString LocalVersionInfo = FString::Printf(TEXT("%s %d"), *Branch, CL);
 
+	if (!AddCookOnTheFlyPlatform(PlatformName))
+	{
+		UE_LOG(LogCook, Warning, TEXT("Unrecognized PlatformName '%s', CookOnTheFly requests for this platform will fail."), *PlatformName);
+		return false;
+	}
+
 	UE_LOG(LogCook, Display, TEXT("Connection received of version %s local version %s"), *VersionInfo, *LocalVersionInfo);
 
 	if (LocalVersionInfo != VersionInfo)
@@ -7458,10 +7979,10 @@ bool UCookOnTheFlyServer::HandleNetworkFileServerNewConnection(const FString& Ve
 	return true;
 }
 
-void UCookOnTheFlyServer::GetCookOnTheFlyUnsolicitedFiles(const FName& PlatformName, TArray<FString>& UnsolicitedFiles, const FString& Filename)
+void UCookOnTheFlyServer::GetCookOnTheFlyUnsolicitedFiles(const ITargetPlatform* TargetPlatform, TArray<FString>& UnsolicitedFiles, const FString& Filename)
 {
 	TArray<FName> UnsolicitedFilenames;
-	PackageTracker->UnsolicitedCookedPackages.GetPackagesForPlatformAndRemove(PlatformName, UnsolicitedFilenames);
+	PackageTracker->UnsolicitedCookedPackages.GetPackagesForPlatformAndRemove(TargetPlatform, UnsolicitedFilenames);
 
 	for (const FName& UnsolicitedFile : UnsolicitedFilenames)
 	{
@@ -7472,7 +7993,7 @@ void UCookOnTheFlyServer::GetCookOnTheFlyUnsolicitedFiles(const FName& PlatformN
 		// this can happen if the package was saved but the async writer thread hasn't finished writing it to disk yet
 
 		FString SandboxFilename = ConvertToFullSandboxPath(*Filename, true);
-		SandboxFilename.ReplaceInline(TEXT("[Platform]"), *PlatformName.ToString());
+		SandboxFilename.ReplaceInline(TEXT("[Platform]"), *TargetPlatform->PlatformName());
 		if (IFileManager::Get().FileExists(*SandboxFilename))
 		{
 			UnsolicitedFiles.Add(StandardFilename);
@@ -7490,23 +8011,31 @@ void UCookOnTheFlyServer::HandleNetworkFileServerFileRequest(const FString& File
 {
 	check(IsCookOnTheFlyMode());
 
-	const bool bIsCookable = FPackageName::IsPackageExtension(*FPaths::GetExtension(Filename, true));
-
-	FName PlatformFname = FName(*PlatformName);
+	ITargetPlatform* TargetPlatform = AddCookOnTheFlyPlatform(PlatformName);
+	if (!TargetPlatform)
+	{
+		UE_LOG(LogCook, Warning, TEXT("Unrecognized PlatformName '%s', CookOnTheFly FileServerRequest requests for this platform will fail."), *PlatformName);
+		return;
+	}
 
+	const bool bIsCookable = FPackageName::IsPackageExtension(*FPaths::GetExtension(Filename, true));
 	if (!bIsCookable)
 	{
-		GetCookOnTheFlyUnsolicitedFiles(PlatformFname, UnsolicitedFiles, Filename);
+		// Wait for the Platform to be added if this is the first time; it is not legal to call GetCookOnTheFlyUnsolicitedFiles until after the platform has been added
+		while (!PlatformManager->IsPlatformInitialized(TargetPlatform))
+		{
+			FPlatformProcess::Sleep(0.001f);
+		}
+		GetCookOnTheFlyUnsolicitedFiles(TargetPlatform, UnsolicitedFiles, Filename);
 		return;
 	}
 
 	FString StandardFileName = Filename;
 	FPaths::MakeStandardFilename( StandardFileName );
-
-	FName StandardFileFname = FName( *StandardFileName );
-	TArray<FName> Platforms;
-	Platforms.Add( PlatformFname );
-	FFilePlatformRequest FileRequest( StandardFileFname, Platforms);
+	FName StandardFileFname = FName(*StandardFileName);
+	TArray<const ITargetPlatform*, TInlineAllocator<1>> TargetPlatforms;
+	TargetPlatforms.Add(TargetPlatform);
+	FFilePlatformRequest FileRequest( StandardFileFname, TargetPlatforms);
 
 #if PROFILE_NETWORK
 	double StartTime = FPlatformTime::Seconds();
@@ -7516,7 +8045,12 @@ void UCookOnTheFlyServer::HandleNetworkFileServerFileRequest(const FString& File
 	
 	UE_LOG(LogCook, Display, TEXT("Requesting file from cooker %s"), *StandardFileName);
 
-	PackageTracker->CookRequests.EnqueueUnique(FileRequest, true);
+	{
+		check(&PlatformManager->GetSessionLock() == &RequestLock && &PackageTracker->GetRequestLock() == &RequestLock);
+		FScopeLock ScopeLock(&RequestLock);
+		PlatformManager->AddRefCookOnTheFlyPlatform(TargetPlatform, *PackageTracker);
+		PackageTracker->GetThreadUnsafeCookRequests().EnqueueUnique(FileRequest, true);
+	}
 	
 	if (PackageTracker->CookRequestEvent)
 	{
@@ -7549,6 +8083,13 @@ void UCookOnTheFlyServer::HandleNetworkFileServerFileRequest(const FString& File
 		FPlatformProcess::Sleep(0.001f);
 	}
 
+
+	{
+		FScopeLock ScopeLock(&PlatformManager->GetSessionLock());
+		PlatformManager->ReleaseCookOnTheFlyPlatform(TargetPlatform);
+	}
+
+
 #if PROFILE_NETWORK
 	if ( bFoundNetworkEventWait )
 	{
@@ -7558,7 +8099,7 @@ void UCookOnTheFlyServer::HandleNetworkFileServerFileRequest(const FString& File
 #endif
 	UE_LOG( LogCook, Display, TEXT("Cook complete %s"), *FileRequest.GetFilename().ToString())
 
-	GetCookOnTheFlyUnsolicitedFiles(PlatformFname, UnsolicitedFiles, Filename);
+	GetCookOnTheFlyUnsolicitedFiles(TargetPlatform, UnsolicitedFiles, Filename);
 
 #if PROFILE_NETWORK
 	WaitForAsyncFilesWrites += FPlatformTime::Seconds() - StartTime;
@@ -7578,10 +8119,16 @@ FString UCookOnTheFlyServer::HandleNetworkGetSandboxPath()
 
 void UCookOnTheFlyServer::HandleNetworkGetPrecookedList(const FString& PlatformName, TMap<FString, FDateTime>& PrecookedFileList)
 {
-	FName PlatformFName = FName(*PlatformName);
+	ITargetPlatformManagerModule& TPM = GetTargetPlatformManagerRef();
+	const ITargetPlatform* TargetPlatform = TPM.FindTargetPlatform(PlatformName);
+	if (!TargetPlatform)
+	{
+		UE_LOG(LogCook, Warning, TEXT("Unrecognized PlatformName '%s' in HandleNetworkGetPrrequests, returning 0 files."), *PlatformName);
+		return;
+	}
 
 	TArray<FName> CookedPlatformFiles;
-	PackageTracker->CookedPackages.GetCookedFilesForPlatform(PlatformFName, CookedPlatformFiles, /* include failed */ true, /* include successful */ true);
+	PackageTracker->CookedPackages.GetCookedFilesForPlatform(TargetPlatform, CookedPlatformFiles, /* include failed */ true, /* include successful */ true);
 
 
 	for ( const FName& CookedFile : CookedPlatformFiles)
@@ -7681,16 +8228,7 @@ uint32 UCookOnTheFlyServer::FullLoadAndSave(uint32& CookedPackageCount)
 
 	uint32 Result = 0;
 
-	TArray<const ITargetPlatform*> AllTargetPlatforms;
-	{
-		UE_LOG(LogCook, Display, TEXT("Gathering target platforms..."));
-		SCOPE_TIMER(FullLoadAndSave_TargetPlatforms);
-		ITargetPlatformManagerModule& TPM = GetTargetPlatformManagerRef();
-		for (const FName& TargetPlatformName : CookByTheBookOptions->TargetPlatformNames)
-		{
-			AllTargetPlatforms.Add(TPM.FindTargetPlatform(TargetPlatformName.ToString()));
-		}
-	}
+	const TArray<const ITargetPlatform*>& TargetPlatforms = PlatformManager->GetSessionPlatforms();
 
 	{
 		UE_LOG(LogCook, Display, TEXT("Loading requested packages..."));
@@ -7698,8 +8236,17 @@ uint32 UCookOnTheFlyServer::FullLoadAndSave(uint32& CookedPackageCount)
 		while (HasCookRequests())
 		{
 			FFilePlatformRequest ToBuild;
-			PackageTracker->CookRequests.Dequeue(/* out */ ToBuild);
-			check(ToBuild.IsValid());
+			TArray<FPackageTracker::FTickCommand> TickCommands;
+			FPackageTracker::ERequestType RequestType = PackageTracker->DequeueRequest(TickCommands, /* out */ ToBuild);
+			if (RequestType == FPackageTracker::ERequestType::TickCommand)
+			{
+				for (FPackageTracker::FTickCommand& TickCommand : TickCommands)
+				{
+					TickCommand();
+				}
+				continue;
+			}
+			check(RequestType == FPackageTracker::ERequestType::Cook && ToBuild.IsValid());
 
 			const FName BuildFilenameFName = ToBuild.GetFilename();
 			if (!PackageTracker->NeverCookPackageList.Contains(BuildFilenameFName))
@@ -7780,6 +8327,7 @@ uint32 UCookOnTheFlyServer::FullLoadAndSave(uint32& CookedPackageCount)
 
 				if (Package->GetOuter() != nullptr)
 				{
+					UE_LOG(LogCook, Warning, TEXT("Skipping package %s with outermost %s"), *Package->GetName(), *Package->GetOutermost()->GetName());
 					continue;
 				}
 
@@ -7837,7 +8385,7 @@ uint32 UCookOnTheFlyServer::FullLoadAndSave(uint32& CookedPackageCount)
 
 							bool bAllPlatformDataLoaded = true;
 							bool bIsTexture = Obj->IsA(UTexture::StaticClass());
-							for (const ITargetPlatform* TargetPlatform : AllTargetPlatforms)
+							for (const ITargetPlatform* TargetPlatform : TargetPlatforms)
 							{
 								if (bSaveConcurrent)
 								{
@@ -7977,7 +8525,7 @@ uint32 UCookOnTheFlyServer::FullLoadAndSave(uint32& CookedPackageCount)
 			{
 				UObject* Obj = ObjectsToWaitForCookedPlatformData[ObjIdx];
 				bool bAllPlatformDataLoaded = true;
-				for (const ITargetPlatform* TargetPlatform : AllTargetPlatforms)
+				for (const ITargetPlatform* TargetPlatform : TargetPlatforms)
 				{
 					if (!Obj->IsCachedCookedPlatformDataLoaded(TargetPlatform))
 					{
@@ -8010,7 +8558,7 @@ uint32 UCookOnTheFlyServer::FullLoadAndSave(uint32& CookedPackageCount)
 		}
 
 		int64 ParallelSavedPackages = 0;
-		ParallelFor(PackagesToSave.Num(), [this, &PackagesToSave, &AllTargetPlatforms, &ParallelSavedPackages, SaveFlags, bSaveConcurrent](int32 PackageIdx)
+		ParallelFor(PackagesToSave.Num(), [this, &PackagesToSave, &TargetPlatforms ,&ParallelSavedPackages, SaveFlags, bSaveConcurrent](int32 PackageIdx)
 		{
 			UPackage* Package = PackagesToSave[PackageIdx];
 			check(Package);
@@ -8035,7 +8583,7 @@ uint32 UCookOnTheFlyServer::FullLoadAndSave(uint32& CookedPackageCount)
 							continue;
 						}
 
-						for (const ITargetPlatform* TargetPlatform : AllTargetPlatforms)
+						for (const ITargetPlatform* TargetPlatform : TargetPlatforms)
 						{
 							Obj->BeginCacheForCookedPlatformData(TargetPlatform);
 						}
@@ -8076,10 +8624,10 @@ uint32 UCookOnTheFlyServer::FullLoadAndSave(uint32& CookedPackageCount)
 				uint32 OriginalPackageFlags = Package->GetPackageFlags();
 
 				TArray<bool> SavePackageSuccessPerPlatform;
-				SavePackageSuccessPerPlatform.SetNum(AllTargetPlatforms.Num());
-				for (int32 PlatformIndex = 0; PlatformIndex < AllTargetPlatforms.Num(); ++PlatformIndex)
+				SavePackageSuccessPerPlatform.SetNum(TargetPlatforms.Num());
+				for (int32 PlatformIndex = 0; PlatformIndex < TargetPlatforms.Num(); ++PlatformIndex)
 				{
-					const ITargetPlatform* Target = AllTargetPlatforms[PlatformIndex];
+					const ITargetPlatform* Target = TargetPlatforms[PlatformIndex];
 
 					// don't save Editor resources from the Engine if the target doesn't have editoronly data
 					bool bCookPackage = (!bExcludeFromNonEditorTargets || Target->HasEditorOnlyData());
@@ -8119,7 +8667,7 @@ uint32 UCookOnTheFlyServer::FullLoadAndSave(uint32& CookedPackageCount)
 						const bool bSucceededSavePackage = (SaveResult == ESavePackageResult::Success || SaveResult == ESavePackageResult::GenerateStub || SaveResult == ESavePackageResult::ReplaceCompletely);
 						if (bSucceededSavePackage)
 						{
-							FAssetRegistryGenerator* Generator = RegistryGenerators.FindRef(FName(*Target->PlatformName()));
+							FAssetRegistryGenerator* Generator = PlatformManager->GetPlatformData(Target)->RegistryGenerator.Get();
 							UpdateAssetRegistryPackageData(Generator, Package->GetFName(), SaveResult);
 
 							FPlatformAtomics::InterlockedIncrement(&ParallelSavedPackages);
@@ -8142,7 +8690,7 @@ uint32 UCookOnTheFlyServer::FullLoadAndSave(uint32& CookedPackageCount)
 
 				check(CachedPackageFilename->StandardFileFName != NAME_None);
 
-				FFilePlatformCookedPackage FileRequest(CachedPackageFilename->StandardFileFName, CookByTheBookOptions->TargetPlatformNames, SavePackageSuccessPerPlatform);
+				FFilePlatformCookedPackage FileRequest(CachedPackageFilename->StandardFileFName, TargetPlatforms, SavePackageSuccessPerPlatform);
 				PackageTracker->CookedPackages.Add(FileRequest);
 
 				if (SavePackageSuccessPerPlatform.Contains(false))
