[partial]

Copying //UE4/Dev-Tools-Staging at cl 11029060 to Dev-Main (//UE4/Dev-Main)

#rnx

diff --git a/Engine/Source/Editor/UnrealEd/Private/Cooker/AsyncIODelete.cpp b/Engine/Source/Editor/UnrealEd/Private/Cooker/AsyncIODelete.cpp
--- a/Engine/Source/Editor/UnrealEd/Private/Cooker/AsyncIODelete.cpp
+++ b/Engine/Source/Editor/UnrealEd/Private/Cooker/AsyncIODelete.cpp
@@ -52,6 +52,7 @@ void FAsyncIODelete::SetTempRoot(const FStringView& InOwnedTempRoot)
 void FAsyncIODelete::SetDeletesPaused(bool bInPaused)
 {
 	bPaused = bInPaused;
+#if ASYNCIODELETE_ASYNC_ENABLED
 	if (!bPaused)
 	{
 		IFileManager& FileManager = IFileManager::Get();
@@ -67,6 +68,7 @@ void FAsyncIODelete::SetDeletesPaused(bool bInPaused)
 		}
 		PausedDeletes.Empty();
 	}
+#endif
 }
 
 bool FAsyncIODelete::Setup()
@@ -82,6 +84,7 @@ bool FAsyncIODelete::Setup()
 		return false;
 	}
 
+#if ASYNCIODELETE_ASYNC_ENABLED
 	// Delete the TempRoot directory to clear the results from any previous process using the same TempRoot that did not shut down cleanly
 	uint32 ErrorCode;
 	if (!DeleteTempRootDirectory(ErrorCode))
@@ -107,6 +110,7 @@ bool FAsyncIODelete::Setup()
 	// TempRoot and bPaused are preserved across setup/teardown and may have any value
 	check(PausedDeletes.Num() == 0);
 	check(DeleteCounter == 0);
+#endif
 
 	// We are now setup and ready to create DeleteTasks
 	bInitialized = true;
@@ -121,6 +125,7 @@ void FAsyncIODelete::Teardown()
 		return;
 	}
 
+#if ASYNCIODELETE_ASYNC_ENABLED
 	// Clear task variables
 	WaitForAllTasks();
 	check(ActiveTaskCount == 0 && TasksComplete != nullptr && TasksComplete->Wait(0));
@@ -138,20 +143,34 @@ void FAsyncIODelete::Teardown()
 	// Clear delete variables; we don't need to run the tasks for the remaining pauseddeletes because synchronously deleting the temp directory above did the work they were going to do
 	PausedDeletes.Empty();
 	DeleteCounter = 0;
+#endif
 
 	// We are now torn down and ready for a new setup
 	bInitialized = false;
 }
 
-void FAsyncIODelete::WaitForAllTasks()
+bool FAsyncIODelete::WaitForAllTasks(float TimeLimitSeconds)
 {
+#if ASYNCIODELETE_ASYNC_ENABLED
 	if (!bInitialized)
 	{
-		return;
+		return true;
 	}
 
+	if (TimeLimitSeconds <= 0.f)
+	{
 	TasksComplete->Wait();
+	}
+	else
+	{
+		if (!TasksComplete->Wait(FTimespan::FromSeconds(TimeLimitSeconds)))
+		{
+			return false;
+		}
+	}
 	check(ActiveTaskCount == 0);
+#endif
+	return true;
 }
 
 bool FAsyncIODelete::Delete(const FStringView& PathToDelete, EPathType ExpectedType)
@@ -184,16 +203,19 @@ bool FAsyncIODelete::Delete(const FStringView& PathToDelete, EPathType ExpectedT
 		return false;
 	}
 
+#if ASYNCIODELETE_ASYNC_ENABLED
 	if (DeleteCounter == UINT32_MAX)
 	{
 		Teardown();
 	}
+#endif
 	if (!Setup())
 	{
 		// Setup failed; we are not able to provide asynchronous deletes; fall back to synchronous
 		return SynchronousDelete(PathToDeleteSZ, ExpectedType);
 	}
 
+#if ASYNCIODELETE_ASYNC_ENABLED
 	const FString TempPath = FPaths::Combine(TempRoot, FString::Printf(TEXT("%u"), DeleteCounter));
 	DeleteCounter++;
 
@@ -217,8 +239,12 @@ bool FAsyncIODelete::Delete(const FStringView& PathToDelete, EPathType ExpectedT
 		CreateDeleteTask(TempPath, ExpectedType);
 	}
 	return true;
+#else
+	return SynchronousDelete(PathToDeleteSZ, ExpectedType);
+#endif
 }
 
+#if ASYNCIODELETE_ASYNC_ENABLED
 void FAsyncIODelete::CreateDeleteTask(const FStringView& InDeletePath, EPathType PathType)
 {
 	{
@@ -243,6 +269,7 @@ void FAsyncIODelete::OnTaskComplete()
 		TasksComplete->Trigger();
 	}
 }
+#endif
 
 bool FAsyncIODelete::SynchronousDelete(const TCHAR* InDeletePath, EPathType PathType)
 {
@@ -266,6 +293,7 @@ bool FAsyncIODelete::SynchronousDelete(const TCHAR* InDeletePath, EPathType Path
 	return Result;
 }
 
+#if ASYNCIODELETE_ASYNC_ENABLED
 bool FAsyncIODelete::DeleteTempRootDirectory(uint32& OutErrorCode)
 {
 	OutErrorCode = 0;
@@ -304,6 +332,7 @@ bool FAsyncIODelete::DeleteTempRootDirectory(uint32& OutErrorCode)
 	}
 	return bDeleteSucceeded;
 }
+#endif
 
 #if WITH_ASYNCIODELETE_DEBUG
 void FAsyncIODelete::AddTempRoot(const FStringView& InTempRoot)
diff --git a/Engine/Source/Editor/UnrealEd/Private/Cooker/AsyncIODelete.h b/Engine/Source/Editor/UnrealEd/Private/Cooker/AsyncIODelete.h
--- a/Engine/Source/Editor/UnrealEd/Private/Cooker/AsyncIODelete.h
+++ b/Engine/Source/Editor/UnrealEd/Private/Cooker/AsyncIODelete.h
@@ -7,6 +7,7 @@
 #include "Async/Future.h"
 #include "Containers/Array.h"
 #include "Containers/StringView.h"
+#include "HAL/Platform.h"
 #include "Misc/AssertionMacros.h"
 
 #ifndef WITH_ASYNCIODELETE_DEBUG
@@ -17,6 +18,13 @@
 	#endif
 #endif
 
+// Temporarily disable AsyncIODelete on Mac until we can diagnose why it's failing
+#if PLATFORM_MAC
+#define ASYNCIODELETE_ASYNC_ENABLED 0
+#else
+#define ASYNCIODELETE_ASYNC_ENABLED 1
+#endif
+
 class FEvent;
 
 /**
@@ -90,11 +98,13 @@ public:
 	}
 
 	/** Synchronously wait for all deletes to complete.  */
-	void WaitForAllTasks();
+	bool WaitForAllTasks(float TimeLimitSeconds = 0.0f);
 
 private:
+#if ASYNCIODELETE_ASYNC_ENABLED
 	/** Update the ActiveTaskCount and events for a task completing */
 	void OnTaskComplete();
+#endif
 
 	enum class EPathType
 	{
@@ -104,20 +114,26 @@ private:
 	/** Asynchronously delete a file or directory.  We handle both in the same function, but we want the interface to be explicit. */
 	bool Delete(const FStringView& PathToDelete, EPathType ExpectedType);
 
+#if ASYNCIODELETE_ASYNC_ENABLED
 	/** Create and store the task to delete the given Directory or File */
 	void CreateDeleteTask(const FStringView& InDeletePath, EPathType PathType);
+#endif
 
 	/** Delete the given path synchronously; called from a task or in error fallback cases from the public thread */
 	bool SynchronousDelete(const TCHAR* InDeletePath, EPathType PathType);
 
+#if ASYNCIODELETE_ASYNC_ENABLED
 	bool DeleteTempRootDirectory(uint32& OutErrorCode);
+#endif
 
 	FString	TempRoot;
+#if ASYNCIODELETE_ASYNC_ENABLED
 	TArray<FString> PausedDeletes;
 	FCriticalSection CriticalSection; // We use a CriticalSection instead of a TAtomic ActiveTaskCount so that we can atomically { trigger TasksComplete if ActiveTaskCount == 0 }
 	FEvent* TasksComplete = nullptr;
 	uint32 ActiveTaskCount = 0;
 	uint32 DeleteCounter = 0;
+#endif
 	bool bInitialized = false;
 	bool bPaused = false;
 
