Fix a mismerge from CL 9528156 as that fix is somewhat out of date (::IsUsingPackageStore isn't used for creating the Bulkdata manifest for example) although I did keep the ::IsCookByTheBookMode() check from 9528156

#rnx

diff --git a/Engine/Source/Editor/UnrealEd/Private/CookOnTheFlyServer.cpp b/Engine/Source/Editor/UnrealEd/Private/CookOnTheFlyServer.cpp
--- a/Engine/Source/Editor/UnrealEd/Private/CookOnTheFlyServer.cpp
+++ b/Engine/Source/Editor/UnrealEd/Private/CookOnTheFlyServer.cpp
@@ -4533,7 +4533,7 @@ void UCookOnTheFlyServer::SaveCookedPackage(UPackage* Package, uint32 SaveFlags,
 						}
 						else
 						{
-							FSavePackageContext* const SavePackageContext = (IsCookByTheBookMode() && IsUsingPackageStore()) ? SavePackageContexts[PlatformIndex] : nullptr;
+							FSavePackageContext* const SavePackageContext = (IsCookByTheBookMode() && SavePackageContexts.Num() > 0) ? SavePackageContexts[PlatformIndex] : nullptr;
 
 							Result = GEditor->Save(	Package, World, FlagsToCook, *PlatFilename, 
 													GError, nullptr, bSwap, false, SaveFlags, Target, 
