Fix CIS staticanalysis build
#rnx

diff --git a/Engine/Source/Editor/UnrealEd/Private/ObjectTools.cpp b/Engine/Source/Editor/UnrealEd/Private/ObjectTools.cpp
--- a/Engine/Source/Editor/UnrealEd/Private/ObjectTools.cpp
+++ b/Engine/Source/Editor/UnrealEd/Private/ObjectTools.cpp
@@ -1555,7 +1555,11 @@ namespace ObjectTools
 								Notification->SetCompletionState( CollectionCreated ? SNotificationItem::CS_Success : SNotificationItem::CS_Fail );
 							}
 						}
-					} //-V773
+						if (ContentHelper)
+						{
+							delete ContentHelper;
+						}
+					}
 				}
 			}
 			else
