Fixed non determinism in cooked UMapBuildDataRegistry

diff --git a/Engine/Source/Editor/UnrealEd/Private/StaticLightingSystem/StaticLightingSystem.cpp b/Engine/Source/Editor/UnrealEd/Private/StaticLightingSystem/StaticLightingSystem.cpp
--- a/Engine/Source/Editor/UnrealEd/Private/StaticLightingSystem/StaticLightingSystem.cpp
+++ b/Engine/Source/Editor/UnrealEd/Private/StaticLightingSystem/StaticLightingSystem.cpp
@@ -749,7 +749,10 @@ bool FStaticLightingSystem::BeginLightmassProcess()
 	else
 	{
 		InvalidateStaticLighting();
-		ApplyNewLightingData(true);
+
+		// Calling ApplyNewLightingData() results in creating an empty MapBuildData which is causing cooking issues
+		// Disabled for now
+		//ApplyNewLightingData(true);
 	}
 	
 	if (!bForceNoPrecomputedLighting)
@@ -776,9 +779,6 @@ void FStaticLightingSystem::InvalidateStaticLighting()
 			continue;
 		}
 
-		// Clear all the atmosphere guids from the MapBuildData when starting a new build.
-		Level->GetOrCreateMapBuildData()->ClearSkyAtmosphereBuildData();
-
 		const bool bBuildLightingForLevel = Options.ShouldBuildLightingForLevel( Level );
 		
 		if (bBuildLightingForLevel)
diff --git a/Engine/Source/Runtime/Engine/Private/MapBuildData.cpp b/Engine/Source/Runtime/Engine/Private/MapBuildData.cpp
--- a/Engine/Source/Runtime/Engine/Private/MapBuildData.cpp
+++ b/Engine/Source/Runtime/Engine/Private/MapBuildData.cpp
@@ -705,6 +705,9 @@ void UMapBuildDataRegistry::InvalidateStaticLighting(UWorld* World, bool bRecrea
 		MarkPackageDirty();
 	}
 
+	// Clear all the atmosphere guids from the MapBuildData when starting a new build.
+	ClearSkyAtmosphereBuildData();
+
 	bSetupResourceClusters = false;
 }
 
