Asset Import: Replace Existing Settings added to Asset Import Task to keep or discard existing asset settings.
- Texture import setup to use new Replace Existing Settings option

diff --git a/Engine/Source/Editor/UnrealEd/Private/Factories/EditorFactories.cpp b/Engine/Source/Editor/UnrealEd/Private/Factories/EditorFactories.cpp
--- a/Engine/Source/Editor/UnrealEd/Private/Factories/EditorFactories.cpp
+++ b/Engine/Source/Editor/UnrealEd/Private/Factories/EditorFactories.cpp
@@ -273,6 +273,7 @@
 #include "DesktopPlatformModule.h"
 #include "Interfaces/IMainFrameModule.h"
 #include "Factories/TextureImportSettings.h"
+#include "AssetImportTask.h"
 
 #include "SkinWeightsUtilities.h"
 
@@ -4284,6 +4285,25 @@ UObject* UTextureFactory::FactoryCreateBinary
 	{
 		bUsingExistingSettings = false;
 	}
+	else if (AssetImportTask && AssetImportTask->bAutomated)
+	{
+		if (ExistingTexture)
+		{
+			if (!AssetImportTask->bReplaceExisting)
+			{
+				GEditor->GetEditorSubsystem<UImportSubsystem>()->BroadcastAssetPostImport(this, nullptr);
+				return nullptr;
+			}
+			else
+			{
+				bUsingExistingSettings = !AssetImportTask->bReplaceExistingSettings;
+			}
+		}
+		else
+		{
+			bUsingExistingSettings = false;
+		}
+	}
 	else
 	{
 		bUsingExistingSettings = bSuppressImportOverwriteDialog;
diff --git a/Engine/Source/Editor/UnrealEd/Public/AssetImportTask.h b/Engine/Source/Editor/UnrealEd/Public/AssetImportTask.h
--- a/Engine/Source/Editor/UnrealEd/Public/AssetImportTask.h
+++ b/Engine/Source/Editor/UnrealEd/Public/AssetImportTask.h
@@ -39,6 +39,10 @@ public:
 	UPROPERTY(BlueprintReadWrite, Category = "Asset Import Task")
 	bool bReplaceExisting;
 
+	/** Replace existing settings when overwriting existing assets  */
+	UPROPERTY(BlueprintReadWrite, Category = "Asset Import Task")
+	bool bReplaceExistingSettings;
+
 	/** Avoid dialogs */
 	UPROPERTY(BlueprintReadWrite, Category = "Asset Import Task")
 	bool bAutomated;
