Fix FBX import of UDIM UVs
- V axix is inverted within each UDIM tile, but tile layout is unchanged

diff --git a/Engine/Source/Editor/UnrealEd/Private/Fbx/FbxSkeletalMeshImport.cpp b/Engine/Source/Editor/UnrealEd/Private/Fbx/FbxSkeletalMeshImport.cpp
--- a/Engine/Source/Editor/UnrealEd/Private/Fbx/FbxSkeletalMeshImport.cpp
+++ b/Engine/Source/Editor/UnrealEd/Private/Fbx/FbxSkeletalMeshImport.cpp
@@ -3540,9 +3540,13 @@ bool UnFbx::FFbxImporter::FillSkelMeshImporterFromFbx( FSkeletalMeshImportData&
 					int32 UVIndex = (UVReferenceMode[UVLayerIndex] == FbxLayerElement::eDirect) ? 
 							UVMapIndex : LayerElementUV[UVLayerIndex]->GetIndexArray().GetAt(UVMapIndex);
 					FbxVector2	UVVector = LayerElementUV[UVLayerIndex]->GetDirectArray().GetAt(UVIndex);
+					const float U = static_cast<float>(UVVector[0]);
+					const float V = static_cast<float>(UVVector[1]);
+					const float VTile = FMath::FloorToFloat(V);
+					const float VOffset = V - VTile;
 
-					TmpWedges[UnrealVertexIndex].UVs[ UVLayerIndex ].X = static_cast<float>(UVVector[0]);
-					TmpWedges[UnrealVertexIndex].UVs[ UVLayerIndex ].Y = 1.f - static_cast<float>(UVVector[1]);
+					TmpWedges[UnrealVertexIndex].UVs[UVLayerIndex].X = U;
+					TmpWedges[UnrealVertexIndex].UVs[UVLayerIndex].Y = VTile + (1.f - VOffset);   //flip the Y of UVs for DirectX
 				}
 			}
 			else if( UVLayerIndex == 0 )
diff --git a/Engine/Source/Editor/UnrealEd/Private/Fbx/FbxStaticMeshImport.cpp b/Engine/Source/Editor/UnrealEd/Private/Fbx/FbxStaticMeshImport.cpp
--- a/Engine/Source/Editor/UnrealEd/Private/Fbx/FbxStaticMeshImport.cpp
+++ b/Engine/Source/Editor/UnrealEd/Private/Fbx/FbxStaticMeshImport.cpp
@@ -723,8 +723,13 @@ bool UnFbx::FFbxImporter::BuildStaticMeshFromGeometry(FbxNode* Node, UStaticMesh
 								UVMapIndex : FBXUVs.LayerElementUV[UVLayerIndex]->GetIndexArray().GetAt(UVMapIndex);
 
 							FbxVector2	UVVector = FBXUVs.LayerElementUV[UVLayerIndex]->GetDirectArray().GetAt(UVIndex);
-							FinalUVVector.X = static_cast<float>(UVVector[0]);
-							FinalUVVector.Y = 1.f - static_cast<float>(UVVector[1]);   //flip the Y of UVs for DirectX
+							const float U = static_cast<float>(UVVector[0]);
+							const float V = static_cast<float>(UVVector[1]);
+							const float VTile = FMath::FloorToFloat(V);
+							const float VOffset = V - VTile;
+
+							FinalUVVector.X = U;
+							FinalUVVector.Y = VTile + (1.f - VOffset);   //flip the Y of UVs for DirectX
 						}
 						VertexInstanceUVs.Set(AddedVertexInstanceId, UVLayerIndex, FinalUVVector);
 					}
