Panel Extension Subsystem

diff --git a/Engine/Source/Editor/LevelEditor/Private/LevelEditor.cpp b/Engine/Source/Editor/LevelEditor/Private/LevelEditor.cpp
--- a/Engine/Source/Editor/LevelEditor/Private/LevelEditor.cpp
+++ b/Engine/Source/Editor/LevelEditor/Private/LevelEditor.cpp
@@ -39,6 +39,7 @@
 #include "Interfaces/IMainFrameModule.h"
 #include "Framework/Commands/GenericCommands.h"
 #include "Misc/EngineBuildSettings.h"
+#include "Subsystems/PanelExtensionSubsystem.h"
 #include "Classes/EditorStyleSettings.h"
 
 #define LOCTEXT_NAMESPACE "LevelEditor"
@@ -136,6 +137,12 @@ public:
 
 		SetToolTipText(RightContentTooltip);
 
+		TSharedRef<SWidget> DefaultNamePlate = SNew(STextBlock)
+			.Text(RightContentText)
+			.Visibility(EVisibility::HitTestInvisible)
+			.TextStyle(FEditorStyle::Get(), "SProjectBadge.Text")
+			.ColorAndOpacity(BadgeTextColor);
+
 		SBox::Construct(SBox::FArguments()
 			.HAlign(HAlign_Right)
 			.VAlign(VAlign_Top)
@@ -147,11 +154,9 @@ public:
 				.BorderBackgroundColor(BadgeBackgroundColor)
 				.VAlign(VAlign_Top)
 				[
-					SNew(STextBlock)
-					.Text(RightContentText)
-					.Visibility(EVisibility::HitTestInvisible)
-					.TextStyle(FEditorStyle::Get(),"SProjectBadge.Text")
-					.ColorAndOpacity(BadgeTextColor)
+					SNew(SExtensionPanel)
+					.ExtensionPanelID("LevelEditorProjectNamePlate")
+					.DefaultWidget(DefaultNamePlate)
 				]
 			]);
 	}
diff --git a/Engine/Source/Editor/LevelEditor/Private/SActorDetails.cpp b/Engine/Source/Editor/LevelEditor/Private/SActorDetails.cpp
--- a/Engine/Source/Editor/LevelEditor/Private/SActorDetails.cpp
+++ b/Engine/Source/Editor/LevelEditor/Private/SActorDetails.cpp
@@ -29,6 +29,7 @@
 #include "ScopedTransaction.h"
 #include "SourceCodeNavigation.h"
 #include "Widgets/Docking/SDockTab.h"
+#include "Subsystems/PanelExtensionSubsystem.h"
 
 
 IConsoleVariable* SActorDetails::ShowComponents = IConsoleManager::Get().RegisterConsoleVariable(TEXT("ShowFlag.DetailsPanelComponents"), true, TEXT("Show components in editor details panel."), ECVF_Cheat);
@@ -153,6 +154,13 @@ void SActorDetails::Construct(const FArguments& InArgs, const FName TabIdentifie
 		[
 			DetailsView->GetNameAreaWidget().ToSharedRef()
 		]
+		+ SVerticalBox::Slot()
+		.Padding(0.0f, 0.0f, 0.0f, 0.0f)
+		.AutoHeight()
+		[
+			SNew(SExtensionPanel)
+			.ExtensionPanelID("ActorDetailsPanel")
+		]
 		+SVerticalBox::Slot()
 		[
 			SAssignNew(DetailsSplitter, SSplitter)
diff --git a/Engine/Source/Editor/UnrealEd/Private/Subsystems/PanelExtensionSubsystem.cpp b/Engine/Source/Editor/UnrealEd/Private/Subsystems/PanelExtensionSubsystem.cpp
new file mode 100644
--- /dev/null
+++ b/Engine/Source/Editor/UnrealEd/Private/Subsystems/PanelExtensionSubsystem.cpp
@@ -0,0 +1,122 @@
+// Copyright Epic Games, Inc. All Rights Reserved.
+
+#include "Subsystems/PanelExtensionSubsystem.h"
+
+#include "Editor.h"
+#include "AssetToolsModule.h"
+#include "ObjectTools.h"
+#include "FileHelpers.h"
+#include "EditorReimportHandler.h"
+
+FPanelExtensionFactory::FPanelExtensionFactory()
+{
+
+}
+
+SExtensionPanel::~SExtensionPanel()
+{
+	if (GEditor && ExtensionPanelID != NAME_None)
+	{
+		if (UPanelExtensionSubsystem* PanelExtensionSubsystem = GEditor->GetEditorSubsystem<UPanelExtensionSubsystem>())
+		{
+			PanelExtensionSubsystem->OnPanelFactoryRegistryChanged(ExtensionPanelID).RemoveAll(this);
+		}
+	}
+}
+
+void SExtensionPanel::Construct(const FArguments& InArgs)
+{
+	ExtensionPanelID = InArgs._ExtensionPanelID.Get(NAME_None);
+	DefaultWidget = InArgs._DefaultWidget.Get(nullptr);
+
+	if (GEditor && ExtensionPanelID != NAME_None)
+	{
+		if (UPanelExtensionSubsystem* PanelExtensionSubsystem = GEditor->GetEditorSubsystem<UPanelExtensionSubsystem>())
+		{
+			PanelExtensionSubsystem->OnPanelFactoryRegistryChanged(ExtensionPanelID).AddRaw(this, &SExtensionPanel::RebuildWidget);
+			RebuildWidget();
+		}
+	}
+}
+
+void SExtensionPanel::RebuildWidget()
+{
+	if (GEditor && ExtensionPanelID != NAME_None)
+	{
+		if (UPanelExtensionSubsystem* PanelExtensionSubsystem = GEditor->GetEditorSubsystem<UPanelExtensionSubsystem>())
+		{
+			TSharedRef<SWidget> Widget = PanelExtensionSubsystem->GetWidget(ExtensionPanelID);
+			if (Widget == SNullWidget::NullWidget && DefaultWidget.IsValid())
+			{
+				Widget = DefaultWidget.ToSharedRef();
+			}
+
+			ChildSlot
+			[
+				Widget
+			];
+		}
+	}
+}
+
+
+UPanelExtensionSubsystem::UPanelExtensionSubsystem()
+	: UEditorSubsystem()
+{
+
+}
+
+void UPanelExtensionSubsystem::Initialize(FSubsystemCollectionBase& Collection)
+{
+
+}
+
+void UPanelExtensionSubsystem::Deinitialize()
+{
+
+}
+
+void UPanelExtensionSubsystem::RegisterPanelFactory(FName ExtensionPanelID, const FPanelExtensionFactory& InPanelExtensionFactory)
+{
+	TArray<FPanelExtensionFactory>& Factories = ExtensionPointMap.FindOrAdd(ExtensionPanelID);
+	Factories.Add(InPanelExtensionFactory);
+
+	UPanelExtensionSubsystem::FPanelFactoryRegistryChanged& PanelFactoryRegsitered = OnPanelFactoryRegistryChanged(ExtensionPanelID);
+	PanelFactoryRegsitered.Broadcast();
+}
+
+void UPanelExtensionSubsystem::UnregisterPanelFactory(FName Identifier, FName ExtensionPanelID)
+{
+	for (auto& AssetPair : ExtensionPointMap)
+	{
+		if (ExtensionPanelID == NAME_None || AssetPair.Key == ExtensionPanelID)
+		{
+			int32 numRemoved = AssetPair.Value.RemoveAll([Identifier](const FPanelExtensionFactory& Factory) { return Factory.Identifier == Identifier; });
+			if (numRemoved > 0)
+			{
+				UPanelExtensionSubsystem::FPanelFactoryRegistryChanged& PanelFactoryRegsitered = OnPanelFactoryRegistryChanged(AssetPair.Key);
+				PanelFactoryRegsitered.Broadcast();
+			}
+		}
+	}
+}
+
+TSharedRef<SWidget> UPanelExtensionSubsystem::GetWidget(FName ExtensionPanelID)
+{
+	const TArray<FPanelExtensionFactory>* ExtensionArray = ExtensionPointMap.Find(ExtensionPanelID);
+	if (ExtensionArray && ExtensionArray->Num() > 0)
+	{
+		const FPanelExtensionFactory& Extension = (*ExtensionArray)[0]; // TODO: We need to add support for multiple widgets
+		if (Extension.CreateWidget.IsBound())
+		{
+			TArray<UObject*> Temp;
+			return Extension.CreateWidget.Execute(Temp);
+		}
+	}
+	return SNullWidget::NullWidget;
+}
+
+UPanelExtensionSubsystem::FPanelFactoryRegistryChanged& UPanelExtensionSubsystem::OnPanelFactoryRegistryChanged(FName ExtensionPanelID)
+{
+	return PanelFactoryRegistryChangedCallbackMap.FindOrAdd(ExtensionPanelID);
+}
\ No newline at end of file
diff --git a/Engine/Source/Editor/UnrealEd/Public/Subsystems/PanelExtensionSubsystem.h b/Engine/Source/Editor/UnrealEd/Public/Subsystems/PanelExtensionSubsystem.h
new file mode 100644
--- /dev/null
+++ b/Engine/Source/Editor/UnrealEd/Public/Subsystems/PanelExtensionSubsystem.h
@@ -0,0 +1,83 @@
+// Copyright Epic Games, Inc. All Rights Reserved.
+
+#pragma once
+
+#include "EditorSubsystem.h"
+
+#include "UObject/ObjectMacros.h"
+#include "UObject/UObjectGlobals.h"
+#include "Templates/SubclassOf.h"
+#include "Containers/Queue.h"
+
+#include "PanelExtensionSubsystem.generated.h"
+
+struct UNREALED_API FPanelExtensionFactory
+{
+public:
+	FPanelExtensionFactory();
+
+	DECLARE_DELEGATE_RetVal_OneParam(TSharedRef<SWidget>, FGenericCreateWidget, const TArray<UObject*>&);
+
+	/** An identifier to allow removal later on. */
+	FName Identifier;
+
+	/** Delegate that generates a the widget. */
+	FGenericCreateWidget CreateWidget;
+};
+
+class UNREALED_API SExtensionPanel : public SCompoundWidget
+{
+public:
+	~SExtensionPanel();
+
+	SLATE_BEGIN_ARGS(SExtensionPanel)
+		: _ExtensionPanelID()
+		, _DefaultWidget()
+	{}
+
+		/** The ID to identify this Extension point */
+		SLATE_ATTRIBUTE(FName, ExtensionPanelID)
+		SLATE_ATTRIBUTE(TSharedPtr<SWidget>, DefaultWidget)
+
+	SLATE_END_ARGS()
+
+	/** Constructs the widget */
+	void Construct(const FArguments& InArgs);
+
+private:
+	void RebuildWidget();
+
+	FName ExtensionPanelID;
+	TSharedPtr<SWidget> DefaultWidget;
+};
+
+/**
+ * UPanelExtensionSubsystem
+ * Subsystem for creating extensible panels in the Editor
+ */
+UCLASS()
+class UNREALED_API UPanelExtensionSubsystem : public UEditorSubsystem
+{
+	GENERATED_BODY()
+
+public:
+	UPanelExtensionSubsystem();
+
+	virtual void Initialize(FSubsystemCollectionBase& Collection) override;
+	virtual void Deinitialize() override;
+
+	void RegisterPanelFactory(FName ExtensionPanelID, const FPanelExtensionFactory& InPanelExtensionFactory);
+	void UnregisterPanelFactory(FName Identifier, FName ExtensionPanelID = NAME_None);
+
+protected:
+	friend class SExtensionPanel;
+	TSharedRef<SWidget> GetWidget(FName ExtensionPanelID);
+
+	DECLARE_MULTICAST_DELEGATE(FPanelFactoryRegistryChanged);
+	FPanelFactoryRegistryChanged& OnPanelFactoryRegistryChanged(FName ExtensionPanelID);
+
+private:
+	TMap<FName, TArray<FPanelExtensionFactory>> ExtensionPointMap;
+	TMap<FName, FPanelFactoryRegistryChanged> PanelFactoryRegistryChangedCallbackMap;
+
+};
