Fixed missing include. For real this time.

diff --git a/Engine/Source/Editor/UnrealEd/Public/Subsystems/PanelExtensionSubsystem.h b/Engine/Source/Editor/UnrealEd/Public/Subsystems/PanelExtensionSubsystem.h
--- a/Engine/Source/Editor/UnrealEd/Public/Subsystems/PanelExtensionSubsystem.h
+++ b/Engine/Source/Editor/UnrealEd/Public/Subsystems/PanelExtensionSubsystem.h
@@ -8,6 +8,7 @@
 #include "UObject/UObjectGlobals.h"
 #include "Templates/SubclassOf.h"
 #include "Containers/Queue.h"
+#include "Widgets/SCompoundWidget.h"
 
 #include "PanelExtensionSubsystem.generated.h"
 
