Update cached material data earlier, fixes failed check() when renaming parameters

diff --git a/Engine/Source/Editor/UnrealEd/Private/PreviewMaterial.cpp b/Engine/Source/Editor/UnrealEd/Private/PreviewMaterial.cpp
--- a/Engine/Source/Editor/UnrealEd/Private/PreviewMaterial.cpp
+++ b/Engine/Source/Editor/UnrealEd/Private/PreviewMaterial.cpp
@@ -346,6 +346,9 @@ void UMaterialEditorPreviewParameters::RegenerateArrays()
 		// Only operate on base materials
 		UMaterial* ParentMaterial = PreviewMaterial;
 
+		// Ensure all cached data is up-to-date before looping over parameters
+		PreviewMaterial->UpdateCachedExpressionData();
+
 		// Loop through all types of parameters for this material and add them to the parameter arrays.
 		TArray<FMaterialParameterInfo> ParameterInfo;
 		TArray<FGuid> Guids;
