Update cached parameters, avoid failed check when modify MI layer asset

diff --git a/Engine/Source/Editor/UnrealEd/Private/PreviewMaterial.cpp b/Engine/Source/Editor/UnrealEd/Private/PreviewMaterial.cpp
--- a/Engine/Source/Editor/UnrealEd/Private/PreviewMaterial.cpp
+++ b/Engine/Source/Editor/UnrealEd/Private/PreviewMaterial.cpp
@@ -1003,6 +1003,7 @@ void UMaterialEditorInstanceConstant::RegenerateArrays()
 		// Only operate on base materials
 		UMaterial* ParentMaterial = Parent->GetMaterial();
 		SourceInstance->UpdateParameterNames();	// Update any parameter names that may have changed.
+		SourceInstance->UpdateCachedLayerParameters();
 
 		// Get all static parameters from the source instance.  This will handle inheriting parent values.	
 		FStaticParameterSet SourceStaticParameters;
