Move FAssetEditorToolkit::GetToolbarTabId() up to the IAssetEditorInstance level.

diff --git a/Engine/Source/Editor/UnrealEd/Private/MiniCurveEditor.cpp b/Engine/Source/Editor/UnrealEd/Private/MiniCurveEditor.cpp
--- a/Engine/Source/Editor/UnrealEd/Private/MiniCurveEditor.cpp
+++ b/Engine/Source/Editor/UnrealEd/Private/MiniCurveEditor.cpp
@@ -79,6 +79,12 @@ bool SMiniCurveEditor::CloseWindow()
 	return true;
 }
 
+FName SMiniCurveEditor::GetToolbarTabId() const
+{
+	//@TODO: This editor should probably derive from FAssetEditorToolkit instead!
+	return NAME_None;
+}
+
 TSharedPtr<class FTabManager> SMiniCurveEditor::GetAssociatedTabManager()
 {
 	//@TODO: This editor should probably derive from FAssetEditorToolkit instead!
diff --git a/Engine/Source/Editor/UnrealEd/Public/MiniCurveEditor.h b/Engine/Source/Editor/UnrealEd/Public/MiniCurveEditor.h
--- a/Engine/Source/Editor/UnrealEd/Public/MiniCurveEditor.h
+++ b/Engine/Source/Editor/UnrealEd/Public/MiniCurveEditor.h
@@ -34,6 +34,7 @@ public:
 	virtual bool CloseWindow() override;
 	virtual bool IsPrimaryEditor() const override { return true; };
 	virtual void InvokeTab(const struct FTabId& TabId) override {}
+	virtual FName GetToolbarTabId() const override;
 	virtual TSharedPtr<class FTabManager> GetAssociatedTabManager() override;
 	virtual double GetLastActivationTime() override;
 	virtual void RemoveEditingAsset(UObject* Asset) override;
diff --git a/Engine/Source/Editor/UnrealEd/Public/Subsystems/AssetEditorSubsystem.h b/Engine/Source/Editor/UnrealEd/Public/Subsystems/AssetEditorSubsystem.h
--- a/Engine/Source/Editor/UnrealEd/Public/Subsystems/AssetEditorSubsystem.h
+++ b/Engine/Source/Editor/UnrealEd/Public/Subsystems/AssetEditorSubsystem.h
@@ -30,6 +30,7 @@ public:
 	virtual bool CloseWindow() = 0;
 	virtual bool IsPrimaryEditor() const = 0;
 	virtual void InvokeTab(const struct FTabId& TabId) = 0;
+	virtual FName GetToolbarTabId() const = 0;
 	virtual TSharedPtr<class FTabManager> GetAssociatedTabManager() = 0;
 	virtual double GetLastActivationTime() = 0;
 	virtual void RemoveEditingAsset(UObject* Asset) = 0;
diff --git a/Engine/Source/Editor/UnrealEd/Public/Toolkits/AssetEditorToolkit.h b/Engine/Source/Editor/UnrealEd/Public/Toolkits/AssetEditorToolkit.h
--- a/Engine/Source/Editor/UnrealEd/Public/Toolkits/AssetEditorToolkit.h
+++ b/Engine/Source/Editor/UnrealEd/Public/Toolkits/AssetEditorToolkit.h
@@ -138,6 +138,7 @@ public:
 	virtual bool CloseWindow() override;
 	virtual bool IsPrimaryEditor() const override { return true; };
 	virtual void InvokeTab(const FTabId& TabId) override;
+	virtual FName GetToolbarTabId() const override { return ToolbarTabId; }
 	virtual TSharedPtr<FTabManager> GetAssociatedTabManager() override;
 	virtual double GetLastActivationTime() override;
 	virtual void RemoveEditingAsset(UObject* Asset) override;
@@ -228,9 +229,6 @@ public:
 	void AddToolbarWidget(TSharedRef<SWidget> Widget);
 	void RemoveAllToolbarWidgets();
 
-	/** Gets the toolbar tab id */
-	FName GetToolbarTabId() const {return ToolbarTabId;}
-
 	/** True if this actually is editing an asset */
 	bool IsActuallyAnAsset() const;
 
