Fix for uninitialized warning.

diff --git a/Engine/Source/Editor/UnrealEd/Private/Commandlets/ContentCommandlets.cpp b/Engine/Source/Editor/UnrealEd/Private/Commandlets/ContentCommandlets.cpp
--- a/Engine/Source/Editor/UnrealEd/Private/Commandlets/ContentCommandlets.cpp
+++ b/Engine/Source/Editor/UnrealEd/Private/Commandlets/ContentCommandlets.cpp
@@ -1959,7 +1959,7 @@ FArchive& operator<<(FArchive& Ar, FPackageObjects& PackageObjects)
 
 	if (Ar.IsLoading())
 	{
-		int32 NumObjects;
+		int32 NumObjects = 0;
 		FString ObjectName;
 		FString ClassName;
 
