GetDirtyWorldPackages now marks map built data as dirty if the world is dirty.
Necessary to allow source control operations on map built data.
Also necessary to save built data when changing maps and closing the editor.

diff --git a/Engine/Source/Editor/UnrealEd/Private/FileHelpers.cpp b/Engine/Source/Editor/UnrealEd/Private/FileHelpers.cpp
--- a/Engine/Source/Editor/UnrealEd/Private/FileHelpers.cpp
+++ b/Engine/Source/Editor/UnrealEd/Private/FileHelpers.cpp
@@ -4057,22 +4057,32 @@ void FEditorFileUtils::GetDirtyWorldPackages(TArray<UPackage*>& OutDirtyPackages
 			{
 				UPackage* BuiltDataPackage = WorldIt->PersistentLevel->MapBuildData->GetOutermost();
 
-				if (BuiltDataPackage->IsDirty() && BuiltDataPackage != WorldPackage)
+				if (BuiltDataPackage != WorldPackage)
 				{
-					// If built data package does not have a name yet add the world package so a user is prompted to have a name chosen
-					if (!WorldPackage->IsDirty())
+					if (WorldPackage->IsDirty() && !BuiltDataPackage->IsDirty())
 					{
-						const FString WorldPackageName = WorldPackage->GetName();
-						const bool bIncludeReadOnlyRoots = false;
-						const bool bIsValidPath = FPackageName::IsValidLongPackageName(WorldPackageName, bIncludeReadOnlyRoots);
-						if (!bIsValidPath)
+						// Must become dirty because new prompts to save should not be brought up after each individual map saves
+						// We also cannot bring up a second prompt to save because a recursion guard blocks it
+						BuiltDataPackage->MarkPackageDirty();
+					}
+
+					if (BuiltDataPackage->IsDirty())
+					{
+						// If built data package does not have a name yet add the world package so a user is prompted to have a name chosen
+						if (!WorldPackage->IsDirty())
 						{
-							WorldPackage->MarkPackageDirty();
-							OutDirtyPackages.Add(WorldPackage);
+							const FString WorldPackageName = WorldPackage->GetName();
+							const bool bIncludeReadOnlyRoots = false;
+							const bool bIsValidPath = FPackageName::IsValidLongPackageName(WorldPackageName, bIncludeReadOnlyRoots);
+							if (!bIsValidPath)
+							{
+								WorldPackage->MarkPackageDirty();
+								OutDirtyPackages.Add(WorldPackage);
+							}
 						}
-					}
 
-					OutDirtyPackages.Add(BuiltDataPackage);
+						OutDirtyPackages.Add(BuiltDataPackage);
+					}
 				}
 			}
 		}
