UEditorLoadingAndSavingUtils::SavePackages no longer prompts user about files that are checked out in another workspace and silently attempts to check it out (unlike UEditorLoadingAndSavingUtils::SavePackagesWithDialog).

diff --git a/Engine/Source/Editor/UnrealEd/Private/FileHelpers.cpp b/Engine/Source/Editor/UnrealEd/Private/FileHelpers.cpp
--- a/Engine/Source/Editor/UnrealEd/Private/FileHelpers.cpp
+++ b/Engine/Source/Editor/UnrealEd/Private/FileHelpers.cpp
@@ -1760,7 +1760,7 @@ bool FEditorFileUtils::PromptToCheckoutPackages(bool bCheckDirty, const TArray<U
 	return bResult;
 }
 
-ECommandResult::Type FEditorFileUtils::CheckoutPackages(const TArray<UPackage*>& PkgsToCheckOut, TArray<UPackage*>* OutPackagesCheckedOut, const bool bErrorIfAlreadyCheckedOut)
+ECommandResult::Type FEditorFileUtils::CheckoutPackages(const TArray<UPackage*>& PkgsToCheckOut, TArray<UPackage*>* OutPackagesCheckedOut, const bool bErrorIfAlreadyCheckedOut, const bool bConfirmPackageBranchCheckOutStatus)
 {
 	ECommandResult::Type CheckOutResult = ECommandResult::Succeeded;
 	FString PkgsWhichFailedCheckout;
@@ -1779,7 +1779,7 @@ ECommandResult::Type FEditorFileUtils::CheckoutPackages(const TArray<UPackage*>&
 	if(CheckOutResult != ECommandResult::Cancelled)
 	{
 		// If any packages are checked out or modified in another branch, prompt for confirmation
-		if (!ConfirmPackageBranchCheckOutStatus(PkgsToCheckOut))
+		if (bConfirmPackageBranchCheckOutStatus && !ConfirmPackageBranchCheckOutStatus(PkgsToCheckOut))
 		{
 			return ECommandResult::Cancelled;
 		}
@@ -4237,8 +4237,8 @@ static bool InternalCheckoutAndSavePackages(const TArray<UPackage*>& PackagesToS
 			TGuardValue<bool> UnattendedScriptGuard(GIsRunningUnattendedScript, true);
 
 			const bool bErrorIfAlreadyCheckedOut = false;
-			const bool bShowDialogIfFailure = false;
-			FEditorFileUtils::CheckoutPackages(PackagesToSave, nullptr, bErrorIfAlreadyCheckedOut);
+			const bool bConfirmPackageBranchCheckOutStatus = bUseDialog;
+			FEditorFileUtils::CheckoutPackages(PackagesToSave, nullptr, bErrorIfAlreadyCheckedOut, bConfirmPackageBranchCheckOutStatus);
 
 			TArray<UPackage*> FailedPackages;
 			FEditorFileUtils::EPromptReturnCode ReturnResponse = InternalPromptForCheckoutAndSave(PackagesToSave, bUseDialog, FailedPackages);
diff --git a/Engine/Source/Editor/UnrealEd/Public/FileHelpers.h b/Engine/Source/Editor/UnrealEd/Public/FileHelpers.h
--- a/Engine/Source/Editor/UnrealEd/Public/FileHelpers.h
+++ b/Engine/Source/Editor/UnrealEd/Public/FileHelpers.h
@@ -440,10 +440,11 @@ public:
 	 * @param	PkgsToCheckOut							Reference to array of packages to check out 
 	 * @param	OutPackagesCheckedOut					If not NULL, this array will be populated with packages that were checked out.
 	 * @param	bErrorIfAlreadyCheckedOut				true to consider being unable to checkout a package because it is already checked out an error, false to allow this without error
+	 * @param	bConfirmPackageBranchCheckOutStatus		true to prompt user on whether a package that is checked out or modified in another branch should be checked out, false to silently attempt check out.
 	 *
 	 * @return	true if all the packages were checked out successfully
 	 */
-	UNREALED_API static ECommandResult::Type CheckoutPackages(const TArray<UPackage*>& PkgsToCheckOut, TArray<UPackage*>* OutPackagesCheckedOut = NULL, const bool bErrorIfAlreadyCheckedOut = true);
+	UNREALED_API static ECommandResult::Type CheckoutPackages(const TArray<UPackage*>& PkgsToCheckOut, TArray<UPackage*>* OutPackagesCheckedOut = NULL, const bool bErrorIfAlreadyCheckedOut = true, const bool bConfirmPackageBranchCheckOutStatus = true);
 
 	/**
 	 * Check out the specified packages from source control and report any errors while checking out
