Merging //UE4/Main[REDACTED_HANDLE] to Dev-RenderPlat-Staging

diff --git a/Engine/Source/Editor/UnrealEd/Private/Fbx/FbxSkeletalMeshImport.cpp b/Engine/Source/Editor/UnrealEd/Private/Fbx/FbxSkeletalMeshImport.cpp
--- a/Engine/Source/Editor/UnrealEd/Private/Fbx/FbxSkeletalMeshImport.cpp
+++ b/Engine/Source/Editor/UnrealEd/Private/Fbx/FbxSkeletalMeshImport.cpp
@@ -27,6 +27,9 @@
 #include "Animation/MorphTarget.h"
 #include "PhysicsAssetUtils.h"
 #include "Rendering/SkeletalMeshModel.h"
+#include "Rendering/SkeletalMeshRenderData.h"
+#include "Rendering/SkeletalMeshLODRenderData.h"
+#include "Misc/ScopedSlowTask.h"
 
 #include "SkelImport.h"
 #include "Logging/TokenizedMessage.h"
@@ -52,8 +55,9 @@
 #include "Misc/FbxErrors.h"
 #include "PhysicsEngine/PhysicsAsset.h"
 #include "Engine/SkeletalMeshSocket.h"
-#include "Assets/ClothingAsset.h"
+#include "ClothingAsset.h"
 #include "LODUtilities.h"
+#include "IMeshBuilderModule.h"
 
 #define LOCTEXT_NAMESPACE "FBXImpoter"
 
@@ -1433,7 +1437,7 @@ bool UnFbx::FFbxImporter::FillSkeletalMeshImportData(TArray<FbxNode*>& NodeArray
 
 	if (ImportOptions->bPreserveSmoothingGroups)
 	{
-		bool bDuplicateUnSmoothWedges = (ImportOptions->NormalGenerationMethod != EFBXNormalGenerationMethod::MikkTSpace);
+		bool bDuplicateUnSmoothWedges = false; //We deprecate legacy build
 		DoUnSmoothVerts(*SkelMeshImportDataPtr, bDuplicateUnSmoothWedges);
 	}
 	else
@@ -1754,28 +1758,40 @@ USkeletalMesh* UnFbx::FFbxImporter::ImportSkeletalMesh(FImportSkeletalMeshArgs &
 		return nullptr;
 	}
 
+	//Stack the PostEditChange call, it will call post edit change when it will go out of scope
+	FScopedSkeletalMeshPostEditChange ScopedPostEditChange(ExistingSkelMesh);
+
+	ESkeletalMeshGeoImportVersions GeoImportVersion = ESkeletalMeshGeoImportVersions::LatestVersion;
+	ESkeletalMeshSkinningImportVersions SkinningImportVersion = ESkeletalMeshSkinningImportVersions::LatestVersion;
 	//Adjust the import data from the import options
 	if (ExistingSkelMesh != nullptr)
 	{
+		int32 SafeLODIndex = ImportSkeletalMeshArgs.LodIndex < 0 ? 0 : ImportSkeletalMeshArgs.LodIndex;
+		if (ExistingSkelMesh->GetImportedModel() && ExistingSkelMesh->GetImportedModel()->LODModels.IsValidIndex(SafeLODIndex))
+		{
+			const FRawSkeletalMeshBulkData& RawSkeletalMeshBulkData = ExistingSkelMesh->GetImportedModel()->LODModels[SafeLODIndex].RawSkeletalMeshBulkData;
+			GeoImportVersion = RawSkeletalMeshBulkData.GeoImportVersion;
+			SkinningImportVersion = RawSkeletalMeshBulkData.SkinningImportVersion;
+		}
+
 		if ((ImportOptions->bImportAsSkeletalSkinning || ImportOptions->bImportAsSkeletalGeometry))
 		{
 			if (ImportOptions->bImportAsSkeletalSkinning)
 			{
 				//Replace geometry import data by original existing skel mesh geometry data
-				ReplaceSkeletalMeshGeometryImportData(ExistingSkelMesh, SkelMeshImportDataPtr, ImportSkeletalMeshArgs.LodIndex);
+				ReplaceSkeletalMeshGeometryImportData(ExistingSkelMesh, SkelMeshImportDataPtr, SafeLODIndex);
 			}
 			else if (ImportOptions->bImportAsSkeletalGeometry)
 			{
 				//Replace skinning import data by original existing skel mesh skinning data
-				ReplaceSkeletalMeshSkinningImportData(ExistingSkelMesh, SkelMeshImportDataPtr, ImportSkeletalMeshArgs.LodIndex);
+				ReplaceSkeletalMeshSkinningImportData(ExistingSkelMesh, SkelMeshImportDataPtr, SafeLODIndex);
 			}
 		}
 
 		//Reuse the vertex color in case we are not re-importing weights only and there is some vertexcolor in the existing skeleMesh
 		if (ExistingSkelMesh->bHasVertexColors && ImportOptions->VertexColorImportOption == EVertexColorImportOption::Ignore)
 		{
-			int32 LODIndex = ImportSkeletalMeshArgs.LodIndex < 0 ? 0 : ImportSkeletalMeshArgs.LodIndex;
-			RemapSkeletalMeshVertexColorToImportData(ExistingSkelMesh, LODIndex, SkelMeshImportDataPtr);
+			RemapSkeletalMeshVertexColorToImportData(ExistingSkelMesh, SafeLODIndex, SkelMeshImportDataPtr);
 		}
 	}
 
@@ -1783,15 +1799,6 @@ USkeletalMesh* UnFbx::FFbxImporter::ImportSkeletalMesh(FImportSkeletalMeshArgs &
 
 	// Create initial bounding box based on expanded version of reference pose for meshes without physics assets. Can be overridden by artist.
 	FBox BoundingBox(SkelMeshImportDataPtr->Points.GetData(), SkelMeshImportDataPtr->Points.Num());
-	FBox Temp = BoundingBox;
-	FVector MidMesh = 0.5f*(Temp.Min + Temp.Max);
-	BoundingBox.Min = Temp.Min + 1.0f*(Temp.Min - MidMesh);
-	BoundingBox.Max = Temp.Max + 1.0f*(Temp.Max - MidMesh);
-	// Tuck up the bottom as this rarely extends lower than a reference pose's (e.g. having its feet on the floor).
-	// Maya has Y in the vertical, other packages have Z.
-	//BEN const int32 CoordToTuck = bAssumeMayaCoordinates ? 1 : 2;
-	//BEN BoundingBox.Min[CoordToTuck]	= Temp.Min[CoordToTuck] + 0.1f*(Temp.Min[CoordToTuck] - MidMesh[CoordToTuck]);
-	BoundingBox.Min[2] = Temp.Min[2] + 0.1f*(Temp.Min[2] - MidMesh[2]);
 	const FVector BoundingBoxSize = BoundingBox.GetSize();
 
 	if (SkelMeshImportDataPtr->Points.Num() > 2 && BoundingBoxSize.X < THRESH_POINTS_ARE_SAME && BoundingBoxSize.Y < THRESH_POINTS_ARE_SAME && BoundingBoxSize.Z < THRESH_POINTS_ARE_SAME)
@@ -1805,18 +1812,8 @@ USkeletalMesh* UnFbx::FFbxImporter::ImportSkeletalMesh(FImportSkeletalMeshArgs &
 		return nullptr;
 	}
 
-	// Array of re-import contexts for components using this mesh
-	// Will unregister before import, then re-register afterwards
-	TIndirectArray<FComponentReregisterContext> ComponentContexts;
-	TUniquePtr< FSkinnedMeshComponentRecreateRenderStateContext > SkinnedMeshComponentRecreateRenderStateContext;
-
-	FSkinnedMeshComponentRecreateRenderStateContext* RecreateExistingRenderStateContext = ExistingSkelMesh ? new FSkinnedMeshComponentRecreateRenderStateContext(ExistingSkelMesh, false) : nullptr;
-	//Backup the data before importing the new one
 	if (ExistingSkelMesh)
 	{
-		//Release the Renderdata resources before reimporting the skeletal mesh
-		SkinnedMeshComponentRecreateRenderStateContext = MakeUnique<FSkinnedMeshComponentRecreateRenderStateContext>(ExistingSkelMesh);
-
 		ExistingSkelMesh->PreEditChange(NULL);
 		//The backup of the skeletal mesh data empty the LOD array in the ImportedResource of the skeletal mesh
 		//If the import fail after this step the editor can crash when updating the bone later since the LODModel will not exist anymore
@@ -1825,12 +1822,18 @@ USkeletalMesh* UnFbx::FFbxImporter::ImportSkeletalMesh(FImportSkeletalMeshArgs &
 
 	if (SkeletalMesh == nullptr)
 	{
+		//Backup the PostEditchangeStackCounter before overriding the skeletalmesh
+		int32 PostEditChangeStackCounter = ExistingSkelMesh != nullptr ? ExistingSkelMesh->GetPostEditChangeStackCounter() : 0;
 		// Create the new mesh after saving the old data, since it will replace the old skeletalmesh
 		// This should happen only when doing a re-import. Otherwise the skeletal mesh is create before.
 		SkeletalMesh = NewObject<USkeletalMesh>(ImportSkeletalMeshArgs.InParent, ImportSkeletalMeshArgs.Name, ImportSkeletalMeshArgs.Flags);
+		SkeletalMesh->SetPostEditChangeStackCounter(PostEditChangeStackCounter);
 	}
 
+
 	SkeletalMesh->PreEditChange(NULL);
+	//Dirty the DDC Key for any imported Skeletal Mesh
+	SkeletalMesh->InvalidateDeriveDataCacheGUID();
 
 	FSkeletalMeshModel *ImportedResource = SkeletalMesh->GetImportedModel();
 	check(ImportedResource->LODModels.Num() == 0);
@@ -1858,9 +1861,24 @@ USkeletalMesh* UnFbx::FFbxImporter::ImportSkeletalMesh(FImportSkeletalMeshArgs &
 	// process bone influences from import data
 	ProcessImportMeshInfluences(*SkelMeshImportDataPtr);
 
-	//Store the original fbx import data
+	//Store the original fbx import data the SkelMeshImportDataPtr should not be modified after this
 	LODModel.RawSkeletalMeshBulkData.SaveRawMesh(*SkelMeshImportDataPtr);
-
+	if (ImportOptions->bImportAsSkeletalSkinning)
+	{
+		LODModel.RawSkeletalMeshBulkData.SkinningImportVersion = ESkeletalMeshSkinningImportVersions::LatestVersion;
+		LODModel.RawSkeletalMeshBulkData.GeoImportVersion = GeoImportVersion;
+	}
+	else if(ImportOptions->bImportAsSkeletalGeometry)
+	{
+		LODModel.RawSkeletalMeshBulkData.GeoImportVersion = ESkeletalMeshGeoImportVersions::LatestVersion;
+		LODModel.RawSkeletalMeshBulkData.SkinningImportVersion = SkinningImportVersion;
+	}
+	else
+	{
+		//We reimport both
+		LODModel.RawSkeletalMeshBulkData.SkinningImportVersion = ESkeletalMeshSkinningImportVersions::LatestVersion;
+		LODModel.RawSkeletalMeshBulkData.GeoImportVersion = ESkeletalMeshGeoImportVersions::LatestVersion;
+	}
 	SkeletalMesh->ResetLODInfo();
 	FSkeletalMeshLODInfo& NewLODInfo = SkeletalMesh->AddLODInfo();
 	NewLODInfo.ReductionSettings.NumOfTrianglesPercentage = 1.0f;
@@ -1886,38 +1904,61 @@ USkeletalMesh* UnFbx::FFbxImporter::ImportSkeletalMesh(FImportSkeletalMeshArgs &
 		TArray<int32> LODPointToRawMap;
 		SkelMeshImportDataPtr->CopyLODImportData(LODPoints,LODWedges,LODFaces,LODInfluences,LODPointToRawMap);
 
-		IMeshUtilities::MeshBuildOptions BuildOptions;
-		BuildOptions.OverlappingThresholds = ImportOptions->OverlappingThresholds;
-		BuildOptions.bComputeNormals = !ImportOptions->ShouldImportNormals() || !SkelMeshImportDataPtr->bHasNormals;
-		BuildOptions.bComputeTangents = !ImportOptions->ShouldImportTangents() || !SkelMeshImportDataPtr->bHasTangents;
+		FSkeletalMeshBuildSettings BuildOptions;
+		BuildOptions.bBuildAdjacencyBuffer = true;
+		BuildOptions.bRecomputeNormals = !ImportOptions->ShouldImportNormals() || !SkelMeshImportDataPtr->bHasNormals;
+		BuildOptions.bRecomputeTangents = !ImportOptions->ShouldImportTangents() || !SkelMeshImportDataPtr->bHasTangents;
 		BuildOptions.bUseMikkTSpace = (ImportOptions->NormalGenerationMethod == EFBXNormalGenerationMethod::MikkTSpace) && (!ImportOptions->ShouldImportNormals() || !ImportOptions->ShouldImportTangents());
 		BuildOptions.bComputeWeightedNormals = ImportOptions->bComputeWeightedNormals;
-		BuildOptions.bRemoveDegenerateTriangles = false;
-
-		IMeshUtilities& MeshUtilities = FModuleManager::Get().LoadModuleChecked<IMeshUtilities>("MeshUtilities");
-		
-		TArray<FText> WarningMessages;
-		TArray<FName> WarningNames;
-		// Create actual rendering data.
-		bool bBuildSuccess = MeshUtilities.BuildSkeletalMesh(ImportedResource->LODModels[0],SkeletalMesh->RefSkeleton,LODInfluences, LODWedges,LODFaces,LODPoints,LODPointToRawMap,BuildOptions, &WarningMessages, &WarningNames);
+		BuildOptions.bRemoveDegenerates = ImportOptions->bRemoveDegenerates;
+		BuildOptions.ThresholdPosition = ImportOptions->OverlappingThresholds.ThresholdPosition;
+		BuildOptions.ThresholdTangentNormal = ImportOptions->OverlappingThresholds.ThresholdTangentNormal;
+		BuildOptions.ThresholdUV = ImportOptions->OverlappingThresholds.ThresholdUV;
 
-		// temporary hack of message/names, should be one token or a struct
-		if(WarningMessages.Num() > 0 && WarningNames.Num() == WarningMessages.Num())
+		bool bUseLegacyMeshUtilities = false;
+		bool bBuildSuccess = false;
+		if (bUseLegacyMeshUtilities)
 		{
-			EMessageSeverity::Type MessageSeverity = bBuildSuccess ? EMessageSeverity::Warning : EMessageSeverity::Error;
+			IMeshUtilities::MeshBuildOptions LegacyBuildOptions;
+			LegacyBuildOptions.FillOptions(BuildOptions);
 
-			for(int32 MessageIdx = 0; MessageIdx<WarningMessages.Num(); ++MessageIdx)	
+			IMeshUtilities& MeshUtilities = FModuleManager::Get().LoadModuleChecked<IMeshUtilities>("MeshUtilities");
+
+			TArray<FText> WarningMessages;
+			TArray<FName> WarningNames;
+			// Create actual rendering data.
+			bBuildSuccess = MeshUtilities.BuildSkeletalMesh(ImportedResource->LODModels[0], SkeletalMesh->RefSkeleton, LODInfluences, LODWedges, LODFaces, LODPoints, LODPointToRawMap, LegacyBuildOptions, &WarningMessages, &WarningNames);
+
+			// temporary hack of message/names, should be one token or a struct
+			if (WarningMessages.Num() > 0 && WarningNames.Num() == WarningMessages.Num())
 			{
-				AddTokenizedErrorMessage(FTokenizedMessage::Create(MessageSeverity, WarningMessages[MessageIdx]), WarningNames[MessageIdx]);
+				EMessageSeverity::Type MessageSeverity = bBuildSuccess ? EMessageSeverity::Warning : EMessageSeverity::Error;
+
+				for (int32 MessageIdx = 0; MessageIdx < WarningMessages.Num(); ++MessageIdx)
+				{
+					AddTokenizedErrorMessage(FTokenizedMessage::Create(MessageSeverity, WarningMessages[MessageIdx]), WarningNames[MessageIdx]);
+				}
 			}
 		}
+		else
+		{
+			//The imported LOD is always 0 here, the LOD custom import will import the LOD alone(in a temporary skeletalmesh) and add it to the base skeletal mesh later
+			const int32 LODIndex = 0;
+			check(SkeletalMesh->GetLODInfo(LODIndex) != nullptr);
+			//Set the build options
+			SkeletalMesh->GetLODInfo(LODIndex)->BuildSettings = BuildOptions;
+			//New MeshDescription build process
+			IMeshBuilderModule& MeshBuilderModule = FModuleManager::Get().LoadModuleChecked<IMeshBuilderModule>("MeshBuilder");
+			//We must build the LODModel so we can restore properly the mesh, but we do not have to regenerate LODs
+			bBuildSuccess = MeshBuilderModule.BuildSkeletalMesh(SkeletalMesh, LODIndex, false);
+		}
 
 		if( !bBuildSuccess )
 		{
 			SkeletalMesh->MarkPendingKill();
 			return NULL;
 		}
-
+		
 		// Presize the per-section shadow casting array with the number of sections in the imported LOD.
 		const int32 NumSections = LODModel.Sections.Num();
 
@@ -1975,18 +2016,12 @@ USkeletalMesh* UnFbx::FFbxImporter::ImportSkeletalMesh(FImportSkeletalMeshArgs &
 		}
 
 		SkeletalMesh->CalculateInvRefMatrices();
-		SkeletalMesh->PostEditChange();
-		SkeletalMesh->MarkPackageDirty();
-
-		// Now iterate over all skeletal mesh components re-initialising them.
-		for (TObjectIterator<USkinnedMeshComponent> It; It; ++It)
+		if ((!SkeletalMesh->GetResourceForRendering() || !SkeletalMesh->GetResourceForRendering()->LODRenderData.IsValidIndex(0)) && ImportOptions->bCreatePhysicsAsset)
 		{
-			USkinnedMeshComponent* SkinComp = *It;
-			if (SkinComp->SkeletalMesh == SkeletalMesh)
-			{
-				ComponentContexts.Add(new FComponentReregisterContext(SkinComp));
-			}
+			//We need to have a valid render data to create physic asset
+			SkeletalMesh->Build();
 		}
+		SkeletalMesh->MarkPackageDirty();
 	}
 
 	if(ImportSkeletalMeshArgs.LodIndex == 0)
@@ -2142,13 +2177,6 @@ USkeletalMesh* UnFbx::FFbxImporter::ImportSkeletalMesh(FImportSkeletalMeshArgs &
 		ImportNodeCustomProperties(SkeletalMesh->Skeleton, SkeletonNode, true);
 	}
 
-	// ComponentContexts will now go out of scope, causing components to be re-registered
-
-	if (RecreateExistingRenderStateContext)
-	{
-		delete RecreateExistingRenderStateContext;
-	}
-
 	return SkeletalMesh;
 }
 
@@ -2352,8 +2380,6 @@ USkeletalMesh* UnFbx::FFbxImporter::ReimportSkeletalMesh(USkeletalMesh* Mesh, UF
 		return NULL;
 	}
 
-	char MeshName[1024];
-	FCStringAnsi::Strcpy(MeshName, TCHAR_TO_UTF8(*Mesh->GetName()));
 	TArray<FbxNode*>* FbxNodes = NULL;
 	USkeletalMesh* NewMesh = NULL;
 
@@ -2362,16 +2388,35 @@ USkeletalMesh* UnFbx::FFbxImporter::ReimportSkeletalMesh(USkeletalMesh* Mesh, UF
 	bool Old_ImportTextures = ImportOptions->bImportTextures;
 	bool Old_ImportAnimations = ImportOptions->bImportAnimations;
 
+	auto CleanUpImportOptionsData = [&]()
+	{
+		ImportOptions->bImportRigidMesh = Old_ImportRigidMesh;
+		ImportOptions->bImportMaterials = Old_ImportMaterials;
+		ImportOptions->bImportTextures = Old_ImportTextures;
+		ImportOptions->bImportAnimations = Old_ImportAnimations;
+	};
+
 	// support to update rigid animation mesh
 	ImportOptions->bImportRigidMesh = true;
 
-	// Unbind any existing clothing assets before we reimport the geometry
-	TArray<ClothingAssetUtils::FClothingAssetMeshBinding> ClothingBindings;
-	ClothingAssetUtils::GetMeshClothingAssetBindings(Mesh, ClothingBindings);
-
-	for(ClothingAssetUtils::FClothingAssetMeshBinding& Binding : ClothingBindings)
+	UFbxSkeletalMeshImportData* SKImportData = Cast<UFbxSkeletalMeshImportData>(Mesh->AssetImportData);
+	if (SKImportData)
 	{
-		Binding.Asset->UnbindFromSkeletalMesh(Mesh, Binding.LODIndex);
+		FSkeletalMeshLODInfo* LODInfo = Mesh->GetLODInfo(0);
+		if(LODInfo && Mesh->GetImportedModel() && Mesh->GetImportedModel()->LODModels.IsValidIndex(0))
+		{
+			const FSkeletalMeshLODModel& LODModel = Mesh->GetImportedModel()->LODModels[0];
+			if (!LODModel.RawSkeletalMeshBulkData.IsBuildDataAvailable())
+			{
+				//Set the build settings
+				LODInfo->BuildSettings.bBuildAdjacencyBuffer = true;
+				LODInfo->BuildSettings.bComputeWeightedNormals = SKImportData->bComputeWeightedNormals;
+				LODInfo->BuildSettings.bRecomputeNormals = SKImportData->NormalImportMethod == EFBXNormalImportMethod::FBXNIM_ComputeNormals;
+				LODInfo->BuildSettings.bRecomputeTangents = SKImportData->NormalImportMethod != EFBXNormalImportMethod::FBXNIM_ImportNormalsAndTangents;
+				LODInfo->BuildSettings.bRemoveDegenerates = true;
+				LODInfo->BuildSettings.bUseMikkTSpace = SKImportData->NormalGenerationMethod == EFBXNormalGenerationMethod::MikkTSpace;
+			}
+		}
 	}
 
 	// get meshes in Fbx file
@@ -2427,26 +2472,53 @@ USkeletalMesh* UnFbx::FFbxImporter::ReimportSkeletalMesh(USkeletalMesh* Mesh, UF
 			if (FbxNodes != nullptr)
 				break;
 		}
-		if (FbxNodes == nullptr)
-			return nullptr;
 	}
 	else
 	{
 		// if there is only one mesh, use it without name checking 
 		// (because the "Used As Full Name" option enables users name the Unreal mesh by themselves
-		if (FbxSkelMeshArray.Num() > 0)
+		if (FbxSkelMeshArray.Num() == 1)
 		{
 			FbxNodes = FbxSkelMeshArray[0];
 		}
-		else
+		else if (FbxSkelMeshArray.Num() > 1)
 		{
-			// @todo - FBX Importing - We need proper support for reimport if the file contains more than one skeletal mesh. 
-
+			//Search a set corresponding
+			for (TArray<FbxNode*>*SkeletalMeshNodes : FbxSkelMeshArray)
+			{
+				FbxNode* Node = GetMeshNodesFromName(Mesh->GetName(), *SkeletalMeshNodes);
+				if (Node != nullptr)
+				{
+					FbxNodes = SkeletalMeshNodes;
+					break;
+				}
+			}
+			if (FbxNodes == nullptr)
+			{
+				//Re import the first skeletalmesh found in the fbx file
+				FbxNodes = FbxSkelMeshArray[0];
+			}
 		}
 	}
-
+	
 	if (FbxNodes)
 	{
+		//Posteditchange will be call at the end of the reimport which will reallocate the render ressource
+		FScopedSkeletalMeshPostEditChange ScopedPostEditChange(Mesh);
+
+		FScopedSlowTask SlowTask(100.0f, NSLOCTEXT("SkeltalMeshReimportFactory", "SkeletalMeshReimportTask", "Reimporting skeletal mesh..."));
+		SlowTask.MakeDialog();
+
+		//The fbx file is valid we can reimport the skeletal mesh
+		Mesh->PreEditChange(nullptr);
+
+		// Unbind any existing clothing assets before we reimport the geometry
+		TArray<ClothingAssetUtils::FClothingAssetMeshBinding> ClothingBindings;
+		FLODUtilities::UnbindClothingAndBackup(Mesh, ClothingBindings);
+
+		//Reimport must force a new guid
+		Mesh->InvalidateDeriveDataCacheGUID();
+
 		for (int32 LODCounter = 1; LODCounter < Mesh->GetLODNum(); ++LODCounter)
 		{
 			//Reset all import with base mesh flags
@@ -2491,8 +2563,10 @@ USkeletalMesh* UnFbx::FFbxImporter::ReimportSkeletalMesh(USkeletalMesh* Mesh, UF
 		int32 LODIndex;
 		int32 SuccessfulLodIndex = 0;
 		bool bOperationCanceled = false;
+		float ProgressStep = 90.0 / MaxLODLevel;
 		for (LODIndex = 0; LODIndex < MaxLODLevel; LODIndex++)
 		{
+			SlowTask.EnterProgressFrame(ProgressStep);
 			int32 ImportedSuccessfulLodIndex = INDEX_NONE;
 			TArray<FbxNode*> SkelMeshNodeArray;
 			for (int32 j = 0; j < (*FbxNodes).Num(); j++)
@@ -2542,8 +2616,9 @@ USkeletalMesh* UnFbx::FFbxImporter::ReimportSkeletalMesh(USkeletalMesh* Mesh, UF
 				NewMesh = ImportSkeletalMesh( ImportSkeletalMeshArgs );
 				if (bOperationCanceled)
 				{
-					// User cancelled, clean up and return
-					return nullptr;
+					// User cancel, clean up and return
+					NewMesh = nullptr;
+					break;
 				}
 				if (NewMesh)
 				{
@@ -2590,7 +2665,7 @@ USkeletalMesh* UnFbx::FFbxImporter::ReimportSkeletalMesh(USkeletalMesh* Mesh, UF
 				ImportSkeletalMeshArgs.OutData = &OutData;
 
 				UObject *LODObject = ImportSkeletalMesh( ImportSkeletalMeshArgs );
-				bool bImportSucceeded = !bOperationCanceled && ImportSkeletalMeshLOD( Cast<USkeletalMesh>(LODObject), BaseSkeletalMesh, SuccessfulLodIndex, false, nullptr, TemplateImportData);
+				bool bImportSucceeded = !bOperationCanceled && ImportSkeletalMeshLOD( Cast<USkeletalMesh>(LODObject), BaseSkeletalMesh, SuccessfulLodIndex, TemplateImportData);
 
 				if (bImportSucceeded)
 				{
@@ -2626,47 +2701,41 @@ USkeletalMesh* UnFbx::FFbxImporter::ReimportSkeletalMesh(USkeletalMesh* Mesh, UF
 
 			if (NewMesh)
 			{
-				// import morph target
 				if ((ImportOptions->bImportSkeletalMeshLODs || LODIndex == 0) &&
 					ImportOptions->bImportMorph &&
 					ImportedSuccessfulLodIndex != INDEX_NONE &&
 					NewMesh->GetImportedModel() &&
-					NewMesh->GetImportedModel()->LODModels.IsValidIndex(LODIndex))
+					NewMesh->GetImportedModel()->LODModels.IsValidIndex(ImportedSuccessfulLodIndex))
 				{
-					// @fixme: @question : where do they import this morph? where to? What morph target sets?
 					ImportFbxMorphTarget(SkelMeshNodeArray, NewMesh, ImportedSuccessfulLodIndex, OutData);
 				}
 			}
 		}
-		if (NewMesh)
+
+		auto RebindClothing = [&ClothingBindings](USkeletalMesh* SkelMesh)
 		{
-			FSkeletalMeshModel* NewMeshResource = NewMesh->GetImportedModel();
-			if(NewMeshResource && ClothingBindings.Num() > 0)
+			FSkeletalMeshModel* MeshResource = SkelMesh->GetImportedModel();
+			if (MeshResource && ClothingBindings.Num() > 0)
 			{
-				NewMesh->PreEditChange(nullptr);
+				FLODUtilities::RestoreClothingFromBackup(SkelMesh, ClothingBindings);
+			}
+		};
 
-				for(ClothingAssetUtils::FClothingAssetMeshBinding& Binding : ClothingBindings)
-				{
-					if(NewMeshResource->LODModels.IsValidIndex(Binding.LODIndex) &&
-						NewMeshResource->LODModels[Binding.LODIndex].Sections.IsValidIndex(Binding.SectionIndex))
-					{
-						Binding.Asset->BindToSkeletalMesh(NewMesh, Binding.LODIndex, Binding.SectionIndex, Binding.AssetInternalLodIndex);
-					}
-				}
 
-				NewMesh->PostEditChange();
-			}
+		if (!bOperationCanceled && NewMesh)
+		{
+			RebindClothing(NewMesh);
 
 			//Update the import data so we can re-import correctly
 			UpdateSkeletalMeshImportData(NewMesh, TemplateImportData, INDEX_NONE, &ImportMaterialOriginalNameData, &ImportMeshLodData);
-			//If we have import some morph target we have to rebuild the render resources since morph target are now using GPU
-			if (NewMesh->MorphTargets.Num() > 0)
-			{
-				NewMesh->ReleaseResources();
-				//Rebuild the resources with a post edit change since we have added some morph targets
-				NewMesh->PostEditChange();
-			}
 		}
+		else if(bOperationCanceled && Mesh != nullptr)
+		{
+			//Operation was cancel
+			//rebind the clothing on the original mesh
+			RebindClothing(Mesh);
+		}
+		SlowTask.EnterProgressFrame(10.0f);
 	}
 	else
 	{
@@ -2674,10 +2743,7 @@ USkeletalMesh* UnFbx::FFbxImporter::ReimportSkeletalMesh(USkeletalMesh* Mesh, UF
 		AddTokenizedErrorMessage(FTokenizedMessage::Create(EMessageSeverity::Warning, FText::Format(LOCTEXT("FbxSkeletaLMeshimport_NoFBXMeshMatch", "No FBX mesh matches the Unreal mesh '{0}'."), FText::FromString(Mesh->GetName()))), FFbxErrors::Generic_Mesh_MeshNotFound);
 	}
 
-	ImportOptions->bImportRigidMesh = Old_ImportRigidMesh;
-	ImportOptions->bImportMaterials = Old_ImportMaterials;
-	ImportOptions->bImportTextures = Old_ImportTextures;
-	ImportOptions->bImportAnimations = Old_ImportAnimations;
+	CleanUpImportOptionsData();
 
 	return NewMesh;
 }
@@ -3411,7 +3477,7 @@ bool UnFbx::FFbxImporter::FillSkelMeshImporterFromFbx( FSkeletalMeshImportData&
 			//
 			// normals, tangents and binormals
 			//
-			if( ImportOptions->ShouldImportNormals() && bHasNormalInformation )
+			if( bHasNormalInformation )
 			{
 				int32 TmpIndex = LocalIndex*3 + VertexIndex;
 				//normals may have different reference and mapping mode than tangents and binormals
@@ -3425,7 +3491,7 @@ bool UnFbx::FFbxImporter::FillSkelMeshImporterFromFbx( FSkeletalMeshImportData&
 
 				FbxVector4 TempValue;
 
-				if( ImportOptions->ShouldImportTangents() && bHasTangentInformation )
+				if( bHasTangentInformation )
 				{
 					TempValue = LayerElementTangent->GetDirectArray().GetAt(TangentMapIndex);
 					TempValue = TotalMatrixForNormal.MultT(TempValue);
@@ -3722,6 +3788,9 @@ bool UnFbx::FFbxImporter::FillSkelMeshImporterFromFbx( FSkeletalMeshImportData&
 
 void UnFbx::FFbxImporter::InsertNewLODToBaseSkeletalMesh(USkeletalMesh* InSkeletalMesh, USkeletalMesh* BaseSkeletalMesh, int32 DesiredLOD, UFbxSkeletalMeshImportData* TemplateImportData)
 {
+	//Scope a post edit change
+	FScopedSkeletalMeshPostEditChange ScopedPostEditChange(BaseSkeletalMesh);
+
 	FSkeletalMeshModel* ImportedResource = InSkeletalMesh->GetImportedModel();
 	FSkeletalMeshModel* DestImportedResource = BaseSkeletalMesh->GetImportedModel();
 
@@ -3781,6 +3850,13 @@ void UnFbx::FFbxImporter::InsertNewLODToBaseSkeletalMesh(USkeletalMesh* InSkelet
 
 	// Set up LODMaterialMap to number of materials in new mesh.
 	FSkeletalMeshLODInfo& LODInfo = *(BaseSkeletalMesh->GetLODInfo(DesiredLOD));
+
+	//Copy the build settings
+	if(InSkeletalMesh->GetLODInfo(0))
+	{
+		const FSkeletalMeshLODInfo& ImportedLODInfo = *(InSkeletalMesh->GetLODInfo(0));
+		LODInfo.BuildSettings = ImportedLODInfo.BuildSettings;
+	}
 	
 	LODInfo.LODMaterialMap.Empty();
 	// Now set up the material mapping array.
@@ -3874,12 +3950,9 @@ void UnFbx::FFbxImporter::InsertNewLODToBaseSkeletalMesh(USkeletalMesh* InSkelet
 			}
 		}
 	}
-
-	// rebuild vertex buffers and reinit RHI resources
-	BaseSkeletalMesh->PostEditChange();
 }
 
-bool UnFbx::FFbxImporter::ImportSkeletalMeshLOD(USkeletalMesh* InSkeletalMesh, USkeletalMesh* BaseSkeletalMesh, int32 DesiredLOD, bool bNeedToReregister, TArray<UActorComponent*>* ReregisterAssociatedComponents, UFbxSkeletalMeshImportData* TemplateImportData)
+bool UnFbx::FFbxImporter::ImportSkeletalMeshLOD(USkeletalMesh* InSkeletalMesh, USkeletalMesh* BaseSkeletalMesh, int32 DesiredLOD, UFbxSkeletalMeshImportData* TemplateImportData)
 {
 	check(InSkeletalMesh);
 	check(BaseSkeletalMesh);
@@ -3943,6 +4016,8 @@ bool UnFbx::FFbxImporter::ImportSkeletalMeshLOD(USkeletalMesh* InSkeletalMesh, U
 		}
 	}
 
+	FScopedSkeletalMeshPostEditChange ScopedPostEditChange(BaseSkeletalMesh);
+
 	//Skeleton do not match and we are importing skeletal geometry only so apply the base LOD Skinning
 	if (bApplyBaseSkinning)
 	{
@@ -3950,8 +4025,10 @@ bool UnFbx::FFbxImporter::ImportSkeletalMeshLOD(USkeletalMesh* InSkeletalMesh, U
 		//Apply the base LOD skinning to the new LOD geometry, DestImportedResource here is the base lod and is the source for the apply skinning.
 		//The Dest here mean we merge the just import LOD to the base LodModel.
 		SkeletalMeshHelper::ApplySkinning(InSkeletalMesh, DestImportedResource->LODModels[0], ImportedResource->LODModels[0]);
-		IMeshUtilities& MeshUtilities = FModuleManager::Get().LoadModuleChecked<IMeshUtilities>("MeshUtilities");
-		MeshUtilities.RemoveBonesFromMesh(BaseSkeletalMesh, DesiredLOD, NULL);
+
+		//Excluded bones will be remove when we rebuild the asset
+		//IMeshUtilities& MeshUtilities = FModuleManager::Get().LoadModuleChecked<IMeshUtilities>("MeshUtilities");
+		//MeshUtilities.RemoveBonesFromMesh(BaseSkeletalMesh, DesiredLOD, NULL);
 	}
 
 	FSkeletalMeshLODModel& NewLODModel = ImportedResource->LODModels[0];
@@ -4114,32 +4191,7 @@ bool UnFbx::FFbxImporter::ImportSkeletalMeshLOD(USkeletalMesh* InSkeletalMesh, U
 		}
 	}
 
-	if (bNeedToReregister)
-	{
-		// Shut down the skeletal mesh component that is previewing this mesh.
-		if (ReregisterAssociatedComponents)
-		{
-			FMultiComponentReregisterContext ReregisterContext(*ReregisterAssociatedComponents);
-			// wait until resources are released
-			FlushRenderingCommands();
-
-			InsertNewLODToBaseSkeletalMesh(InSkeletalMesh, BaseSkeletalMesh, DesiredLOD, TemplateImportData);
-
-			// ReregisterContexts go out of scope here, reregistering associated components with the scene.
-		}
-		else
-		{
-			TComponentReregisterContext<USkinnedMeshComponent> ComponentReregisterContext;
-
-			InsertNewLODToBaseSkeletalMesh(InSkeletalMesh, BaseSkeletalMesh, DesiredLOD, TemplateImportData);
-
-			// ReregisterContexts go out of scope here, reregistering skel components with the scene.
-		}
-	}
-	else
-	{
-		InsertNewLODToBaseSkeletalMesh(InSkeletalMesh, BaseSkeletalMesh, DesiredLOD, TemplateImportData);
-	}
+	InsertNewLODToBaseSkeletalMesh(InSkeletalMesh, BaseSkeletalMesh, DesiredLOD, TemplateImportData);
 
 	return true;
 }
@@ -4280,9 +4332,21 @@ void UnFbx::FFbxImporter::ImportMorphTargetsInternal( TArray<FbxNode*>& SkelMesh
 		check(BaseImportData.MorphTargetNames.Num() == BaseImportData.MorphTargets.Num() && BaseImportData.MorphTargetNames.Num() == BaseImportData.MorphTargetModifiedPoints.Num());
 	}
 
-	//Build MorphTargets
-	// Calculate overlapping corners and tangents
-	FLODUtilities::BuildMorphTargets(BaseSkelMesh, BaseImportData, LODIndex, ImportOptions->ShouldImportNormals(), ImportOptions->ShouldImportTangents(), (ImportOptions->NormalGenerationMethod == EFBXNormalGenerationMethod::MikkTSpace));
+	if (BaseSkelMesh->GetImportedModel() && BaseSkelMesh->GetImportedModel()->LODModels.IsValidIndex(LODIndex))
+	{
+		//If we can build the skeletal mesh there is no need to build the morph target now, all the necessary build morph target data was copied before.
+		if (!BaseSkelMesh->GetImportedModel()->LODModels[LODIndex].RawSkeletalMeshBulkData.IsBuildDataAvailable())
+		{
+			//Build MorphTargets
+			FLODUtilities::BuildMorphTargets(
+				BaseSkelMesh,
+				BaseImportData,
+				LODIndex,
+				ImportOptions->ShouldImportNormals(),
+				ImportOptions->ShouldImportTangents(),
+				(ImportOptions->NormalGenerationMethod == EFBXNormalGenerationMethod::MikkTSpace));
+		}
+	}
 
 	GWarn->EndSlowTask();
 }
@@ -4290,6 +4354,9 @@ void UnFbx::FFbxImporter::ImportMorphTargetsInternal( TArray<FbxNode*>& SkelMesh
 // Import Morph target
 void UnFbx::FFbxImporter::ImportFbxMorphTarget(TArray<FbxNode*> &SkelMeshNodeArray, USkeletalMesh* BaseSkelMesh, int32 LODIndex, FSkeletalMeshImportData &BaseSkeletalMeshImportData)
 {
+	//Stack the PostEditChange call, it will call post edit change when it will go out of scope
+	FScopedSkeletalMeshPostEditChange ScopedPostEditChange(BaseSkelMesh);
+
 	bool bHasMorph = false;
 	int32 NodeIndex;
 	// check if there are morph in this geometry
@@ -4308,7 +4375,12 @@ void UnFbx::FFbxImporter::ImportFbxMorphTarget(TArray<FbxNode*> &SkelMeshNodeArr
 	
 	if (bHasMorph)
 	{
-		ImportMorphTargetsInternal( SkelMeshNodeArray, BaseSkelMesh, LODIndex, BaseSkeletalMeshImportData );
+		ImportMorphTargetsInternal( SkelMeshNodeArray, BaseSkelMesh, LODIndex, BaseSkeletalMeshImportData);
+		//Save the rawMesh
+		if (BaseSkelMesh->GetImportedModel() && BaseSkelMesh->GetImportedModel()->LODModels.IsValidIndex(LODIndex))
+		{
+			BaseSkelMesh->GetImportedModel()->LODModels[LODIndex].RawSkeletalMeshBulkData.SaveRawMesh(BaseSkeletalMeshImportData);
+		}
 	}
 }
 
diff --git a/Engine/Source/Editor/UnrealEd/Private/FbxMeshUtils.cpp b/Engine/Source/Editor/UnrealEd/Private/FbxMeshUtils.cpp
--- a/Engine/Source/Editor/UnrealEd/Private/FbxMeshUtils.cpp
+++ b/Engine/Source/Editor/UnrealEd/Private/FbxMeshUtils.cpp
@@ -33,8 +33,9 @@
 #include "Framework/Notifications/NotificationManager.h"
 #include "Widgets/Notifications/SNotificationList.h"
 
-#include "Assets/ClothingAsset.h"
+#include "ClothingAsset.h"
 #include "SkinWeightsUtilities.h"
+#include "LODUtilities.h"
 
 DEFINE_LOG_CATEGORY_STATIC(LogExportMeshUtils, Log, All);
 
@@ -280,10 +281,13 @@ namespace FbxMeshUtils
 			return false;
 		}
 
+		FScopedSkeletalMeshPostEditChange ScopePostEditChange(SelectedSkelMesh);
+
 		//If the imported LOD already exist, we will need to reimport all the skin weight profiles
 		bool bMustReimportAlternateSkinWeightProfile = false;
 
 		// Get a list of all the clothing assets affecting this LOD so we can re-apply later
+		TArray<ClothingAssetUtils::FClothingAssetMeshBinding> ClothingBindings;
 		TArray<UClothingAssetBase*> ClothingAssetsInUse;
 		TArray<int32> ClothingAssetSectionIndices;
 		TArray<int32> ClothingAssetInternalLodIndices;
@@ -292,47 +296,16 @@ namespace FbxMeshUtils
 		if(ImportedResource && ImportedResource->LODModels.IsValidIndex(LODLevel))
 		{
 			bMustReimportAlternateSkinWeightProfile = true;
-			FSkeletalMeshLODModel& LodModel = ImportedResource->LODModels[LODLevel];
-
-			const int32 NumSections = LodModel.Sections.Num();
-
-			for(int32 SectionIndex = 0; SectionIndex < NumSections; ++SectionIndex)
-			{
-				FSkelMeshSection& Section = LodModel.Sections[SectionIndex];
-
-				if(Section.HasClothingData())
-				{
-					UClothingAssetBase* AssetInUse = SelectedSkelMesh->GetSectionClothingAsset(LODLevel, SectionIndex);
-					ClothingAssetsInUse.Add(AssetInUse);
-					ClothingAssetSectionIndices.Add(SectionIndex);
-					ClothingAssetInternalLodIndices.Add(Section.ClothingData.AssetLodIndex);
-				}
-			}
-		}
-
-		// Remove our clothing assets while we import this LOD
-		for(UClothingAssetBase* ClothingAsset : ClothingAssetsInUse)
-		{
-			ClothingAsset->UnbindFromSkeletalMesh(SelectedSkelMesh, LODLevel);
+			FLODUtilities::UnbindClothingAndBackup(SelectedSkelMesh, ClothingBindings, LODLevel);
 		}
 
 		//Lambda to call to re-apply the clothing
-		auto ReapplyClothing = [&SelectedSkelMesh, &ClothingAssetsInUse, &ClothingAssetSectionIndices, &ClothingAssetInternalLodIndices, &ImportedResource, &LODLevel]()
+		auto ReapplyClothing = [&SelectedSkelMesh, &ClothingBindings, &ImportedResource, &LODLevel]()
 		{
-			// Re-apply our clothing assets
-			int32 NumClothingAssetsToApply = ClothingAssetsInUse.Num();
 			if (ImportedResource && ImportedResource->LODModels.IsValidIndex(LODLevel))
 			{
-				FSkeletalMeshLODModel& LodModel = ImportedResource->LODModels[LODLevel];
-				for (int32 AssetIndex = 0; AssetIndex < NumClothingAssetsToApply; ++AssetIndex)
-				{
-					// Only if the same section exists
-					if (LodModel.Sections.IsValidIndex(ClothingAssetSectionIndices[AssetIndex]))
-					{
-						UClothingAssetBase* AssetToApply = ClothingAssetsInUse[AssetIndex];
-						AssetToApply->BindToSkeletalMesh(SelectedSkelMesh, LODLevel, ClothingAssetSectionIndices[AssetIndex], ClothingAssetInternalLodIndices[AssetIndex]);
-					}
-				}
+				// Re-apply our clothing assets
+				FLODUtilities::RestoreClothingFromBackup(SelectedSkelMesh, ClothingBindings, LODLevel);
 			}
 		};
 
@@ -395,6 +368,24 @@ namespace FbxMeshUtils
 			TArray<FString> LODStrings;
 			TArray<FbxNode*>* MeshObject = NULL;;
 
+			//Set the build options if the BuildDat is not available so it is the same option we use to import the LOD
+			if (ImportedResource && ImportedResource->LODModels.IsValidIndex(LODLevel) && !ImportedResource->LODModels[LODLevel].RawSkeletalMeshBulkData.IsBuildDataAvailable())
+			{
+				FSkeletalMeshLODInfo* LODInfo = SelectedSkelMesh->GetLODInfo(LODLevel);
+				if (LODInfo)
+				{
+					LODInfo->BuildSettings.bBuildAdjacencyBuffer = true;
+					LODInfo->BuildSettings.bRecomputeNormals = !ImportOptions->ShouldImportNormals();
+					LODInfo->BuildSettings.bRecomputeTangents = !ImportOptions->ShouldImportTangents();
+					LODInfo->BuildSettings.bUseMikkTSpace = (ImportOptions->NormalGenerationMethod == EFBXNormalGenerationMethod::MikkTSpace) && (!ImportOptions->ShouldImportNormals() || !ImportOptions->ShouldImportTangents());
+					LODInfo->BuildSettings.bComputeWeightedNormals = ImportOptions->bComputeWeightedNormals;
+					LODInfo->BuildSettings.bRemoveDegenerates = ImportOptions->bRemoveDegenerates;
+					LODInfo->BuildSettings.ThresholdPosition = ImportOptions->OverlappingThresholds.ThresholdPosition;
+					LODInfo->BuildSettings.ThresholdTangentNormal = ImportOptions->OverlappingThresholds.ThresholdTangentNormal;
+					LODInfo->BuildSettings.ThresholdUV = ImportOptions->OverlappingThresholds.ThresholdUV;
+				}
+			}
+
 			// Populate the mesh array
 			FFbxImporter->FillFbxSkelMeshArrayInScene(FFbxImporter->Scene->GetRootNode(), MeshArray, false, ImportOptions->bImportAsSkeletalGeometry || ImportOptions->bImportAsSkeletalSkinning, ImportOptions->bImportScene);
 
@@ -524,12 +515,8 @@ namespace FbxMeshUtils
 
 				TempSkelMesh = (USkeletalMesh*)FFbxImporter->ImportSkeletalMesh( ImportSkeletalMeshArgs );
 				// Add the new imported LOD to the existing model (check skeleton compatibility)
-				if( TempSkelMesh  && FFbxImporter->ImportSkeletalMeshLOD(TempSkelMesh, SelectedSkelMesh, SelectedLOD, true, nullptr, TempAssetImportData))
+				if( TempSkelMesh  && FFbxImporter->ImportSkeletalMeshLOD(TempSkelMesh, SelectedSkelMesh, SelectedLOD, TempAssetImportData))
 				{
-					TComponentReregisterContext<USkinnedMeshComponent> ReregisterContext;
-					SelectedSkelMesh->ReleaseResources();
-					SelectedSkelMesh->ReleaseResourcesFence.Wait();
-					
 					//Update the import data for this lod
 					UnFbx::FFbxImporter::UpdateSkeletalMeshImportData(SelectedSkelMesh, nullptr, SelectedLOD, &ImportMaterialOriginalNameData, &ImportMeshLodData);
 
@@ -555,13 +542,9 @@ namespace FbxMeshUtils
 					{
 						//We cannot use anymore the FFbxImporter after the cleanup
 						CleanUpScene();
-						FSkinWeightsUtilities::ReimportAlternateSkinWeight(SelectedSkelMesh, SelectedLOD, false);
+						FSkinWeightsUtilities::ReimportAlternateSkinWeight(SelectedSkelMesh, SelectedLOD);
 					}
 
-					SelectedSkelMesh->PostEditChange();
-					SelectedSkelMesh->InitResources();
-					SelectedSkelMesh->MarkPackageDirty();
-
 					// Notification of success
 					FNotificationInfo NotificationInfo(FText::GetEmpty());
 					NotificationInfo.Text = FText::Format(NSLOCTEXT("UnrealEd", "LODImportSuccessful", "Mesh for LOD {0} imported successfully!"), FText::AsNumber(SelectedLOD));
@@ -633,7 +616,7 @@ namespace FbxMeshUtils
 		return ChosenFilname;
 	}
 
-	bool ImportMeshLODDialog(class UObject* SelectedMesh, int32 LODLevel)
+	bool ImportMeshLODDialog(class UObject* SelectedMesh, int32 LODLevel, bool bNotifyCB /*= true*/)
 	{
 		if(!SelectedMesh)
 		{
@@ -732,7 +715,7 @@ namespace FbxMeshUtils
 			FMessageDialog::Open(EAppMsgType::Ok, FText::Format(LOCTEXT("LODImport_Failure", "Failed to import LOD{0}"), FText::AsNumber(LODLevel)));
 		}
 
-		if (bImportSuccess)
+		if (bImportSuccess && bNotifyCB)
 		{
 			if (SkeletalMesh)
 			{
diff --git a/Engine/Source/Runtime/ClothingSystemRuntimeInterface/Public/ClothingAssetBase.h b/Engine/Source/Runtime/ClothingSystemRuntimeInterface/Public/ClothingAssetBase.h
--- a/Engine/Source/Runtime/ClothingSystemRuntimeInterface/Public/ClothingAssetBase.h
+++ b/Engine/Source/Runtime/ClothingSystemRuntimeInterface/Public/ClothingAssetBase.h
@@ -23,7 +23,7 @@ public:
 	* @param InSubmeshIdx Submesh in this asset to replace section with
 	* @param InAssetLodIndex Internal clothing LOD to use
 	*/
-	virtual bool BindToSkeletalMesh(USkeletalMesh* InSkelMesh, int32 InMeshLodIndex, int32 InSectionIndex, int32 InAssetLodIndex, bool bCallPostEditChange = true)
+	virtual bool BindToSkeletalMesh(USkeletalMesh* InSkelMesh, const int32 InMeshLodIndex, const int32 InSectionIndex, const int32 InAssetLodIndex, const bool bCallPostEditChange = true)
 	PURE_VIRTUAL(UClothingAssetBase::BindToSkeletalMesh, return false;);
 
 	/**
@@ -38,7 +38,7 @@ public:
 	* @param InSkelMesh skeletal mesh to remove from
 	* @param InMeshLodIndex Mesh LOD to remove this asset from (could still be bound to other LODs)
 	*/
-	virtual void UnbindFromSkeletalMesh(USkeletalMesh* InSkelMesh, int32 InMeshLodIndex)
+	virtual void UnbindFromSkeletalMesh(USkeletalMesh* InSkelMesh, const int32 InMeshLodIndex)
 	PURE_VIRTUAL(UClothingAssetBase::UnbindFromSkeletalMesh,);
 
 	/** 
@@ -65,10 +65,26 @@ public:
 	virtual int32 GetNumLods()
 	PURE_VIRTUAL(UClothingAssetBase::GetNumLods(), return 0;);
 
+	/** Add a new LOD class instance. */
+	virtual int32 AddNewLod()
+	PURE_VIRTUAL(UClothingAssetBase::AddNewLod(), return INDEX_NONE;);
+
+	/** Called from \c FPersonaMeshDetails::OnFinishedChangingClothingProperties(). */
+	virtual void PostPropertyChangeCb(const FPropertyChangedEvent& Event)
+	PURE_VIRTUAL(UClothingAssetBase::PostPropertyChangeCb, );
+
 	/** Builds self collision data */
 	virtual void BuildSelfCollisionData()
 	PURE_VIRTUAL(UClothingAssetBase::BuildSelfCollisionData(), );
 
+	/**
+	*	Builds the LOD transition data
+	*	When we transition between LODs we skin the incoming mesh to the outgoing mesh
+	*	in exactly the same way the render mesh is skinned to create a smooth swap
+	*/
+	virtual void BuildLodTransitionData()
+	PURE_VIRTUAL(UClothingAssetBase::BuildLodTransitionData(), );
+
 	/** Get the guid identifying this asset */
 	const FGuid& GetAssetGuid() const
 	{
diff --git a/Engine/Source/Runtime/Engine/Private/SkeletalRenderGPUSkin.cpp b/Engine/Source/Runtime/Engine/Private/SkeletalRenderGPUSkin.cpp
--- a/Engine/Source/Runtime/Engine/Private/SkeletalRenderGPUSkin.cpp
+++ b/Engine/Source/Runtime/Engine/Private/SkeletalRenderGPUSkin.cpp
@@ -333,6 +333,7 @@ void FSkeletalMeshObjectGPUSkin::UpdateDynamicData_RenderThread(FGPUSkinCache* G
 
 #if RHI_RAYTRACING
 	bRequireRecreatingRayTracingGeometry = (DynamicData == nullptr || // Newly created
+		RayTracingGeometry.Initializer.PositionVertexBuffer == nullptr ||
 		(DynamicData != nullptr && DynamicData->LODIndex != InDynamicData->LODIndex)); // LOD level changed
 #endif
 
@@ -574,7 +575,7 @@ void FSkeletalMeshObjectGPUSkin::ProcessUpdatedDynamicData(FGPUSkinCache* GPUSki
 			// Create a uniform buffer from the bone transforms.
 			TArray<FMatrix>& ReferenceToLocalMatrices = DynamicData->ReferenceToLocal;
 			bool bNeedFence = ShaderData.UpdateBoneData(RHICmdList, ReferenceToLocalMatrices, Section.BoneMap, RevisionNumber, false, FeatureLevel, bUseSkinCache);
-#if WITH_APEX_CLOTHING
+#if WITH_APEX_CLOTHING || WITH_CHAOS_CLOTHING
 			// Update uniform buffer for APEX cloth simulation mesh positions and normals
 			if( bClothFactory )
 			{				
@@ -587,7 +588,7 @@ void FSkeletalMeshObjectGPUSkin::ProcessUpdatedDynamicData(FGPUSkinCache* GPUSki
 					bNeedFence = ClothShaderData.UpdateClothSimulData(RHICmdList, DynamicData->ClothingSimData[ActorIdx].Positions, DynamicData->ClothingSimData[ActorIdx].Normals, FrameNumberToPrepare, FeatureLevel) || bNeedFence;
 				}
 			}
-#endif // WITH_APEX_CLOTHING
+#endif // WITH_APEX_CLOTHING || WITH_CHAOS_CLOTHING
 			// Try to use the GPU skinning cache if possible
 			if (bUseSkinCache)
 			{
diff --git a/Engine/Source/Runtime/Experimental/Chaos/Public/Chaos/KinematicGeometryParticles.h b/Engine/Source/Runtime/Experimental/Chaos/Public/Chaos/KinematicGeometryParticles.h
--- a/Engine/Source/Runtime/Experimental/Chaos/Public/Chaos/KinematicGeometryParticles.h
+++ b/Engine/Source/Runtime/Experimental/Chaos/Public/Chaos/KinematicGeometryParticles.h
@@ -10,7 +10,7 @@ template<class T, int d, EGeometryParticlesSimType SimType>
 class TKinematicGeometryParticlesImp : public TGeometryParticlesImp<T, d, SimType>
 {
   public:
-	TKinematicGeometryParticlesImp()
+	CHAOS_API TKinematicGeometryParticlesImp()
 	    : TGeometryParticlesImp<T, d, SimType>()
 	{
 		this->MParticleType = EParticleType::Kinematic;
@@ -18,14 +18,14 @@ class TKinematicGeometryParticlesImp : public TGeometryParticlesImp<T, d, SimTyp
 		TArrayCollection::AddArray(&MW);
 	}
 	TKinematicGeometryParticlesImp(const TKinematicGeometryParticlesImp<T, d, SimType>& Other) = delete;
-	TKinematicGeometryParticlesImp(TKinematicGeometryParticlesImp<T, d, SimType>&& Other)
+	CHAOS_API TKinematicGeometryParticlesImp(TKinematicGeometryParticlesImp<T, d, SimType>&& Other)
 	    : TGeometryParticlesImp<T, d, SimType>(MoveTemp(Other)), MV(MoveTemp(Other.MV)), MW(MoveTemp(Other.MW))
 	{
 		this->MParticleType = EParticleType::Kinematic;
 		TArrayCollection::AddArray(&MV);
 		TArrayCollection::AddArray(&MW);
 	}
-	~TKinematicGeometryParticlesImp() {}
+	CHAOS_API virtual ~TKinematicGeometryParticlesImp();
 
 	const TVector<T, d>& V(const int32 Index) const { return MV[Index]; }
 	TVector<T, d>& V(const int32 Index) { return MV[Index]; }
@@ -45,7 +45,7 @@ class TKinematicGeometryParticlesImp : public TGeometryParticlesImp<T, d, SimTyp
 	//cannot be reference because double pointer would allow for badness, but still useful to have non const access to handle
 	THandleType* Handle(int32 Index);
 
-	void Serialize(FChaosArchive& Ar)
+	CHAOS_API virtual void Serialize(FChaosArchive& Ar) override
 	{
 		TGeometryParticlesImp<T, d, SimType>::Serialize(Ar);
 		Ar << MV << MW;
@@ -63,8 +63,13 @@ FChaosArchive& operator<<(FChaosArchive& Ar, TKinematicGeometryParticlesImp<T, d
 	return Ar;
 }
 
+#ifdef __clang__
+extern template class CHAOS_API Chaos::TKinematicGeometryParticlesImp<float, 3, Chaos::EGeometryParticlesSimType::RigidBodySim>;
+extern template class CHAOS_API Chaos::TKinematicGeometryParticlesImp<float, 3, Chaos::EGeometryParticlesSimType::Other>;
+#else
 extern template class Chaos::TKinematicGeometryParticlesImp<float, 3, Chaos::EGeometryParticlesSimType::RigidBodySim>;
 extern template class Chaos::TKinematicGeometryParticlesImp<float, 3, Chaos::EGeometryParticlesSimType::Other>;
+#endif
 
 template <typename T, int d>
 using TKinematicGeometryParticles = TKinematicGeometryParticlesImp<T, d, EGeometryParticlesSimType::RigidBodySim>;
diff --git a/Engine/Source/Runtime/MovieSceneTracks/Public/Evaluation/MovieSceneParameterTemplate.h b/Engine/Source/Runtime/MovieSceneTracks/Public/Evaluation/MovieSceneParameterTemplate.h
--- a/Engine/Source/Runtime/MovieSceneTracks/Public/Evaluation/MovieSceneParameterTemplate.h
+++ b/Engine/Source/Runtime/MovieSceneTracks/Public/Evaluation/MovieSceneParameterTemplate.h
@@ -36,6 +36,8 @@ struct FEvaluatedParameterSectionValues
 	TArray<FVectorParameterNameAndValue, TInlineAllocator<2>> VectorValues;
 	/** Array of evaluated color values */
 	TArray<FColorParameterNameAndValue, TInlineAllocator<2>> ColorValues;
+	/** Array of evaluated transform values */
+	TArray<FTransformParameterNameAndValue, TInlineAllocator<2>> TransformValues;
 };
 
 /** Template that performs evaluation of parameter sections */
@@ -54,7 +56,7 @@ protected:
 	/** Evaluate our curves, outputting evaluated values into the specified container */
 	MOVIESCENETRACKS_API void EvaluateCurves(const FMovieSceneContext& Context, FEvaluatedParameterSectionValues& OutValues) const;
 
-private:
+protected:
 
 	/** The scalar parameter names and their associated curves. */
 	UPROPERTY()
@@ -67,6 +69,9 @@ private:
 	/** The color parameter names and their associated curves. */
 	UPROPERTY()
 	TArray<FColorParameterNameAndCurves> Colors;
+
+	UPROPERTY()
+	TArray<FTransformParameterNameAndCurves> Transforms;
 };
 
 /** Default accessor type for use with TMaterialTrackExecutionToken */
@@ -120,12 +125,12 @@ struct TMaterialTrackExecutionToken : IMovieSceneExecutionToken
 				continue;
 			}
 
+			// Save the old instance
+			Player.SavePreAnimatedState(*Object, Accessor.GetAnimTypeID(), FPreAnimatedTokenProducer(Accessor));
+
 			UMaterialInstanceDynamic* DynamicMaterialInstance = Cast<UMaterialInstanceDynamic>(Material);
 			if (!DynamicMaterialInstance)
 			{
-				// Save the old instance
-				Player.SavePreAnimatedState(*Object, Accessor.GetAnimTypeID(), FPreAnimatedTokenProducer(Accessor));
-
 				FString DynamicName = Material->GetName() + "_Animated";
 				FName UniqueDynamicName = MakeUniqueObjectName( Object, UMaterialInstanceDynamic::StaticClass() , *DynamicName );
 				DynamicMaterialInstance = UMaterialInstanceDynamic::Create( Material, Object, UniqueDynamicName );
diff --git a/Engine/Source/Runtime/Renderer/Private/PostProcess/PostProcessCombineLUTs.cpp b/Engine/Source/Runtime/Renderer/Private/PostProcess/PostProcessCombineLUTs.cpp
--- a/Engine/Source/Runtime/Renderer/Private/PostProcess/PostProcessCombineLUTs.cpp
+++ b/Engine/Source/Runtime/Renderer/Private/PostProcess/PostProcessCombineLUTs.cpp
@@ -232,6 +232,11 @@ public:
 	class FBlendCount : SHADER_PERMUTATION_RANGE_INT("BLENDCOUNT", 1, 5);
 	using FPermutationDomain = TShaderPermutationDomain<FBlendCount>;
 
+	static bool ShouldCompilePermutation(const FGlobalShaderPermutationParameters& Parameters)
+	{
+		return true;
+	}
+
 	static void ModifyCompilationEnvironment(const FGlobalShaderPermutationParameters& Parameters, FShaderCompilerEnvironment& OutEnvironment)
 	{
 		FGlobalShader::ModifyCompilationEnvironment(Parameters, OutEnvironment);
