Apply PIE Advanced options Safe Zone to mobile PIE standalone process using command line overrides
Use -featureleveles31 for standalone mobile PIE by default

diff --git a/Engine/Source/Editor/UnrealEd/Private/PlayLevelNewProcess.cpp b/Engine/Source/Editor/UnrealEd/Private/PlayLevelNewProcess.cpp
--- a/Engine/Source/Editor/UnrealEd/Private/PlayLevelNewProcess.cpp
+++ b/Engine/Source/Editor/UnrealEd/Private/PlayLevelNewProcess.cpp
@@ -132,8 +132,8 @@ void UEditorEngine::LaunchNewProcess(const FRequestPlaySessionParams& InParams,
 		}
 		else
 		{
-			// Otherwise, we'll fall back to ES2 emulation.
-			CommandLine += TEXT(" -featureleveles2");
+			// Otherwise, we'll fall back to ES31 emulation.
+			CommandLine += TEXT(" -featureleveles31");
 		}
 
 		// If we're currently running in OpenGL mode, pass that onto our newly spawned processes.
@@ -148,6 +148,22 @@ void UEditorEngine::LaunchNewProcess(const FRequestPlaySessionParams& InParams,
 		// Ensure the executable writes out a differently named config file to avoid multiple instances overwriting each other.
 		// ToDo: Should this be on all multi-client launches?
 		CommandLine += TEXT(" -MultiprocessSaveConfig");
+
+		// In order for the mobile previewer to adjust its safe zone according to the device profile specified in the editor play settings,
+		// we need to pass the PIESafeZoneOverride's values as command line variables to the new process that we are about to launch.
+		FMargin PIESafeZoneOverride = InParams.EditorPlaySettings->PIESafeZoneOverride;
+		CommandLine += FString::Printf(TEXT(" -SafeZonePaddingLeft=%f -SafeZonePaddingRight=%f -SafeZonePaddingTop=%f -SafeZonePaddingBottom=%f"),
+			PIESafeZoneOverride.Left,
+			PIESafeZoneOverride.Right,
+			PIESafeZoneOverride.Top,
+			PIESafeZoneOverride.Bottom
+		);
+
+		if (!PIESafeZoneOverride.GetDesiredSize().IsZero())
+		{
+			// Send -DrawUnSafeZones so that the red "unsafe" zones show up in the mobile previewer.
+			CommandLine += TEXT(" -DrawUnSafeZones");
+		}
 	}
 
 	if (InParams.SessionPreviewTypeOverride.Get(EPlaySessionPreviewType::NoPreview) == EPlaySessionPreviewType::VulkanPreview)
