Fix improper assumtion that a component's outer is of type AActorwhen checking for native component editing support

#rnx

diff --git a/Engine/Source/Editor/UnrealEd/Private/Kismet2/ComponentEditorUtils.cpp b/Engine/Source/Editor/UnrealEd/Private/Kismet2/ComponentEditorUtils.cpp
--- a/Engine/Source/Editor/UnrealEd/Private/Kismet2/ComponentEditorUtils.cpp
+++ b/Engine/Source/Editor/UnrealEd/Private/Kismet2/ComponentEditorUtils.cpp
@@ -173,7 +173,7 @@ FProperty* FComponentEditorUtils::GetPropertyForEditableNativeComponent(const UA
 	// Note: We aren't concerned with whether the component is marked editable - the component itself is responsible for determining which of its properties are editable	
 	UObject* ComponentOuter = (NativeComponent ? NativeComponent->GetOuter() : nullptr);
 	UClass* OwnerClass = (ComponentOuter ? ComponentOuter->GetClass() : nullptr);
-	AActor* OwnerCDO = (OwnerClass ? CastChecked<AActor>(OwnerClass->GetDefaultObject()) : nullptr);
+	UObject* OwnerCDO = (OwnerClass ? OwnerClass->GetDefaultObject() : nullptr);
 
 	if (OwnerClass != nullptr)
 	{
