Added OnClosed delegate to SCustomDialog, which will always trigger when the dialog is closed. This was causing issues when the X button was used to close a non-modal dialog and the calling code did not know the dialog had been dismissed.

#rnx

diff --git a/Engine/Source/Editor/UnrealEd/Private/Dialogs/CustomDialog.cpp b/Engine/Source/Editor/UnrealEd/Private/Dialogs/CustomDialog.cpp
--- a/Engine/Source/Editor/UnrealEd/Private/Dialogs/CustomDialog.cpp
+++ b/Engine/Source/Editor/UnrealEd/Private/Dialogs/CustomDialog.cpp
@@ -26,6 +26,8 @@ void SCustomDialog::Construct(const FArguments& InArgs)
 
 	check(InArgs._Buttons.Num() > 0);
 	
+	OnClosed = InArgs._OnClosed;
+
 	TSharedPtr<SHorizontalBox> ContentBox;
 	TSharedPtr<SHorizontalBox> ButtonBox;
 
@@ -148,7 +150,12 @@ int SCustomDialog::ShowModal()
 
 void SCustomDialog::Show()
 {
-	FSlateApplication::Get().AddWindow(StaticCastSharedRef<SWindow>(this->AsShared()), true);
+	TSharedRef<SWindow> Window = FSlateApplication::Get().AddWindow(StaticCastSharedRef<SWindow>(this->AsShared()), true);
+
+	if (OnClosed.IsBound())
+	{
+		Window->GetOnWindowClosedEvent().AddLambda([this](const TSharedRef<SWindow>& Window) { OnClosed.Execute(); });
+	}
 }
 
 /** Handle the button being clicked */
diff --git a/Engine/Source/Editor/UnrealEd/Public/Dialogs/CustomDialog.h b/Engine/Source/Editor/UnrealEd/Public/Dialogs/CustomDialog.h
--- a/Engine/Source/Editor/UnrealEd/Public/Dialogs/CustomDialog.h
+++ b/Engine/Source/Editor/UnrealEd/Public/Dialogs/CustomDialog.h
@@ -58,6 +58,9 @@ public:
 		/** Content for the dialog */
 		SLATE_ARGUMENT(TSharedPtr<SWidget>, DialogContent)
 
+		/** Event triggered when the dialog is closed, either because one of the buttons is pressed, or the windows is closed. */
+		SLATE_EVENT(FSimpleDelegate, OnClosed)
+
 	SLATE_END_ARGS()
 
 	void Construct(const FArguments& InArgs);
@@ -77,4 +80,6 @@ private:
 
 	/** The index of the button that was pressed last. */
 	int LastPressedButton = -1;
+
+	FSimpleDelegate OnClosed;
 };
\ No newline at end of file
