Make a utility function to make checks for whether something is a blueprint spawnable component consistent.
#rnx

diff --git a/Engine/Source/Editor/BlueprintGraph/Private/BlueprintComponentNodeSpawner.cpp b/Engine/Source/Editor/BlueprintGraph/Private/BlueprintComponentNodeSpawner.cpp
--- a/Engine/Source/Editor/BlueprintGraph/Private/BlueprintComponentNodeSpawner.cpp
+++ b/Engine/Source/Editor/BlueprintGraph/Private/BlueprintComponentNodeSpawner.cpp
@@ -7,6 +7,7 @@
 #include "EdGraphSchema_K2.h"
 #include "K2Node_AddComponent.h"
 #include "Kismet2/BlueprintEditorUtils.h"
+#include "Kismet2/KismetEditorUtilities.h"
 #include "Styling/SlateIconFinder.h"
 #include "BlueprintNodeTemplateCache.h"
 #include "ComponentAssetBroker.h"
@@ -79,7 +80,7 @@ UBlueprintComponentNodeSpawner* UBlueprintComponentNodeSpawner::Create(const FCo
 		return NodeSpawner;
 	}
 
-	if (ComponentClass->HasAnyClassFlags(CLASS_Abstract) || !ComponentClass->HasMetaData(FBlueprintMetadata::MD_BlueprintSpawnableComponent))
+	if (!FKismetEditorUtilities::IsClassABlueprintSpawnableComponent(ComponentClass))
 	{
 		// loaded class that is marked as abstract or not spawnable, don't create an entry:
 		return nullptr;
diff --git a/Engine/Source/Editor/BlueprintGraph/Private/K2Node_AddComponent.cpp b/Engine/Source/Editor/BlueprintGraph/Private/K2Node_AddComponent.cpp
--- a/Engine/Source/Editor/BlueprintGraph/Private/K2Node_AddComponent.cpp
+++ b/Engine/Source/Editor/BlueprintGraph/Private/K2Node_AddComponent.cpp
@@ -11,6 +11,7 @@
 #include "Engine/BlueprintGeneratedClass.h"
 #include "EdGraphSchema_K2.h"
 #include "Kismet2/BlueprintEditorUtils.h"
+#include "Kismet2/KismetEditorUtilities.h"
 #include "UObject/ReleaseObjectVersion.h"
 #include "KismetCompilerMisc.h"
 #include "KismetCompiler.h"
@@ -210,7 +211,7 @@ void UK2Node_AddComponent::ValidateNodeDuringCompilation(FCompilerResultsLog& Me
 	const UClass* TemplateClass = GetSpawnedType();
 	if (TemplateClass)
 	{
-		if (!TemplateClass->IsChildOf(UActorComponent::StaticClass()) || TemplateClass->HasAnyClassFlags(CLASS_Abstract) || !TemplateClass->HasMetaData(FBlueprintMetadata::MD_BlueprintSpawnableComponent) )
+		if (!FKismetEditorUtilities::IsClassABlueprintSpawnableComponent(TemplateClass))
 		{
 			FFormatNamedArguments Args;
 			Args.Add(TEXT("TemplateClass"), FText::FromString(TemplateClass->GetName()));
diff --git a/Engine/Source/Editor/UnrealEd/Private/ComponentTypeRegistry.cpp b/Engine/Source/Editor/UnrealEd/Private/ComponentTypeRegistry.cpp
--- a/Engine/Source/Editor/UnrealEd/Private/ComponentTypeRegistry.cpp
+++ b/Engine/Source/Editor/UnrealEd/Private/ComponentTypeRegistry.cpp
@@ -287,8 +287,7 @@ void FComponentTypeRegistryData::ForceRefreshComponentList()
 				!bBlueprintSkeletonClass &&
 				bPassesAllowedClasses)
 			{
-				if (!Class->HasAnyClassFlags(CLASS_Abstract) &&
-					Class->HasMetaData(FBlueprintMetadata::MD_BlueprintSpawnableComponent)) //@TODO: Fold this logic together with the one in UEdGraphSchema_K2::GetAddComponentClasses
+				if (FKismetEditorUtilities::IsClassABlueprintSpawnableComponent(Class))
 				{
 					TArray<FString> ClassGroupNames;
 					Class->GetClassGroupNames(ClassGroupNames);
diff --git a/Engine/Source/Editor/UnrealEd/Private/Kismet2/ComponentEditorUtils.cpp b/Engine/Source/Editor/UnrealEd/Private/Kismet2/ComponentEditorUtils.cpp
--- a/Engine/Source/Editor/UnrealEd/Private/Kismet2/ComponentEditorUtils.cpp
+++ b/Engine/Source/Editor/UnrealEd/Private/Kismet2/ComponentEditorUtils.cpp
@@ -112,7 +112,7 @@ protected:
 		else
 		{
 			// Actor component classes should not be abstract and must also be tagged as BlueprintSpawnable
-			bCanCreate = !ObjectClass->HasAnyClassFlags(CLASS_Abstract) && ObjectClass->HasMetaData(FBlueprintMetadata::MD_BlueprintSpawnableComponent);
+			bCanCreate = FKismetEditorUtilities::IsClassABlueprintSpawnableComponent(ObjectClass);
 		}
 
 		return bCanCreate;
@@ -429,8 +429,7 @@ bool FComponentEditorUtils::CanCopyComponent(const UActorComponent* ComponentToC
 		check(ComponentClass != nullptr);
 
 		// Component class cannot be abstract and must also be tagged as BlueprintSpawnable
-		return !ComponentClass->HasAnyClassFlags(CLASS_Abstract)
-			&& ComponentClass->HasMetaData(FBlueprintMetadata::MD_BlueprintSpawnableComponent);
+		return FKismetEditorUtilities::IsClassABlueprintSpawnableComponent(ComponentClass);
 	}
 
 	return false;
diff --git a/Engine/Source/Editor/UnrealEd/Private/Kismet2/Kismet2.cpp b/Engine/Source/Editor/UnrealEd/Private/Kismet2/Kismet2.cpp
--- a/Engine/Source/Editor/UnrealEd/Private/Kismet2/Kismet2.cpp
+++ b/Engine/Source/Editor/UnrealEd/Private/Kismet2/Kismet2.cpp
@@ -1155,7 +1155,7 @@ void FKismetEditorUtilities::AddComponentsToBlueprint(UBlueprint* Blueprint, con
 		UActorComponent* ActorComponent = Components[CompIndex];
 
 		// Filter out nulls and the components we won't be able to create.
-		if (ActorComponent && ActorComponent->GetClass()->HasMetaData(FBlueprintMetadata::MD_BlueprintSpawnableComponent))
+		if (ActorComponent && FKismetEditorUtilities::IsClassABlueprintSpawnableComponent(ActorComponent->GetClass()))
 		{
 			USceneComponent* SceneComponent = Cast<USceneComponent>(Components[CompIndex]);
 			if (SceneComponent)
@@ -2181,6 +2181,13 @@ bool FKismetEditorUtilities::CanBlueprintImplementInterface(UBlueprint const* Bl
 	return bCanImplementInterface;
 }
 
+bool FKismetEditorUtilities::IsClassABlueprintSpawnableComponent(const UClass* Class)
+{
+	return (!Class->HasAnyClassFlags(CLASS_Abstract) &&
+	        Class->IsChildOf<UActorComponent>() &&
+	        Class->HasMetaData(FBlueprintMetadata::MD_BlueprintSpawnableComponent));
+}
+
 bool FKismetEditorUtilities::IsClassABlueprintSkeleton(const UClass* Class)
 {
 	// Find generating blueprint for a class
diff --git a/Engine/Source/Editor/UnrealEd/Public/Kismet2/KismetEditorUtilities.h b/Engine/Source/Editor/UnrealEd/Public/Kismet2/KismetEditorUtilities.h
--- a/Engine/Source/Editor/UnrealEd/Public/Kismet2/KismetEditorUtilities.h
+++ b/Engine/Source/Editor/UnrealEd/Public/Kismet2/KismetEditorUtilities.h
@@ -278,6 +278,9 @@ public:
 	/** Check to see if a given class is blueprint skeleton class. */
 	static bool IsClassABlueprintSkeleton (const UClass* Class);
 
+	/** Check to see if a given class is blueprint spawnable component class. */
+	static bool IsClassABlueprintSpawnableComponent(const UClass* Class);
+
 	/** Check to see if a given class is a blueprint macro library */
 	static bool IsClassABlueprintMacroLibrary(const UClass* Class);
 
