Source Control: Unattended package save now handles adding new files to source control in more situations.
- CheckoutPackages now handles adding new files
- InternalCheckoutAndSavePackages now handles adding new files

diff --git a/Engine/Source/Editor/UnrealEd/Private/FileHelpers.cpp b/Engine/Source/Editor/UnrealEd/Private/FileHelpers.cpp
--- a/Engine/Source/Editor/UnrealEd/Private/FileHelpers.cpp
+++ b/Engine/Source/Editor/UnrealEd/Private/FileHelpers.cpp
@@ -1768,6 +1768,7 @@ ECommandResult::Type FEditorFileUtils::CheckoutPackages(const TArray<UPackage*>&
 	ISourceControlProvider& SourceControlProvider = ISourceControlModule::Get().GetProvider();
 
 	TArray<UPackage*> FinalPackageCheckoutList;
+	TArray<UPackage*> FinalPackageMarkForAddList;
 
 	// Source control may have been enabled in the package checkout dialog.
 	// Ensure the status is up to date
@@ -1806,6 +1807,15 @@ ECommandResult::Type FEditorFileUtils::CheckoutPackages(const TArray<UPackage*>&
 					bShowCheckoutError = false;
 					FinalPackageCheckoutList.Add(PackageToCheckOut);
 				}
+				else if (SourceControlState->CanAdd())
+				{
+					// Cannot add unsaved packages to source control
+					if (!PackageToCheckOut->HasAnyPackageFlags(PKG_NewlyCreated))
+					{
+						bShowCheckoutError = false;
+						FinalPackageMarkForAddList.Add(PackageToCheckOut);
+					}
+				}
 				else if( !bErrorIfAlreadyCheckedOut && SourceControlState->IsCheckedOut() && !SourceControlState->IsCheckedOutOther() )
 				{
 					bShowCheckoutError = false;
@@ -1826,12 +1836,32 @@ ECommandResult::Type FEditorFileUtils::CheckoutPackages(const TArray<UPackage*>&
 	if(FinalPackageCheckoutList.Num() > 0)
 	{
 		CheckOutResult = SourceControlProvider.Execute(ISourceControlOperation::Create<FCheckOut>(), FinalPackageCheckoutList);
-		if(CheckOutResult != ECommandResult::Cancelled)
+	}
+
+	// Attempt to mark for add each package the user specified that is not already tracked by source control
+	ECommandResult::Type MarkForAddResult = ECommandResult::Cancelled;
+	if (FinalPackageMarkForAddList.Num() > 0)
+	{
+		MarkForAddResult = SourceControlProvider.Execute(ISourceControlOperation::Create<FMarkForAdd>(), FinalPackageMarkForAddList);
+	}
+
+	TArray<UPackage*> CombinedPackageList = FinalPackageCheckoutList;
+	CombinedPackageList.Append(FinalPackageMarkForAddList);
+
+	if (CombinedPackageList.Num() > 0)
+	{
 		{
 			// Checked out some or all files successfully, so check their state
-			for( auto PkgsToCheckOutIter = FinalPackageCheckoutList.CreateConstIterator(); PkgsToCheckOutIter; ++PkgsToCheckOutIter )
+			for (int32 i=0; i < CombinedPackageList.Num(); ++i)
 			{
-				UPackage* CurPackage = *PkgsToCheckOutIter;
+				const bool bCheckedOut = (i < FinalPackageCheckoutList.Num()) && (CheckOutResult != ECommandResult::Cancelled);
+				const bool bMarkedForAdd = (i >= FinalPackageMarkForAddList.Num()) && (MarkForAddResult != ECommandResult::Cancelled);
+				if (!(bCheckedOut || bMarkedForAdd))
+				{
+					continue;
+				}
+
+				UPackage* CurPackage = CombinedPackageList[i];
 
 				// If we're ignoring the package due to the user ignoring it for saving, remove it from the ignore list
 				// as getting here means we've explicitly decided to save the asset.
@@ -1839,7 +1869,7 @@ ECommandResult::Type FEditorFileUtils::CheckoutPackages(const TArray<UPackage*>&
 				PackagesNotSavedDuringSaveAll.Remove(CurPackageName);
 
 				FSourceControlStatePtr SourceControlState = SourceControlProvider.GetState(CurPackage, EStateCacheUsage::Use);
-				if(SourceControlState.IsValid() && SourceControlState->IsCheckedOut())
+				if (SourceControlState.IsValid() && (SourceControlState->IsCheckedOut() || SourceControlState->IsAdded()))
 				{
 					if ( OutPackagesCheckedOut )
 					{
@@ -1900,6 +1930,7 @@ ECommandResult::Type FEditorFileUtils::CheckoutPackages(const TArray<FString>& P
 	}
 
 	TArray<FString> FinalPackageCheckoutList;
+	TArray<FString> FinalPackageMarkForAddList;
 	if(CheckOutResult != ECommandResult::Cancelled)
 	{
 		// Assemble a final list of packages to check out
@@ -1933,6 +1964,11 @@ ECommandResult::Type FEditorFileUtils::CheckoutPackages(const TArray<FString>& P
 					bShowCheckoutError = false;
 					FinalPackageCheckoutList.Add(PackageToCheckOutName);
 				}
+				else if (SourceControlState->CanAdd())
+				{
+					bShowCheckoutError = false;
+					FinalPackageMarkForAddList.Add(PackageToCheckOutName);
+				}
 				else if( !bErrorIfAlreadyCheckedOut && SourceControlState->IsCheckedOut() && !SourceControlState->IsCheckedOutOther() )
 				{
 					bShowCheckoutError = false;
@@ -1948,34 +1984,57 @@ ECommandResult::Type FEditorFileUtils::CheckoutPackages(const TArray<FString>& P
 		}
 	}
 
-	// Attempt to check out each package the user specified to be checked out that is not read only
-	if ( FinalPackageCheckoutList.Num() > 0 )
+	// We have an array of package names, but the SCC needs an array of their corresponding filenames
+	auto GetFilenamesFromPackageNames = [](const TArray<FString>& PackageNames)
 	{
-		{
-			// We have an array of package names, but the SCC needs an array of their corresponding filenames
-			TArray<FString> FinalPackageCheckoutListFilenames;
-			FinalPackageCheckoutListFilenames.Reserve(FinalPackageCheckoutList.Num());
+		TArray<FString> Filenames;
+		Filenames.Reserve(PackageNames.Num());
 
-			for( auto PkgsToCheckOutIter = FinalPackageCheckoutList.CreateConstIterator(); PkgsToCheckOutIter; ++PkgsToCheckOutIter )
+		for (const FString& PackageName : PackageNames)
+		{
+			FString PackageFilename;
+			if (FPackageName::DoesPackageExist(PackageName, nullptr, &PackageFilename))
 			{
-				const FString& PackageToCheckOutName = *PkgsToCheckOutIter;
-
-				FString PackageFilename;
-				if(FPackageName::DoesPackageExist(PackageToCheckOutName, nullptr, &PackageFilename))
-				{
-					FinalPackageCheckoutListFilenames.Add(PackageFilename);
-				}
+				Filenames.Add(PackageFilename);
 			}
-
-			CheckOutResult = SourceControlProvider.Execute(ISourceControlOperation::Create<FCheckOut>(), FinalPackageCheckoutListFilenames);
 		}
 
-		if(CheckOutResult != ECommandResult::Cancelled)
+		return Filenames;
+	};
+
+	// Attempt to check out each package the user specified to be checked out that is not read only
+	if (FinalPackageCheckoutList.Num() > 0)
+	{
+		// We have an array of package names, but the SCC needs an array of their corresponding filenames
+		TArray<FString> Filenames = GetFilenamesFromPackageNames(FinalPackageCheckoutList);
+		CheckOutResult = SourceControlProvider.Execute(ISourceControlOperation::Create<FCheckOut>(), Filenames);
+	}
+
+	// Attempt to mark for add each package the user specified not already tracked by source control
+	ECommandResult::Type MarkForAddResult = ECommandResult::Succeeded;
+	if (FinalPackageMarkForAddList.Num() > 0)
+	{
+		TArray<FString> Filenames = GetFilenamesFromPackageNames(FinalPackageMarkForAddList);
+		MarkForAddResult = SourceControlProvider.Execute(ISourceControlOperation::Create<FMarkForAdd>(), Filenames);
+	}
+
+	TArray<FString> CombinedPackageList = FinalPackageCheckoutList;
+	CombinedPackageList.Append(FinalPackageMarkForAddList);
+
+	if (CombinedPackageList.Num() > 0)
+	{
 		{
 			// Checked out some or all files successfully, so check their state
-			for( auto PkgsToCheckOutIter = FinalPackageCheckoutList.CreateConstIterator(); PkgsToCheckOutIter; ++PkgsToCheckOutIter )
+			for (int32 i = 0; i < CombinedPackageList.Num(); ++i)
 			{
-				const FString& CurPackageName = *PkgsToCheckOutIter;
+				const bool bCheckedOut = (i < FinalPackageCheckoutList.Num()) && (CheckOutResult != ECommandResult::Cancelled);
+				const bool bMarkedForAdd = (i >= FinalPackageCheckoutList.Num()) && (MarkForAddResult != ECommandResult::Cancelled);
+				if (!(bCheckedOut || bMarkedForAdd))
+				{
+					continue;
+				}
+
+				const FString& CurPackageName = CombinedPackageList[i];
 
 				// If we're ignoring the package due to the user ignoring it for saving, remove it from the ignore list
 				// as getting here means we've explicitly decided to save the asset.
@@ -1991,7 +2050,7 @@ ECommandResult::Type FEditorFileUtils::CheckoutPackages(const TArray<FString>& P
 					SourceControlState = SourceControlProvider.GetState(PackageFilename, EStateCacheUsage::Use);
 				}
 
-				if(SourceControlState.IsValid() && SourceControlState->IsCheckedOut())
+				if (SourceControlState.IsValid() && (SourceControlState->IsCheckedOut() || SourceControlState->IsAdded()))
 				{
 					if ( OutPackagesCheckedOut )
 					{
@@ -4236,12 +4295,32 @@ static bool InternalCheckoutAndSavePackages(const TArray<UPackage*>& PackagesToS
 			// Prevent modal window if not requested.
 			TGuardValue<bool> UnattendedScriptGuard(GIsRunningUnattendedScript, true);
 
+			TArray<UPackage*> PackagesCheckedOut;
 			const bool bErrorIfAlreadyCheckedOut = false;
 			const bool bConfirmPackageBranchCheckOutStatus = bUseDialog;
-			FEditorFileUtils::CheckoutPackages(PackagesToSave, nullptr, bErrorIfAlreadyCheckedOut, bConfirmPackageBranchCheckOutStatus);
+			FEditorFileUtils::CheckoutPackages(PackagesToSave, &PackagesCheckedOut, bErrorIfAlreadyCheckedOut, bConfirmPackageBranchCheckOutStatus);
+
+			// Cannot mark files for add until after packages saved
+			TArray<UPackage*> PackagesToMarkForAdd;
+			for (UPackage* Package : PackagesToSave)
+			{
+				// List unsaved packages that were not checked out
+				if (Package->HasAnyPackageFlags(PKG_NewlyCreated) && !PackagesCheckedOut.Contains(Package))
+				{
+					PackagesToMarkForAdd.Add(Package);
+				}
+			}
 
 			TArray<UPackage*> FailedPackages;
 			FEditorFileUtils::EPromptReturnCode ReturnResponse = InternalPromptForCheckoutAndSave(PackagesToSave, bUseDialog, FailedPackages);
+
+			// Mark files for add now that packages have saved
+			PackagesToMarkForAdd.RemoveAll([&FailedPackages](UPackage* Package) { return FailedPackages.Contains(Package); });
+			if (PackagesToMarkForAdd.Num() > 0)
+			{
+				FEditorFileUtils::CheckoutPackages(PackagesToMarkForAdd, nullptr, bErrorIfAlreadyCheckedOut);
+			}
+
 			bResult = (ReturnResponse == FEditorFileUtils::PR_Success);
 		}
 	}
