Capitalized labels of buttons in Bug Reports preferences pane

## Summary

The labels for the buttons in "Editor Preferences" > "Privacy" > "Bug Reports" were not capitalized contrary to other buttons in this UI.

## Test Plan

- Searched codebase for occurrences of "don't send" and "send unattended" just in case this copy was also elsewhere
- Built and ran editor on Mac and verified button labels were now capitalized

diff --git a/Engine/Source/Editor/UnrealEd/Private/Analytics/CrashReportsPrivacySettings.cpp b/Engine/Source/Editor/UnrealEd/Private/Analytics/CrashReportsPrivacySettings.cpp
--- a/Engine/Source/Editor/UnrealEd/Private/Analytics/CrashReportsPrivacySettings.cpp
+++ b/Engine/Source/Editor/UnrealEd/Private/Analytics/CrashReportsPrivacySettings.cpp
@@ -27,7 +27,7 @@ void UCrashReportsPrivacySettings::GetToggleCategoryAndPropertyNames(FName& OutC
 
 FText UCrashReportsPrivacySettings::GetFalseStateLabel() const
 {
-	return LOCTEXT("FalseStateLabel", "Don't send");
+	return LOCTEXT("FalseStateLabel", "Don't Send");
 };
 
 FText UCrashReportsPrivacySettings::GetFalseStateTooltip() const
@@ -42,7 +42,7 @@ FText UCrashReportsPrivacySettings::GetFalseStateDescription() const
 
 FText UCrashReportsPrivacySettings::GetTrueStateLabel() const
 {
-	return LOCTEXT("TrueStateLabel", "Send unattended bug reports");
+	return LOCTEXT("TrueStateLabel", "Send Unattended Bug Reports");
 };
 
 FText UCrashReportsPrivacySettings::GetTrueStateTooltip() const
