UEditorLoadingAndSavingUtils::SavePackages (no dialog version): don't ConfirmPackageBranchCheckOutStatus for files marked for add.

diff --git a/Engine/Source/Editor/UnrealEd/Private/FileHelpers.cpp b/Engine/Source/Editor/UnrealEd/Private/FileHelpers.cpp
--- a/Engine/Source/Editor/UnrealEd/Private/FileHelpers.cpp
+++ b/Engine/Source/Editor/UnrealEd/Private/FileHelpers.cpp
@@ -4295,9 +4295,10 @@ static bool InternalCheckoutAndSavePackages(const TArray<UPackage*>& PackagesToS
 			// Prevent modal window if not requested.
 			TGuardValue<bool> UnattendedScriptGuard(GIsRunningUnattendedScript, true);
 
-			TArray<UPackage*> PackagesCheckedOut;
 			const bool bErrorIfAlreadyCheckedOut = false;
-			const bool bConfirmPackageBranchCheckOutStatus = bUseDialog;
+			const bool bConfirmPackageBranchCheckOutStatus = false;
+
+			TArray<UPackage*> PackagesCheckedOut;
 			FEditorFileUtils::CheckoutPackages(PackagesToSave, &PackagesCheckedOut, bErrorIfAlreadyCheckedOut, bConfirmPackageBranchCheckOutStatus);
 
 			// Cannot mark files for add until after packages saved
@@ -4318,7 +4319,7 @@ static bool InternalCheckoutAndSavePackages(const TArray<UPackage*>& PackagesToS
 			PackagesToMarkForAdd.RemoveAll([&FailedPackages](UPackage* Package) { return FailedPackages.Contains(Package); });
 			if (PackagesToMarkForAdd.Num() > 0)
 			{
-				FEditorFileUtils::CheckoutPackages(PackagesToMarkForAdd, nullptr, bErrorIfAlreadyCheckedOut);
+				FEditorFileUtils::CheckoutPackages(PackagesToMarkForAdd, nullptr, bErrorIfAlreadyCheckedOut, bConfirmPackageBranchCheckOutStatus);
 			}
 
 			bResult = (ReturnResponse == FEditorFileUtils::PR_Success);
