Fixed bug that would cause references in UI code to be updated to point at the child blueprint after compiling the child blueprint (or blueprint that resulted from duplication)
Flip default behavior of CPFUO to tell tools to replace references.. also changes behavior of FAddComponentsToBlueprintImpl::MakeComponentCopy

diff --git a/Engine/Plugins/Enterprise/DataprepEditor/Source/DataprepCore/Private/Parameterization/DataprepParameterization.cpp b/Engine/Plugins/Enterprise/DataprepEditor/Source/DataprepCore/Private/Parameterization/DataprepParameterization.cpp
--- a/Engine/Plugins/Enterprise/DataprepEditor/Source/DataprepCore/Private/Parameterization/DataprepParameterization.cpp
+++ b/Engine/Plugins/Enterprise/DataprepEditor/Source/DataprepCore/Private/Parameterization/DataprepParameterization.cpp
@@ -1377,7 +1377,9 @@ void UDataprepParameterization::DoReinstancing(UClass* OldClass, bool bMigrateDa
 
 		if ( bMigrateData )
 		{
-			GEngine->CopyPropertiesForUnrelatedObjects( OldClass->GetDefaultObject(), CustomContainerClass->GetDefaultObject() );
+			UEngine::FCopyPropertiesForUnrelatedObjectsParams Options;
+			Options.bNotifyObjectReplacement = true;
+			UEngine::CopyPropertiesForUnrelatedObjects( OldClass->GetDefaultObject(), CustomContainerClass->GetDefaultObject(), Options );
 		}
 
 		// For the instances
@@ -1399,7 +1401,9 @@ void UDataprepParameterization::DoReinstancing(UClass* OldClass, bool bMigrateDa
 				
 				if ( bMigrateData )
 				{
-					GEngine->CopyPropertiesForUnrelatedObjects( OldObject, Object );
+					UEngine::FCopyPropertiesForUnrelatedObjectsParams Options;
+					Options.bNotifyObjectReplacement = true;
+					UEngine::CopyPropertiesForUnrelatedObjects( OldObject, Object, Options );
 				}
 
 				OldToNew.Add( OldObject, Object );
diff --git a/Engine/Plugins/Experimental/ControlRig/Source/ControlRig/Private/Sequencer/ControlRigSequence.cpp b/Engine/Plugins/Experimental/ControlRig/Source/ControlRig/Private/Sequencer/ControlRigSequence.cpp
--- a/Engine/Plugins/Experimental/ControlRig/Source/ControlRig/Private/Sequencer/ControlRigSequence.cpp
+++ b/Engine/Plugins/Experimental/ControlRig/Source/ControlRig/Private/Sequencer/ControlRigSequence.cpp
@@ -73,9 +73,7 @@ UObject* UControlRigSequence::MakeSpawnableTemplateFromInstance(UObject& InSourc
 {
 	UObject* NewInstance = NewObject<UObject>(MovieScene, InSourceObject.GetClass(), ObjectName);
 
-	UEngine::FCopyPropertiesForUnrelatedObjectsParams CopyParams;
-	CopyParams.bNotifyObjectReplacement = false;
-	UEngine::CopyPropertiesForUnrelatedObjects(&InSourceObject, NewInstance, CopyParams);
+	UEngine::CopyPropertiesForUnrelatedObjects(&InSourceObject, NewInstance);
 
 	return NewInstance;
 }
diff --git a/Engine/Source/Editor/KismetCompiler/Private/KismetCompiler.cpp b/Engine/Source/Editor/KismetCompiler/Private/KismetCompiler.cpp
--- a/Engine/Source/Editor/KismetCompiler/Private/KismetCompiler.cpp
+++ b/Engine/Source/Editor/KismetCompiler/Private/KismetCompiler.cpp
@@ -4265,6 +4265,7 @@ void FKismetCompilerContext::CompileFunctions(EInternalCompilerFlags InternalFla
 
 					UEditorEngine::FCopyPropertiesForUnrelatedObjectsParams CopyDetails;
 					CopyDetails.bCopyDeprecatedProperties = Blueprint->bIsRegeneratingOnLoad;
+					CopyDetails.bNotifyObjectReplacement = true; 
 					UEditorEngine::CopyPropertiesForUnrelatedObjects(OldCDO, NewCDO, CopyDetails);
 					FBlueprintEditorUtils::PatchCDOSubobjectsIntoExport(OldCDO, NewCDO);
 				}
diff --git a/Engine/Source/Editor/KismetCompiler/Private/KismetCompilerModule.cpp b/Engine/Source/Editor/KismetCompiler/Private/KismetCompilerModule.cpp
--- a/Engine/Source/Editor/KismetCompiler/Private/KismetCompilerModule.cpp
+++ b/Engine/Source/Editor/KismetCompiler/Private/KismetCompilerModule.cpp
@@ -170,6 +170,7 @@ void FKismet2CompilerModule::RefreshVariables(UBlueprint* Blueprint)
 
 		// cpfuo:
 		UEngine::FCopyPropertiesForUnrelatedObjectsParams Params;
+		Params.bNotifyObjectReplacement = true;
 		UEngine::CopyPropertiesForUnrelatedObjects(OldCDO, Blueprint->GeneratedClass->ClassDefaultObject, Params);
 	}
 }
diff --git a/Engine/Source/Editor/Persona/Private/SAnimNotifyPanel.cpp b/Engine/Source/Editor/Persona/Private/SAnimNotifyPanel.cpp
--- a/Engine/Source/Editor/Persona/Private/SAnimNotifyPanel.cpp
+++ b/Engine/Source/Editor/Persona/Private/SAnimNotifyPanel.cpp
@@ -4484,6 +4484,7 @@ void SAnimNotifyPanel::OnReplaceSelectedWithNotify(FString NewNotifyName, UClass
 			if ((OldEventPayload != nullptr) && (NewEvent.Notify != nullptr))
 			{
 				UEngine::FCopyPropertiesForUnrelatedObjectsParams CopyParams;
+				CopyParams.bNotifyObjectReplacement = true;
 				UEngine::CopyPropertiesForUnrelatedObjects(OldEventPayload, NewEvent.Notify, CopyParams);
 			}
 
diff --git a/Engine/Source/Editor/UnrealEd/Private/EditorEngine.cpp b/Engine/Source/Editor/UnrealEd/Private/EditorEngine.cpp
--- a/Engine/Source/Editor/UnrealEd/Private/EditorEngine.cpp
+++ b/Engine/Source/Editor/UnrealEd/Private/EditorEngine.cpp
@@ -5168,7 +5168,9 @@ void UEditorEngine::ReplaceActors(UActorFactory* Factory, const FAssetData& Asse
 			if (Blueprint->GeneratedClass->IsChildOf(OldActorClass) && NewActor != NULL)
 			{
 				NewActor->UnregisterAllComponents();
-				UEditorEngine::CopyPropertiesForUnrelatedObjects(OldActor, NewActor);
+				FCopyPropertiesForUnrelatedObjectsParams Options;
+				Options.bNotifyObjectReplacement = true;
+				UEditorEngine::CopyPropertiesForUnrelatedObjects(OldActor, NewActor, Options);
 				NewActor->RegisterAllComponents();
 			}
 		}
diff --git a/Engine/Source/Editor/UnrealEd/Private/Kismet2/KismetReinstanceUtilities.cpp b/Engine/Source/Editor/UnrealEd/Private/Kismet2/KismetReinstanceUtilities.cpp
--- a/Engine/Source/Editor/UnrealEd/Private/Kismet2/KismetReinstanceUtilities.cpp
+++ b/Engine/Source/Editor/UnrealEd/Private/Kismet2/KismetReinstanceUtilities.cpp
@@ -1699,7 +1699,9 @@ static void ReplaceObjectHelper(UObject*& OldObject, UClass* OldClass, UObject*&
 
 	InstancedPropertyUtils::FInstancedPropertyMap InstancedPropertyMap;
 	InstancedPropertyUtils::FArchiveInstancedSubObjCollector  InstancedSubObjCollector(OldObject, InstancedPropertyMap);
-	UEditorEngine::CopyPropertiesForUnrelatedObjects(OldObject, NewUObject);
+	UEngine::FCopyPropertiesForUnrelatedObjectsParams Options;
+	Options.bNotifyObjectReplacement = true;
+	UEditorEngine::CopyPropertiesForUnrelatedObjects(OldObject, NewUObject, Options);
 	InstancedPropertyUtils::FArchiveInsertInstancedSubObjects InstancedSubObjSpawner(NewUObject, InstancedPropertyMap);
 
 	UWorld* RegisteredWorld = nullptr;
@@ -1853,6 +1855,7 @@ static void ReplaceActorHelper(AActor* OldActor, UClass* OldClass, UObject*& New
 	UEngine::FCopyPropertiesForUnrelatedObjectsParams Params;
 	Params.bPreserveRootComponent = bPreserveRootComponent;
 	Params.bAggressiveDefaultSubobjectReplacement = true;
+	Params.bNotifyObjectReplacement = true;
 	UEngine::CopyPropertiesForUnrelatedObjects(OldActor, NewActor, Params);
 
 	// reset properties/streams
@@ -2363,6 +2366,7 @@ void FBlueprintCompileReinstancer::CopyPropertiesForUnrelatedObjects(UObject* Ol
 	Params.bCopyDeprecatedProperties = true;
 	Params.bSkipCompilerGeneratedDefaults = true;
 	Params.bClearReferences = bClearExternalReferences;
+	Params.bNotifyObjectReplacement = true;
 	UEngine::CopyPropertiesForUnrelatedObjects(OldObject, NewObject, Params);
 
 	InstancedPropertyUtils::FArchiveInsertInstancedSubObjects InstancedSubObjSpawner(NewObject, InstancedPropertyMap);
diff --git a/Engine/Source/Runtime/Engine/Classes/Engine/Engine.h b/Engine/Source/Runtime/Engine/Classes/Engine/Engine.h
--- a/Engine/Source/Runtime/Engine/Classes/Engine/Engine.h
+++ b/Engine/Source/Runtime/Engine/Classes/Engine/Engine.h
@@ -2669,7 +2669,7 @@ public:
 			, bCopyDeprecatedProperties(false)
 			, bPreserveRootComponent(true)
 			, bSkipCompilerGeneratedDefaults(false)
-			, bNotifyObjectReplacement(true)
+			, bNotifyObjectReplacement(false)
 			, bClearReferences(true)
 		{}
 	};
diff --git a/Engine/Source/Runtime/Engine/Private/Components/ChildActorComponent.cpp b/Engine/Source/Runtime/Engine/Private/Components/ChildActorComponent.cpp
--- a/Engine/Source/Runtime/Engine/Private/Components/ChildActorComponent.cpp
+++ b/Engine/Source/Runtime/Engine/Private/Components/ChildActorComponent.cpp
@@ -440,7 +440,9 @@ void UChildActorComponent::SetChildActorClass(TSubclassOf<AActor> Class, AActor*
 				{
 					if (ActorTemplate == nullptr)
 					{
-						UEngine::CopyPropertiesForUnrelatedObjects(ChildActorTemplate, NewChildActorTemplate);
+						UEngine::FCopyPropertiesForUnrelatedObjectsParams Options;
+						Options.bNotifyObjectReplacement = true;
+						UEngine::CopyPropertiesForUnrelatedObjects(ChildActorTemplate, NewChildActorTemplate, Options);
 					}
 					ChildActorTemplate->Rename(nullptr, GetTransientPackage(), REN_DontCreateRedirectors);
 				}
diff --git a/Engine/Source/Runtime/Engine/Private/World.cpp b/Engine/Source/Runtime/Engine/Private/World.cpp
--- a/Engine/Source/Runtime/Engine/Private/World.cpp
+++ b/Engine/Source/Runtime/Engine/Private/World.cpp
@@ -1316,7 +1316,9 @@ void UWorld::RepairWorldSettings()
 		if (ExistingWorldSettings)
 		{
 			NewWorldSettings->UnregisterAllComponents();
-			UEngine::CopyPropertiesForUnrelatedObjects(ExistingWorldSettings, NewWorldSettings);
+			UEngine::FCopyPropertiesForUnrelatedObjectsParams CopyParams;
+			CopyParams.bNotifyObjectReplacement = true;
+			UEngine::CopyPropertiesForUnrelatedObjects(ExistingWorldSettings, NewWorldSettings, CopyParams);
 			NewWorldSettings->RegisterAllComponents();
 		}
 
