Fix transaction buffer being cleared when hiding levels

diff --git a/Engine/Source/Editor/UnrealEd/Private/EditorEngine.cpp b/Engine/Source/Editor/UnrealEd/Private/EditorEngine.cpp
--- a/Engine/Source/Editor/UnrealEd/Private/EditorEngine.cpp
+++ b/Engine/Source/Editor/UnrealEd/Private/EditorEngine.cpp
@@ -6677,17 +6677,19 @@ void UEditorEngine::OnLevelRemovedFromWorld(ULevel* InLevel, UWorld* InWorld)
 		}
 	}
 	// UEngine::LoadMap broadcast this event with InLevel==NULL, before cleaning up the world during travel in Multiplayer
-	else if (InWorld->IsPlayInEditor() && Trans)
+	else if (Trans)
 	{
-		// Each additional instance of PIE in a multiplayer game will add another barrier, so if the event is triggered then this is the case and we need to lift it
-		// Otherwise there will be an imbalance between barriers set and barriers removed and we won't be able to undo when we return.
-		Trans->RemoveUndoBarrier();
-	}
-
-	if (!InWorld->IsPlayInEditor())
-	{
-		// If we're in editor mode, reset transactions buffer, to ensure that there are no references to a world which is about to be destroyed
-		ResetTransaction(NSLOCTEXT("UnrealEd", "LevelRemovedFromWorldEditorCallback", "Level removed from world"));
+		if (InWorld->IsPlayInEditor())
+		{
+			// Each additional instance of PIE in a multiplayer game will add another barrier, so if the event is triggered then this is the case and we need to lift it
+			// Otherwise there will be an imbalance between barriers set and barriers removed and we won't be able to undo when we return.
+			Trans->RemoveUndoBarrier();
+		}
+		else
+		{
+			// If we're in editor mode, reset transactions buffer, to ensure that there are no references to a world which is about to be destroyed
+			ResetTransaction(NSLOCTEXT("UnrealEd", "LevelRemovedFromWorldEditorCallback", "Level removed from world"));
+		}
 	}
 }
 
