Integrating CL#11313264 from Dev-Enterprise
Fixing Fbx import option "Material Search Location" sometimes being ignored.

diff --git a/Engine/Plugins/FX/Niagara/Source/NiagaraEditor/Private/TypeEditorUtilities/NiagaraEnumTypeEditorUtilities.cpp b/Engine/Plugins/FX/Niagara/Source/NiagaraEditor/Private/TypeEditorUtilities/NiagaraEnumTypeEditorUtilities.cpp
--- a/Engine/Plugins/FX/Niagara/Source/NiagaraEditor/Private/TypeEditorUtilities/NiagaraEnumTypeEditorUtilities.cpp
+++ b/Engine/Plugins/FX/Niagara/Source/NiagaraEditor/Private/TypeEditorUtilities/NiagaraEnumTypeEditorUtilities.cpp
@@ -23,9 +23,12 @@ public:
 
 		ChildSlot
 		[
-			SNew(SEnumCombobox, Enum)
-			.CurrentValue(this, &SNiagaraEnumParameterEditor::GetValue) 
-			.OnEnumSelectionChanged(SEnumCombobox::FOnEnumSelectionChanged::CreateSP(this, &SNiagaraEnumParameterEditor::ValueChanged))
+			SNew(SEnumComboBox, Enum)
+			.CurrentValue(this, &SNiagaraEnumParameterEditor::GetValue)
+			.ButtonStyle(FEditorStyle::Get(), "FlatButton.Light")
+			.ContentPadding(FMargin(2, 0))
+			.Font(FEditorStyle::GetFontStyle("Sequencer.AnimationOutliner.RegularFont"))
+			.OnEnumSelectionChanged(SEnumComboBox::FOnEnumSelectionChanged::CreateSP(this, &SNiagaraEnumParameterEditor::ValueChanged))
 		];
 	}
 
diff --git a/Engine/Source/Editor/DetailCustomizations/Private/FbxImportUIDetails.cpp b/Engine/Source/Editor/DetailCustomizations/Private/FbxImportUIDetails.cpp
--- a/Engine/Source/Editor/DetailCustomizations/Private/FbxImportUIDetails.cpp
+++ b/Engine/Source/Editor/DetailCustomizations/Private/FbxImportUIDetails.cpp
@@ -1,30 +1,33 @@
 // Copyright Epic Games, Inc. All Rights Reserved.
 
 #include "FbxImportUIDetails.h"
-#include "Misc/Attribute.h"
-#include "Misc/Guid.h"
-#include "Widgets/DeclarativeSyntaxSupport.h"
-#include "Factories/FbxSkeletalMeshImportData.h"
+
+#include "DetailCategoryBuilder.h"
+#include "DetailLayoutBuilder.h"
+#include "DetailWidgetRow.h"
+#include "Editor.h"
+#include "Engine/StaticMesh.h"
+#include "IDetailGroup.h"
+#include "IDetailPropertyRow.h"
 #include "Factories/FbxAnimSequenceImportData.h"
+#include "Factories/FbxSkeletalMeshImportData.h"
 #include "Factories/FbxStaticMeshImportData.h"
 #include "Factories/FbxTextureImportData.h"
-#include "Materials/MaterialInterface.h"
 #include "Materials/Material.h"
-#include "Engine/StaticMesh.h"
-#include "DetailWidgetRow.h"
-#include "Widgets/SBoxPanel.h"
-#include "Widgets/Text/STextBlock.h"
+#include "Materials/MaterialInterface.h"
+#include "Misc/Attribute.h"
+#include "Misc/Guid.h"
 #include "PropertyHandle.h"
-#include "DetailLayoutBuilder.h"
-#include "IDetailPropertyRow.h"
-#include "DetailCategoryBuilder.h"
-#include "Widgets/Layout/SBox.h"
+#include "SEnumCombobox.h"
+#include "Templates/Tuple.h"
+#include "Widgets/DeclarativeSyntaxSupport.h"
+#include "Widgets/Images/SImage.h"
 #include "Widgets/Input/STextComboBox.h"
-#include "Widgets/SToolTip.h"
 #include "Widgets/Input/SButton.h"
-#include "Widgets/Images/SImage.h"
-#include "Editor.h"
-#include "IDetailGroup.h"
+#include "Widgets/Layout/SBox.h"
+#include "Widgets/Text/STextBlock.h"
+#include "Widgets/SBoxPanel.h"
+#include "Widgets/SToolTip.h"
 
 #define LOCTEXT_NAMESPACE "FbxImportUIDetails"
 
@@ -33,10 +36,6 @@
 
 static FString DoNotOverrideString = FText(LOCTEXT("BaseColorPropertyDoNotOverride", "Do Not Override")).ToString();
 
-static FString CreateNewMaterialsString = FText(LOCTEXT("MaterialImportMethodCreateNewMaterials", "Create New Materials")).ToString();
-static FString CreateNewInstancedMaterialsString = FText(LOCTEXT("MaterialImportMethodCreateNewInstancedMaterials", "Create New Instanced Materials")).ToString();
-static FString DoNotCreateMaterialString = FText(LOCTEXT("MaterialImportMethodDoNotCreateMaterial", "Do Not Create Material")).ToString();
-
 //If the String is contain in the StringArray, it return the index. Otherwise return INDEX_NONE
 int FindString(const TArray<TSharedPtr<FString>> &StringArray, const FString &String) {
 	for (int i = 0; i < StringArray.Num(); i++)
@@ -132,7 +131,7 @@ bool FFbxImportUIDetails::ShowCompareResult()
 		{
 			FName ConflictCategoryName = TEXT("Conflicts");
 			IDetailCategoryBuilder& CategoryBuilder = CachedDetailBuilder->EditCategory(ConflictCategoryName, LOCTEXT("CategoryConflictsName", "Conflicts"), ECategoryPriority::Important);
-			auto BuildConflictRow = [this, &CategoryBuilder](const FText& CategoryName, const FText& Conflict_NameContent, const FText& Conflict_ButtonTooltip, const FText& Conflict_ButtonText, ConflictDialogType DialogType, const FSlateBrush* Brush, const FText& Conflict_IconTooltip)
+			auto BuildConflictRow = [this, &CategoryBuilder](const FText& CategoryName, const FText& Conflict_NameContent, const FText& Conflict_ButtonTooltip, const FText& Conflict_ButtonText, EConflictDialogType DialogType, const FSlateBrush* Brush, const FText& Conflict_IconTooltip)
 			{
 				CategoryBuilder.AddCustomRow(CategoryName).WholeRowContent()
 				[
@@ -178,7 +177,7 @@ bool FFbxImportUIDetails::ShowCompareResult()
 					LOCTEXT("MaterialConflict_NameContent", "Unmatched Materials"),
 					LOCTEXT("MaterialConflict_ButtonShowTooltip", "Show a detailed view of the materials conflict."),
 					LOCTEXT("MaterialConflict_ButtonShow", "Show Conflict"),
-					ConflictDialogType::Conflict_Material,
+					EConflictDialogType::Conflict_Material,
 					FEditorStyle::GetBrush("Icons.Error"),
 					LOCTEXT("MaterialConflict_IconTooltip", "There is one or more material(s) that do not match.")
 				);
@@ -208,7 +207,7 @@ bool FFbxImportUIDetails::ShowCompareResult()
 					LOCTEXT("SkeletonConflict_NameContent", "Unmatched Skeleton joints"),
 					LOCTEXT("SkeletonConflict_ButtonShowTooltip", "Show a detailed view of the skeleton joints conflict."),
 					LOCTEXT("SkeletonConflict_ButtonShow", "Show Conflict"),
-					ConflictDialogType::Conflict_Skeleton,
+					EConflictDialogType::Conflict_Skeleton,
 					Brush,
 					IconTooltip);
 			}
@@ -687,22 +686,16 @@ void FFbxImportUIDetails::AddSubCategory(IDetailLayoutBuilder& DetailBuilder, FN
 
 void FFbxImportUIDetails::ConstructMaterialImportMethod(TSharedPtr<IPropertyHandle> ImportMaterialPropHandle, IDetailCategoryBuilder& MaterialCategory)
 {
-	//The import material is represent by a combobox with 3 choices
-	//1. Create New Materials
-	//2. Create New Instanced Materials (Using an existing base material)
-	//3. Do not Create Materials
-	ImportMethodNames.Reset();
-	ImportMethodNames.Add(MakeShareable(new FString(CreateNewMaterialsString)));
-	ImportMethodNames.Add(MakeShareable(new FString(CreateNewInstancedMaterialsString)));
-	ImportMethodNames.Add(MakeShareable(new FString(DoNotCreateMaterialString)));
-
 	if(ImportUI->TextureImportData->BaseMaterialName.IsValid())
 	{
 		//When we load the UI the first time we set this boolean to true in case the BaseMaterialName is valid.
 		ImportUI->TextureImportData->bUseBaseMaterial = true;
 	}
 
-	int32 InitialSelect = ImportUI->bImportMaterials ? (ImportUI->TextureImportData->bUseBaseMaterial ? 1 : 0) : 2;
+	const UEnum* MaterialImportMethodEnum = StaticEnum< EMaterialImportMethod >();
+	SelectedMaterialImportMethod = ImportUI->bImportMaterials
+		? (ImportUI->TextureImportData->bUseBaseMaterial ? EMaterialImportMethod::CreateNewInstancedMaterials : EMaterialImportMethod::CreateNewMaterials) 
+		: EMaterialImportMethod::DoNotCreateMaterialString;
 
 	MaterialCategory.AddCustomRow(LOCTEXT("MaterialImportMethod", "Material Import Method"))
 	.NameContent()
@@ -720,10 +713,9 @@ void FFbxImportUIDetails::ConstructMaterialImportMethod(TSharedPtr<IPropertyHand
 		[
 			SNew(SBox)
 			[
-				SNew(STextComboBox)
-				.OptionsSource(&ImportMethodNames)
-				.OnSelectionChanged(this, &FFbxImportUIDetails::OnMaterialImportMethodChanged)
-				.InitiallySelectedItem(ImportMethodNames[InitialSelect])
+				SNew(SEnumComboBox, MaterialImportMethodEnum)
+				.OnEnumSelectionChanged(SEnumComboBox::FOnEnumSelectionChanged::CreateSP(this, &FFbxImportUIDetails::OnMaterialImportMethodChanged))
+				.CurrentValue(this, &FFbxImportUIDetails::GetMaterialImportMethodValue)
 				.Font(IDetailLayoutBuilder::GetDetailFont())
 			]
 		]
@@ -1221,32 +1213,40 @@ FReply FFbxImportUIDetails::MaterialBaseParamClearAllProperties()
 	return FReply::Handled();
 }
 
-void FFbxImportUIDetails::OnMaterialImportMethodChanged(TSharedPtr<FString> Selection, ESelectInfo::Type SelectInfo)
+int32 FFbxImportUIDetails::GetMaterialImportMethodValue() const
+{
+	return static_cast<int32>(SelectedMaterialImportMethod);
+}
+
+void FFbxImportUIDetails::OnMaterialImportMethodChanged(int32 Value, ESelectInfo::Type SelectInfo)
 {
-	FString SelectName = *Selection.Get();
-	if (SelectName.Equals(CreateNewMaterialsString))
+	SelectedMaterialImportMethod = static_cast<EMaterialImportMethod>(Value);
+
+	switch (SelectedMaterialImportMethod)
 	{
+	case EMaterialImportMethod::CreateNewMaterials:
 		ImportUI->bImportMaterials = true;
 		//Reset the base material and the UseBaseMaterial flag to hide the base material name property
 		ImportUI->TextureImportData->bUseBaseMaterial = false;
 		ImportUI->TextureImportData->BaseMaterialName.Reset();
-	}
-	else if (SelectName.Equals(CreateNewInstancedMaterialsString))
-	{
+		break;
+	case EMaterialImportMethod::CreateNewInstancedMaterials:
 		ImportUI->bImportMaterials = true;
 		ImportUI->TextureImportData->bUseBaseMaterial = true;
-	}
-	else
-	{
+		break;
+	case EMaterialImportMethod::DoNotCreateMaterialString:
 		ImportUI->bImportMaterials = false;
 		//Reset the base material and the UseBaseMaterial flag to hide the base material name property
 		ImportUI->TextureImportData->bUseBaseMaterial = false;
 		ImportUI->TextureImportData->BaseMaterialName.Reset();
+		break;
+	default:
+		break;
 	}
 	RefreshCustomDetail();
 }
 
-FReply FFbxImportUIDetails::ShowConflictDialog(ConflictDialogType DialogType)
+FReply FFbxImportUIDetails::ShowConflictDialog(EConflictDialogType DialogType)
 {
 	if (DialogType == Conflict_Material)
 	{
diff --git a/Engine/Source/Editor/DetailCustomizations/Private/FbxImportUIDetails.h b/Engine/Source/Editor/DetailCustomizations/Private/FbxImportUIDetails.h
--- a/Engine/Source/Editor/DetailCustomizations/Private/FbxImportUIDetails.h
+++ b/Engine/Source/Editor/DetailCustomizations/Private/FbxImportUIDetails.h
@@ -14,12 +14,20 @@ class IDetailLayoutBuilder;
 class IDetailPropertyRow;
 class IPropertyHandle;
 
-enum ConflictDialogType
+enum EConflictDialogType
 {
 	Conflict_Material,
 	Conflict_Skeleton
 };
 
+UENUM()
+enum class EMaterialImportMethod : int32 
+{
+	CreateNewMaterials  UMETA(DisplayName = "Create New Materials", ToolTip = "A new material will be created from the imported data."),
+	CreateNewInstancedMaterials UMETA(DisplayName = "Create New Instanced Materials", Tooltip = "A new material instance of the specified base material will be created and set with the imported material data."),
+	DoNotCreateMaterialString UMETA(DisplayName = "Do Not Create Material", Tooltip = "No materials will be created from the import data."),
+};
+
 class FFbxImportUIDetails : public IDetailCustomization, public FEditorUndoClient
 {
 public:
@@ -81,7 +89,8 @@ public:
 	void OnOpacityTextureColor(TSharedPtr<FString> Selection, ESelectInfo::Type SelectInfo);
 	FReply MaterialBaseParamClearAllProperties();
 
-	void OnMaterialImportMethodChanged(TSharedPtr<FString> Selection, ESelectInfo::Type SelectInfo);
+	int32 GetMaterialImportMethodValue() const;
+	void OnMaterialImportMethodChanged(int32 Value, ESelectInfo::Type SelectInfo);
 	
 	TWeakObjectPtr<UFbxImportUI> ImportUI;		// The UI data object being customised
 	IDetailLayoutBuilder* CachedDetailBuilder;	// The detail builder for this cusomtomisation
@@ -101,7 +110,7 @@ private:
 	bool GetVertexOverrideColorEnabledState() const;
 	bool GetSkeletalMeshVertexOverrideColorEnabledState() const;
 
-	FReply ShowConflictDialog(ConflictDialogType DialogType);
+	FReply ShowConflictDialog(EConflictDialogType DialogType);
 	bool ShowCompareResult();
 
 	/** LOD group options. */
@@ -115,8 +124,7 @@ private:
 	TSharedPtr<IPropertyHandle> VertexColorImportOptionHandle;
 	TSharedPtr<IPropertyHandle> SkeletalMeshVertexColorImportOptionHandle;
 
-	TArray< TSharedPtr< FString > > ImportMethodNames;
-	TSharedPtr<class STextComboBox> MaterialImportMethodComboBox;
+	EMaterialImportMethod SelectedMaterialImportMethod;
 	bool bShowBaseMaterialUI;
 
 	TArray< TSharedPtr< FString > > BaseColorNames;
diff --git a/Engine/Source/Editor/EditorWidgets/Public/SEnumCombobox.h b/Engine/Source/Editor/EditorWidgets/Public/SEnumCombobox.h
--- a/Engine/Source/Editor/EditorWidgets/Public/SEnumCombobox.h
+++ b/Engine/Source/Editor/EditorWidgets/Public/SEnumCombobox.h
@@ -3,22 +3,31 @@
 #pragma once
 
 #include "CoreMinimal.h"
-#include "Widgets/Input/SComboBox.h"
-#include "Types/SlateEnums.h"
 #include "EditorStyleSet.h"
+#include "Types/SlateEnums.h"
+#include "Widgets/Input/SComboBox.h"
+#include "Widgets/SToolTip.h"
 
 /*
 * Editor only widget for creating a combo-box which allows the user to pick the values from a dropdown combo-box. 
 */
-class SEnumCombobox : public SComboBox<TSharedPtr<int32>>
-{
+class SEnumComboBox : public SComboBox<TSharedPtr<int32>>
+ {
 public:
 	DECLARE_DELEGATE_TwoParams(FOnEnumSelectionChanged, int32 /*Selection*/, ESelectInfo::Type /*SelectionType*/);
 public:
-	SLATE_BEGIN_ARGS(SEnumCombobox) {}
+	SLATE_BEGIN_ARGS(SEnumComboBox)
+		: _CurrentValue()
+		, _ContentPadding(FMargin(4.0, 2.0))
+		, _OnEnumSelectionChanged()
+		, _ButtonStyle(nullptr)
+	{}
 
 		SLATE_ATTRIBUTE(int32, CurrentValue)
+		SLATE_ATTRIBUTE(FMargin, ContentPadding)
+		SLATE_ATTRIBUTE(FSlateFontInfo, Font)
 		SLATE_ARGUMENT(FOnEnumSelectionChanged, OnEnumSelectionChanged)
+		SLATE_STYLE_ARGUMENT(FButtonStyle, ButtonStyle)
 
 	SLATE_END_ARGS()
 
@@ -28,7 +37,7 @@ public:
 		CurrentValue = InArgs._CurrentValue;
 		check(CurrentValue.IsBound());
 		OnEnumSelectionChangedDelegate = InArgs._OnEnumSelectionChanged;
-
+		Font = InArgs._Font;
 		bUpdatingSelectionInternally = false;
 
 		for (int32 i = 0; i < Enum->NumEnums() - 1; i++)
@@ -38,22 +47,25 @@ public:
 				VisibleEnumNameIndices.Add(MakeShareable(new int32(i)));
 			}
 		}
-
+		
 		SComboBox::Construct(SComboBox<TSharedPtr<int32>>::FArguments()
-			.ButtonStyle(FEditorStyle::Get(), "FlatButton.Light")
+			.ButtonStyle(InArgs._ButtonStyle)
 			.OptionsSource(&VisibleEnumNameIndices)
 			.OnGenerateWidget_Lambda([this](TSharedPtr<int32> InItem)
-		{
-			return SNew(STextBlock)
-				.Text(Enum->GetDisplayNameTextByIndex(*InItem));
-		})
-			.OnSelectionChanged(this, &SEnumCombobox::OnComboSelectionChanged)
-			.OnComboBoxOpening(this, &SEnumCombobox::OnComboMenuOpening)
-			.ContentPadding(FMargin(2, 0))
+			{
+				return SNew(STextBlock)
+					.Font(Font)
+					.Text(Enum->GetDisplayNameTextByIndex(*InItem))
+					.ToolTipText(Enum->GetToolTipTextByIndex(*InItem));
+			})
+			.OnSelectionChanged(this, &SEnumComboBox::OnComboSelectionChanged)
+			.OnComboBoxOpening(this, &SEnumComboBox::OnComboMenuOpening)
+			.ContentPadding(InArgs._ContentPadding)
 			[
 				SNew(STextBlock)
-				.Font(FEditorStyle::GetFontStyle("Sequencer.AnimationOutliner.RegularFont"))
-			.Text(this, &SEnumCombobox::GetCurrentValue)
+					.Font(Font)
+					.Text(this, &SEnumComboBox::GetCurrentValue)
+					.ToolTipText(this, &SEnumComboBox::GetCurrentValueTooltip)
 			]);
 	}
 
@@ -64,10 +76,17 @@ private:
 		return Enum->GetDisplayNameTextByIndex(CurrentNameIndex);
 	}
 
+	FText GetCurrentValueTooltip() const
+	{
+		int32 CurrentNameIndex = Enum->GetIndexByValue(CurrentValue.Get());
+		return Enum->GetToolTipTextByIndex(CurrentNameIndex);
+	}
+
 	TSharedRef<SWidget> OnGenerateWidget(TSharedPtr<int32> InItem)
 	{
 		return SNew(STextBlock)
-			.Text(Enum->GetDisplayNameTextByIndex(*InItem));
+			.Text(Enum->GetDisplayNameTextByIndex(*InItem))
+			.Font(Font);
 	}
 
 	void OnComboSelectionChanged(TSharedPtr<int32> InSelectedItem, ESelectInfo::Type SelectInfo)
@@ -103,6 +122,8 @@ private:
 
 	TAttribute<int32> CurrentValue;
 
+	TAttribute< FSlateFontInfo > Font;
+
 	TArray<TSharedPtr<int32>> VisibleEnumNameIndices;
 
 	bool bUpdatingSelectionInternally;
diff --git a/Engine/Source/Editor/MovieSceneTools/Private/CurveKeyEditors/SEnumKeyEditor.cpp b/Engine/Source/Editor/MovieSceneTools/Private/CurveKeyEditors/SEnumKeyEditor.cpp
--- a/Engine/Source/Editor/MovieSceneTools/Private/CurveKeyEditors/SEnumKeyEditor.cpp
+++ b/Engine/Source/Editor/MovieSceneTools/Private/CurveKeyEditors/SEnumKeyEditor.cpp
@@ -16,7 +16,7 @@ void SEnumCurveKeyEditor::Construct(const FArguments& InArgs, const TSequencerKe
 		MovieSceneToolHelpers::MakeEnumComboBox(
 			Enum,
 			TAttribute<int32>::Create(TAttribute<int32>::FGetter::CreateSP(this, &SEnumCurveKeyEditor::OnGetCurrentValueAsInt)),
-			SEnumCombobox::FOnEnumSelectionChanged::CreateSP(this, &SEnumCurveKeyEditor::OnChangeKey)
+			SEnumComboBox::FOnEnumSelectionChanged::CreateSP(this, &SEnumCurveKeyEditor::OnChangeKey)
 		)
 	];
 }
diff --git a/Engine/Source/Editor/MovieSceneTools/Private/MovieSceneToolHelpers.cpp b/Engine/Source/Editor/MovieSceneTools/Private/MovieSceneToolHelpers.cpp
--- a/Engine/Source/Editor/MovieSceneTools/Private/MovieSceneToolHelpers.cpp
+++ b/Engine/Source/Editor/MovieSceneTools/Private/MovieSceneToolHelpers.cpp
@@ -528,10 +528,13 @@ int32 MovieSceneToolHelpers::FindAvailableRowIndex(UMovieSceneTrack* InTrack, UM
 	return InTrack->GetMaxRowIndex() + 1;
 }
 
-TSharedRef<SWidget> MovieSceneToolHelpers::MakeEnumComboBox(const UEnum* InEnum, TAttribute<int32> InCurrentValue, SEnumCombobox::FOnEnumSelectionChanged InOnSelectionChanged)
+TSharedRef<SWidget> MovieSceneToolHelpers::MakeEnumComboBox(const UEnum* InEnum, TAttribute<int32> InCurrentValue, SEnumComboBox::FOnEnumSelectionChanged InOnSelectionChanged)
 {
-	return SNew(SEnumCombobox, InEnum)
+	return SNew(SEnumComboBox, InEnum)
 		.CurrentValue(InCurrentValue)
+		.ButtonStyle(FEditorStyle::Get(), "FlatButton.Light")
+		.ContentPadding(FMargin(2, 0))
+		.Font(FEditorStyle::GetFontStyle("Sequencer.AnimationOutliner.RegularFont"))
 		.OnEnumSelectionChanged(InOnSelectionChanged);
 }
 
diff --git a/Engine/Source/Editor/MovieSceneTools/Public/MovieSceneToolHelpers.h b/Engine/Source/Editor/MovieSceneTools/Public/MovieSceneToolHelpers.h
--- a/Engine/Source/Editor/MovieSceneTools/Public/MovieSceneToolHelpers.h
+++ b/Engine/Source/Editor/MovieSceneTools/Public/MovieSceneToolHelpers.h
@@ -158,7 +158,7 @@ public:
 	 * @param OnSelectionChanged Delegate fired when selection is changed
 	 * @return The new widget
 	 */
-	static TSharedRef<SWidget> MakeEnumComboBox(const UEnum* Enum, TAttribute<int32> CurrentValue, SEnumCombobox::FOnEnumSelectionChanged OnSelectionChanged);
+	static TSharedRef<SWidget> MakeEnumComboBox(const UEnum* Enum, TAttribute<int32> CurrentValue, SEnumComboBox::FOnEnumSelectionChanged OnSelectionChanged);
 
 
 	/**
diff --git a/Engine/Source/Editor/UnrealEd/Classes/Factories/FbxImportUI.h b/Engine/Source/Editor/UnrealEd/Classes/Factories/FbxImportUI.h
--- a/Engine/Source/Editor/UnrealEd/Classes/Factories/FbxImportUI.h
+++ b/Engine/Source/Editor/UnrealEd/Classes/Factories/FbxImportUI.h
@@ -185,12 +185,12 @@ public:
 	UPROPERTY(BlueprintReadWrite, Category = Miscellaneous)
 	uint32 bImportRigidMesh:1;
 
-	/** Whether to automatically create Unreal materials for materials found in the FBX scene */
+	/** If no existing materials are found, whether to automatically create Unreal materials for materials found in the FBX scene */
 	UPROPERTY(EditAnywhere, BlueprintReadWrite, config, Category = Material, meta=(ImportType="GeoOnly"))
 	uint32 bImportMaterials:1;
 
-	/** The option works only when option "Import Material" is OFF. If "Import Material" is ON, textures are always imported. */
-	UPROPERTY(EditAnywhere, BlueprintReadWrite, config, Category=Material, meta = (ImportType = "GeoOnly"))
+	/** Whether or not we should import textures. This option is disabled when we are importing materials because textures are always imported in that case. */
+	UPROPERTY(EditAnywhere, BlueprintReadWrite, config, Category=Material, meta = (ImportType = "GeoOnly", EditCondition="!bImportMaterials"))
 	uint32 bImportTextures:1;
 
 	/** Import data used when importing static meshes */
diff --git a/Engine/Source/Editor/UnrealEd/Classes/Factories/MaterialImportHelpers.h b/Engine/Source/Editor/UnrealEd/Classes/Factories/MaterialImportHelpers.h
--- a/Engine/Source/Editor/UnrealEd/Classes/Factories/MaterialImportHelpers.h
+++ b/Engine/Source/Editor/UnrealEd/Classes/Factories/MaterialImportHelpers.h
@@ -20,7 +20,9 @@ enum class EMaterialSearchLocation : uint8
 	/** Search for matching material recursively from root folder. */
 	UnderRoot,
 	/** Search for matching material in all assets folders. */
-	AllAssets
+	AllAssets,
+	/** Do not search for existing matching materials */
+	DoNotSearch,
 };
 
 UCLASS(transient)
diff --git a/Engine/Source/Editor/UnrealEd/Private/Factories/MaterialImportHelpers.cpp b/Engine/Source/Editor/UnrealEd/Private/Factories/MaterialImportHelpers.cpp
--- a/Engine/Source/Editor/UnrealEd/Private/Factories/MaterialImportHelpers.cpp
+++ b/Engine/Source/Editor/UnrealEd/Private/Factories/MaterialImportHelpers.cpp
@@ -9,6 +9,11 @@
 
 UMaterialInterface* UMaterialImportHelpers::FindExistingMaterialFromSearchLocation(const FString& MaterialFullName, const FString& BasePackagePath, EMaterialSearchLocation SearchLocation, FText& OutError)
 {
+	if (SearchLocation == EMaterialSearchLocation::DoNotSearch)
+	{
+		return nullptr;
+	}
+
 	//Search in memory
 	constexpr bool bExactClass = false;
 	UMaterialInterface* FoundMaterial = FindObject<UMaterialInterface>(nullptr, *MaterialFullName, bExactClass);
diff --git a/Engine/Source/Editor/UnrealEd/Private/Fbx/FbxMainImport.cpp b/Engine/Source/Editor/UnrealEd/Private/Fbx/FbxMainImport.cpp
--- a/Engine/Source/Editor/UnrealEd/Private/Fbx/FbxMainImport.cpp
+++ b/Engine/Source/Editor/UnrealEd/Private/Fbx/FbxMainImport.cpp
@@ -513,15 +513,15 @@ void ApplyImportUIToImportOptions(UFbxImportUI* ImportUI, FBXImportOptions& InOu
 	}
 }
 
-void FImportedMaterialData::AddImportedMaterial( FbxSurfaceMaterial& FbxMaterial, UMaterialInterface& UnrealMaterial )
+void FImportedMaterialData::AddImportedMaterial( const FbxSurfaceMaterial& FbxMaterial, UMaterialInterface& UnrealMaterial )
 {
 	FbxToUnrealMaterialMap.Add( &FbxMaterial, &UnrealMaterial );
 	ImportedMaterialNames.Add( *UnrealMaterial.GetPathName() );
 }
 
-bool FImportedMaterialData::IsUnique( FbxSurfaceMaterial& FbxMaterial, FName ImportedMaterialName ) const
+bool FImportedMaterialData::IsAlreadyImported( const FbxSurfaceMaterial& FbxMaterial, FName ImportedMaterialName ) const
 {
-	UMaterialInterface* FoundMaterial = GetUnrealMaterial( FbxMaterial );
+	const UMaterialInterface* FoundMaterial = GetUnrealMaterial( FbxMaterial );
 
 	return FoundMaterial != NULL || ImportedMaterialNames.Contains( ImportedMaterialName );
 }
diff --git a/Engine/Source/Editor/UnrealEd/Private/Fbx/FbxMaterialImport.cpp b/Engine/Source/Editor/UnrealEd/Private/Fbx/FbxMaterialImport.cpp
--- a/Engine/Source/Editor/UnrealEd/Private/Fbx/FbxMaterialImport.cpp
+++ b/Engine/Source/Editor/UnrealEd/Private/Fbx/FbxMaterialImport.cpp
@@ -250,7 +250,7 @@ void LogPropertyAndChild(FbxSurfaceMaterial& FbxMaterial, const FbxProperty &Pro
 #endif
 
 bool UnFbx::FFbxImporter::CreateAndLinkExpressionForMaterialProperty(
-							FbxSurfaceMaterial& FbxMaterial,
+							const FbxSurfaceMaterial& FbxMaterial,
 							UMaterial* UnrealMaterial,
 							const char* MaterialProperty ,
 							FExpressionInput& MaterialInput, 
@@ -365,7 +365,7 @@ bool UnFbx::FFbxImporter::CreateAndLinkExpressionForMaterialProperty(
 //-------------------------------------------------------------------------
 //
 //-------------------------------------------------------------------------
-void UnFbx::FFbxImporter::FixupMaterial( FbxSurfaceMaterial& FbxMaterial, UMaterial* UnrealMaterial )
+void UnFbx::FFbxImporter::FixupMaterial( const FbxSurfaceMaterial& FbxMaterial, UMaterial* UnrealMaterial )
 {
 	// add a basic diffuse color if no texture is linked to diffuse
 	if (UnrealMaterial->BaseColor.Expression == NULL)
@@ -418,7 +418,7 @@ void UnFbx::FFbxImporter::FixupMaterial( FbxSurfaceMaterial& FbxMaterial, UMater
 //
 //-------------------------------------------------------------------------
 
-FString UnFbx::FFbxImporter::GetMaterialFullName(FbxSurfaceMaterial& FbxMaterial)
+FString UnFbx::FFbxImporter::GetMaterialFullName(const FbxSurfaceMaterial& FbxMaterial) const
 {
 	FString MaterialFullName = UTF8_TO_TCHAR(MakeName(FbxMaterial.GetName()));
 
@@ -443,12 +443,32 @@ FString UnFbx::FFbxImporter::GetMaterialFullName(FbxSurfaceMaterial& FbxMaterial
 	return MaterialFullName;
 }
 
+FString UnFbx::FFbxImporter::GetMaterialBasePackageName(const FString& MaterialFullName) const
+{
+	FString BasePackageName = FPackageName::GetLongPackagePath(Parent->GetOutermost()->GetName());
+	if (ImportOptions->MaterialBasePath != NAME_None)
+	{
+		BasePackageName = ImportOptions->MaterialBasePath.ToString();
+	}
+	else
+	{
+		BasePackageName += TEXT("/");
+	}
+
+	BasePackageName += MaterialFullName;
+	BasePackageName = UPackageTools::SanitizePackageName(BasePackageName);
+
+	return BasePackageName;
+}
+
+
 bool UnFbx::FFbxImporter::LinkMaterialProperty(
-	FbxSurfaceMaterial& FbxMaterial,
+	const FbxSurfaceMaterial& FbxMaterial,
 	UMaterialInstanceConstant* UnrealMaterial,
 	const char* MaterialProperty,
 	FName ParameterValue,
-	bool bSetupAsNormalMap) {
+	bool bSetupAsNormalMap) 
+{
 	bool bCreated = false;
 	FbxProperty FbxProperty = FbxMaterial.FindProperty(MaterialProperty);
 	if (FbxProperty.IsValid())
@@ -483,7 +503,7 @@ bool UnFbx::FFbxImporter::LinkMaterialProperty(
 	return bCreated;
 }
 
-bool CanUseMaterialWithInstance(FbxSurfaceMaterial& FbxMaterial, const char* MaterialProperty, FString ParameterValueName, UMaterialInterface *BaseMaterial, TArray<FString>& UVSet) {
+bool CanUseMaterialWithInstance(const FbxSurfaceMaterial& FbxMaterial, const char* MaterialProperty, FString ParameterValueName, UMaterialInterface *BaseMaterial, TArray<FString>& UVSet) {
 	FbxProperty FbxProperty = FbxMaterial.FindProperty(MaterialProperty);
 	if (FbxProperty.IsValid())
 	{
@@ -523,59 +543,39 @@ bool CanUseMaterialWithInstance(FbxSurfaceMaterial& FbxMaterial, const char* Mat
 	return true;
 }
 
-
-void UnFbx::FFbxImporter::CreateUnrealMaterial(FbxSurfaceMaterial& FbxMaterial, TArray<UMaterialInterface*>& OutMaterials, TArray<FString>& UVSets, bool bForSkeletalMesh)
+UMaterialInterface* UnFbx::FFbxImporter::FindExistingMaterialFromFbxMaterial(const FbxSurfaceMaterial& FbxMaterial, EMaterialSearchLocation MaterialSearchLocation)
 {
 	// Make sure we have a parent
-	if ( !ensure(Parent.IsValid()) )
+	if (!ensure(Parent.IsValid()))
 	{
-		return;
+		return nullptr;
 	}
-	if (ImportOptions->OverrideMaterials.Contains(FbxMaterial.GetUniqueID()))
+
+	if (UMaterialInterface** OverrideMaterial = ImportOptions->OverrideMaterials.Find(FbxMaterial.GetUniqueID()))
 	{
-		UMaterialInterface* FoundMaterial = *(ImportOptions->OverrideMaterials.Find(FbxMaterial.GetUniqueID()));
-		if (ImportedMaterialData.IsUnique(FbxMaterial, FName(*FoundMaterial->GetPathName())) == false)
+		if (!ImportedMaterialData.IsAlreadyImported(FbxMaterial, FName(*(*OverrideMaterial)->GetPathName())))
 		{
-			ImportedMaterialData.AddImportedMaterial(FbxMaterial, *FoundMaterial);
+			// Add the material override to the list of imported materials
+			ImportedMaterialData.AddImportedMaterial(FbxMaterial, **OverrideMaterial);
 		}
-		// The material is override add the existing one
-		OutMaterials.Add(FoundMaterial);
-		return;
-	}
-	FString MaterialFullName = GetMaterialFullName(FbxMaterial);
-	FString BasePackageName = FPackageName::GetLongPackagePath(Parent->GetOutermost()->GetName());
-	if (ImportOptions->MaterialBasePath != NAME_None)
-	{
-		BasePackageName = ImportOptions->MaterialBasePath.ToString();
-	}
-	else
-	{
-		BasePackageName += TEXT("/");
+		return *OverrideMaterial;
 	}
-	BasePackageName += MaterialFullName;
-	
-	BasePackageName = UPackageTools::SanitizePackageName(BasePackageName);
 
-	// The material could already exist in the project
-	FName ObjectPath = *(BasePackageName + TEXT(".") + MaterialFullName);
+	const FString MaterialFullName = GetMaterialFullName(FbxMaterial);
+	const FString BasePackageName = GetMaterialBasePackageName(MaterialFullName);
+	const FName ObjectPath = *(BasePackageName + TEXT(".") + MaterialFullName);
 
-	if( ImportedMaterialData.IsUnique( FbxMaterial, ObjectPath ) )
+	// The material could already exist in the project
+	UMaterialInterface* FoundMaterial = nullptr;
+	if (ImportedMaterialData.IsAlreadyImported(FbxMaterial, ObjectPath))
 	{
-		UMaterialInterface* FoundMaterial = ImportedMaterialData.GetUnrealMaterial( FbxMaterial );
-		if (FoundMaterial)
-		{
-			// The material was imported from this FBX.  Reuse it
-			OutMaterials.Add(FoundMaterial);
-			return;
-		}
+		FoundMaterial = ImportedMaterialData.GetUnrealMaterial(FbxMaterial);
 	}
 	else
 	{
-		FBXImportOptions* FbxImportOptions = GetImportOptions();
-
 		FText Error;
-		UMaterialInterface* FoundMaterial = UMaterialImportHelpers::FindExistingMaterialFromSearchLocation(MaterialFullName, BasePackageName, FbxImportOptions->MaterialSearchLocation, Error);
-		
+		FoundMaterial = UMaterialImportHelpers::FindExistingMaterialFromSearchLocation(MaterialFullName, BasePackageName, MaterialSearchLocation, Error);
+
 		if (!Error.IsEmpty())
 		{
 			AddTokenizedErrorMessage(
@@ -585,57 +585,69 @@ void UnFbx::FFbxImporter::CreateUnrealMaterial(FbxSurfaceMaterial& FbxMaterial,
 						Error)),
 				FFbxErrors::Generic_LoadingSceneFailed);
 		}
-		// do not override existing materials
+		
 		if (FoundMaterial)
 		{
-			ImportedMaterialData.AddImportedMaterial( FbxMaterial, *FoundMaterial );
-			OutMaterials.Add(FoundMaterial);
-			return;
+			ImportedMaterialData.AddImportedMaterial(FbxMaterial, *FoundMaterial);
 		}
 	}
-	
+
+	return FoundMaterial;
+}
+
+
+UMaterialInterface* UnFbx::FFbxImporter::CreateUnrealMaterial(const FbxSurfaceMaterial& FbxMaterial, TArray<FString>& OutUVSets, bool bForSkeletalMesh)
+{
+	// Make sure we have a parent
+	if ( !ensure(Parent.IsValid()) )
+	{
+		return nullptr;
+	}
+
+	const FString BasePackageName = GetMaterialBasePackageName(GetMaterialFullName(FbxMaterial));
 	const FString Suffix(TEXT(""));
-	FAssetToolsModule& AssetToolsModule = FModuleManager::LoadModuleChecked<FAssetToolsModule>("AssetTools");
 	FString FinalPackageName;
-	AssetToolsModule.Get().CreateUniqueAssetName(BasePackageName, Suffix, FinalPackageName, MaterialFullName);
+	FString FinalMaterialName;
+	FAssetToolsModule& AssetToolsModule = FModuleManager::LoadModuleChecked<FAssetToolsModule>("AssetTools");
+	AssetToolsModule.Get().CreateUniqueAssetName(BasePackageName, Suffix, FinalPackageName, FinalMaterialName);
 
-	UPackage* Package = CreatePackage(NULL, *FinalPackageName);
+	UPackage* Package = CreatePackage(nullptr, *FinalPackageName);
 	
 	// Check if we can use the specified base material to instance from it
 	FBXImportOptions* FbxImportOptions = GetImportOptions();
 	bool bCanInstance = false;
 	if (FbxImportOptions->BaseMaterial)
 	{
-		bCanInstance = false;
 		// try to use the material as a base for the new material to instance from
 		FbxProperty FbxDiffuseProperty = FbxMaterial.FindProperty(FbxSurfaceMaterial::sDiffuse);
 		if (FbxDiffuseProperty.IsValid())
 		{
-			bCanInstance = CanUseMaterialWithInstance(FbxMaterial, FbxSurfaceMaterial::sDiffuse, FbxImportOptions->BaseDiffuseTextureName, FbxImportOptions->BaseMaterial, UVSets);
+			bCanInstance = CanUseMaterialWithInstance(FbxMaterial, FbxSurfaceMaterial::sDiffuse, FbxImportOptions->BaseDiffuseTextureName, FbxImportOptions->BaseMaterial, OutUVSets);
 		}
 		else
 		{
 			bCanInstance = !FbxImportOptions->BaseColorName.IsEmpty();
 		}
+
 		FbxProperty FbxEmissiveProperty = FbxMaterial.FindProperty(FbxSurfaceMaterial::sEmissive);
-		if (FbxDiffuseProperty.IsValid())
+		if (FbxEmissiveProperty.IsValid())
 		{
-			bCanInstance &= CanUseMaterialWithInstance(FbxMaterial, FbxSurfaceMaterial::sEmissive, FbxImportOptions->BaseEmmisiveTextureName, FbxImportOptions->BaseMaterial, UVSets);
+			bCanInstance &= CanUseMaterialWithInstance(FbxMaterial, FbxSurfaceMaterial::sEmissive, FbxImportOptions->BaseEmmisiveTextureName, FbxImportOptions->BaseMaterial, OutUVSets);
 		}
 		else
 		{
 			bCanInstance &= !FbxImportOptions->BaseEmissiveColorName.IsEmpty();
 		}
-		bCanInstance &= CanUseMaterialWithInstance(FbxMaterial, FbxSurfaceMaterial::sSpecular, FbxImportOptions->BaseSpecularTextureName, FbxImportOptions->BaseMaterial, UVSets);
-		bCanInstance &= CanUseMaterialWithInstance(FbxMaterial, FbxSurfaceMaterial::sNormalMap, FbxImportOptions->BaseNormalTextureName, FbxImportOptions->BaseMaterial, UVSets);
-		bCanInstance &= CanUseMaterialWithInstance(FbxMaterial, FbxSurfaceMaterial::sTransparentColor, FbxImportOptions->BaseOpacityTextureName, FbxImportOptions->BaseMaterial, UVSets);
+		bCanInstance &= CanUseMaterialWithInstance(FbxMaterial, FbxSurfaceMaterial::sSpecular, FbxImportOptions->BaseSpecularTextureName, FbxImportOptions->BaseMaterial, OutUVSets);
+		bCanInstance &= CanUseMaterialWithInstance(FbxMaterial, FbxSurfaceMaterial::sNormalMap, FbxImportOptions->BaseNormalTextureName, FbxImportOptions->BaseMaterial, OutUVSets);
+		bCanInstance &= CanUseMaterialWithInstance(FbxMaterial, FbxSurfaceMaterial::sTransparentColor, FbxImportOptions->BaseOpacityTextureName, FbxImportOptions->BaseMaterial, OutUVSets);
 	}
 
 	UMaterialInterface* UnrealMaterialFinal = nullptr;
 	if (bCanInstance) {
 		auto MaterialInstanceFactory = NewObject<UMaterialInstanceConstantFactoryNew>();
 		MaterialInstanceFactory->InitialParent = FbxImportOptions->BaseMaterial;
-		UMaterialInstanceConstant* UnrealMaterialConstant = (UMaterialInstanceConstant*)MaterialInstanceFactory->FactoryCreateNew(UMaterialInstanceConstant::StaticClass(), Package, *MaterialFullName, RF_Standalone | RF_Public, NULL, GWarn);
+		UMaterialInstanceConstant* UnrealMaterialConstant = (UMaterialInstanceConstant*)MaterialInstanceFactory->FactoryCreateNew(UMaterialInstanceConstant::StaticClass(), Package, *FinalMaterialName, RF_Standalone | RF_Public, NULL, GWarn);
 		if (UnrealMaterialConstant != NULL)
 		{
 			UnrealMaterialFinal = UnrealMaterialConstant;
@@ -735,7 +747,7 @@ void UnFbx::FFbxImporter::CreateUnrealMaterial(FbxSurfaceMaterial& FbxMaterial,
 		auto MaterialFactory = NewObject<UMaterialFactoryNew>();
 	
 		UMaterial* UnrealMaterial = (UMaterial*)MaterialFactory->FactoryCreateNew(
-			UMaterial::StaticClass(), Package, *MaterialFullName, RF_Standalone|RF_Public, NULL, GWarn );
+			UMaterial::StaticClass(), Package, *FinalMaterialName, RF_Standalone|RF_Public, NULL, GWarn );
 
 		if (UnrealMaterial != NULL)
 		{
@@ -762,19 +774,19 @@ void UnFbx::FFbxImporter::CreateUnrealMaterial(FbxSurfaceMaterial& FbxMaterial,
 				UE_LOG(LogFbxMaterialImport, Display, TEXT("-------------------------------"));
 			}
 #endif
-			CreateAndLinkExpressionForMaterialProperty(FbxMaterial, UnrealMaterial, FbxSurfaceMaterial::sDiffuse, UnrealMaterial->BaseColor, false, UVSets, FVector2D(240, -320));
-			CreateAndLinkExpressionForMaterialProperty(FbxMaterial, UnrealMaterial, FbxSurfaceMaterial::sEmissive, UnrealMaterial->EmissiveColor, false, UVSets, FVector2D(240, -64));
-			CreateAndLinkExpressionForMaterialProperty(FbxMaterial, UnrealMaterial, FbxSurfaceMaterial::sSpecular, UnrealMaterial->Specular, false, UVSets, FVector2D(240, -128));
-			CreateAndLinkExpressionForMaterialProperty(FbxMaterial, UnrealMaterial, FbxSurfaceMaterial::sSpecularFactor, UnrealMaterial->Roughness, false, UVSets, FVector2D(240, -180));
-			CreateAndLinkExpressionForMaterialProperty(FbxMaterial, UnrealMaterial, FbxSurfaceMaterial::sShininess, UnrealMaterial->Metallic, false, UVSets, FVector2D(240, -210));
-			if (!CreateAndLinkExpressionForMaterialProperty(FbxMaterial, UnrealMaterial, FbxSurfaceMaterial::sNormalMap, UnrealMaterial->Normal, true, UVSets, FVector2D(240, 256)))
+			CreateAndLinkExpressionForMaterialProperty(FbxMaterial, UnrealMaterial, FbxSurfaceMaterial::sDiffuse, UnrealMaterial->BaseColor, false, OutUVSets, FVector2D(240, -320));
+			CreateAndLinkExpressionForMaterialProperty(FbxMaterial, UnrealMaterial, FbxSurfaceMaterial::sEmissive, UnrealMaterial->EmissiveColor, false, OutUVSets, FVector2D(240, -64));
+			CreateAndLinkExpressionForMaterialProperty(FbxMaterial, UnrealMaterial, FbxSurfaceMaterial::sSpecular, UnrealMaterial->Specular, false, OutUVSets, FVector2D(240, -128));
+			CreateAndLinkExpressionForMaterialProperty(FbxMaterial, UnrealMaterial, FbxSurfaceMaterial::sSpecularFactor, UnrealMaterial->Roughness, false, OutUVSets, FVector2D(240, -180));
+			CreateAndLinkExpressionForMaterialProperty(FbxMaterial, UnrealMaterial, FbxSurfaceMaterial::sShininess, UnrealMaterial->Metallic, false, OutUVSets, FVector2D(240, -210));
+			if (!CreateAndLinkExpressionForMaterialProperty(FbxMaterial, UnrealMaterial, FbxSurfaceMaterial::sNormalMap, UnrealMaterial->Normal, true, OutUVSets, FVector2D(240, 256)))
 			{
-				CreateAndLinkExpressionForMaterialProperty(FbxMaterial, UnrealMaterial, FbxSurfaceMaterial::sBump, UnrealMaterial->Normal, true, UVSets, FVector2D(240, 256)); // no bump in unreal, use as normal map
+				CreateAndLinkExpressionForMaterialProperty(FbxMaterial, UnrealMaterial, FbxSurfaceMaterial::sBump, UnrealMaterial->Normal, true, OutUVSets, FVector2D(240, 256)); // no bump in unreal, use as normal map
 			}
-			if (CreateAndLinkExpressionForMaterialProperty(FbxMaterial, UnrealMaterial, FbxSurfaceMaterial::sTransparentColor, UnrealMaterial->Opacity, false, UVSets, FVector2D(200, 256)))
+			if (CreateAndLinkExpressionForMaterialProperty(FbxMaterial, UnrealMaterial, FbxSurfaceMaterial::sTransparentColor, UnrealMaterial->Opacity, false, OutUVSets, FVector2D(200, 256)))
 			{
 				UnrealMaterial->BlendMode = BLEND_Translucent;
-				CreateAndLinkExpressionForMaterialProperty(FbxMaterial, UnrealMaterial, FbxSurfaceMaterial::sTransparencyFactor, UnrealMaterial->OpacityMask, false, UVSets, FVector2D(150, 256));
+				CreateAndLinkExpressionForMaterialProperty(FbxMaterial, UnrealMaterial, FbxSurfaceMaterial::sTransparencyFactor, UnrealMaterial->OpacityMask, false, OutUVSets, FVector2D(150, 256));
 
 			}
 			FixupMaterial(FbxMaterial, UnrealMaterial); // add random diffuse if none exists
@@ -783,26 +795,24 @@ void UnFbx::FFbxImporter::CreateUnrealMaterial(FbxSurfaceMaterial& FbxMaterial,
 		// compile shaders for PC (from UPrecompileShadersCommandlet::ProcessMaterial
 		// and FMaterialEditor::UpdateOriginalMaterial)
 	}
+
 	if (UnrealMaterialFinal)
 	{
 		// let the material update itself if necessary
 		UnrealMaterialFinal->PreEditChange(NULL);
 		UnrealMaterialFinal->PostEditChange();
-
 		ImportedMaterialData.AddImportedMaterial(FbxMaterial, *UnrealMaterialFinal);
-
-		OutMaterials.Add(UnrealMaterialFinal);
 	}
+	
+	return UnrealMaterialFinal;
 }
 
-int32 UnFbx::FFbxImporter::CreateNodeMaterials(FbxNode* FbxNode, TArray<UMaterialInterface*>& OutMaterials, TArray<FString>& UVSets, bool bForSkeletalMesh)
+void UnFbx::FFbxImporter::FindOrImportMaterialsFromNode(FbxNode* FbxNode, TArray<UMaterialInterface*>& OutMaterials, TArray<FString>& UVSets, bool bForSkeletalMesh)
 {
-	int32 MaterialCount = FbxNode->GetMaterialCount();
-	TArray<FbxSurfaceMaterial*> UsedSurfaceMaterials;
-	FbxMesh *MeshNode = FbxNode->GetMesh();
-	TSet<int32> UsedMaterialIndexes;
-	if (MeshNode)
+	if (FbxMesh *MeshNode = FbxNode->GetMesh())
 	{
+		TSet<int32> UsedMaterialIndexes;
+		
 		for (int32 ElementMaterialIndex = 0; ElementMaterialIndex < MeshNode->GetElementMaterialCount(); ++ElementMaterialIndex)
 		{
 			FbxGeometryElementMaterial *ElementMaterial = MeshNode->GetElementMaterial(ElementMaterialIndex);
@@ -826,25 +836,35 @@ int32 UnFbx::FFbxImporter::CreateNodeMaterials(FbxNode* FbxNode, TArray<UMateria
 				break;
 			}
 		}
-	}
-	for(int32 MaterialIndex=0; MaterialIndex < MaterialCount; ++MaterialIndex)
-	{
-		//Create only the material used by the mesh element material
-		if (MeshNode == nullptr || UsedMaterialIndexes.Contains(MaterialIndex))
+
+		for (int32 MaterialIndex = 0, MaterialCount = FbxNode->GetMaterialCount(); MaterialIndex < MaterialCount; ++MaterialIndex)
 		{
-			FbxSurfaceMaterial *FbxMaterial = FbxNode->GetMaterial(MaterialIndex);
+			const FbxSurfaceMaterial *FbxMaterial = FbxNode->GetMaterial(MaterialIndex);
+			UMaterialInterface* MaterialImported = nullptr;
 
-			if (FbxMaterial)
+			//Only create the material used by the mesh element material
+			if (FbxMaterial && UsedMaterialIndexes.Contains(MaterialIndex))
 			{
-				CreateUnrealMaterial(*FbxMaterial, OutMaterials, UVSets, bForSkeletalMesh);
+				if (UMaterialInterface* ExistingMaterial = FindExistingMaterialFromFbxMaterial(*FbxMaterial, ImportOptions->MaterialSearchLocation))
+				{
+					MaterialImported = ExistingMaterial;
+				}
+				else if (ImportOptions->bImportMaterials)
+				{
+					//Only create a new material if we are importing them and we could not find an existing one.
+					MaterialImported = CreateUnrealMaterial(*FbxMaterial, UVSets, bForSkeletalMesh);
+				}
 			}
+
+			//The fbxMaterial is not valid.
+			OutMaterials.Add(MaterialImported);			
 		}
-		else
-		{
-			OutMaterials.Add(nullptr);
-		}
 	}
-	return MaterialCount;
+	else
+	{
+		//Could not import the materials, no mesh found.
+		OutMaterials.AddDefaulted(FbxNode->GetMaterialCount());
+	}
 }
 
 
diff --git a/Engine/Source/Editor/UnrealEd/Private/Fbx/FbxSkeletalMeshImport.cpp b/Engine/Source/Editor/UnrealEd/Private/Fbx/FbxSkeletalMeshImport.cpp
--- a/Engine/Source/Editor/UnrealEd/Private/Fbx/FbxSkeletalMeshImport.cpp
+++ b/Engine/Source/Editor/UnrealEd/Private/Fbx/FbxSkeletalMeshImport.cpp
@@ -3105,21 +3105,18 @@ bool UnFbx::FFbxImporter::FillSkelMeshImporterFromFbx( FSkeletalMeshImportData&
 
 
 	TArray<UMaterialInterface*> Materials;
-	if (ImportOptions->bImportMaterials)
-	{
-		bool bForSkeletalMesh = true;
-		CreateNodeMaterials(Node, Materials, UVSets, bForSkeletalMesh);
-	}
-	else if (ImportOptions->bImportTextures)
+	const bool bForSkeletalMesh = true;
+
+	FindOrImportMaterialsFromNode(Node, Materials, UVSets, bForSkeletalMesh);
+	if (!ImportOptions->bImportMaterials && ImportOptions->bImportTextures)
 	{
+		//If we are not importing any new material, we might still want to import new textures.
 		ImportTexturesFromNode(Node);
 	}
 
 	// Maps local mesh material index to global material index
+	const int32 MaterialCount = Node->GetMaterialCount();
 	TArray<int32> MaterialMapping;
-
-	int32 MaterialCount = Node->GetMaterialCount();
-
 	MaterialMapping.AddUninitialized(MaterialCount);
 
 	for( int32 MaterialIndex = 0; MaterialIndex < MaterialCount; ++MaterialIndex )
diff --git a/Engine/Source/Editor/UnrealEd/Private/Fbx/FbxStaticMeshImport.cpp b/Engine/Source/Editor/UnrealEd/Private/Fbx/FbxStaticMeshImport.cpp
--- a/Engine/Source/Editor/UnrealEd/Private/Fbx/FbxStaticMeshImport.cpp
+++ b/Engine/Source/Editor/UnrealEd/Private/Fbx/FbxStaticMeshImport.cpp
@@ -310,18 +310,17 @@ bool UnFbx::FFbxImporter::BuildStaticMeshFromGeometry(FbxNode* Node, UStaticMesh
 		TRACE_CPUPROFILER_EVENT_SCOPE(CreateMaterials);
 		
 		TArray<UMaterialInterface*> Materials;
-		if (ImportOptions->bImportMaterials)
-		{
-			bool bForSkeletalMesh = false;
-			CreateNodeMaterials(Node, Materials, FBXUVs.UVSets, bForSkeletalMesh);
-		}
-		else if (ImportOptions->bImportTextures)
+		const bool bForSkeletalMesh = false;
+		
+		FindOrImportMaterialsFromNode(Node, Materials, FBXUVs.UVSets, bForSkeletalMesh);
+		if (!ImportOptions->bImportMaterials && ImportOptions->bImportTextures)
 		{
+			//If we are not importing any new material, we might still want to import new textures.
 			ImportTexturesFromNode(Node);
 		}
 
 		MaterialCount = Node->GetMaterialCount();
-		check(!ImportOptions->bImportMaterials || Materials.Num() == MaterialCount);
+		check(Materials.Num() == MaterialCount);
 	
 		// Used later to offset the material indices on the raw triangle data
 		MaterialIndexOffset = MeshMaterials.Num();
@@ -331,46 +330,14 @@ bool UnFbx::FFbxImporter::BuildStaticMeshFromGeometry(FbxNode* Node, UStaticMesh
 			FFbxMaterial* NewMaterial = new(MeshMaterials) FFbxMaterial;
 			FbxSurfaceMaterial *FbxMaterial = Node->GetMaterial(MaterialIndex);
 			NewMaterial->FbxMaterial = FbxMaterial;
-			if (ImportOptions->bImportMaterials)
+			
+			if (Materials[MaterialIndex])
 			{
 				NewMaterial->Material = Materials[MaterialIndex];
 			}
 			else
 			{
-				FString MaterialFullName = GetMaterialFullName(*FbxMaterial);
-				FString BasePackageName = UPackageTools::SanitizePackageName(FPackageName::GetLongPackagePath(StaticMesh->GetOutermost()->GetName()) / MaterialFullName);
-				FString MaterialPackagePath = BasePackageName + TEXT(".") + MaterialFullName;
-				UMaterialInterface* UnrealMaterialInterface = FindObject<UMaterialInterface>(NULL, *MaterialPackagePath);
-				if (UnrealMaterialInterface == nullptr)
-				{
-					// Try loading the object if its package exists on disk
-					FSoftObjectPath ObjectPath(MaterialPackagePath);
-
-					FString LongPackageName = ObjectPath.GetAssetName().IsEmpty() ? ObjectPath.ToString() : ObjectPath.GetLongPackageName();
-					if (FPackageName::DoesPackageExist(LongPackageName))
-					{
-						UnrealMaterialInterface = Cast<UMaterialInterface>(ObjectPath.TryLoad());
-					}
-				}
-				if (UnrealMaterialInterface == NULL)
-				{
-					//In case we do not found the material we can see if the material is in the material list of the static mesh material
-					FName MaterialFbxFullName = UTF8_TO_TCHAR(MakeName(FbxMaterial->GetName()));
-					for (const FStaticMaterial &StaticMaterial : StaticMesh->StaticMaterials)
-					{
-						if (StaticMaterial.ImportedMaterialSlotName == MaterialFbxFullName)
-						{
-							UnrealMaterialInterface = StaticMaterial.MaterialInterface;
-							break;
-						}
-					}
-				
-					if (UnrealMaterialInterface == NULL)
-					{
-						UnrealMaterialInterface = UMaterial::GetDefaultMaterial(MD_Surface);
-					}
-				}
-				NewMaterial->Material = UnrealMaterialInterface;
+				NewMaterial->Material = UMaterial::GetDefaultMaterial(MD_Surface);
 			}
 		}
 
diff --git a/Engine/Source/Editor/UnrealEd/Private/Kismet2/DebuggerCommands.cpp b/Engine/Source/Editor/UnrealEd/Private/Kismet2/DebuggerCommands.cpp
--- a/Engine/Source/Editor/UnrealEd/Private/Kismet2/DebuggerCommands.cpp
+++ b/Engine/Source/Editor/UnrealEd/Private/Kismet2/DebuggerCommands.cpp
@@ -939,9 +939,12 @@ TSharedRef< SWidget > FPlayWorldCommands::GeneratePlayMenuContent(TSharedRef<FUI
 		{
 			const UEnum* PlayNetModeEnum = FindObject<UEnum>(ANY_PACKAGE, TEXT("EPlayNetMode"));
 			
-			TSharedRef<SWidget> NetMode = SNew(SEnumCombobox, PlayNetModeEnum)
+			TSharedRef<SWidget> NetMode = SNew(SEnumComboBox, PlayNetModeEnum)
 				.CurrentValue(TAttribute<int32>::Create(TAttribute<int32>::FGetter::CreateStatic(&FInternalPlayWorldCommandCallbacks::GetNetPlayMode)))
-				.OnEnumSelectionChanged(SEnumCombobox::FOnEnumSelectionChanged::CreateStatic(&FInternalPlayWorldCommandCallbacks::SetNetPlayMode))
+				.ButtonStyle(FEditorStyle::Get(), "FlatButton.Light")
+				.ContentPadding(FMargin(2, 0))
+				.Font(FEditorStyle::GetFontStyle("Sequencer.AnimationOutliner.RegularFont"))
+				.OnEnumSelectionChanged(SEnumComboBox::FOnEnumSelectionChanged::CreateStatic(&FInternalPlayWorldCommandCallbacks::SetNetPlayMode))
 				.ToolTipText(LOCTEXT("NetworkModeToolTip", "Which network mode should the clients launch in? A server will automatically be started if needed."));
 
 			MenuBuilder.AddWidget(NetMode, LOCTEXT("NetworkModeMenuWidget", "Net Mode"));
diff --git a/Engine/Source/Editor/UnrealEd/Public/FbxImporter.h b/Engine/Source/Editor/UnrealEd/Public/FbxImporter.h
--- a/Engine/Source/Editor/UnrealEd/Public/FbxImporter.h
+++ b/Engine/Source/Editor/UnrealEd/Public/FbxImporter.h
@@ -528,13 +528,13 @@ UNREALED_API void ApplyImportUIToImportOptions(UFbxImportUI* ImportUI, FBXImport
 struct FImportedMaterialData
 {
 public:
-	void AddImportedMaterial( FbxSurfaceMaterial& FbxMaterial, UMaterialInterface& UnrealMaterial );
-	bool IsUnique( FbxSurfaceMaterial& FbxMaterial, FName ImportedMaterialName ) const;
+	void AddImportedMaterial( const FbxSurfaceMaterial& FbxMaterial, UMaterialInterface& UnrealMaterial );
+	bool IsAlreadyImported( const FbxSurfaceMaterial& FbxMaterial, FName ImportedMaterialName ) const;
 	UMaterialInterface* GetUnrealMaterial( const FbxSurfaceMaterial& FbxMaterial ) const;
 	void Clear();
 private:
 	/** Mapping of FBX material to Unreal material.  Some materials in FBX have the same name so we use this map to determine if materials are unique */
-	TMap<FbxSurfaceMaterial*, TWeakObjectPtr<UMaterialInterface> > FbxToUnrealMaterialMap;
+	TMap<const FbxSurfaceMaterial*, TWeakObjectPtr<UMaterialInterface> > FbxToUnrealMaterialMap;
 	TSet<FName> ImportedMaterialNames;
 };
 
@@ -1228,14 +1228,6 @@ public:
 		FString GetName() const { return FbxMaterial ? ANSI_TO_TCHAR(FbxMaterial->GetName()) : (Material != nullptr ? Material->GetName() : TEXT("None")); }
 	};
 
-	/**
-	* Make material Unreal asset name from the Fbx material
-	*
-	* @param FbxMaterial Material from the Fbx node
-	* @return Sanitized asset name
-	*/
-	FString GetMaterialFullName(FbxSurfaceMaterial& FbxMaterial);
-
 	FbxGeometryConverter* GetGeometryConverter() { return GeometryConverter; }
 
 	/*
@@ -1257,6 +1249,16 @@ protected:
 		FIXEDANDCONVERTED,
 	};
 	
+	/**
+	* Make material Unreal asset name from the Fbx material
+	*
+	* @param FbxMaterial Material from the Fbx node
+	* @return Sanitized asset name
+	*/
+	FString GetMaterialFullName(const FbxSurfaceMaterial& FbxMaterial) const;
+
+	FString GetMaterialBasePackageName(const FString& MaterialFullName) const;
+
 	static TSharedPtr<FFbxImporter> StaticInstance;
 	static TSharedPtr<FFbxImporter> StaticPreviewInstance;
 	
@@ -1487,7 +1489,7 @@ protected:
 	 * @param UVSet
 	 * @return bool	
 	 */
-	bool CreateAndLinkExpressionForMaterialProperty(	FbxSurfaceMaterial& FbxMaterial,
+	bool CreateAndLinkExpressionForMaterialProperty(	const FbxSurfaceMaterial& FbxMaterial,
 														UMaterial* UnrealMaterial,
 														const char* MaterialProperty ,
 														FExpressionInput& MaterialInput, 
@@ -1504,7 +1506,7 @@ protected:
 	* @param bSetupAsNormalMap
 	* @return bool
 	*/
-	bool LinkMaterialProperty(FbxSurfaceMaterial& FbxMaterial,
+	bool LinkMaterialProperty(const FbxSurfaceMaterial& FbxMaterial,
 		UMaterialInstanceConstant* UnrealMaterial,
 		const char* MaterialProperty,
 		FName ParameterValue,
@@ -1514,7 +1516,7 @@ protected:
 	 *
 	 * @param unMaterial Unreal material object.
 	 */
-	void FixupMaterial( FbxSurfaceMaterial& FbxMaterial, UMaterial* unMaterial);
+	void FixupMaterial( const FbxSurfaceMaterial& FbxMaterial, UMaterial* unMaterial);
 	
 	/**
 	 * Get material mapping array according "Skinxx" flag in material name
@@ -1536,13 +1538,25 @@ protected:
 	 * Create materials from Fbx node.
 	 * Only setup channels that connect to texture, and setup the UV coordinate of texture.
 	 * If diffuse channel has no texture, one default node will be created with constant.
+	 * If a material cannot be imported a nullptr will be insterted in the outMaterials array in its place.
 	 *
 	 * @param FbxNode  Fbx node
 	 * @param outMaterials Unreal Materials we created
 	 * @param UVSets UV set name list
 	 * @return int32 material count that created from the Fbx node
 	 */
-	int32 CreateNodeMaterials(FbxNode* FbxNode, TArray<UMaterialInterface*>& outMaterials, TArray<FString>& UVSets, bool bForSkeletalMesh);
+	void FindOrImportMaterialsFromNode(FbxNode* FbxNode, TArray<UMaterialInterface*>& outMaterials, TArray<FString>& UVSets, bool bForSkeletalMesh);
+
+	/**
+	 * Tries to find an existing UnrealMaterial from the FbxMaterial, returns nullptr if could not find a material.
+	 * The function will look for materials imported by the FbxFactory first,
+	 * and then search into the asset database using the passed MaterialSearchLocation search scope.
+	 * 
+	 * @param FbxMaterial				The FbxMaterial used to search the UnrealMaterial
+	 * @param MaterialSearchLocation	The asset database search scope.
+	 * @return							The UMaterialInterfaceFound, returns nullptr if no material was found.
+	 */
+	UMaterialInterface* FindExistingMaterialFromFbxMaterial(const FbxSurfaceMaterial& FbxMaterial, EMaterialSearchLocation MaterialSearchLocation);
 
 	/**
 	 * Create Unreal material from Fbx material.
@@ -1550,10 +1564,12 @@ protected:
 	 * If diffuse channel has no texture, one default node will be created with constant.
 	 *
 	 * @param KFbxSurfaceMaterial*  Fbx material
-	 * @param outMaterials Unreal Materials we created
 	 * @param outUVSets
+	 * @param bForSkeletalMesh		If set to true, the material target usage will be set to "SkeletalMesh".
+	 * @return						The created material.
 	 */
-	void CreateUnrealMaterial(FbxSurfaceMaterial& FbxMaterial, TArray<UMaterialInterface*>& OutMaterials, TArray<FString>& UVSets, bool bForSkeletalMesh);
+	UMaterialInterface* CreateUnrealMaterial(const FbxSurfaceMaterial& FbxMaterial, TArray<FString>& OutUVSets, bool bForSkeletalMesh);
+
 
 	/**
 	 * Visit all materials of one node, import textures from materials.
