Integrating CL# 11313960 and 11314589 from Dev-Enterprise
Adding a warning when using the wrong type of options in python when importing FBX, instead of silently failing to apply them.

diff --git a/Engine/Source/Editor/UnrealEd/Private/Fbx/FbxFactory.cpp b/Engine/Source/Editor/UnrealEd/Private/Fbx/FbxFactory.cpp
--- a/Engine/Source/Editor/UnrealEd/Private/Fbx/FbxFactory.cpp
+++ b/Engine/Source/Editor/UnrealEd/Private/Fbx/FbxFactory.cpp
@@ -283,6 +283,11 @@ UObject* UFbxFactory::FactoryCreateFile
 		}
 		ImportUI = OverrideImportUI;
 	}
+	else if(AssetImportTask && AssetImportTask->Options)
+	{
+		UE_LOG(LogFbx, Warning, TEXT("The options set in the Asset Import Task are not of type UFbxImportUI and will be ignored"));
+	}
+
 	//We are not re-importing
 	ImportUI->bIsReimport = false;
 	ImportUI->ReimportMesh = nullptr;
