Integrating CL# 11214374 from Dev-Enterprise
SkeletalMesh Import: Fixing custom properties on bones being applied on unrelated skeletal meshes

Alexis.Matte

diff --git a/Engine/Source/Editor/UnrealEd/Private/Fbx/FbxFactory.cpp b/Engine/Source/Editor/UnrealEd/Private/Fbx/FbxFactory.cpp
--- a/Engine/Source/Editor/UnrealEd/Private/Fbx/FbxFactory.cpp
+++ b/Engine/Source/Editor/UnrealEd/Private/Fbx/FbxFactory.cpp
@@ -606,13 +606,9 @@ UObject* UFbxFactory::FactoryCreateFile
 							{
 								FName OutputName = FbxImporter->MakeNameForMesh(Name.ToString(), SkelMeshNodeArray[0]);
 
-								TArray<FbxNode*> SkeletonNodeArray;
-								FbxImporter->FillFbxSkeletonArray(RootNodeToImport, SkeletonNodeArray);
-
 								UnFbx::FFbxImporter::FImportSkeletalMeshArgs ImportSkeletalMeshArgs;
 								ImportSkeletalMeshArgs.InParent = InParent;
 								ImportSkeletalMeshArgs.NodeArray = SkelMeshNodeArray;
-								ImportSkeletalMeshArgs.BoneNodeArray = SkeletonNodeArray;
 								ImportSkeletalMeshArgs.Name = OutputName;
 								ImportSkeletalMeshArgs.Flags = Flags;
 								ImportSkeletalMeshArgs.TemplateImportData = ImportUI->SkeletalMeshImportData;
diff --git a/Engine/Source/Editor/UnrealEd/Private/Fbx/FbxMainImport.cpp b/Engine/Source/Editor/UnrealEd/Private/Fbx/FbxMainImport.cpp
--- a/Engine/Source/Editor/UnrealEd/Private/Fbx/FbxMainImport.cpp
+++ b/Engine/Source/Editor/UnrealEd/Private/Fbx/FbxMainImport.cpp
@@ -2021,19 +2021,6 @@ void FFbxImporter::FillFbxMeshArray(FbxNode* Node, TArray<FbxNode*>& outMeshArra
 	}
 }
 
-void FFbxImporter::FillFbxSkeletonArray(FbxNode* Node, TArray<FbxNode*>& OutMeshArray)
-{
-	if (Node->GetSkeleton())
-	{
-		OutMeshArray.Add(Node);
-	}
-
-	for (int32 ChildIndex = 0; ChildIndex < Node->GetChildCount(); ++ChildIndex)
-	{
-		FillFbxSkeletonArray(Node->GetChild(ChildIndex), OutMeshArray);
-	}
-}
-
 void FFbxImporter::FillFbxMeshAndLODGroupArray(FbxNode* Node, TArray<FbxNode*>& outLODGroupArray, TArray<FbxNode*>& outMeshArray)
 {
 	// Is this node an LOD group
diff --git a/Engine/Source/Editor/UnrealEd/Private/Fbx/FbxSceneImportFactory.cpp b/Engine/Source/Editor/UnrealEd/Private/Fbx/FbxSceneImportFactory.cpp
--- a/Engine/Source/Editor/UnrealEd/Private/Fbx/FbxSceneImportFactory.cpp
+++ b/Engine/Source/Editor/UnrealEd/Private/Fbx/FbxSceneImportFactory.cpp
@@ -1957,14 +1957,10 @@ UObject* UFbxSceneImportFactory::ImportOneSkeletalMesh(void* VoidRootNodeToImpor
 			RootNodeInfo->AttributeInfo->SetOriginalImportPath(PackageName);
 			FName SkeletalMeshFName = FName(*SkeletalMeshName);
 
-			TArray<FbxNode*> SkeletonNodeArray;
-			FbxImporter->FillFbxSkeletonArray(RootNodeToImport, SkeletonNodeArray);
-
 			//Import the skeletal mesh
 			UnFbx::FFbxImporter::FImportSkeletalMeshArgs ImportSkeletalMeshArgs;
 			ImportSkeletalMeshArgs.InParent = Pkg;
 			ImportSkeletalMeshArgs.NodeArray = bUseSkelMeshNodePivotArray ? SkelMeshNodePivotArray : SkelMeshNodeArray;
-			ImportSkeletalMeshArgs.BoneNodeArray = SkeletonNodeArray;
 			ImportSkeletalMeshArgs.Name = SkeletalMeshFName;
 			ImportSkeletalMeshArgs.Flags = Flags;
 			ImportSkeletalMeshArgs.TemplateImportData = SkeletalMeshImportData;
diff --git a/Engine/Source/Editor/UnrealEd/Private/Fbx/FbxSkeletalMeshImport.cpp b/Engine/Source/Editor/UnrealEd/Private/Fbx/FbxSkeletalMeshImport.cpp
--- a/Engine/Source/Editor/UnrealEd/Private/Fbx/FbxSkeletalMeshImport.cpp
+++ b/Engine/Source/Editor/UnrealEd/Private/Fbx/FbxSkeletalMeshImport.cpp
@@ -1326,15 +1326,13 @@ bool UnFbx::FFbxImporter::ImportBone(TArray<FbxNode*>& NodeArray, FSkeletalMeshI
 	return true;
 }
 
-bool UnFbx::FFbxImporter::FillSkeletalMeshImportData(TArray<FbxNode*>& NodeArray, UFbxSkeletalMeshImportData* TemplateImportData, TArray<FbxShape*> *FbxShapeArray, FSkeletalMeshImportData* OutData, TArray<FName> &LastImportedMaterialNames, const bool bIsReimport, const TMap<FVector, FColor>& ExistingVertexColorData)
+bool UnFbx::FFbxImporter::FillSkeletalMeshImportData(TArray<FbxNode*>& NodeArray, UFbxSkeletalMeshImportData* TemplateImportData, TArray<FbxShape*> *FbxShapeArray, FSkeletalMeshImportData* OutData, TArray<FbxNode*>& OutImportedSkeletonLinkNodes, TArray<FName> &LastImportedMaterialNames, const bool bIsReimport, const TMap<FVector, FColor>& ExistingVertexColorData)
 {
 	if (NodeArray.Num() == 0)
 	{
 		return false;
 	}
 
-	int32 SkelType = 0; // 0 for skeletal mesh, 1 for rigid mesh
-
 	FbxNode* Node = NodeArray[0];
 	// find the mesh by its name
 	FbxMesh* FbxMesh = Node->GetMesh();
@@ -1344,16 +1342,11 @@ bool UnFbx::FFbxImporter::FillSkeletalMeshImportData(TArray<FbxNode*>& NodeArray
 		return false;
 	}
 
-	if (FbxMesh->GetDeformerCount(FbxDeformer::eSkin) == 0)
-	{
-		SkelType = 1;
-	}
-
 	// Fill with data from buffer - contains the full .FBX file. 	
 	FSkeletalMeshImportData* SkelMeshImportDataPtr = OutData;
 
-
-	TArray<FbxNode*> SortedLinkArray;
+	OutImportedSkeletonLinkNodes.Empty();
+	TArray<FbxNode*>& SortedLinkArray = OutImportedSkeletonLinkNodes;
 	FbxArray<FbxAMatrix> GlobalsPerLink;
 
 	SkelMeshImportDataPtr->bUseT0AsRefPose = ImportOptions->bUseT0AsRefPose;
@@ -1621,9 +1614,6 @@ USkeletalMesh* UnFbx::FFbxImporter::ImportSkeletalMesh(FImportSkeletalMeshArgs &
 	}
 
 	int32 SafeLODIndex = ImportSkeletalMeshArgs.LodIndex < 0 ? 0 : ImportSkeletalMeshArgs.LodIndex;
-
-	int32 SkelType = 0; // 0 for skeletal mesh, 1 for rigid mesh
-	
 	FbxNode* Node = ImportSkeletalMeshArgs.NodeArray[0];
 	// find the mesh by its name
 	FbxMesh* FbxMesh = Node->GetMesh();
@@ -1633,10 +1623,7 @@ USkeletalMesh* UnFbx::FFbxImporter::ImportSkeletalMesh(FImportSkeletalMeshArgs &
 		AddTokenizedErrorMessage(FTokenizedMessage::Create(EMessageSeverity::Warning, FText::Format(LOCTEXT("FbxSkeletaLMeshimport_NodeInvalidSkeletalMesh", "Fbx node: '{0}' is not a valid skeletal mesh"), FText::FromString(Node->GetName()))), FFbxErrors::Generic_Mesh_MeshNotFound);
 		return nullptr;
 	}
-	if (FbxMesh->GetDeformerCount(FbxDeformer::eSkin) == 0)
-	{
-		SkelType = 1;
-	}
+
 	//Make sure the render thread is done
 	FlushRenderingCommands();
 
@@ -1687,8 +1674,8 @@ USkeletalMesh* UnFbx::FFbxImporter::ImportSkeletalMesh(FImportSkeletalMeshArgs &
 	//////////////////////////////////////////////////////////////////////////
 	// We must do a maximum of fail test before backing up the data since the backup is destructive on the existing skeletal mesh.
 	// See the comment later when we call the following function (SaveExistingSkelMeshData)
-
-	if (FillSkeletalMeshImportData(ImportSkeletalMeshArgs.NodeArray, ImportSkeletalMeshArgs.TemplateImportData, ImportSkeletalMeshArgs.FbxShapeArray, SkelMeshImportDataPtr, LastImportedMaterialNames, ExistingSkelMesh != nullptr, ExistingVertexColorData) == false)
+	TArray<FbxNode*> ImportedSkeletonLinkNodes;
+	if (FillSkeletalMeshImportData(ImportSkeletalMeshArgs.NodeArray, ImportSkeletalMeshArgs.TemplateImportData, ImportSkeletalMeshArgs.FbxShapeArray, SkelMeshImportDataPtr, ImportedSkeletonLinkNodes, LastImportedMaterialNames, ExistingSkelMesh != nullptr, ExistingVertexColorData) == false)
 	{
 		AddTokenizedErrorMessage(FTokenizedMessage::Create(EMessageSeverity::Warning, LOCTEXT("FbxSkeletaLMeshimport_FillupImportData", "Get Import Data has failed.")), FFbxErrors::SkeletalMesh_FillImportDataFailed);
 		if (SkeletalMesh)
@@ -2111,7 +2098,7 @@ USkeletalMesh* UnFbx::FFbxImporter::ImportSkeletalMesh(FImportSkeletalMeshArgs &
 	}
 
 	// Import bone metadata to Skeleton
-	for (FbxNode* SkeletonNode : ImportSkeletalMeshArgs.BoneNodeArray)
+	for (FbxNode* SkeletonNode : ImportedSkeletonLinkNodes)
 	{
 		ImportNodeCustomProperties(SkeletalMesh->Skeleton, SkeletonNode, true);
 	}
@@ -2538,14 +2525,10 @@ USkeletalMesh* UnFbx::FFbxImporter::ReimportSkeletalMesh(USkeletalMesh* Mesh, UF
 			FSkeletalMeshImportData OutData;
 			if (LODIndex == 0)
 			{
-				TArray<FbxNode*> SkeletonNodeArray;
-				FillFbxSkeletonArray(Scene->GetRootNode(), SkeletonNodeArray);
-
 				ImportMeshLodData.AddZeroed();
 				UnFbx::FFbxImporter::FImportSkeletalMeshArgs ImportSkeletalMeshArgs;
 				ImportSkeletalMeshArgs.InParent = Mesh->GetOuter();
 				ImportSkeletalMeshArgs.NodeArray = SkelMeshNodeArray;
-				ImportSkeletalMeshArgs.BoneNodeArray = SkeletonNodeArray;
 				ImportSkeletalMeshArgs.Name = *Mesh->GetName();
 				ImportSkeletalMeshArgs.Flags = RF_Public | RF_Standalone;
 				ImportSkeletalMeshArgs.TemplateImportData = TemplateImportData;
diff --git a/Engine/Source/Editor/UnrealEd/Public/FbxImporter.h b/Engine/Source/Editor/UnrealEd/Public/FbxImporter.h
--- a/Engine/Source/Editor/UnrealEd/Public/FbxImporter.h
+++ b/Engine/Source/Editor/UnrealEd/Public/FbxImporter.h
@@ -833,7 +833,6 @@ public:
 
 		UObject* InParent;
 		TArray<FbxNode*> NodeArray;
-		TArray<FbxNode*> BoneNodeArray;
 		FName Name;
 		EObjectFlags Flags;
 		UFbxSkeletalMeshImportData* TemplateImportData;
@@ -1002,14 +1001,6 @@ public:
 	*/
 	UNREALED_API void FillFbxMeshArray(FbxNode* Node, TArray<FbxNode*>& outMeshArray, UnFbx::FFbxImporter* FFbxImporter);
 
-	/**
-	* Get all Fbx Skeleton nodes
-	*
-	* @param Node Root node to find skeleton nodes
-	* @param outNodeArray return skeleton nodes
-	*/
-	UNREALED_API void FillFbxSkeletonArray(FbxNode* Node, TArray<FbxNode*>& OutNodeArray);
-
 	/**
 	* Get all Fbx mesh objects not under a LOD group and all LOD group node
 	*
@@ -1400,7 +1391,9 @@ public:
 	*
 	* @returns bool*	true if import successfully.
 	*/
-	bool FillSkeletalMeshImportData(TArray<FbxNode*>& NodeArray, UFbxSkeletalMeshImportData* TemplateImportData, TArray<FbxShape*> *FbxShapeArray, FSkeletalMeshImportData* OutData, TArray<FName> &LastImportedMaterialNames, const bool bIsReimport, const TMap<FVector, FColor>& ExistingVertexColorData);
+	bool FillSkeletalMeshImportData(TArray<FbxNode*>& NodeArray, UFbxSkeletalMeshImportData* TemplateImportData, TArray<FbxShape*> *FbxShapeArray,
+									FSkeletalMeshImportData* OutData, TArray<FbxNode*>& OutImportedSkeletonLinkNodes, TArray<FName> &LastImportedMaterialNames, 
+									const bool bIsReimport, const TMap<FVector, FColor>& ExistingVertexColorData);
 
 protected:
 
