Integrating CL#11425305 from Dev-Enterprise
Fix: Datasmith static mesh cannot be reimported after importing a LOD in the static mesh editor

diff --git a/Engine/Source/Editor/UnrealEd/Private/Fbx/FbxStaticMeshImport.cpp b/Engine/Source/Editor/UnrealEd/Private/Fbx/FbxStaticMeshImport.cpp
--- a/Engine/Source/Editor/UnrealEd/Private/Fbx/FbxStaticMeshImport.cpp
+++ b/Engine/Source/Editor/UnrealEd/Private/Fbx/FbxStaticMeshImport.cpp
@@ -1677,26 +1677,24 @@ UStaticMesh* UnFbx::FFbxImporter::ImportStaticMeshAsSingle(UObject* InParent, TA
 				StaticMesh->GetSourceModel(ModelLODIndex).ReductionSettings = LODGroup.GetDefaultSettings(ModelLODIndex);
 			}
 			StaticMesh->LightMapResolution = LODGroup.GetDefaultLightMapResolution();
-		}
-
-		UFbxStaticMeshImportData* ImportData = UFbxStaticMeshImportData::GetImportDataForStaticMesh(StaticMesh, TemplateImportData);
 
-		//@third party BEGIN SIMPLYGON
-		/* ImportData->Update(UFactory::GetCurrentFilename());
-		Developer Note: Update method above computed Hash internally. Hash is calculated based on the file size.
-		Doing this for CAD files with thousands of components hugely increases the time.
-		The following method uses a precomputed hash (once per file). Huge time savings.
-		*/
-		FString FactoryCurrentFileName = UFactory::GetCurrentFilename();
-		if (!FactoryCurrentFileName.IsEmpty() && LODIndex == 0)
-		{
-			//The factory is instantiate only when importing or re-importing the LOD 0
-			//The LOD re-import is not using the factory so the static function UFactory::GetCurrentFilename()
-			//will return the last fbx imported asset name or no name if there was no imported asset before.
-			ImportData->Update(FactoryCurrentFileName, UFactory::GetFileHash());
+			//@third party BEGIN SIMPLYGON
+			/* ImportData->Update(UFactory::GetCurrentFilename());
+			Developer Note: Update method above computed Hash internally. Hash is calculated based on the file size.
+			Doing this for CAD files with thousands of components hugely increases the time.
+			The following method uses a precomputed hash (once per file). Huge time savings.
+			*/
+			UFbxStaticMeshImportData* ImportData = UFbxStaticMeshImportData::GetImportDataForStaticMesh(StaticMesh, TemplateImportData);
+			FString FactoryCurrentFileName = UFactory::GetCurrentFilename();
+			if (!FactoryCurrentFileName.IsEmpty())
+			{
+				//The factory is instantiate only when importing or re-importing the LOD 0
+				//The LOD re-import is not using the factory so the static function UFactory::GetCurrentFilename()
+				//will return the last fbx imported asset name or no name if there was no imported asset before.
+				ImportData->Update(FactoryCurrentFileName, UFactory::GetFileHash());
+			}
+			//@third party END SIMPLYGON
 		}
-		//@third party END SIMPLYGON
-
 
 		// @todo This overrides restored values currently but we need to be able to import over the existing settings if the user chose to do so.
 		SrcModel.BuildSettings.bRemoveDegenerates = ImportOptions->bRemoveDegenerates;
