Added logic to silently renaming assets in UEditorLevelUtils::MoveActorsToLevel: Currently only the copy can be done silently. This change makes  UEditorLevelUtils::MoveActorsToLevel completely silent if needed.

note: Was CL #11451391 against Enterprise branch

diff --git a/Engine/Source/Editor/UnrealEd/Private/EditorLevelUtils.cpp b/Engine/Source/Editor/UnrealEd/Private/EditorLevelUtils.cpp
--- a/Engine/Source/Editor/UnrealEd/Private/EditorLevelUtils.cpp
+++ b/Engine/Source/Editor/UnrealEd/Private/EditorLevelUtils.cpp
@@ -55,12 +55,12 @@ DEFINE_LOG_CATEGORY(LogLevelTools);
 #define LOCTEXT_NAMESPACE "EditorLevelUtils"
 
 
-int32 UEditorLevelUtils::MoveActorsToLevel(const TArray<AActor*>& ActorsToMove, ULevelStreaming* DestStreamingLevel, bool bWarnAboutReferences)
+int32 UEditorLevelUtils::MoveActorsToLevel(const TArray<AActor*>& ActorsToMove, ULevelStreaming* DestStreamingLevel, bool bWarnAboutReferences, bool bWarnAboutRenaming)
 {
-	return MoveActorsToLevel(ActorsToMove, DestStreamingLevel ? DestStreamingLevel->GetLoadedLevel() : nullptr, bWarnAboutReferences);
+	return MoveActorsToLevel(ActorsToMove, DestStreamingLevel ? DestStreamingLevel->GetLoadedLevel() : nullptr, bWarnAboutReferences, bWarnAboutRenaming);
 }
 
-int32 UEditorLevelUtils::MoveActorsToLevel(const TArray<AActor*>& ActorsToMove, ULevel* DestLevel, bool bWarnAboutReferences)
+int32 UEditorLevelUtils::MoveActorsToLevel(const TArray<AActor*>& ActorsToMove, ULevel* DestLevel, bool bWarnAboutReferences, bool bWarnAboutRenaming)
 {
 	int32 NumMovedActors = 0;
 
@@ -221,7 +221,14 @@ int32 UEditorLevelUtils::MoveActorsToLevel(const TArray<AActor*>& ActorsToMove,
 					
 				if (RenameData.Num() > 0)
 				{
-					AssetToolsModule.Get().RenameAssetsWithDialog(RenameData);
+					if(bWarnAboutRenaming)
+					{
+						AssetToolsModule.Get().RenameAssetsWithDialog(RenameData);
+					}
+					else
+					{
+						AssetToolsModule.Get().RenameAssets(RenameData);
+					}
 				}
 
 				// Restore new level visibility to previous state
diff --git a/Engine/Source/Editor/UnrealEd/Public/EditorLevelUtils.h b/Engine/Source/Editor/UnrealEd/Public/EditorLevelUtils.h
--- a/Engine/Source/Editor/UnrealEd/Public/EditorLevelUtils.h
+++ b/Engine/Source/Editor/UnrealEd/Public/EditorLevelUtils.h
@@ -62,7 +62,7 @@ public:
 	 * @return							The number of actors that were successfully moved to the new level
 	 */
 	UFUNCTION(BlueprintCallable, Category = "Editor Scripting | Level Creation")
-	static UNREALED_API int32 MoveActorsToLevel(const TArray<AActor*>& ActorsToMove, ULevelStreaming* DestStreamingLevel, bool bWarnAboutReferences = true);
+	static UNREALED_API int32 MoveActorsToLevel(const TArray<AActor*>& ActorsToMove, ULevelStreaming* DestStreamingLevel, bool bWarnAboutReferences = true, bool bWarnAboutRenaming = true);
 
 	/**
 	 * Moves the currently selected actors to the specified streaming level. The new actors will be selected
@@ -86,7 +86,7 @@ public:
 	static UNREALED_API void MakeLevelCurrent(ULevel* InLevel, bool bEvenIfLocked = false);
 
 	
-	static UNREALED_API int32 MoveActorsToLevel(const TArray<AActor*>& ActorsToMove, ULevel* DestLevel, bool bWarnAboutReferences = true);
+	static UNREALED_API int32 MoveActorsToLevel(const TArray<AActor*>& ActorsToMove, ULevel* DestLevel, bool bWarnAboutReferences = true, bool bWarnAboutRenaming = true);
 
 	static UNREALED_API int32 MoveSelectedActorsToLevel(ULevel* DestLevel, bool bWarnAboutReferences = true);
 
