Source Control: Unattended package save now handles adding new files to source control in more situations.
- CheckoutPackages now handles adding new files
- InternalCheckoutAndSavePackages now handles adding new files

Editgrate from 11420249

diff --git a/Engine/Source/Editor/UnrealEd/Private/FileHelpers.cpp b/Engine/Source/Editor/UnrealEd/Private/FileHelpers.cpp
--- a/Engine/Source/Editor/UnrealEd/Private/FileHelpers.cpp
+++ b/Engine/Source/Editor/UnrealEd/Private/FileHelpers.cpp
@@ -4295,6 +4295,7 @@ static bool InternalCheckoutAndSavePackages(const TArray<UPackage*>& PackagesToS
 			// Prevent modal window if not requested.
 			TGuardValue<bool> UnattendedScriptGuard(GIsRunningUnattendedScript, true);
 
+			TArray<UPackage*> PackagesCheckedOut;
 			const bool bErrorIfAlreadyCheckedOut = false;
 			const bool bConfirmPackageBranchCheckOutStatus = false;
 
