Fixed Play in Editor button being disabled if you canceled PIE startup due to a blueprint compile error via the dialog.

diff --git a/Engine/Source/Editor/UnrealEd/Private/PlayLevel.cpp b/Engine/Source/Editor/UnrealEd/Private/PlayLevel.cpp
--- a/Engine/Source/Editor/UnrealEd/Private/PlayLevel.cpp
+++ b/Engine/Source/Editor/UnrealEd/Private/PlayLevel.cpp
@@ -2326,6 +2326,7 @@ void UEditorEngine::StartPlayInEditorSession(FRequestPlaySessionParams& InReques
 		Info.ExpireDuration = 5.0f;
 		Info.bUseLargeFont = true;
 		FSlateNotificationManager::Get().AddNotification(Info);
+		CancelRequestPlaySession();
 		return;
 	}
 
@@ -2333,6 +2334,7 @@ void UEditorEngine::StartPlayInEditorSession(FRequestPlaySessionParams& InReques
 	// to close Matinee, we can't PIE.
 	if (!PromptMatineeClose())
 	{
+		CancelRequestPlaySession();
 		return;
 	}
 
@@ -2370,6 +2372,7 @@ void UEditorEngine::StartPlayInEditorSession(FRequestPlaySessionParams& InReques
 
 			FEditorDelegates::EndPIE.Broadcast(InRequestParams.WorldType == EPlaySessionWorldType::SimulateInEditor);
 			FNavigationSystem::OnPIEEnd(*InWorld);
+			CancelRequestPlaySession();
 			return;
 		}
 		else
@@ -2569,6 +2572,7 @@ void UEditorEngine::StartPlayInEditorSession(FRequestPlaySessionParams& InReques
 			CreateNewPlayInEditorInstance(InRequestParams, bRunAsDedicated, LocalNetMode);
 
 			// If there was an error creating an instance it will call EndPlay which invalidates the session info.
+			// This also broadcasts the EndPIE event so no need to do that here (to make a matching call to our BeginPIE)
 			if (!PlayInEditorSessionInfo.IsSet())
 			{
 				return;
