Fixes a regression where using Launch blocked HotReload/PIE buttons from being enabled.

diff --git a/Engine/Source/Editor/UnrealEd/Private/PlayLevelViaLauncher.cpp b/Engine/Source/Editor/UnrealEd/Private/PlayLevelViaLauncher.cpp
--- a/Engine/Source/Editor/UnrealEd/Private/PlayLevelViaLauncher.cpp
+++ b/Engine/Source/Editor/UnrealEd/Private/PlayLevelViaLauncher.cpp
@@ -62,11 +62,13 @@ void UEditorEngine::StartPlayUsingLauncherSession(FRequestPlaySessionParams& InR
 
 	if (!ensureAlwaysMsgf(PlaySessionRequest->LauncherTargetDevice.IsSet(), TEXT("PlayUsingLauncher should not be called without a target device set!")))
 	{
+		CancelRequestPlaySession();
 		return;
 	}
 
 	if (!ensureAlwaysMsgf(LastPlayUsingLauncherDeviceId.Len() > 0, TEXT("PlayUsingLauncher should not be called without a target device id set!")))
 	{
+		CancelRequestPlaySession();
 		return;
 	}
 
@@ -294,6 +296,11 @@ void UEditorEngine::StartPlayUsingLauncherSession(FRequestPlaySessionParams& InR
 		)
 	);
 
+	// Launch doesn't block PIE/Compile requests as it's an async background process, so we just
+	// cancel the request to denote it as having been handled. This has to come after we've used
+	// anything we might need from the original request.
+	CancelRequestPlaySession();
+
 	TSharedPtr<SNotificationItem> NotificationItem = FSlateNotificationManager::Get().AddNotification(Info);
 
 	if (!NotificationItem.IsValid())
@@ -366,8 +373,6 @@ void UEditorEngine::AutomationPlayUsingLauncher(const FString& InLauncherDeviceI
 */
 void UEditorEngine::CancelPlayUsingLauncher()
 {
-	CancelRequestPlaySession();
-
 	FText LaunchingText = LOCTEXT("LauncherTaskInProgressNotificationNotAuthorized", "Cannot launch to this device until this computer is authorized from the device");
 	FNotificationInfo Info(LaunchingText);
 	Info.ExpireDuration = 5.0f;
