Fix for redefinition of TArray

diff --git a/Engine/Source/Editor/UnrealEd/Private/FileHelpers.cpp b/Engine/Source/Editor/UnrealEd/Private/FileHelpers.cpp
--- a/Engine/Source/Editor/UnrealEd/Private/FileHelpers.cpp
+++ b/Engine/Source/Editor/UnrealEd/Private/FileHelpers.cpp
@@ -4299,7 +4299,6 @@ static bool InternalCheckoutAndSavePackages(const TArray<UPackage*>& PackagesToS
 			const bool bErrorIfAlreadyCheckedOut = false;
 			const bool bConfirmPackageBranchCheckOutStatus = false;
 
-			TArray<UPackage*> PackagesCheckedOut;
 			FEditorFileUtils::CheckoutPackages(PackagesToSave, &PackagesCheckedOut, bErrorIfAlreadyCheckedOut, bConfirmPackageBranchCheckOutStatus);
 
 			// Cannot mark files for add until after packages saved
