Reworked how realtime is overriden in the editor.  This fixes issues where a temporarily set realtime state is saved between editor sessions causing users to be unaware that their viewport is not realtime.  Shutting down the editor during PIE or remote desktop were two such cases.

Now there is a separate api on editor viewports to set and remove a temporary override state.  Deprecated existing methods and fixed up use cases.

diff --git a/Engine/Source/Editor/FoliageEdit/Private/FoliageEdMode.cpp b/Engine/Source/Editor/FoliageEdit/Private/FoliageEdMode.cpp
--- a/Engine/Source/Editor/FoliageEdit/Private/FoliageEdMode.cpp
+++ b/Engine/Source/Editor/FoliageEdit/Private/FoliageEdMode.cpp
@@ -402,8 +402,7 @@ void FEdModeFoliage::Enter()
 	// Force real-time viewports.  We'll back up the current viewport state so we can restore it when the
 	// user exits this mode.
 	const bool bWantRealTime = true;
-	const bool bRememberCurrentState = true;
-	ForceRealTimeViewports(bWantRealTime, bRememberCurrentState);
+	ForceRealTimeViewports(bWantRealTime);
 
 	if (!Toolkit.IsValid())
 	{
@@ -516,9 +515,7 @@ void FEdModeFoliage::Exit()
 	SphereBrushComponent->UnregisterComponent();
 
 	// Restore real-time viewport state if we changed it
-	const bool bWantRealTime = false;
-	const bool bRememberCurrentState = false;
-	ForceRealTimeViewports(bWantRealTime, bRememberCurrentState);
+	ForceRealTimeViewports(false);
 
 	// Clear the cache (safety, should be empty!)
 	LandscapeLayerCaches.Empty();
@@ -3955,7 +3952,7 @@ void FEdModeFoliage::ActorSelectionChangeNotify()
 
 
 /** Forces real-time perspective viewports */
-void FEdModeFoliage::ForceRealTimeViewports(const bool bEnable, const bool bStoreCurrentState)
+void FEdModeFoliage::ForceRealTimeViewports(const bool bEnable)
 {
 	FLevelEditorModule& LevelEditorModule = FModuleManager::GetModuleChecked<FLevelEditorModule>("LevelEditor");
 	TSharedPtr< IAssetViewport > ViewportWindow = LevelEditorModule.GetFirstActiveViewport();
@@ -3966,14 +3963,13 @@ void FEdModeFoliage::ForceRealTimeViewports(const bool bEnable, const bool bStor
 		{
 			if (bEnable)
 			{
-				Viewport.SetRealtime(bEnable, bStoreCurrentState);
+				const bool bShouldBeRealtime = false;
+				Viewport.SetRealtimeOverride(bShouldBeRealtime, LOCTEXT("RealtimeOverrideMessage_Foliage", "Foliage Mode"));
 			}
 			else
 			{
-				const bool bAllowDisable = true;
-				Viewport.RestoreRealtime(bAllowDisable);
+				Viewport.RemoveRealtimeOverride();
 			}
-
 		}
 	}
 }
diff --git a/Engine/Source/Editor/FoliageEdit/Private/FoliageEdMode.h b/Engine/Source/Editor/FoliageEdit/Private/FoliageEdMode.h
--- a/Engine/Source/Editor/FoliageEdit/Private/FoliageEdMode.h
+++ b/Engine/Source/Editor/FoliageEdit/Private/FoliageEdMode.h
@@ -414,7 +414,7 @@ public:
 	void OnObjectsReplaced(const TMap<UObject*, UObject*>& ReplacementMap);
 
 	/** Forces real-time perspective viewports */
-	void ForceRealTimeViewports(const bool bEnable, const bool bStoreCurrentState);
+	void ForceRealTimeViewports(const bool bEnable);
 
 	/** Start foliage tracing */
 	void StartFoliageBrushTrace(FEditorViewportClient* ViewportClient, class UViewportInteractor* Interactor = nullptr);
diff --git a/Engine/Source/Editor/Kismet/Private/SSCSEditorViewport.cpp b/Engine/Source/Editor/Kismet/Private/SSCSEditorViewport.cpp
--- a/Engine/Source/Editor/Kismet/Private/SSCSEditorViewport.cpp
+++ b/Engine/Source/Editor/Kismet/Private/SSCSEditorViewport.cpp
@@ -395,12 +395,13 @@ void SSCSEditorViewport::EnablePreview(bool bEnable)
 	if(bEnable)
 	{
 		// Restore the previously-saved realtime setting
-		ViewportClient->RestoreRealtime();
+		ViewportClient->RemoveRealtimeOverride();
 	}
 	else
 	{
 		// Disable and store the current realtime setting. This will bypass real-time rendering in the preview viewport (see UEditorEngine::UpdateSingleViewportClient).
-		ViewportClient->SetRealtime(false, true);
+		const bool bShouldBeRealtime = false;
+		ViewportClient->SetRealtimeOverride(bShouldBeRealtime, NSLOCTEXT("BlueprintEditor", "RealtimeOverrideMessage_Blueprints", "the active blueprint mode"));
 	}
 }
 
diff --git a/Engine/Source/Editor/LandscapeEditor/Private/LandscapeEdMode.cpp b/Engine/Source/Editor/LandscapeEditor/Private/LandscapeEdMode.cpp
--- a/Engine/Source/Editor/LandscapeEditor/Private/LandscapeEdMode.cpp
+++ b/Engine/Source/Editor/LandscapeEditor/Private/LandscapeEdMode.cpp
@@ -591,8 +591,7 @@ void FEdModeLandscape::Enter()
 	// Force real-time viewports.  We'll back up the current viewport state so we can restore it when the
 	// user exits this mode.
 	const bool bWantRealTime = true;
-	const bool bRememberCurrentState = true;
-	ForceRealTimeViewports(bWantRealTime, bRememberCurrentState);
+	ForceRealTimeViewports(bWantRealTime);
 
 	CurrentBrush->EnterBrush();
 	if (GizmoBrush)
@@ -657,9 +656,8 @@ void FEdModeLandscape::Exit()
 	}
 
 	// Restore real-time viewport state if we changed it
-	const bool bWantRealTime = false;
-	const bool bRememberCurrentState = false;
-	ForceRealTimeViewports(bWantRealTime, bRememberCurrentState);
+	const bool bWantRealTimeOverride = false;
+	ForceRealTimeViewports(bWantRealTimeOverride);
 
 	if (Toolkit.IsValid())
 	{
@@ -3366,8 +3364,7 @@ void FEdModeLandscape::PostUndo()
 }
 
 /** Forces all level editor viewports to realtime mode */
-void FEdModeLandscape::ForceRealTimeViewports(const bool bEnable, const bool bStoreCurrentState)
-
+void FEdModeLandscape::ForceRealTimeViewports(const bool bEnable)
 {
 	FLevelEditorModule& LevelEditorModule = FModuleManager::GetModuleChecked<FLevelEditorModule>("LevelEditor");
 	TSharedPtr<ILevelEditor> LevelEditor = LevelEditorModule.GetFirstLevelEditor();
@@ -3381,7 +3378,7 @@ void FEdModeLandscape::ForceRealTimeViewports(const bool bEnable, const bool bSt
 				FEditorViewportClient& Viewport = ViewportWindow->GetAssetViewportClient();
 				if (bEnable)
 				{
-					Viewport.SetRealtime(bEnable, bStoreCurrentState);
+					Viewport.SetRealtimeOverride(bEnable, LOCTEXT("RealtimeOverrideMessage_Landscape", "Landscape Mode"));
 
 					// @todo vreditor: Force game view to true in VREditor since we can't use hitproxies and debug objects yet
 					UVREditorMode* VREditorMode = Cast<UVREditorMode>( GEditor->GetEditorWorldExtensionsManager()->GetEditorWorldExtensions( GetWorld() )->FindExtension( UVREditorMode::StaticClass() ) );
@@ -3396,8 +3393,7 @@ void FEdModeLandscape::ForceRealTimeViewports(const bool bEnable, const bool bSt
 				}
 				else
 				{
-					const bool bAllowDisable = true;
-					Viewport.RestoreRealtime(bAllowDisable);
+					Viewport.RemoveRealtimeOverride();
 				}
 			}
 		}
diff --git a/Engine/Source/Editor/LandscapeEditor/Private/LandscapeEdMode.h b/Engine/Source/Editor/LandscapeEditor/Private/LandscapeEdMode.h
--- a/Engine/Source/Editor/LandscapeEditor/Private/LandscapeEdMode.h
+++ b/Engine/Source/Editor/LandscapeEditor/Private/LandscapeEdMode.h
@@ -487,7 +487,7 @@ public:
 	virtual bool PostConvertMouseMovement(FEditorViewportClient* InViewportClient) override;
 
 	/** Forces real-time perspective viewports */
-	void ForceRealTimeViewports(const bool bEnable, const bool bStoreCurrentState);
+	void ForceRealTimeViewports(const bool bEnable);
 
 	/** Trace under the mouse cursor and return the landscape hit and the hit location (in landscape quad space) */
 	bool LandscapeMouseTrace(FEditorViewportClient* ViewportClient, float& OutHitX, float& OutHitY);
diff --git a/Engine/Source/Editor/LevelEditor/Private/LevelViewportLayout2x2.cpp b/Engine/Source/Editor/LevelEditor/Private/LevelViewportLayout2x2.cpp
--- a/Engine/Source/Editor/LevelEditor/Private/LevelViewportLayout2x2.cpp
+++ b/Engine/Source/Editor/LevelEditor/Private/LevelViewportLayout2x2.cpp
@@ -68,7 +68,7 @@ TSharedRef<SWidget> FLevelViewportLayout2x2::MakeViewportLayout(const FString& L
 	TSharedPtr< ILevelViewportLayoutEntity > ViewportTL = LevelEditor.FactoryViewport(*TopLeftType, Args);
 
 	// Persp viewport
-	Args.bRealtime = !FPlatformMisc::IsRemoteSession();
+	Args.bRealtime = true;
 	Args.ConfigKey = *BottomLeftKey;
 	Args.ViewportType = LVT_Perspective;
 	TSharedPtr< ILevelViewportLayoutEntity > ViewportBL = LevelEditor.FactoryViewport(*BottomLeftType, Args);
diff --git a/Engine/Source/Editor/LevelEditor/Private/LevelViewportLayoutFourPanes.cpp b/Engine/Source/Editor/LevelEditor/Private/LevelViewportLayoutFourPanes.cpp
--- a/Engine/Source/Editor/LevelEditor/Private/LevelViewportLayoutFourPanes.cpp
+++ b/Engine/Source/Editor/LevelEditor/Private/LevelViewportLayoutFourPanes.cpp
@@ -69,7 +69,7 @@ TSharedRef<SWidget> FLevelViewportLayoutFourPanes::MakeViewportLayout(const FStr
 	Args.ParentLevelEditor = ParentLevelEditor;
 	Args.IsEnabled = FSlateApplication::Get().GetNormalExecutionAttribute();
 
-	Args.bRealtime = !FPlatformMisc::IsRemoteSession();
+	Args.bRealtime = true;
 	Args.ConfigKey = *ViewportKey0;
 	Args.ViewportType = LVT_Perspective;
 	TSharedRef<ILevelViewportLayoutEntity> Viewport0 = LevelEditor.FactoryViewport(*ViewportType0, Args);
diff --git a/Engine/Source/Editor/LevelEditor/Private/LevelViewportLayoutOnePane.cpp b/Engine/Source/Editor/LevelEditor/Private/LevelViewportLayoutOnePane.cpp
--- a/Engine/Source/Editor/LevelEditor/Private/LevelViewportLayoutOnePane.cpp
+++ b/Engine/Source/Editor/LevelEditor/Private/LevelViewportLayoutOnePane.cpp
@@ -52,7 +52,7 @@ TSharedRef<SWidget> FLevelViewportLayoutOnePane::MakeViewportLayout(const FStrin
 	Args.ParentLayout = AsShared();
 	Args.ParentLevelEditor = ParentLevelEditor;
 	Args.IsEnabled = FSlateApplication::Get().GetNormalExecutionAttribute();
-	Args.bRealtime = !FPlatformMisc::IsRemoteSession();
+	Args.bRealtime = true;
 	Args.ConfigKey = *ViewportKey;
 	Args.ViewportType = LVT_Perspective;
 	TSharedRef<ILevelViewportLayoutEntity> Viewport = LevelEditor.FactoryViewport(*ViewportType, Args);
diff --git a/Engine/Source/Editor/LevelEditor/Private/LevelViewportLayoutThreePanes.cpp b/Engine/Source/Editor/LevelEditor/Private/LevelViewportLayoutThreePanes.cpp
--- a/Engine/Source/Editor/LevelEditor/Private/LevelViewportLayoutThreePanes.cpp
+++ b/Engine/Source/Editor/LevelEditor/Private/LevelViewportLayoutThreePanes.cpp
@@ -64,7 +64,7 @@ TSharedRef<SWidget> FLevelViewportLayoutThreePanes::MakeViewportLayout(const FSt
 	Args.ParentLevelEditor = ParentLevelEditor;
 	Args.IsEnabled = FSlateApplication::Get().GetNormalExecutionAttribute();
 
-	Args.bRealtime = !FPlatformMisc::IsRemoteSession();
+	Args.bRealtime = true;
 	Args.ConfigKey = *ViewportKey0;
 	Args.ViewportType = LVT_Perspective;
 	TSharedRef<ILevelViewportLayoutEntity> Viewport0 = LevelEditor.FactoryViewport(*ViewportType0, Args);
diff --git a/Engine/Source/Editor/LevelEditor/Private/LevelViewportLayoutTwoPanes.cpp b/Engine/Source/Editor/LevelEditor/Private/LevelViewportLayoutTwoPanes.cpp
--- a/Engine/Source/Editor/LevelEditor/Private/LevelViewportLayoutTwoPanes.cpp
+++ b/Engine/Source/Editor/LevelEditor/Private/LevelViewportLayoutTwoPanes.cpp
@@ -57,7 +57,7 @@ TSharedRef<SWidget> TLevelViewportLayoutTwoPanes<TOrientation>::MakeViewportLayo
 	Args.ViewportType = LVT_OrthoXY;
 	TSharedRef<ILevelViewportLayoutEntity> Viewport0 = LevelEditor.FactoryViewport(*ViewportType0, Args);
 
-	Args.bRealtime = !FPlatformMisc::IsRemoteSession();
+	Args.bRealtime = true;
 	Args.ConfigKey = *ViewportKey1;
 	Args.ViewportType = LVT_Perspective;
 	TSharedRef<ILevelViewportLayoutEntity> Viewport1 = LevelEditor.FactoryViewport(*ViewportType1, Args);
diff --git a/Engine/Source/Editor/LevelEditor/Private/SLevelViewport.cpp b/Engine/Source/Editor/LevelEditor/Private/SLevelViewport.cpp
--- a/Engine/Source/Editor/LevelEditor/Private/SLevelViewport.cpp
+++ b/Engine/Source/Editor/LevelEditor/Private/SLevelViewport.cpp
@@ -462,11 +462,6 @@ void SLevelViewport::ConstructLevelEditorViewportClient( const FArguments& InArg
 			ViewportInstanceSettings.BufferVisualizationMode = NAME_None;
 		}
 
-		// Disable realtime viewports by default for remote sessions
-		if (FPlatformMisc::IsRemoteSession())
-		{
-			ViewportInstanceSettings.bIsRealtime = false;
-		}
 	}
 
 	if(ViewportInstanceSettings.ViewportType == LVT_Perspective)
@@ -536,6 +531,14 @@ void SLevelViewport::ConstructLevelEditorViewportClient( const FArguments& InArg
 	LevelViewportClient->SetViewModes(ViewportInstanceSettings.PerspViewModeIndex, ViewportInstanceSettings.OrthoViewModeIndex );
 
 	bShowFullToolbar = ViewportInstanceSettings.bShowFullToolbar;
+
+
+	// Disable realtime viewports by default for remote sessions
+	if (FPlatformMisc::IsRemoteSession())
+	{
+		bool bShouldBeRealtime = false;
+		LevelViewportClient->SetRealtimeOverride(bShouldBeRealtime, LOCTEXT("RealtimeOverrideMessage_RDP", "Remote Desktop"));
+	}
 }
 
 FSceneViewport* SLevelViewport::GetGameSceneViewport() const
@@ -2186,11 +2189,9 @@ void SLevelViewport::SaveConfig(const FString& ConfigName) const
 		ViewportInstanceSettings.RayTracingDebugVisualizationMode = LevelViewportClient->CurrentRayTracingDebugVisualizationMode;
 		ViewportInstanceSettings.ExposureSettings = LevelViewportClient->ExposureSettings;
 		ViewportInstanceSettings.FOVAngle = LevelViewportClient->FOVAngle;
-		if (!FPlatformMisc::IsRemoteSession())
-		{
-			// Only save this when we're not a remote session, as remote sessions force realtime to be disabled
-			ViewportInstanceSettings.bIsRealtime = LevelViewportClient->IsRealtime();
-		}
+	
+		LevelViewportClient->SaveRealtimeStateToConfig(ViewportInstanceSettings.bIsRealtime);
+	
 		ViewportInstanceSettings.bShowOnScreenStats = LevelViewportClient->ShouldShowStats();
 		ViewportInstanceSettings.FarViewPlane = LevelViewportClient->GetFarClipPlaneOverride();
 		ViewportInstanceSettings.bShowFullToolbar = bShowFullToolbar;
@@ -2288,12 +2289,6 @@ FLevelEditorViewportInstanceSettings SLevelViewport::LoadLegacyConfigFromIni(con
 	GConfig->GetBool(*IniSection, *(InConfigKey + TEXT(".bWantFPS")), ViewportInstanceSettings.bShowFPS_DEPRECATED, GEditorPerProjectIni);
 	GConfig->GetFloat(*IniSection, *(InConfigKey + TEXT(".FOVAngle")), ViewportInstanceSettings.FOVAngle, GEditorPerProjectIni);
 
-	// Disable realtime viewports by default for remote sessions
-	if (FPlatformMisc::IsRemoteSession())
-	{
-		ViewportInstanceSettings.bIsRealtime = false;
-	}
-
 	return ViewportInstanceSettings;
 }
 
diff --git a/Engine/Source/Editor/LevelEditor/Private/SLevelViewportToolBar.cpp b/Engine/Source/Editor/LevelEditor/Private/SLevelViewportToolBar.cpp
--- a/Engine/Source/Editor/LevelEditor/Private/SLevelViewportToolBar.cpp
+++ b/Engine/Source/Editor/LevelEditor/Private/SLevelViewportToolBar.cpp
@@ -504,7 +504,7 @@ static void OnGenerateBookmarkMenu( FMenuBuilder& MenuBuilder , TWeakPtr<class S
 	MenuBuilder.EndSection();
 }
 
-TSharedRef<SWidget> SLevelViewportToolBar::GenerateOptionsMenu() const
+TSharedRef<SWidget> SLevelViewportToolBar::GenerateOptionsMenu() 
 {
 	Viewport.Pin()->OnFloatingButtonClicked();
 
@@ -522,6 +522,25 @@ TSharedRef<SWidget> SLevelViewportToolBar::GenerateOptionsMenu() const
 		OptionsMenuBuilder.BeginSection("LevelViewportViewportOptions", LOCTEXT("OptionsMenuHeader", "Viewport Options") );
 		{
 			OptionsMenuBuilder.AddMenuEntry( FEditorViewportCommands::Get().ToggleRealTime );
+
+			// Add an option to disable the temporary override if there is one
+			{
+				FUIAction Action;
+				Action.ExecuteAction.BindSP(this, &SLevelViewportToolBar::OnDisableRealtimeOverride);
+				Action.IsActionVisibleDelegate.BindSP(this, &SLevelViewportToolBar::IsRealtimeOverrideToggleVisible);
+
+				TAttribute<FText> Tooltip(this, &SLevelViewportToolBar::GetRealtimeOverrideTooltip);
+
+				OptionsMenuBuilder.AddMenuEntry(
+					LOCTEXT("DisableRealtimeOverride", "Disable Realtime Override"),
+					Tooltip,
+					FSlateIcon(),
+					Action);
+
+				OptionsMenuBuilder.AddMenuSeparator();
+					
+			}
+		
 			OptionsMenuBuilder.AddMenuEntry( FEditorViewportCommands::Get().ToggleStats );
 			OptionsMenuBuilder.AddMenuEntry( FEditorViewportCommands::Get().ToggleFPS );
 			OptionsMenuBuilder.AddMenuEntry( LevelViewportActions.ToggleViewportToolbar );
@@ -1329,6 +1348,34 @@ void SLevelViewportToolBar::CreateViewMenuExtensions(FMenuBuilder& MenuBuilder)
 }
 
 
+void SLevelViewportToolBar::OnDisableRealtimeOverride()
+{
+	if (TSharedPtr<SLevelViewport> ViewportPinned = Viewport.Pin())
+	{
+		ViewportPinned->GetLevelViewportClient().RemoveRealtimeOverride();
+	}
+}
+
+bool SLevelViewportToolBar::IsRealtimeOverrideToggleVisible() const
+{
+	if (TSharedPtr<SLevelViewport> ViewportPinned = Viewport.Pin())
+	{
+		return ViewportPinned->GetLevelViewportClient().IsRealtimeOverrideSet();
+	}
+
+	return false;
+}
+
+FText SLevelViewportToolBar::GetRealtimeOverrideTooltip() const
+{
+	if (TSharedPtr<SLevelViewport> ViewportPinned = Viewport.Pin())
+	{
+		return FText::Format(LOCTEXT("DisableRealtimeOverrideToolTip", "Realtime is currently overridden by \"{0}\".  Click to remove that override"), ViewportPinned->GetLevelViewportClient().GetRealtimeOverrideMessage());
+	}
+
+	return FText::GetEmpty();
+}
+
 bool SLevelViewportToolBar::IsLandscapeLODSettingChecked(int32 Value) const
 {
 	return Viewport.Pin()->GetLevelViewportClient().LandscapeLODOverride == Value;
diff --git a/Engine/Source/Editor/LevelEditor/Private/SLevelViewportToolBar.h b/Engine/Source/Editor/LevelEditor/Private/SLevelViewportToolBar.h
--- a/Engine/Source/Editor/LevelEditor/Private/SLevelViewportToolBar.h
+++ b/Engine/Source/Editor/LevelEditor/Private/SLevelViewportToolBar.h
@@ -104,7 +104,7 @@ private:
 	 *
 	 * @return The widget containing the options menu content
 	 */
-	TSharedRef<SWidget> GenerateOptionsMenu() const;
+	TSharedRef<SWidget> GenerateOptionsMenu();
 
 	/**
 	 * Generates the toolbar camera menu content 
@@ -231,6 +231,11 @@ private:
 	TSharedPtr<FExtender> GetViewMenuExtender();
 
 	void CreateViewMenuExtensions(FMenuBuilder& MenuBuilder);
+
+	/** Called when the user disables realtime override from the toolbar */
+	void OnDisableRealtimeOverride();
+	bool IsRealtimeOverrideToggleVisible() const;
+	FText GetRealtimeOverrideTooltip() const;
 private:
 	/** The viewport that we are in */
 	TWeakPtr<class SLevelViewport> Viewport;
diff --git a/Engine/Source/Editor/MaterialEditor/Private/MaterialEditor.cpp b/Engine/Source/Editor/MaterialEditor/Private/MaterialEditor.cpp
--- a/Engine/Source/Editor/MaterialEditor/Private/MaterialEditor.cpp
+++ b/Engine/Source/Editor/MaterialEditor/Private/MaterialEditor.cpp
@@ -1650,9 +1650,9 @@ void FMaterialEditor::LoadEditorSettings()
 			PreviewViewport->TogglePreviewGrid();
 		}
 
-		if (EditorOptions->bRealtimeMaterialViewport)
+		if (EditorOptions->bRealtimeMaterialViewport && PreviewViewport->GetViewportClient())
 		{
-			PreviewViewport->OnToggleRealtime();
+			PreviewViewport->GetViewportClient()->SetRealtime(true);
 		}
 	}
 	
diff --git a/Engine/Source/Editor/Matinee/Private/Matinee.cpp b/Engine/Source/Editor/Matinee/Private/Matinee.cpp
--- a/Engine/Source/Editor/Matinee/Private/Matinee.cpp
+++ b/Engine/Source/Editor/Matinee/Private/Matinee.cpp
@@ -956,8 +956,9 @@ void FMatinee::InitMatinee(const EToolkitMode::Type Mode, const TSharedPtr< clas
 			// If there is a director group, set the perspective viewports to realtime automatically.
 			if(LevelVC->IsPerspective() && LevelVC->AllowsCinematicControl())
 			{				
-				//Ensure Realtime is turned on and store the original setting so we can restore it later.
-				LevelVC->SetRealtime(true, true);
+				//Ensure Realtime is turned on temporarily
+				const bool bShouldBeRealtime = true;
+				LevelVC->SetRealtimeOverride(bShouldBeRealtime, LOCTEXT("RealtimeOverrideMessage_Matinee", "Matinee"));
 			}
 		}
 	}
@@ -2084,8 +2085,7 @@ void FMatinee::OnClose()
 			// Turn off realtime when exiting.
 			if( LevelVC->IsPerspective() && LevelVC->AllowsCinematicControl() )
 			{				
-				//Specify true so RestoreRealtime will allow us to disable Realtime if it was original disabled
-				LevelVC->RestoreRealtime(true);
+				LevelVC->RemoveRealtimeOverride();
 			}
 		}
 	}
diff --git a/Engine/Source/Editor/MeshPaint/Private/MeshPaintHelpers.cpp b/Engine/Source/Editor/MeshPaint/Private/MeshPaintHelpers.cpp
--- a/Engine/Source/Editor/MeshPaint/Private/MeshPaintHelpers.cpp
+++ b/Engine/Source/Editor/MeshPaint/Private/MeshPaintHelpers.cpp
@@ -1161,12 +1161,11 @@ void MeshPaintHelpers::SetRealtimeViewport(bool bRealtime)
 		{
 			if (bRealtime)
 			{
-				Viewport.SetRealtime(bRealtime, bRememberCurrentState);
+				Viewport.SetRealtimeOverride(bRealtime, NSLOCTEXT("MeshPaint", "RealtimeOverrideMessage_MeshPaint", "Mesh Paint"));
 			}
 			else
 			{
-				const bool bAllowDisable = true;
-				Viewport.RestoreRealtime(bAllowDisable);
+				Viewport.RemoveRealtimeOverride();
 			}
 		}
 	}
diff --git a/Engine/Source/Editor/Persona/Private/AnimationEditorViewportClient.cpp b/Engine/Source/Editor/Persona/Private/AnimationEditorViewportClient.cpp
--- a/Engine/Source/Editor/Persona/Private/AnimationEditorViewportClient.cpp
+++ b/Engine/Source/Editor/Persona/Private/AnimationEditorViewportClient.cpp
@@ -113,7 +113,8 @@ FAnimationViewportClient::FAnimationViewportClient(const TSharedRef<IPersonaPrev
 	SetRealtime(true);
 	if(GEditor->PlayWorld)
 	{
-		SetRealtime(false,true); // We are PIE, don't start in realtime mode
+		const bool bShouldBeRealtime = false;
+		SetRealtimeOverride(bShouldBeRealtime, LOCTEXT("RealtimeOverride_PIE", "Play in Editor"));
 	}
 
 	// @todo double define - fix it
diff --git a/Engine/Source/Editor/PixelInspector/Private/PixelInspector.cpp b/Engine/Source/Editor/PixelInspector/Private/PixelInspector.cpp
--- a/Engine/Source/Editor/PixelInspector/Private/PixelInspector.cpp
+++ b/Engine/Source/Editor/PixelInspector/Private/PixelInspector.cpp
@@ -425,7 +425,8 @@ namespace PixelInspector
 				{
 					if (!EditorViewport->IsRealtime())
 					{
-						EditorViewport->SetRealtime(true);
+						const bool bShouldBeRealtime = true;
+						EditorViewport->SetRealtimeOverride(bShouldBeRealtime, LOCTEXT("RealtimeOverrideMessage_PixelInspector", "Pixel Inspector"));
 					}
 				}
 			}
diff --git a/Engine/Source/Editor/Sequencer/Private/LevelEditorSequencerIntegration.cpp b/Engine/Source/Editor/Sequencer/Private/LevelEditorSequencerIntegration.cpp
--- a/Engine/Source/Editor/Sequencer/Private/LevelEditorSequencerIntegration.cpp
+++ b/Engine/Source/Editor/Sequencer/Private/LevelEditorSequencerIntegration.cpp
@@ -1049,8 +1049,8 @@ void FLevelEditorSequencerIntegration::ActivateRealtimeViewports()
 			// If there is a director group, set the perspective viewports to realtime automatically.
 			if (LevelVC->IsPerspective() && LevelVC->AllowsCinematicControl())
 			{				
-				// Ensure Realtime is turned on and store the original setting so we can restore it later.
-				LevelVC->SetRealtime(true, true);
+				const bool bShouldBeRealtime = true;
+				LevelVC->SetRealtimeOverride(bShouldBeRealtime, LOCTEXT("RealtimeOverrideMessage_Sequencer", "Sequencer"));
 			}
 		}
 	}
@@ -1068,8 +1068,7 @@ void FLevelEditorSequencerIntegration::RestoreRealtimeViewports()
 			// Turn off realtime when exiting.
 			if( LevelVC->IsPerspective() && LevelVC->AllowsCinematicControl() )
 			{				
-				// Specify true so RestoreRealtime will allow us to disable Realtime if it was original disabled
-				LevelVC->RestoreRealtime(true);
+				LevelVC->RemoveRealtimeOverride();
 			}
 		}
 	}
diff --git a/Engine/Source/Editor/StaticMeshEditor/Private/SStaticMeshEditorViewport.h b/Engine/Source/Editor/StaticMeshEditor/Private/SStaticMeshEditorViewport.h
--- a/Engine/Source/Editor/StaticMeshEditor/Private/SStaticMeshEditorViewport.h
+++ b/Engine/Source/Editor/StaticMeshEditor/Private/SStaticMeshEditorViewport.h
@@ -167,9 +167,6 @@ private:
 	/** Callback for checking the physical material mask show flag. */
 	bool IsInViewModePhysicalMaterialMasksChecked() const;
 
-	/** Callback for toggling the realtime preview flag. */
-	void SetRealtimePreview();
-
 	/** Callback for updating preview socket meshes if the static mesh or socket has been modified. */
 	void OnObjectPropertyChanged(UObject* ObjectBeingModified, FPropertyChangedEvent& PropertyChangedEvent);
 
diff --git a/Engine/Source/Editor/UnrealEd/Classes/Editor/EditorEngine.h b/Engine/Source/Editor/UnrealEd/Classes/Editor/EditorEngine.h
--- a/Engine/Source/Editor/UnrealEd/Classes/Editor/EditorEngine.h
+++ b/Engine/Source/Editor/UnrealEd/Classes/Editor/EditorEngine.h
@@ -1816,10 +1816,25 @@ public:
 	 */
 	virtual void EndPlayOnLocalPc();
 
+	/**
+	 * Overrides the realtime state all viewports until RemoveViewportsRealtimeOverride is called.
+	 * The state of this override is not preserved between editor sessions. 
+	 *
+	 * @param bShouldBeRealtime	If true, this viewport will be realtime, if false this viewport will not be realtime
+	 * @param SystemDisplayName	This display name of whatever system is overriding realtime. This name is displayed to users in the viewport options menu
+	 */
+	void SetViewportsRealtimeOverride(bool bShouldBeRealtime, FText SystemDisplayName);
+
+	/**
+	 * Removes the current realtime override.  If there was another realtime override set it will restore that override
+	 */
+	void RemoveViewportsRealtimeOverride();
+
 	/**
 	 * Disables any realtime viewports that are currently viewing the level.  This will not disable
 	 * things like preview viewports in Cascade, etc. Typically called before running the game.
 	 */
+	UE_DEPRECATED(4.25, "To save and restore realtime state non-permanently use SetViewportsRealtimeOverride and RemoveViewportsRealtimeOverride")
 	void DisableRealtimeViewports();
 
 	/**
@@ -1827,6 +1842,7 @@ public:
 	 * disable viewporst that were realtime when DisableRealtimeViewports has been called and got
 	 * latter toggled to be realtime.
 	 */
+	UE_DEPRECATED(4.25, "To save and restore realtime state non-permanently use SetViewportsRealtimeOverride and RemoveViewportsRealtimeOverride")
 	void RestoreRealtimeViewports();
 
 	/**
diff --git a/Engine/Source/Editor/UnrealEd/Private/EditorEngine.cpp b/Engine/Source/Editor/UnrealEd/Private/EditorEngine.cpp
--- a/Engine/Source/Editor/UnrealEd/Private/EditorEngine.cpp
+++ b/Engine/Source/Editor/UnrealEd/Private/EditorEngine.cpp
@@ -6142,8 +6142,39 @@ void UEditorEngine::NotifyToolsOfObjectReplacement(const TMap<UObject*, UObject*
 	}
 }
 
+void UEditorEngine::SetViewportsRealtimeOverride(bool bShouldBeRealtime, FText SystemDisplayName)
+{
+	for (FEditorViewportClient* VC : AllViewportClients)
+	{
+		if (VC)
+		{
+			VC->SetRealtimeOverride(bShouldBeRealtime, SystemDisplayName);
+		}
+	}
+
+	RedrawAllViewports();
+
+	FEditorSupportDelegates::UpdateUI.Broadcast();
+}
+
+void UEditorEngine::RemoveViewportsRealtimeOverride()
+{
+	for (FEditorViewportClient* VC : AllViewportClients)
+	{
+		if (VC)
+		{
+			VC->RemoveRealtimeOverride();
+		}
+	}
+
+	RedrawAllViewports();
+
+	FEditorSupportDelegates::UpdateUI.Broadcast();
+}
+
 void UEditorEngine::DisableRealtimeViewports()
 {
+PRAGMA_DISABLE_DEPRECATION_WARNINGS
 	for(FEditorViewportClient* VC : AllViewportClients)
 	{
 		if( VC )
@@ -6155,11 +6186,13 @@ void UEditorEngine::DisableRealtimeViewports()
 	RedrawAllViewports();
 
 	FEditorSupportDelegates::UpdateUI.Broadcast();
+PRAGMA_ENABLE_DEPRECATION_WARNINGS
 }
 
 
 void UEditorEngine::RestoreRealtimeViewports()
 {
+PRAGMA_DISABLE_DEPRECATION_WARNINGS
 	for(FEditorViewportClient* VC : AllViewportClients)
 	{
 		if( VC )
@@ -6171,6 +6204,7 @@ void UEditorEngine::RestoreRealtimeViewports()
 	RedrawAllViewports();
 
 	FEditorSupportDelegates::UpdateUI.Broadcast();
+PRAGMA_ENABLE_DEPRECATION_WARNINGS
 }
 
 
diff --git a/Engine/Source/Editor/UnrealEd/Private/EditorViewportClient.cpp b/Engine/Source/Editor/UnrealEd/Private/EditorViewportClient.cpp
--- a/Engine/Source/Editor/UnrealEd/Private/EditorViewportClient.cpp
+++ b/Engine/Source/Editor/UnrealEd/Private/EditorViewportClient.cpp
@@ -358,8 +358,6 @@ FEditorViewportClient::FEditorViewportClient(FEditorModeTools* InModeTools, FPre
 	, bForceAudioRealtime(false)
 	, RealTimeUntilFrameNumber(0)
 	, bIsRealtime(false)
-	, bStoredRealtime(false)
-	, bStoredShowStats(false)
 	, bShowStats(false)
 	, bHasAudioFocus(false)
 	, bShouldCheckHitProxy(false)
@@ -486,50 +484,52 @@ FEditorViewportClient::~FEditorViewportClient()
 	}
 }
 
+void FEditorViewportClient::SetRealtimeOverride(bool bShouldBeRealtime, FText SystemDisplayName)
+{
+	PreviousRealtimeOverrideState = TempRealtimeOverride;
+	TempRealtimeOverride = TPair<bool, FText>(bShouldBeRealtime, SystemDisplayName);
+
+	bShouldInvalidateViewportWidget = true;
+}
+
+void FEditorViewportClient::RemoveRealtimeOverride()
+{
+	TempRealtimeOverride = PreviousRealtimeOverrideState;
+	PreviousRealtimeOverrideState.Reset();
+	bShouldInvalidateViewportWidget = true;
+}
+
 bool FEditorViewportClient::ToggleRealtime()
 {
 	SetRealtime(!bIsRealtime);
 	return bIsRealtime;
 }
 
-void FEditorViewportClient::SetRealtime(bool bInRealtime, bool bStoreCurrentValue)
+void FEditorViewportClient::SetRealtime(bool bInRealtime)
 {
-	if (bStoreCurrentValue)
-	{
-		//Cache the Realtime and ShowStats flags
-		bStoredRealtime = bIsRealtime;
-		bStoredShowStats = bShowStats;
-	}
-
 	bIsRealtime = bInRealtime;
+	bShouldInvalidateViewportWidget = true;
+}
 
-	if (!bIsRealtime)
-	{
-		SetShowStats(false);
-	}
-	else
-	{
-		bShouldInvalidateViewportWidget = true;
-	}
+FText FEditorViewportClient::GetRealtimeOverrideMessage() const
+{
+	return TempRealtimeOverride.IsSet() ? TempRealtimeOverride.GetValue().Value : FText::GetEmpty();
+}
+
+void FEditorViewportClient::SetRealtime(bool bInRealtime, bool bStoreCurrentValue)
+{
+	SetRealtime(bInRealtime);
 }
 
 void FEditorViewportClient::RestoreRealtime(const bool bAllowDisable)
 {
-	if (bAllowDisable)
-	{
-		bIsRealtime = bStoredRealtime;
-		bShowStats = bStoredShowStats;
-	}
-	else
-	{
-		bIsRealtime |= bStoredRealtime;
-		bShowStats |= bStoredShowStats;
-	}
+	RemoveRealtimeOverride();
+}
 
-	if (bIsRealtime)
-	{
-		bShouldInvalidateViewportWidget = true;
-	}
+void FEditorViewportClient::SaveRealtimeStateToConfig(bool& ConfigVar) const
+{
+	// Note, this should not ever look at anything but the true realtime state.  No overrides etc.
+	ConfigVar = bIsRealtime;
 }
 
 void FEditorViewportClient::SetShowStats(bool bWantStats)
@@ -2085,7 +2085,7 @@ void FEditorViewportClient::HandleViewportStatEnabled(const TCHAR* InName)
 	if (GStatProcessingViewportClient == this)
 	{
 		SetShowStats(true);
-		SetRealtime(true);
+		SetRealtimeOverride(true, LOCTEXT("RealtimeOverrideMessage_Stats", "Stats Display"));
 		SetStatEnabled(InName, true);
 	}
 }
@@ -2098,7 +2098,7 @@ void FEditorViewportClient::HandleViewportStatDisabled(const TCHAR* InName)
 		if (SetStatEnabled(InName, false) == 0)
 		{
 			SetShowStats(false);
-			// Note: we can't disable realtime as we don't know the setting it was previously
+			RemoveRealtimeOverride();
 		}
 	}
 }
@@ -2109,8 +2109,8 @@ void FEditorViewportClient::HandleViewportStatDisableAll(const bool bInAnyViewpo
 	if (bInAnyViewport || GStatProcessingViewportClient == this)
 	{
 		SetShowStats(false);
-		// Note: we can't disable realtime as we don't know the setting it was previously
 		SetStatEnabled(NULL, false, true);
+		RemoveRealtimeOverride();
 	}
 }
 
@@ -5197,11 +5197,6 @@ bool FEditorViewportClient::IsSetShowCollisionChecked() const
 	return EngineShowFlags.Collision;
 }
 
-void FEditorViewportClient::SetRealtimePreview()
-{
-	SetRealtime(!IsRealtime());
-	Invalidate();
-}
 
 void FEditorViewportClient::SetViewMode(EViewModeIndex InViewModeIndex)
 {
diff --git a/Engine/Source/Editor/UnrealEd/Private/LevelEditorViewport.cpp b/Engine/Source/Editor/UnrealEd/Private/LevelEditorViewport.cpp
--- a/Engine/Source/Editor/UnrealEd/Private/LevelEditorViewport.cpp
+++ b/Engine/Source/Editor/UnrealEd/Private/LevelEditorViewport.cpp
@@ -2767,7 +2767,7 @@ bool FLevelEditorViewportClient::InputKey(FViewport* InViewport, int32 Controlle
 			{
 				FText TrackingDescription = FText::Format(LOCTEXT("RotatationShortcut", "Rotate Atmosphere Light {0}"), LightIndex);
 				TrackingTransaction.Begin(TrackingDescription, SelectedSunLight->GetOwner());
-				SetRealtime(true, true); // The first time, save that setting for RestoreRealtime
+				SetRealtimeOverride(true, LOCTEXT("RealtimeOverrideMessage_AtmospherelLight", "Atmosphere Light Control"));// The first time, save that setting for RestoreRealtime
 			}
 			bCurrentUserControl = true;
 			UserIsControllingAtmosphericLightTimer = 3.0f; // Keep the widget open for a few seconds even when not tweaking the sun light
@@ -2789,7 +2789,7 @@ bool FLevelEditorViewportClient::InputKey(FViewport* InViewport, int32 Controlle
 	if (bUserIsControllingAtmosphericLight0 || bUserIsControllingAtmosphericLight1)
 	{
 		TrackingTransaction.End();					// End undo/redo translation
-		RestoreRealtime(true);						// Restore previous real-time state
+		RemoveRealtimeOverride();						// Restore previous real-time state
 	}
 	bUserIsControllingAtmosphericLight0 = false;	// Disable all atmospheric light controls
 	bUserIsControllingAtmosphericLight1 = false;
diff --git a/Engine/Source/Editor/UnrealEd/Private/PerformanceMonitor.cpp b/Engine/Source/Editor/UnrealEd/Private/PerformanceMonitor.cpp
--- a/Engine/Source/Editor/UnrealEd/Private/PerformanceMonitor.cpp
+++ b/Engine/Source/Editor/UnrealEd/Private/PerformanceMonitor.cpp
@@ -206,7 +206,14 @@ void FPerformanceMonitor::AutoApplyScalability()
 	const bool bAutoApplied = false;
 	Scalability::RecordQualityLevelsAnalytics(bAutoApplied);
 
-	GEditor->DisableRealtimeViewports();
+	for (FEditorViewportClient* VC : GEditor->GetAllViewportClients())
+	{
+		if (VC)
+		{
+			VC->SetRealtime(false);
+			VC->Invalidate();
+		}
+	}
 
 	// Reset the timers so as not to skew the data with the time it took to do the benchmark	
 	FineMovingAverage.Reset();
diff --git a/Engine/Source/Editor/UnrealEd/Private/PlayLevel.cpp b/Engine/Source/Editor/UnrealEd/Private/PlayLevel.cpp
--- a/Engine/Source/Editor/UnrealEd/Private/PlayLevel.cpp
+++ b/Engine/Source/Editor/UnrealEd/Private/PlayLevel.cpp
@@ -487,9 +487,8 @@ void UEditorEngine::EndPlayMap()
 		GEngine->StereoRenderingDevice->EnableStereo(false);
 	}
 
-
 	// Restores realtime viewports that have been disabled for PIE.
-	RestoreRealtimeViewports();
+	RemoveViewportsRealtimeOverride();
 
 	// Don't actually need to reset this delegate but doing so allows is to check invalid attempts to execute the delegate
 	FScopedConditionalWorldSwitcher::SwitchWorldForPIEDelegate = FOnSwitchWorldForPIE();
@@ -1898,7 +1897,10 @@ void UEditorEngine::ToggleBetweenPIEandSIE( bool bNewSession )
 			GameViewport->GetGameViewport()->SetPlayInEditorIsSimulate(false);
 
 			// The editor viewport client wont be visible so temporarily disable it being realtime
-			EditorViewportClient.SetRealtime( false, true );
+			const bool bShouldBeRealtime = false;
+			// Remove any previous override since we already applied a override when entering PIE
+			EditorViewportClient.RemoveRealtimeOverride();
+			EditorViewportClient.SetRealtimeOverride(bShouldBeRealtime, LOCTEXT("RealtimeOverrideMessage_PIE", "Play in Editor"));
 
 			if (!SlatePlayInEditorSession.EditorPlayer.IsValid())
 			{
@@ -1941,7 +1943,10 @@ void UEditorEngine::ToggleBetweenPIEandSIE( bool bNewSession )
 			bIsSimulatingInEditor = true;
 
 			// Make sure the viewport is in real-time mode
-			EditorViewportClient.SetRealtime( true );
+			const bool bShouldBeRealtime = true;
+			// Remove any previous override since we already applied a override when entering PIE
+			EditorViewportClient.RemoveRealtimeOverride();
+			EditorViewportClient.SetRealtimeOverride(bShouldBeRealtime, LOCTEXT("RealtimeOverrideMessage_PIE", "Play in Editor"));
 
 			// The Simulate window should show stats
 			EditorViewportClient.SetShowStats( true );
@@ -2438,7 +2443,8 @@ void UEditorEngine::StartPlayInEditorSession(FRequestPlaySessionParams& InReques
 	EditorWorld->bAllowAudioPlayback = false;
 
 	// Can't allow realtime viewports whilst in PIE so disable it for ALL viewports here.
-	DisableRealtimeViewports();
+	const bool bShouldBeRealtime = false;
+	SetViewportsRealtimeOverride(bShouldBeRealtime, LOCTEXT("RealtimeOverride_PIE", "Play in Editor"));
 
 	// Allow the global config to override our ability to create multiple PIE worlds.
 	if (!GEditor->bAllowMultiplePIEWorlds)
diff --git a/Engine/Source/Editor/UnrealEd/Private/SEditorViewport.cpp b/Engine/Source/Editor/UnrealEd/Private/SEditorViewport.cpp
--- a/Engine/Source/Editor/UnrealEd/Private/SEditorViewport.cpp
+++ b/Engine/Source/Editor/UnrealEd/Private/SEditorViewport.cpp
@@ -147,7 +147,7 @@ void SEditorViewport::BindCommands()
 	CommandListRef.MapAction( 
 		Commands.ToggleRealTime,
 		FExecuteAction::CreateSP( this, &SEditorViewport::OnToggleRealtime ),
-		FCanExecuteAction(),
+		FCanExecuteAction::CreateSP(this, &SEditorViewport::CanToggleRealtime),
 		FIsActionChecked::CreateSP(this, &SEditorViewport::IsRealtime));
 
 	CommandListRef.MapAction( 
@@ -426,6 +426,11 @@ void SEditorViewport::OnToggleRealtime()
 }
 
 
+bool SEditorViewport::CanToggleRealtime() const
+{
+	return !Client->IsRealtimeOverrideSet();
+}
+
 void SEditorViewport::SetRenderDirectlyToWindow( const bool bInRenderDirectlyToWindow )
 {
 	ViewportWidget->SetRenderDirectlyToWindow( bInRenderDirectlyToWindow );
diff --git a/Engine/Source/Editor/UnrealEd/Public/EditorViewportClient.h b/Engine/Source/Editor/UnrealEd/Public/EditorViewportClient.h
--- a/Engine/Source/Editor/UnrealEd/Public/EditorViewportClient.h
+++ b/Engine/Source/Editor/UnrealEd/Public/EditorViewportClient.h
@@ -294,20 +294,38 @@ public:
 		return PreviewScene;
 	}
 
+	/**
+	 * Overrides the realtime state of this viewport until RemoveViewportsRealtimeOverride is called.
+	 * The state of this override is not preserved between editor sessions. 
+	 *
+	 * @param bShouldBeRealtime	If true, this viewport will be realtime, if false this viewport will not be realtime
+	 * @param SystemDisplayName	This display name of whatever system is overriding realtime. This name is displayed to users in the viewport options menu	
+	 */
+	void SetRealtimeOverride(bool bShouldBeRealtime, FText SystemDisplayName);
+
+	/**
+	 * Removes the current realtime override.  If there was another realtime override set it will restore that override
+	 */
+	void RemoveRealtimeOverride();
+
 	/**
 	 * Toggles whether or not the viewport updates in realtime and returns the updated state.
+	 * Note: This value is saved between editor sessions so it should not be used for temporary states.  For that see SetRealtimeOverride
 	 *
 	 * @return		The current state of the realtime flag.
 	 */
 	bool ToggleRealtime();
 
-	/** Sets whether or not the viewport updates in realtime. */
-	void SetRealtime(bool bInRealtime, bool bStoreCurrentValue = false);
+	/** 
+	 * Sets whether or not the viewport updates in realtime.
+	 * Note: This value is saved between editor sessions so it should not be used for temporary states.  For that see SetRealtimeOverride
+	 */
+	void SetRealtime(bool bInRealtime);
 
 	/** @return		True if viewport is in realtime mode, false otherwise. */
 	bool IsRealtime() const
 	{ 
-		return bIsRealtime || GFrameCounter < RealTimeUntilFrameNumber;
+		return (TempRealtimeOverride.IsSet() ? TempRealtimeOverride.GetValue().Key : bIsRealtime) || GFrameCounter < RealTimeUntilFrameNumber;
 	}
 
 	/**
@@ -321,16 +339,33 @@ public:
 		RealTimeUntilFrameNumber = FMath::Max(GFrameCounter + NumRealTimeFrames, RealTimeUntilFrameNumber);
 	}
 
+	/**
+	 * Saves the realtime state to a config location.  Does not save any temp overrides
+	 */ 
+	void SaveRealtimeStateToConfig(bool& ConfigVar) const;
+
+	/**
+	 * @return true if there are any temp realtime overrides set
+	 */
+	bool IsRealtimeOverrideSet() const { return TempRealtimeOverride.IsSet(); }
+
+	/**
+	 * @return  If an override is set this returns the message indicating what set it
+	 */
+	FText GetRealtimeOverrideMessage() const;
+
+	UE_DEPRECATED(4.25, "SetRealtime no longer takes in bStoreCurrentValue parameter. For temporary overrides use SetRealtimeOverride")
+	void SetRealtime(bool bInRealtime, bool bStoreCurrentValue);
+
 	/**
 	 * Restores realtime setting to stored value. This will only enable realtime and 
 	 * never disable it (unless bAllowDisable is true)
 	 */
+	UE_DEPRECATED(4.25, "To save and restore realtime state non-permanently use SetRealtimeOverride and RemoveRealtimeOverride")
 	void RestoreRealtime(const bool bAllowDisable = false);
 
-
 	// this set ups camera for both orbit and non orbit control
-	void SetCameraSetup(const FVector& LocationForOrbiting, const FRotator& InOrbitRotation, const FVector& InOrbitZoom, const FVector& InOrbitLookAt, 
-			const FVector& InViewLocation, const FRotator &InViewRotation );
+	void SetCameraSetup(const FVector& LocationForOrbiting, const FRotator& InOrbitRotation, const FVector& InOrbitZoom, const FVector& InOrbitLookAt, const FVector& InViewLocation, const FRotator &InViewRotation );
 
 	/** Callback for toggling the camera lock flag. */
 	virtual void SetCameraLock();
@@ -359,9 +394,6 @@ public:
 	/** Callback for checking the collision geometry show flag. */
 	bool IsSetShowCollisionChecked() const;
 
-	/** Callback for toggling the realtime preview flag. */
-	void SetRealtimePreview();
-
 	/** Gets ViewportCameraTransform object for the current viewport type */
 	FViewportCameraTransform& GetViewTransform()
 	{
@@ -1196,7 +1228,20 @@ public:
 
 	/** @return True if the camera speed should be scaled by its view distance. */
 	virtual bool ShouldScaleCameraSpeedByDistance() const;
-	
+
+	/**
+	 * Enable customization of the EngineShowFlags for rendering. After calling this function,
+	 * the provided OverrideFunc will be passed a copy of .EngineShowFlags in ::Draw() just before rendering setup.
+	 * Changes made to the ShowFlags will be used for that frame but .EngineShowFlags will not be modified.
+	 * @param OverrideFunc custom override function that will be called every frame until override is disabled.
+	 */
+	void EnableOverrideEngineShowFlags(TUniqueFunction<void(FEngineShowFlags&)> OverrideFunc);
+
+	/** Disable EngineShowFlags override if enabled */
+	void DisableOverrideEngineShowFlags();
+
+	/** @return true if Override EngineShowFlags are currently enabled */
+	bool IsEngineShowFlagsOverrideEnabled() const { return !! OverrideShowFlagsFunc; }
 protected:
 	/** Invalidates the viewport widget (if valid) to register its active timer */
 	void InvalidateViewportWidget();
@@ -1604,12 +1649,6 @@ protected:
 	/** if the viewport is currently realtime */
 	bool bIsRealtime;
 	
-	/** Cached realtime flag */
-	bool bStoredRealtime;
-	
-	/** Cached show statistics flag */	
-	bool bStoredShowStats;
-
 	/** True if we should draw stats over the viewport */
 	bool bShowStats;
 
@@ -1649,22 +1688,11 @@ protected:
 	FVector DefaultOrbitZoom;
 	FVector DefaultOrbitLookAt;
 
-public:
-	/**
-	 * Enable customization of the EngineShowFlags for rendering. After calling this function,
-	 * the provided OverrideFunc will be passed a copy of .EngineShowFlags in ::Draw() just before rendering setup.
-	 * Changes made to the ShowFlags will be used for that frame but .EngineShowFlags will not be modified.
-	 * @param OverrideFunc custom override function that will be called every frame until override is disabled.
-	 */
-	void EnableOverrideEngineShowFlags(TUniqueFunction<void(FEngineShowFlags&)> OverrideFunc);
-
-	/** Disable EngineShowFlags override if enabled */
-	void DisableOverrideEngineShowFlags();
-
-	/** @return true if Override EngineShowFlags are currently enabled */
-	bool IsEngineShowFlagsOverrideEnabled() const { return !! OverrideShowFlagsFunc; }
+	/** Cached realtime override flag and user message (used for restoring nested overrides) */
+	TOptional<TPair<bool, FText>> PreviousRealtimeOverrideState;
 
-protected:
+	/** A temporary realtime override and user message that is not saved between editor sessions.  If this value is set, this viewport determines realtime from this setting, otherwise it reads the real realtime value */
+	TOptional<TPair<bool, FText>> TempRealtimeOverride;
 	
 	/** Custom override function that will be called every ::Draw() until override is disabled */
 	TUniqueFunction<void(FEngineShowFlags&)> OverrideShowFlagsFunc;
diff --git a/Engine/Source/Editor/UnrealEd/Public/SEditorViewport.h b/Engine/Source/Editor/UnrealEd/Public/SEditorViewport.h
--- a/Engine/Source/Editor/UnrealEd/Public/SEditorViewport.h
+++ b/Engine/Source/Editor/UnrealEd/Public/SEditorViewport.h
@@ -100,6 +100,9 @@ private:
 	EShaderPlatform GetShaderPlatformHelper(const ERHIFeatureLevel::Type InFeatureLevel) const;
 
 protected:
+	/** @return true if realtime can be toggled (it cannot be toggled directly if there is an override in place) */
+	bool CanToggleRealtime() const;
+
 	/** call this function to build a 'text' widget that can display the present feature level */
 	TSharedRef<SWidget> BuildFeatureLevelWidget() const;
 ///////////////////////////////////////////////////////////////////////////////
