Non-unity

#rnx

diff --git a/Engine/Source/Editor/UnrealEd/Private/PerformanceMonitor.cpp b/Engine/Source/Editor/UnrealEd/Private/PerformanceMonitor.cpp
--- a/Engine/Source/Editor/UnrealEd/Private/PerformanceMonitor.cpp
+++ b/Engine/Source/Editor/UnrealEd/Private/PerformanceMonitor.cpp
@@ -27,6 +27,7 @@
 
 #include "Settings/EditorSettings.h"
 #include "Editor/EditorPerformanceSettings.h"
+#include "EditorViewportClient.h"
 
 #define LOCTEXT_NAMESPACE "PerformanceMonitor"
 
