Static analysis fix

diff --git a/Engine/Source/Editor/UnrealEd/Private/ObjectTools.cpp b/Engine/Source/Editor/UnrealEd/Private/ObjectTools.cpp
--- a/Engine/Source/Editor/UnrealEd/Private/ObjectTools.cpp
+++ b/Engine/Source/Editor/UnrealEd/Private/ObjectTools.cpp
@@ -3442,7 +3442,7 @@ namespace ObjectTools
 			// if this object is being renamed out of the MyLevel package into a content package, we need to mark it RF_Standalone
 			// so that it will be saved (UWorld::CleanupWorld() clears this flag for all objects inside the package)
 			if (!Object->HasAnyFlags(RF_Standalone)
-				&&	OldPackage->ContainsMap()
+				&&	(OldPackage && OldPackage->ContainsMap())
 				&&	!NewPackage->GetOutermost()->ContainsMap() )
 			{
 				Object->SetFlags(RF_Standalone);
