Re-map catch-all chunk for platform

diff --git a/Engine/Source/Editor/UnrealEd/Private/Commandlets/AssetRegistryGenerator.cpp b/Engine/Source/Editor/UnrealEd/Private/Commandlets/AssetRegistryGenerator.cpp
--- a/Engine/Source/Editor/UnrealEd/Private/Commandlets/AssetRegistryGenerator.cpp
+++ b/Engine/Source/Editor/UnrealEd/Private/Commandlets/AssetRegistryGenerator.cpp
@@ -1918,7 +1918,7 @@ void FAssetRegistryGenerator::WriteCollection(FName CollectionName, const TArray
 	}
 }
 
-int32 FAssetRegistryGenerator::GetPakchunkIndex(int32 ChunkId)
+int32 FAssetRegistryGenerator::GetPakchunkIndex(int32 ChunkId) const
 {
 	if (ChunkIdPakchunkIndexMapping.Contains(ChunkId))
 	{
diff --git a/Engine/Source/Editor/UnrealEd/Private/CookOnTheFlyServer.cpp b/Engine/Source/Editor/UnrealEd/Private/CookOnTheFlyServer.cpp
--- a/Engine/Source/Editor/UnrealEd/Private/CookOnTheFlyServer.cpp
+++ b/Engine/Source/Editor/UnrealEd/Private/CookOnTheFlyServer.cpp
@@ -7568,7 +7568,7 @@ void UCookOnTheFlyServer::StartCookByTheBook( const FCookByTheBookStartupOptions
 			for (const ITargetPlatform* TargetPlatform : TargetPlatforms)
 			{
 				FAssetRegistryGenerator* RegistryGenerator = PlatformManager->GetPlatformData(TargetPlatform)->RegistryGenerator.Get();
-				RegistryGenerator->RegisterChunkDataGenerator(MakeShared<FLocalizationChunkDataGenerator>(PackagingSettings->LocalizationTargetCatchAllChunkId, LocalizationTargetsToChunk, AllCulturesToCook));
+				RegistryGenerator->RegisterChunkDataGenerator(MakeShared<FLocalizationChunkDataGenerator>(RegistryGenerator->GetPakchunkIndex(PackagingSettings->LocalizationTargetCatchAllChunkId), LocalizationTargetsToChunk, AllCulturesToCook));
 			}
 		}
 	}
diff --git a/Engine/Source/Editor/UnrealEd/Public/Commandlets/AssetRegistryGenerator.h b/Engine/Source/Editor/UnrealEd/Public/Commandlets/AssetRegistryGenerator.h
--- a/Engine/Source/Editor/UnrealEd/Public/Commandlets/AssetRegistryGenerator.h
+++ b/Engine/Source/Editor/UnrealEd/Public/Commandlets/AssetRegistryGenerator.h
@@ -158,6 +158,13 @@ public:
 	 */
 	void AddAssetToFileOrderRecursive(const FName& InPackageName, TArray<FName>& OutFileOrder, TSet<FName>& OutEncounteredNames, const TSet<FName>& InPackageNameSet, const TSet<FName>& InTopLevelAssets);
 
+	/**
+	 * Get pakchunk file index from ChunkID
+	 *
+	 * @param ChunkID
+	 * @return Index of target pakchunk file
+	 */
+	int32 GetPakchunkIndex(int32 ChunkId) const;
 
 private:
 
@@ -267,14 +274,6 @@ private:
 	*/
 	void RemovePackageFromManifest(FName PackageName, int32 ChunkId);
 
-	/**
-	* Get pakchunk file index from ChunkID
-	*
-	* @param ChunkID
-	* @return Index of target pakchunk file
-	*/
-	int32 GetPakchunkIndex(int32 ChunkId);
-
 	/**
 	 * Walks the dependency graph of assets and assigns packages to correct chunks.
 	 * 
