Integrating CL# 11464638 from Dev-Enterprise
Fixing UCX Chaos Collision meshes not being exported

diff --git a/Engine/Source/Editor/UnrealEd/Private/Fbx/FbxMainExport.cpp b/Engine/Source/Editor/UnrealEd/Private/Fbx/FbxMainExport.cpp
--- a/Engine/Source/Editor/UnrealEd/Private/Fbx/FbxMainExport.cpp
+++ b/Engine/Source/Editor/UnrealEd/Private/Fbx/FbxMainExport.cpp
@@ -93,6 +93,12 @@
 #include "PhysicsEngine/BodySetup.h"
 #endif // WITH_PHYSX
 
+#if WITH_CHAOS
+#include "Chaos/Core.h"
+#include "Chaos/Particles.h"
+#include "Chaos/Plane.h"
+#endif
+
 #include "Exporters/FbxExportOption.h"
 #include "FbxExportOptionsWindow.h"
 #include "Widgets/SWindow.h"
@@ -3122,7 +3128,7 @@ void DetermineVertsToWeld(TArray<int32>& VertRemap, TArray<int32>& UniqueVerts,
 	}
 }
 
-#if WITH_PHYSX
+#if (WITH_PHYSX && PHYSICS_INTERFACE_PHYSX) || WITH_CHAOS
 
 class FCollisionFbxExporter
 {
@@ -3290,7 +3296,11 @@ public:
 private:
 	uint32 GetConvexVerticeNumber(const FKConvexElem &ConvexElem)
 	{
+#if WITH_PHYSX && PHYSICS_INTERFACE_PHYSX
 		return ConvexElem.GetConvexMesh() != nullptr ? ConvexElem.GetConvexMesh()->getNbVertices() : 0;
+#elif WITH_CHAOS
+		return ConvexElem.VertexData.Num();
+#endif
 	}
 
 	uint32 GetBoxVerticeNumber() { return 24; }
@@ -3301,6 +3311,7 @@ private:
 
 	void AddConvexVertex(const FKConvexElem &ConvexElem)
 	{
+#if WITH_PHYSX && PHYSICS_INTERFACE_PHYSX
 		const physx::PxConvexMesh* ConvexMesh = ConvexElem.GetConvexMesh();
 		if (ConvexMesh == nullptr)
 		{
@@ -3313,10 +3324,20 @@ private:
 			ControlPoints[CurrentVertexOffset + PosIndex] = FbxVector4(Position.X, -Position.Y, Position.Z);
 		}
 		CurrentVertexOffset += ConvexMesh->getNbVertices();
+#elif WITH_CHAOS
+		const TArray<FVector>& VertexArray = ConvexElem.VertexData;
+		for (int32 PosIndex = 0; PosIndex < VertexArray.Num(); ++PosIndex)
+		{
+			FVector Position = VertexArray[PosIndex];
+			ControlPoints[CurrentVertexOffset + PosIndex] = FbxVector4(Position.X, -Position.Y, Position.Z);
+		}
+		CurrentVertexOffset += VertexArray.Num();
+#endif
 	}
 
 	void AddConvexNormals(const FKConvexElem &ConvexElem)
 	{
+#if WITH_PHYSX && PHYSICS_INTERFACE_PHYSX
 		const physx::PxConvexMesh* ConvexMesh = ConvexElem.GetConvexMesh();
 		if (ConvexMesh == nullptr)
 		{
@@ -3340,10 +3361,29 @@ private:
 				LayerElementNormal->GetDirectArray().Add(FbxNormal);
 			}
 		}
+#elif WITH_CHAOS
+		const auto& ConvexMesh = ConvexElem.GetChaosConvexMesh();
+		if (!ConvexMesh.IsValid())
+		{
+			return;
+		}
+		const TArray<Chaos::TPlaneConcrete<Chaos::FReal, 3>>& Faces = ConvexMesh->GetFaces();
+		for (int32 PolyIndex = 0; PolyIndex < Faces.Num(); ++PolyIndex)
+		{
+			FVector Normal = Faces[PolyIndex].Normal().GetSafeNormal();
+			FbxVector4 FbxNormal = FbxVector4(Normal.X, -Normal.Y, Normal.Z);
+			// add vertices 
+			for (int32 j = 0; j < 3; ++j)
+			{
+				LayerElementNormal->GetDirectArray().Add(FbxNormal);
+			}
+		}
+#endif
 	}
 
 	void AddConvexPolygon(const FKConvexElem &ConvexElem)
 	{
+#if WITH_PHYSX && PHYSICS_INTERFACE_PHYSX
 		const physx::PxConvexMesh* ConvexMesh = ConvexElem.GetConvexMesh();
 		if (ConvexMesh == nullptr)
 		{
@@ -3369,6 +3409,29 @@ private:
 			Mesh->EndPolygon();
 		}
 		CurrentVertexOffset += ConvexMesh->getNbVertices();
+#elif WITH_CHAOS
+		const auto& ConvexMesh = ConvexElem.GetChaosConvexMesh();
+		if (!ConvexMesh.IsValid())
+		{
+			return;
+		}
+		TArray<int32> IndexData(ConvexElem.IndexData);
+		if (IndexData.Num() == 0)
+		{
+			IndexData = ConvexElem.GetChaosConvexIndices();
+		}
+		check(IndexData.Num() % 3 == 0);
+		for (int32 VertexIndex = 0; VertexIndex < IndexData.Num(); VertexIndex += 3)
+		{
+			Mesh->BeginPolygon(ActualMatIndex);
+			Mesh->AddPolygon(CurrentVertexOffset + IndexData[VertexIndex]);
+			Mesh->AddPolygon(CurrentVertexOffset + IndexData[VertexIndex + 1]);
+			Mesh->AddPolygon(CurrentVertexOffset + IndexData[VertexIndex + 2]);
+			Mesh->EndPolygon();
+		}
+
+		CurrentVertexOffset += ConvexElem.VertexData.Num();
+#endif
 	}
 
 	void AddBoxVertex(const FKBoxElem &BoxElem)
@@ -3803,7 +3866,7 @@ FbxNode* FFbxExporter::ExportCollisionMesh(const UStaticMesh* StaticMesh, const
 	FbxActor->SetNodeAttribute(CollisionMesh);
 	return FbxActor;
 }
-#endif
+#endif //WITH_PHYSX
 
 void FFbxExporter::ExportObjectMetadata(const UObject* ObjectToExport, FbxNode* Node)
 {
@@ -4174,7 +4237,7 @@ FbxNode* FFbxExporter::ExportStaticMeshToFbx(const UStaticMesh* StaticMesh, int3
 		}
 	}
 
-#if WITH_PHYSX
+#if (WITH_PHYSX && PHYSICS_INTERFACE_PHYSX) || WITH_CHAOS
 	if ((ExportLOD == 0 || ExportLOD == -1) && GetExportOptions()->Collision)
 	{
 		ExportCollisionMesh(StaticMesh, MeshName, FbxActor);
diff --git a/Engine/Source/Editor/UnrealEd/Private/FbxExporter.h b/Engine/Source/Editor/UnrealEd/Private/FbxExporter.h
--- a/Engine/Source/Editor/UnrealEd/Private/FbxExporter.h
+++ b/Engine/Source/Editor/UnrealEd/Private/FbxExporter.h
@@ -313,7 +313,7 @@ private:
 
 	void ExportModel(UModel* Model, FbxNode* Node, const char* Name);
 	
-#if WITH_PHYSX
+#if (WITH_PHYSX && PHYSICS_INTERFACE_PHYSX) || WITH_CHAOS
 	FbxNode* ExportCollisionMesh(const UStaticMesh* StaticMesh, const TCHAR* MeshName, FbxNode* ParentActor);
 #endif
 
diff --git a/Engine/Source/Runtime/Engine/Classes/PhysicsEngine/ConvexElem.h b/Engine/Source/Runtime/Engine/Classes/PhysicsEngine/ConvexElem.h
--- a/Engine/Source/Runtime/Engine/Classes/PhysicsEngine/ConvexElem.h
+++ b/Engine/Source/Runtime/Engine/Classes/PhysicsEngine/ConvexElem.h
@@ -119,7 +119,9 @@ public:
 
 	ENGINE_API void ResetChaosConvexMesh();
 
-	ENGINE_API void ComputeChaosConvexIndices();
+	ENGINE_API void ComputeChaosConvexIndices(bool bForceCompute = false);
+
+	ENGINE_API TArray<int32> GetChaosConvexIndices() const;
 #endif
 
 	/** Get current transform applied to convex mesh vertices */
diff --git a/Engine/Source/Runtime/Engine/Private/PhysicsEngine/BodySetup.cpp b/Engine/Source/Runtime/Engine/Private/PhysicsEngine/BodySetup.cpp
--- a/Engine/Source/Runtime/Engine/Private/PhysicsEngine/BodySetup.cpp
+++ b/Engine/Source/Runtime/Engine/Private/PhysicsEngine/BodySetup.cpp
@@ -1732,7 +1732,9 @@ float FKConvexElem::GetVolume(const FVector& Scale) const
 void FKConvexElem::SetChaosConvexMesh(TSharedPtr<Chaos::FConvex, ESPMode::ThreadSafe>&& InChaosConvex)
 {
 	ChaosConvex = MoveTemp(InChaosConvex);
-	ComputeChaosConvexIndices();
+	
+	const bool bForceCompute = true;
+	ComputeChaosConvexIndices(bForceCompute);
 }
 
 void FKConvexElem::ResetChaosConvexMesh()
@@ -1740,15 +1742,24 @@ void FKConvexElem::ResetChaosConvexMesh()
 	ChaosConvex.Reset();
 }
 
-ENGINE_API void FKConvexElem::ComputeChaosConvexIndices()
+ENGINE_API void FKConvexElem::ComputeChaosConvexIndices(bool bForceCompute)
 {
+	if (bForceCompute || IndexData.Num() == 0)
+	{
+		IndexData = GetChaosConvexIndices();
+	}
+}
+
+TArray<int32> FKConvexElem::GetChaosConvexIndices() const
+{
+	TArray<int32> ResultIndexData;
 	const int32 NumVerts = VertexData.Num();
-	if(NumVerts > 0 && IndexData.Num() == 0)
+	if (NumVerts > 0)
 	{
 		Chaos::TParticles<Chaos::FReal, 3> ConvexParticles;
 		ConvexParticles.AddParticles(NumVerts);
 
-		for(int32 VertIndex = 0; VertIndex < NumVerts; ++VertIndex)
+		for (int32 VertIndex = 0; VertIndex < NumVerts; ++VertIndex)
 		{
 			ConvexParticles.X(VertIndex) = VertexData[VertIndex];
 		}
@@ -1756,14 +1767,16 @@ ENGINE_API void FKConvexElem::ComputeChaosConvexIndices()
 		TArray<Chaos::TVector<int32, 3>> Triangles;
 		Chaos::FConvexBuilder::BuildConvexHull(ConvexParticles, Triangles);
 
-		IndexData.Reset(Triangles.Num() * 3);
-		for(Chaos::TVector<int32, 3> Tri : Triangles)
+		ResultIndexData.Reserve(Triangles.Num() * 3);
+		for (Chaos::TVector<int32, 3> Tri : Triangles)
 		{
-			IndexData.Add(Tri[0]);
-			IndexData.Add(Tri[1]);
-			IndexData.Add(Tri[2]);
+			ResultIndexData.Add(Tri[0]);
+			ResultIndexData.Add(Tri[1]);
+			ResultIndexData.Add(Tri[2]);
 		}
 	}
+
+	return ResultIndexData;
 }
 #endif
 
