Convert SEditorViewportViewMenu into tool menu
#rnx

diff --git a/Engine/Source/Editor/UnrealEd/Private/SEditorViewportViewMenu.cpp b/Engine/Source/Editor/UnrealEd/Private/SEditorViewportViewMenu.cpp
--- a/Engine/Source/Editor/UnrealEd/Private/SEditorViewportViewMenu.cpp
+++ b/Engine/Source/Editor/UnrealEd/Private/SEditorViewportViewMenu.cpp
@@ -2,7 +2,9 @@
 
 
 #include "SEditorViewportViewMenu.h"
+#include "SEditorViewportViewMenuContext.h"
 #include "Framework/MultiBox/MultiBoxBuilder.h"
+#include "ToolMenus.h"
 #include "EditorStyleSet.h"
 #include "EditorViewportCommands.h"
 #include "RayTracingDebugVisualizationMenuCommands.h"
@@ -191,31 +193,47 @@ const FSlateBrush* SEditorViewportViewMenu::GetViewMenuLabelIcon() const
 
 TSharedRef<SWidget> SEditorViewportViewMenu::GenerateViewMenuContent() const
 {
-	const FEditorViewportCommands& BaseViewportActions = FEditorViewportCommands::Get();
+	static const FName MenuName("UnrealEd.ViewportToolbar.View");
+	if (!UToolMenus::Get()->IsMenuRegistered(MenuName))
+	{
+		UToolMenu* Menu = UToolMenus::Get()->RegisterMenu(MenuName);
+		Menu->AddDynamicSection("DynamicSection", FNewToolMenuDelegate::CreateLambda([](UToolMenu* InMenu)
+		{
+			UEditorViewportViewMenuContext* Context = InMenu->FindContext<UEditorViewportViewMenuContext>();
+			Context->Widget.Pin()->FillViewMenu(InMenu);
+		}));
+	}
+
+	UEditorViewportViewMenuContext* ContextObject = NewObject<UEditorViewportViewMenuContext>();
+	ContextObject->Widget = SharedThis(this);
+
+	FToolMenuContext MenuContext(Viewport.Pin()->GetCommandList(), MenuExtenders, ContextObject);
+	return UToolMenus::Get()->GenerateWidget(MenuName, MenuContext);
+}
 
-	const bool bInShouldCloseWindowAfterMenuSelection = true;
-	FMenuBuilder ViewMenuBuilder(bInShouldCloseWindowAfterMenuSelection, Viewport.Pin()->GetCommandList(), MenuExtenders);
+void SEditorViewportViewMenu::FillViewMenu(UToolMenu* Menu) const
+{
+	const FEditorViewportCommands& BaseViewportActions = FEditorViewportCommands::Get();
 
 	{
 		// View modes
 		{
-			ViewMenuBuilder.BeginSection("ViewMode", LOCTEXT("ViewModeHeader", "View Mode") );
+			FToolMenuSection& Section = Menu->AddSection("ViewMode", LOCTEXT("ViewModeHeader", "View Mode"));
 			{
-				ViewMenuBuilder.AddMenuEntry(BaseViewportActions.LitMode, NAME_None, UViewModeUtils::GetViewModeDisplayName(VMI_Lit));
-				ViewMenuBuilder.AddMenuEntry(BaseViewportActions.UnlitMode, NAME_None, UViewModeUtils::GetViewModeDisplayName(VMI_Unlit));
-				ViewMenuBuilder.AddMenuEntry(BaseViewportActions.WireframeMode, NAME_None, UViewModeUtils::GetViewModeDisplayName(VMI_BrushWireframe));
-				ViewMenuBuilder.AddMenuEntry(BaseViewportActions.DetailLightingMode, NAME_None, UViewModeUtils::GetViewModeDisplayName(VMI_Lit_DetailLighting));
-				ViewMenuBuilder.AddMenuEntry(BaseViewportActions.LightingOnlyMode, NAME_None, UViewModeUtils::GetViewModeDisplayName(VMI_LightingOnly));
-				ViewMenuBuilder.AddMenuEntry(BaseViewportActions.ReflectionOverrideMode, NAME_None, UViewModeUtils::GetViewModeDisplayName(VMI_ReflectionOverride));
-				ViewMenuBuilder.AddMenuEntry(BaseViewportActions.CollisionPawn, NAME_None, UViewModeUtils::GetViewModeDisplayName(VMI_CollisionPawn));
-				ViewMenuBuilder.AddMenuEntry(BaseViewportActions.CollisionVisibility, NAME_None, UViewModeUtils::GetViewModeDisplayName(VMI_CollisionVisibility));
+				Section.AddMenuEntry(BaseViewportActions.LitMode, UViewModeUtils::GetViewModeDisplayName(VMI_Lit));
+				Section.AddMenuEntry(BaseViewportActions.UnlitMode, UViewModeUtils::GetViewModeDisplayName(VMI_Unlit));
+				Section.AddMenuEntry(BaseViewportActions.WireframeMode, UViewModeUtils::GetViewModeDisplayName(VMI_BrushWireframe));
+				Section.AddMenuEntry(BaseViewportActions.DetailLightingMode, UViewModeUtils::GetViewModeDisplayName(VMI_Lit_DetailLighting));
+				Section.AddMenuEntry(BaseViewportActions.LightingOnlyMode, UViewModeUtils::GetViewModeDisplayName(VMI_LightingOnly));
+				Section.AddMenuEntry(BaseViewportActions.ReflectionOverrideMode, UViewModeUtils::GetViewModeDisplayName(VMI_ReflectionOverride));
+				Section.AddMenuEntry(BaseViewportActions.CollisionPawn, UViewModeUtils::GetViewModeDisplayName(VMI_CollisionPawn));
+				Section.AddMenuEntry(BaseViewportActions.CollisionVisibility, UViewModeUtils::GetViewModeDisplayName(VMI_CollisionVisibility));
 			}
 
 #if RHI_RAYTRACING
 			if (IsRayTracingEnabled())
 			{
-				const FEditorViewportCommands& BaseViewportCommands = FEditorViewportCommands::Get();
-				ViewMenuBuilder.AddMenuEntry(BaseViewportCommands.PathTracingMode, NAME_None, UViewModeUtils::GetViewModeDisplayName(VMI_PathTracing));
+				Section.AddMenuEntry(BaseViewportActions.PathTracingMode, UViewModeUtils::GetViewModeDisplayName(VMI_PathTracing));
 			}
 #endif
 
@@ -223,60 +241,61 @@ TSharedRef<SWidget> SEditorViewportViewMenu::GenerateViewMenuContent() const
 			{
 				struct Local
 				{
-					static void BuildOptimizationMenu( FMenuBuilder& Menu, TWeakPtr< SViewportToolBar > InParentToolBar )
+					static void BuildOptimizationMenu( UToolMenu* Menu, TWeakPtr< SViewportToolBar > InParentToolBar )
 					{
 						const FEditorViewportCommands& BaseViewportCommands = FEditorViewportCommands::Get();
 
 						UWorld* World = GWorld;
 						const ERHIFeatureLevel::Type FeatureLevel = (IsInGameThread() && World) ? (ERHIFeatureLevel::Type)World->FeatureLevel : GMaxRHIFeatureLevel;
 
-						Menu.BeginSection("OptimizationViewmodes", LOCTEXT("OptimizationSubMenuHeader", "Optimization Viewmodes"));
 						{
+							FToolMenuSection& Section = Menu->AddSection("OptimizationViewmodes", LOCTEXT("OptimizationSubMenuHeader", "Optimization Viewmodes"));
 							if (FeatureLevel == ERHIFeatureLevel::SM5)
 							{
-								Menu.AddMenuEntry(BaseViewportCommands.LightComplexityMode, NAME_None, UViewModeUtils::GetViewModeDisplayName(VMI_LightComplexity));
-								Menu.AddMenuEntry(BaseViewportCommands.LightmapDensityMode, NAME_None, UViewModeUtils::GetViewModeDisplayName(VMI_LightmapDensity));
-								Menu.AddMenuEntry(BaseViewportCommands.StationaryLightOverlapMode, NAME_None, UViewModeUtils::GetViewModeDisplayName(VMI_StationaryLightOverlap));
+								Section.AddMenuEntry(BaseViewportCommands.LightComplexityMode, UViewModeUtils::GetViewModeDisplayName(VMI_LightComplexity));
+								Section.AddMenuEntry(BaseViewportCommands.LightmapDensityMode, UViewModeUtils::GetViewModeDisplayName(VMI_LightmapDensity));
+								Section.AddMenuEntry(BaseViewportCommands.StationaryLightOverlapMode, UViewModeUtils::GetViewModeDisplayName(VMI_StationaryLightOverlap));
 							}
 
-							Menu.AddMenuEntry(BaseViewportCommands.ShaderComplexityMode, NAME_None, UViewModeUtils::GetViewModeDisplayName(VMI_ShaderComplexity));
+							Section.AddMenuEntry(BaseViewportCommands.ShaderComplexityMode, UViewModeUtils::GetViewModeDisplayName(VMI_ShaderComplexity));
 
 							if (AllowDebugViewShaderMode(DVSM_ShaderComplexityContainedQuadOverhead, GMaxRHIShaderPlatform, FeatureLevel))
 							{
-								Menu.AddMenuEntry(BaseViewportCommands.ShaderComplexityWithQuadOverdrawMode, NAME_None, UViewModeUtils::GetViewModeDisplayName(VMI_ShaderComplexityWithQuadOverdraw));
+								Section.AddMenuEntry(BaseViewportCommands.ShaderComplexityWithQuadOverdrawMode, UViewModeUtils::GetViewModeDisplayName(VMI_ShaderComplexityWithQuadOverdraw));
 							}
 							if (AllowDebugViewShaderMode(DVSM_QuadComplexity, GMaxRHIShaderPlatform, FeatureLevel))
 							{
-								Menu.AddMenuEntry(BaseViewportCommands.QuadOverdrawMode, NAME_None, UViewModeUtils::GetViewModeDisplayName(VMI_QuadOverdraw));
+								Section.AddMenuEntry(BaseViewportCommands.QuadOverdrawMode, UViewModeUtils::GetViewModeDisplayName(VMI_QuadOverdraw));
 							}
 						}
-						Menu.EndSection();
 
-						Menu.BeginSection("TextureStreaming", LOCTEXT("TextureStreamingHeader", "Texture Streaming Accuracy") );
-						if ( AllowDebugViewShaderMode(DVSM_PrimitiveDistanceAccuracy, GMaxRHIShaderPlatform, FeatureLevel) && (!InParentToolBar.IsValid() || InParentToolBar.Pin()->IsViewModeSupported(VMI_PrimitiveDistanceAccuracy)) )
-						{
-							Menu.AddMenuEntry(BaseViewportCommands.TexStreamAccPrimitiveDistanceMode, NAME_None, UViewModeUtils::GetViewModeDisplayName(VMI_PrimitiveDistanceAccuracy));
-						}
-						if ( AllowDebugViewShaderMode(DVSM_MeshUVDensityAccuracy, GMaxRHIShaderPlatform, FeatureLevel) && (!InParentToolBar.IsValid() || InParentToolBar.Pin()->IsViewModeSupported(VMI_MeshUVDensityAccuracy)) )
-						{
-							Menu.AddMenuEntry(BaseViewportCommands.TexStreamAccMeshUVDensityMode, NAME_None, UViewModeUtils::GetViewModeDisplayName(VMI_MeshUVDensityAccuracy));
-						}
-						// TexCoordScale accuracy viewmode requires shaders that are only built in the TextureStreamingBuild, which requires the new metrics to be enabled.
-						if ( AllowDebugViewShaderMode(DVSM_MaterialTextureScaleAccuracy, GMaxRHIShaderPlatform, FeatureLevel) && CVarStreamingUseNewMetrics.GetValueOnAnyThread() != 0 && (!InParentToolBar.IsValid() || InParentToolBar.Pin()->IsViewModeSupported(VMI_MaterialTextureScaleAccuracy)) )
-						{
-							Menu.AddMenuEntry(BaseViewportCommands.TexStreamAccMaterialTextureScaleMode, NAME_None, UViewModeUtils::GetViewModeDisplayName(VMI_MaterialTextureScaleAccuracy));
-						}
-						if ( AllowDebugViewShaderMode(DVSM_RequiredTextureResolution, GMaxRHIShaderPlatform, FeatureLevel) && (!InParentToolBar.IsValid() || InParentToolBar.Pin()->IsViewModeSupported(VMI_MaterialTextureScaleAccuracy)) )
 						{
-							Menu.AddMenuEntry(BaseViewportCommands.RequiredTextureResolutionMode, NAME_None, UViewModeUtils::GetViewModeDisplayName(VMI_RequiredTextureResolution));
+							FToolMenuSection& Section = Menu->AddSection("TextureStreaming", LOCTEXT("TextureStreamingHeader", "Texture Streaming Accuracy"));
+							if (AllowDebugViewShaderMode(DVSM_PrimitiveDistanceAccuracy, GMaxRHIShaderPlatform, FeatureLevel) && (!InParentToolBar.IsValid() || InParentToolBar.Pin()->IsViewModeSupported(VMI_PrimitiveDistanceAccuracy)))
+							{
+								Section.AddMenuEntry(BaseViewportCommands.TexStreamAccPrimitiveDistanceMode, UViewModeUtils::GetViewModeDisplayName(VMI_PrimitiveDistanceAccuracy));
+							}
+							if (AllowDebugViewShaderMode(DVSM_MeshUVDensityAccuracy, GMaxRHIShaderPlatform, FeatureLevel) && (!InParentToolBar.IsValid() || InParentToolBar.Pin()->IsViewModeSupported(VMI_MeshUVDensityAccuracy)))
+							{
+								Section.AddMenuEntry(BaseViewportCommands.TexStreamAccMeshUVDensityMode, UViewModeUtils::GetViewModeDisplayName(VMI_MeshUVDensityAccuracy));
+							}
+							// TexCoordScale accuracy viewmode requires shaders that are only built in the TextureStreamingBuild, which requires the new metrics to be enabled.
+							if (AllowDebugViewShaderMode(DVSM_MaterialTextureScaleAccuracy, GMaxRHIShaderPlatform, FeatureLevel) && CVarStreamingUseNewMetrics.GetValueOnAnyThread() != 0 && (!InParentToolBar.IsValid() || InParentToolBar.Pin()->IsViewModeSupported(VMI_MaterialTextureScaleAccuracy)))
+							{
+								Section.AddMenuEntry(BaseViewportCommands.TexStreamAccMaterialTextureScaleMode, UViewModeUtils::GetViewModeDisplayName(VMI_MaterialTextureScaleAccuracy));
+							}
+							if (AllowDebugViewShaderMode(DVSM_RequiredTextureResolution, GMaxRHIShaderPlatform, FeatureLevel) && (!InParentToolBar.IsValid() || InParentToolBar.Pin()->IsViewModeSupported(VMI_MaterialTextureScaleAccuracy)))
+							{
+								Section.AddMenuEntry(BaseViewportCommands.RequiredTextureResolutionMode, UViewModeUtils::GetViewModeDisplayName(VMI_RequiredTextureResolution));
+							}
 						}
-						Menu.EndSection();
 					}
 				};
 
-				ViewMenuBuilder.AddSubMenu(
+				Section.AddSubMenu(
+					"OptimizationSubMenu",
 					LOCTEXT("OptimizationSubMenu", "Optimization Viewmodes"), LOCTEXT("Optimization_ToolTip", "Select optimization visualizer"),
-					FNewMenuDelegate::CreateStatic(&Local::BuildOptimizationMenu, ParentToolBar),
+					FNewToolMenuDelegate::CreateStatic(&Local::BuildOptimizationMenu, ParentToolBar),
 					FUIAction(FExecuteAction(), FCanExecuteAction(),
 						FIsActionChecked::CreateLambda([this]()
 						{
@@ -292,7 +311,7 @@ TSharedRef<SWidget> SEditorViewportViewMenu::GenerateViewMenuContent() const
 									|| ViewMode == VMI_PrimitiveDistanceAccuracy || ViewMode == VMI_MeshUVDensityAccuracy || ViewMode == VMI_MaterialTextureScaleAccuracy || ViewMode == VMI_RequiredTextureResolution
 								);
 						})),
-					/* InExtensionHook = */ NAME_None, EUserInterfaceActionType::RadioButton,
+					EUserInterfaceActionType::RadioButton,
 					/* bInOpenSubMenuOnClick = */ false, FSlateIcon(FEditorStyle::GetStyleSetName(), "EditorViewport.QuadOverdrawMode"));
 			}
 
@@ -308,27 +327,27 @@ TSharedRef<SWidget> SEditorViewportViewMenu::GenerateViewMenuContent() const
 					}
 				};
 
-				ViewMenuBuilder.AddSubMenu(LOCTEXT("RayTracingDebugSubMenu", "Ray Tracing Debug"), LOCTEXT("RayTracing_ToolTip", "Select ray tracing buffer visualization view modes"), FNewMenuDelegate::CreateStatic(&Local::BuildRayTracingDebugMenu, ParentToolBar));
+				Section.AddSubMenu("RayTracingDebugSubMenu", LOCTEXT("RayTracingDebugSubMenu", "Ray Tracing Debug"), LOCTEXT("RayTracing_ToolTip", "Select ray tracing buffer visualization view modes"), FNewMenuDelegate::CreateStatic(&Local::BuildRayTracingDebugMenu, ParentToolBar));
 			}
 #endif
 
 			{
 				struct Local
 				{
-					static void BuildLODMenu(FMenuBuilder& Menu)
+					static void BuildLODMenu(UToolMenu* Menu)
 					{
-						Menu.BeginSection("LevelViewportLODColoration", LOCTEXT("LODModesHeader", "Level of Detail Coloration"));
 						{
-							Menu.AddMenuEntry(FEditorViewportCommands::Get().LODColorationMode, NAME_None, UViewModeUtils::GetViewModeDisplayName(VMI_LODColoration));
-							Menu.AddMenuEntry(FEditorViewportCommands::Get().HLODColorationMode, NAME_None, UViewModeUtils::GetViewModeDisplayName(VMI_HLODColoration));
+							FToolMenuSection& Section = Menu->AddSection("LevelViewportLODColoration", LOCTEXT("LODModesHeader", "Level of Detail Coloration"));
+							Section.AddMenuEntry(FEditorViewportCommands::Get().LODColorationMode, UViewModeUtils::GetViewModeDisplayName(VMI_LODColoration));
+							Section.AddMenuEntry(FEditorViewportCommands::Get().HLODColorationMode, UViewModeUtils::GetViewModeDisplayName(VMI_HLODColoration));
 						}
-						Menu.EndSection();
 					}
 				};
 
-				ViewMenuBuilder.AddSubMenu(
+				Section.AddSubMenu(
+					"VisualizeGroupedLOD",
 					LOCTEXT("VisualizeGroupedLODDisplayName", "Level of Detail Coloration"), LOCTEXT("GroupedLODMenu_ToolTip", "Select a mode for LOD Coloration"),
-					FNewMenuDelegate::CreateStatic(&Local::BuildLODMenu),
+					FNewToolMenuDelegate::CreateStatic(&Local::BuildLODMenu),
 					FUIAction(FExecuteAction(), FCanExecuteAction(),
 						FIsActionChecked::CreateLambda([this]()
 						{
@@ -338,11 +357,9 @@ TSharedRef<SWidget> SEditorViewportViewMenu::GenerateViewMenuContent() const
 								const EViewModeIndex ViewMode = ViewportClient->GetViewMode();
 								return (ViewMode == VMI_LODColoration || ViewMode == VMI_HLODColoration);
 						})),
-					/* InExtensionHook = */ NAME_None, EUserInterfaceActionType::RadioButton,
+					EUserInterfaceActionType::RadioButton,
 					/* bInOpenSubMenuOnClick = */ false, FSlateIcon(FEditorStyle::GetStyleSetName(), "EditorViewport.GroupLODColorationMode"));
 			}
-
-			ViewMenuBuilder.EndSection();
 		}
 
 		// Auto Exposure
@@ -353,13 +370,11 @@ TSharedRef<SWidget> SEditorViewportViewMenu::GenerateViewMenuContent() const
 			TSharedPtr<FEditorViewportClient> EditorViewPostClient = Viewport.Pin()->GetViewportClient();
 			const bool bIsLevelEditor = EditorViewPostClient.IsValid() && EditorViewPostClient->IsLevelEditorClient();
 
-			ViewMenuBuilder.BeginSection("Exposure", LOCTEXT("ExposureHeader", "Exposure"));
-			ViewMenuBuilder.AddMenuEntry( bIsLevelEditor ? BaseViewportCommands.ToggleInGameExposure : BaseViewportCommands.ToggleAutoExposure, NAME_None );
-			ViewMenuBuilder.AddWidget( FixedEV100Menu, LOCTEXT("FixedEV100", "EV100") );
-			ViewMenuBuilder.EndSection();
+			FToolMenuSection& Section = Menu->AddSection("Exposure", LOCTEXT("ExposureHeader", "Exposure"));
+			Section.AddMenuEntry(bIsLevelEditor ? BaseViewportCommands.ToggleInGameExposure : BaseViewportCommands.ToggleAutoExposure);
+			Section.AddEntry(FToolMenuEntry::InitWidget("FixedEV100", FixedEV100Menu, LOCTEXT("FixedEV100", "EV100")));
 		}
 	}
-	return ViewMenuBuilder.MakeWidget();
 }
 
 #undef LOCTEXT_NAMESPACE
diff --git a/Engine/Source/Editor/UnrealEd/Public/SEditorViewportViewMenu.h b/Engine/Source/Editor/UnrealEd/Public/SEditorViewportViewMenu.h
--- a/Engine/Source/Editor/UnrealEd/Public/SEditorViewportViewMenu.h
+++ b/Engine/Source/Editor/UnrealEd/Public/SEditorViewportViewMenu.h
@@ -12,6 +12,7 @@
 #include "Styling/SlateTypes.h"
 
 struct FSlateBrush;
+class UToolMenu;
 
 class UNREALED_API SEditorViewportViewMenu : public SEditorViewportToolbarMenu
 {
@@ -25,6 +26,7 @@ public:
 private:
 	FText GetViewMenuLabel() const;
 	const FSlateBrush* GetViewMenuLabelIcon() const;
+	void FillViewMenu(UToolMenu* Menu) const;
 
 protected:
 	virtual TSharedRef<SWidget> GenerateViewMenuContent() const;
diff --git a/Engine/Source/Editor/UnrealEd/Public/SEditorViewportViewMenuContext.h b/Engine/Source/Editor/UnrealEd/Public/SEditorViewportViewMenuContext.h
new file mode 100644
--- /dev/null
+++ b/Engine/Source/Editor/UnrealEd/Public/SEditorViewportViewMenuContext.h
@@ -0,0 +1,18 @@
+// Copyright Epic Games, Inc. All Rights Reserved.
+
+#pragma once
+
+#include "CoreMinimal.h"
+
+#include "SEditorViewportViewMenuContext.generated.h"
+
+class SEditorViewportViewMenu;
+
+UCLASS()
+class UNREALED_API UEditorViewportViewMenuContext : public UObject
+{
+	GENERATED_BODY()
+public:
+
+	TWeakPtr<const SEditorViewportViewMenu> Widget;
+};
