Fix non-unity

diff --git a/Engine/Source/Editor/UnrealEd/Public/SEditorViewportToolBarMenu.h b/Engine/Source/Editor/UnrealEd/Public/SEditorViewportToolBarMenu.h
--- a/Engine/Source/Editor/UnrealEd/Public/SEditorViewportToolBarMenu.h
+++ b/Engine/Source/Editor/UnrealEd/Public/SEditorViewportToolBarMenu.h
@@ -11,6 +11,7 @@
 #include "Widgets/SCompoundWidget.h"
 #include "SViewportToolBar.h"
 #include "Framework/SlateDelegates.h"
+#include "Styling/SlateTypes.h"
 
 class SMenuAnchor;
 struct FSlateBrush;
