Fix non-unity

#rnx

diff --git a/Engine/Source/Editor/UnrealEd/Public/SEditorViewportToolBarMenu.h b/Engine/Source/Editor/UnrealEd/Public/SEditorViewportToolBarMenu.h
--- a/Engine/Source/Editor/UnrealEd/Public/SEditorViewportToolBarMenu.h
+++ b/Engine/Source/Editor/UnrealEd/Public/SEditorViewportToolBarMenu.h
@@ -12,6 +12,7 @@
 #include "SViewportToolBar.h"
 #include "Framework/SlateDelegates.h"
 #include "Styling/SlateTypes.h"
+#include "EditorStyleSet.h"
 
 class SMenuAnchor;
 struct FSlateBrush;
