Fix incremental build

#rnx

diff --git a/Engine/Source/Editor/UnrealEd/Public/SEditorViewportToolBarButton.h b/Engine/Source/Editor/UnrealEd/Public/SEditorViewportToolBarButton.h
--- a/Engine/Source/Editor/UnrealEd/Public/SEditorViewportToolBarButton.h
+++ b/Engine/Source/Editor/UnrealEd/Public/SEditorViewportToolBarButton.h
@@ -9,6 +9,9 @@
 #include "Widgets/SCompoundWidget.h"
 #include "Framework/SlateDelegates.h"
 #include "Framework/Commands/UICommandInfo.h"
+#include "Styling/SlateTypes.h"
+#include "EditorStyleSet.h"
+#include "Styling/SlateWidgetStyleAsset.h"
 
 struct FSlateBrush;
 enum class ECheckBoxState : uint8;
diff --git a/Engine/Source/Editor/UnrealEd/Public/SEditorViewportToolBarMenu.h b/Engine/Source/Editor/UnrealEd/Public/SEditorViewportToolBarMenu.h
--- a/Engine/Source/Editor/UnrealEd/Public/SEditorViewportToolBarMenu.h
+++ b/Engine/Source/Editor/UnrealEd/Public/SEditorViewportToolBarMenu.h
@@ -13,6 +13,7 @@
 #include "Framework/SlateDelegates.h"
 #include "Styling/SlateTypes.h"
 #include "EditorStyleSet.h"
+#include "Styling/SlateWidgetStyleAsset.h"
 
 class SMenuAnchor;
 struct FSlateBrush;
