Fix source control errors during rename of file open for add

diff --git a/Engine/Source/Developer/AssetTools/Private/AssetRenameManager.cpp b/Engine/Source/Developer/AssetTools/Private/AssetRenameManager.cpp
--- a/Engine/Source/Developer/AssetTools/Private/AssetRenameManager.cpp
+++ b/Engine/Source/Developer/AssetTools/Private/AssetRenameManager.cpp
@@ -1394,7 +1394,20 @@ void FAssetRenameManager::PerformAssetRename(TArray<FAssetRenameDataWithReferenc
 		// as it will just create a copy and any attempt to load it will result in crashes.
 		if (!RenameData.bOnlyFixSoftReferences && NewPackage && FPackageName::DoesPackageExist(NewPackage->GetName()))
 		{
-			SourceControlHelpers::BranchPackage(NewPackage, OldPackage);
+			if (ISourceControlModule::Get().IsEnabled())
+			{
+				ISourceControlProvider& SourceControlProvider = ISourceControlModule::Get().GetProvider();
+				const FString SourceFilename = USourceControlHelpers::PackageFilename(OldPackage);
+				FSourceControlStatePtr SourceControlState = SourceControlProvider.GetState(SourceFilename, EStateCacheUsage::ForceUpdate);
+				if (SourceControlState.IsValid() && SourceControlState->IsSourceControlled())
+				{
+					// Do not attempt to branch if the old file was open for add
+					if (!SourceControlState->IsAdded())
+					{
+						SourceControlHelpers::BranchPackage(NewPackage, OldPackage);
+					}
+				}
+			}
 		}
 	}
 
diff --git a/Engine/Source/Editor/UnrealEd/Private/FileHelpers.cpp b/Engine/Source/Editor/UnrealEd/Private/FileHelpers.cpp
--- a/Engine/Source/Editor/UnrealEd/Private/FileHelpers.cpp
+++ b/Engine/Source/Editor/UnrealEd/Private/FileHelpers.cpp
@@ -1762,6 +1762,8 @@ bool FEditorFileUtils::PromptToCheckoutPackages(bool bCheckDirty, const TArray<U
 
 ECommandResult::Type FEditorFileUtils::CheckoutPackages(const TArray<UPackage*>& PkgsToCheckOut, TArray<UPackage*>* OutPackagesCheckedOut, const bool bErrorIfAlreadyCheckedOut, const bool bConfirmPackageBranchCheckOutStatus)
 {
+	const bool bErrorIfFileMissing = false;
+
 	ECommandResult::Type CheckOutResult = ECommandResult::Succeeded;
 	FString PkgsWhichFailedCheckout;
 
@@ -1809,12 +1811,18 @@ ECommandResult::Type FEditorFileUtils::CheckoutPackages(const TArray<UPackage*>&
 				}
 				else if (SourceControlState->CanAdd())
 				{
-					// Cannot add unsaved packages to source control
-					if (!PackageToCheckOut->HasAnyPackageFlags(PKG_NewlyCreated))
+					FString Filename;
+					if (FPackageName::DoesPackageExist(PackageToCheckOut->GetName(), nullptr, &Filename))
 					{
 						bShowCheckoutError = false;
 						FinalPackageMarkForAddList.Add(PackageToCheckOut);
 					}
+					else if (!bErrorIfFileMissing)
+					{
+						// Silently skip package that has not been saved yet
+						// Expected when called by InternalCheckoutAndSavePackages before packages saved
+						bShowCheckoutError = false;
+					}
 				}
 				else if( !bErrorIfAlreadyCheckedOut && SourceControlState->IsCheckedOut() && !SourceControlState->IsCheckedOutOther() )
 				{
@@ -4306,7 +4314,7 @@ static bool InternalCheckoutAndSavePackages(const TArray<UPackage*>& PackagesToS
 			for (UPackage* Package : PackagesToSave)
 			{
 				// List unsaved packages that were not checked out
-				if (Package->HasAnyPackageFlags(PKG_NewlyCreated) && !PackagesCheckedOut.Contains(Package))
+				if (!PackagesCheckedOut.Contains(Package))
 				{
 					PackagesToMarkForAdd.Add(Package);
 				}
