Make sure post buid refactored skeletalmesh has there dependent LODs regenerated after importing the base mesh.
#rnx

diff --git a/Engine/Source/Editor/UnrealEd/Private/Factories/SkeletalMeshImport.cpp b/Engine/Source/Editor/UnrealEd/Private/Factories/SkeletalMeshImport.cpp
--- a/Engine/Source/Editor/UnrealEd/Private/Factories/SkeletalMeshImport.cpp
+++ b/Engine/Source/Editor/UnrealEd/Private/Factories/SkeletalMeshImport.cpp
@@ -509,57 +509,6 @@ void RestoreDependentLODs(ExistingSkelMeshData* MeshData, USkeletalMesh* Skeleta
 	}
 }
 
-void TryRegenerateLODs(ExistingSkelMeshData* MeshData, USkeletalMesh* SkeletalMesh)
-{
-	check(SkeletalMesh != nullptr);
-	int32 TotalLOD = MeshData->ExistingLODModels.Num();
-	FSkeletalMeshModel* SkeletalMeshImportedModel = SkeletalMesh->GetImportedModel();
-
-	// see if mesh reduction util is available
-	IMeshReductionManagerModule& Module = FModuleManager::Get().LoadModuleChecked<IMeshReductionManagerModule>("MeshReductionInterface");
-	static bool bAutoMeshReductionAvailable = Module.GetSkeletalMeshReductionInterface() != NULL;
-
-	if (bAutoMeshReductionAvailable)
-	{
-		GWarn->BeginSlowTask(LOCTEXT("RegenLODs", "Generating new LODs"), true);
-		FSkeletalMeshUpdateContext UpdateContext;
-		UpdateContext.SkeletalMesh = SkeletalMesh;
-		TArray<bool> Dependencies;
-		Dependencies.AddZeroed(TotalLOD + 1);
-		Dependencies[0] = true;
-		for (int32 Index = 0; Index < TotalLOD; ++Index)
-		{
-			int32 LODIndex = Index + 1;
-			if (LODIndex >= SkeletalMesh->GetLODInfoArray().Num())
-			{
-				FSkeletalMeshLODInfo& ExistLODInfo = MeshData->ExistingLODInfo[Index];
-				FSkeletalMeshLODModel& ExistLODModel = MeshData->ExistingLODModels[Index];
-				// reset material maps, it won't work anyway. 
-				ExistLODInfo.LODMaterialMap.Empty();
-
-				FSkeletalMeshLODModel* NewLODModel = new FSkeletalMeshLODModel(ExistLODModel);
-				SkeletalMeshImportedModel->LODModels.Add(NewLODModel);
-				// add LOD info back
-				SkeletalMesh->AddLODInfo(ExistLODInfo);
-				check(LODIndex < SkeletalMesh->GetLODInfoArray().Num());
-			}
-			const FSkeletalMeshLODInfo* LODInfo = SkeletalMesh->GetLODInfo(LODIndex);
-			if (LODInfo && LODInfo->bHasBeenSimplified && Dependencies[LODInfo->ReductionSettings.BaseLOD])
-			{
-				Dependencies[LODIndex] = true;
-				// force it to regenerate
-				FLODUtilities::SimplifySkeletalMeshLOD(UpdateContext, LODIndex, false);
-			}
-		}
-		GWarn->EndSlowTask();
-	}
-	else
-	{
-		UnFbx::FFbxImporter* FFbxImporter = UnFbx::FFbxImporter::GetInstance();
-		FFbxImporter->AddTokenizedErrorMessage(FTokenizedMessage::Create(EMessageSeverity::Warning, LOCTEXT("NoCompatibleSkeleton", "New base mesh is not compatible with previous LODs. LOD will be removed.")), FFbxErrors::SkeletalMesh_LOD_MissingBone);
-	}
-}
-
 namespace SkeletalMeshHelper
 {
 	void ApplySkinning(USkeletalMesh* SkeletalMesh, FSkeletalMeshLODModel& SrcLODModel, FSkeletalMeshLODModel& DestLODModel)
@@ -946,8 +895,14 @@ void RestoreExistingSkelMeshData(ExistingSkelMeshData* MeshData, USkeletalMesh*
 					}
 				}
 			}
-			//We just need to restore the LOD model and LOD info the build will regenerate the LODs
+			//We just need to restore the LOD model and LOD info the build should regenerate the LODs
 			RestoreDependentLODs(MeshData, SkeletalMesh);
+			
+			//Old asset cannot use the new build system, we need to regenerate dependent LODs
+			if (SkeletalMeshImportedModel->LODModels[SafeReimportLODIndex].RawSkeletalMeshBulkData.IsBuildDataAvailable() == false)
+			{
+				FLODUtilities::RegenerateDependentLODs(SkeletalMesh, SafeReimportLODIndex);
+			}
 		}
 
 		for (int32 AssetIndex = 0; AssetIndex < MeshData->ExistingPhysicsAssets.Num(); ++AssetIndex)
