Fix build to use the new API
#rnx

diff --git a/Engine/Source/Editor/UnrealEd/Private/Factories/SkeletalMeshImport.cpp b/Engine/Source/Editor/UnrealEd/Private/Factories/SkeletalMeshImport.cpp
--- a/Engine/Source/Editor/UnrealEd/Private/Factories/SkeletalMeshImport.cpp
+++ b/Engine/Source/Editor/UnrealEd/Private/Factories/SkeletalMeshImport.cpp
@@ -898,8 +898,9 @@ void RestoreExistingSkelMeshData(ExistingSkelMeshData* MeshData, USkeletalMesh*
 			//We just need to restore the LOD model and LOD info the build should regenerate the LODs
 			RestoreDependentLODs(MeshData, SkeletalMesh);
 			
+			
 			//Old asset cannot use the new build system, we need to regenerate dependent LODs
-			if (SkeletalMeshImportedModel->LODModels[SafeReimportLODIndex].RawSkeletalMeshBulkData.IsBuildDataAvailable() == false)
+			if (SkeletalMesh->IsLODImportedDataBuildAvailable(SafeReimportLODIndex) == false)
 			{
 				FLODUtilities::RegenerateDependentLODs(SkeletalMesh, SafeReimportLODIndex);
 			}
