Turn on new mesh paint by default

diff --git a/Engine/Source/Editor/UnrealEd/Private/EditorModeManager.cpp b/Engine/Source/Editor/UnrealEd/Private/EditorModeManager.cpp
--- a/Engine/Source/Editor/UnrealEd/Private/EditorModeManager.cpp
+++ b/Engine/Source/Editor/UnrealEd/Private/EditorModeManager.cpp
@@ -1159,7 +1159,9 @@ void FEditorModeTools::ActivateMode(FEditorModeID InID, bool bToggle)
 				Toolkit->GetToolPaletteNames(PaletteNames);
 				for(auto Palette : PaletteNames)
 				{
-					FToolBarBuilder ModeToolbarBuilder(CommandList, FMultiBoxCustomization(ScriptableMode->GetModeInfo().ToolbarCustomizationName), TSharedPtr<FExtender>(), Orient_Horizontal, false);
+					const bool bUniform = true;
+					FToolBarBuilder ModeToolbarBuilder(CommandList, FMultiBoxCustomization(ScriptableMode->GetModeInfo().ToolbarCustomizationName), TSharedPtr<FExtender>(), Orient_Horizontal, false, bUniform);
+					ModeToolbarBuilder.SetStyle(&FEditorStyle::Get(), "PaletteToolBar");
 					Toolkit->BuildToolPalette(Palette, ModeToolbarBuilder);
 
 					ActiveToolBarRows.Emplace(ScriptableMode->GetID(), Palette, Toolkit->GetToolPaletteDisplayName(Palette), ModeToolbarBuilder.MakeWidget());
diff --git a/Engine/Source/Editor/UnrealEd/Private/Settings/SettingsClasses.cpp b/Engine/Source/Editor/UnrealEd/Private/Settings/SettingsClasses.cpp
--- a/Engine/Source/Editor/UnrealEd/Private/Settings/SettingsClasses.cpp
+++ b/Engine/Source/Editor/UnrealEd/Private/Settings/SettingsClasses.cpp
@@ -430,7 +430,7 @@ ULevelEditorMiscSettings::ULevelEditorMiscSettings( const FObjectInitializer& Ob
 	PercentageThresholdForPrompt = 20.0f;
 	MinimumBoundsForCheckingSize = FVector(500.0f, 500.0f, 50.0f);
 	bCreateNewAudioDeviceForPlayInEditor = true;
-	bEnableLegacyMeshPaintMode = true;
+	bEnableLegacyMeshPaintMode = false;
 }
 
 void ULevelEditorMiscSettings::PostEditChangeProperty( struct FPropertyChangedEvent& PropertyChangedEvent )
