Fix path prefix comparison when cooking with ErrorOnEngineContentUse

PR #4686: Fix for not being able to cook plugins content. (Reported by [REDACTED_PERSON])

diff --git a/Engine/Source/Editor/UnrealEd/Private/CookOnTheFlyServer.cpp b/Engine/Source/Editor/UnrealEd/Private/CookOnTheFlyServer.cpp
--- a/Engine/Source/Editor/UnrealEd/Private/CookOnTheFlyServer.cpp
+++ b/Engine/Source/Editor/UnrealEd/Private/CookOnTheFlyServer.cpp
@@ -2707,6 +2707,7 @@ uint32 UCookOnTheFlyServer::TickCookOnTheSide(const float TimeSlice, uint32 &Coo
 		{
 			check(IsCookingDLC());
 			FString DLCPath = FPaths::Combine(*GetBaseDirectoryForDLC(), TEXT("Content"));
+			FPaths::MakeStandardFilename(DLCPath);
 			if (ToBuild.GetFilename().ToString().StartsWith(DLCPath) == false) // if we don't start with the dlc path then we shouldn't be cooking this data 
 			{
 				UE_LOG(LogCook, Error, TEXT("Engine or Game content %s is being referenced by DLC!"), *ToBuild.GetFilename().ToString());
