Changing the delta on the Number of Players slider so it doesn't get stuck on the max value.

diff --git a/Engine/Source/Editor/UnrealEd/Private/Kismet2/DebuggerCommands.cpp b/Engine/Source/Editor/UnrealEd/Private/Kismet2/DebuggerCommands.cpp
--- a/Engine/Source/Editor/UnrealEd/Private/Kismet2/DebuggerCommands.cpp
+++ b/Engine/Source/Editor/UnrealEd/Private/Kismet2/DebuggerCommands.cpp
@@ -929,6 +929,7 @@ TSharedRef< SWidget > FPlayWorldCommands::GeneratePlayMenuContent(TSharedRef<FUI
 				.MaxValue(64)
 				.MinSliderValue(1)
 				.MaxSliderValue(4)
+				.Delta(1)
 				.ToolTipText(LOCTEXT("NumberOfClientsToolTip", "How many client instances do you want to create? The first instance respects the Play Mode location (PIE/PINW) and additional instances respect the RunUnderOneProcess setting."))
 				.Value(FInternalPlayWorldCommandCallbacks::GetNumberOfClients())
 				.OnValueCommitted_Static(&FInternalPlayWorldCommandCallbacks::SetNumberOfClients);
