Fix crash when framerate is zero in the anim timeline

diff --git a/Engine/Source/Editor/UnrealEd/Private/SkeletalMeshEdit.cpp b/Engine/Source/Editor/UnrealEd/Private/SkeletalMeshEdit.cpp
--- a/Engine/Source/Editor/UnrealEd/Private/SkeletalMeshEdit.cpp
+++ b/Engine/Source/Editor/UnrealEd/Private/SkeletalMeshEdit.cpp
@@ -1356,7 +1356,7 @@ bool UnFbx::FFbxImporter::ImportAnimation(USkeleton* Skeleton, UAnimSequence * D
 	}
 
 	FbxNode *SkeletalMeshRootNode = NodeArray.Num() > 0 ? NodeArray[0] : nullptr;
-	
+	int32 CurveAttributeKeyCount = 0;
 	//
 	// import blend shape curves
 	//
@@ -1415,6 +1415,7 @@ bool UnFbx::FFbxImporter::ImportAnimation(USkeleton* Skeleton, UAnimSequence * D
 
 								if (ImportCurveToAnimSequence(DestSeq, *ChannelName, Curve, 0, AnimTimeSpan, 0.01f /** for some reason blend shape values are coming as 100 scaled **/))
 								{
+									CurveAttributeKeyCount = FMath::Max(CurveAttributeKeyCount, Curve->KeyGetCount());
 									// this one doesn't reset Material curve to false, it just accumulate if true. 
 									MySkeleton->AccumulateCurveMetaData(*ChannelName, false, true);
 								}
@@ -1443,7 +1444,7 @@ bool UnFbx::FFbxImporter::ImportAnimation(USkeleton* Skeleton, UAnimSequence * D
 		for(auto Node: SortedLinks)
 		{
 			FbxAnimUtils::ExtractAttributeCurves(Node, ImportOptions->bDoNotImportCurveWithZero, 
-			[this, &CurLinkIndex, &TotalLinks, &DestSeq, &AnimTimeSpan, &ExistingCurveNames](FbxAnimCurve* InCurve, const FString& InCurveName, const FString& InChannelName, int32 InChannelIndex, int32 InNumChannels)
+			[this, &CurLinkIndex, &TotalLinks, &DestSeq, &AnimTimeSpan, &ExistingCurveNames, &CurveAttributeKeyCount](FbxAnimCurve* InCurve, const FString& InCurveName, const FString& InChannelName, int32 InChannelIndex, int32 InNumChannels)
 			{
 				FString FinalCurveName;
 				if (InNumChannels == 1)
@@ -1464,7 +1465,7 @@ bool UnFbx::FFbxImporter::ImportAnimation(USkeleton* Skeleton, UAnimSequence * D
 				if (ImportCurveToAnimSequence(DestSeq, FinalCurveName, InCurve, CurveFlags, AnimTimeSpan))
 				{
 					USkeleton* SeqSkeleton = DestSeq->GetSkeleton();
-
+					CurveAttributeKeyCount = FMath::Max(CurveAttributeKeyCount, InCurve->KeyGetCount());
 					// first let them override material curve if required
 					if (ImportOptions->bSetMaterialDriveParameterOnCustomAttribute)
 					{
@@ -1688,7 +1689,11 @@ bool UnFbx::FFbxImporter::ImportAnimation(USkeleton* Skeleton, UAnimSequence * D
 
 		GWarn->EndSlowTask();
 	}
-
+	else if (CurveAttributeKeyCount > 0)
+	{
+		DestSeq->SetRawNumberOfFrame(CurveAttributeKeyCount);
+		DestSeq->MarkRawDataAsModified();
+	}
 	// compress animation
 	{
 		GWarn->BeginSlowTask( LOCTEXT("BeginCompressAnimation", "Compress Animation"), true);
