Bones metadata from FBX are now transferred to AnimSequence at import.

diff --git a/Engine/Source/Editor/UnrealEd/Private/Fbx/FbxMainImport.cpp b/Engine/Source/Editor/UnrealEd/Private/Fbx/FbxMainImport.cpp
--- a/Engine/Source/Editor/UnrealEd/Private/Fbx/FbxMainImport.cpp
+++ b/Engine/Source/Editor/UnrealEd/Private/Fbx/FbxMainImport.cpp
@@ -3114,12 +3114,42 @@ FString GetFbxPropertyStringValue(const FbxProperty& Property)
 		ValueStr = LexToString(BoolValue);
 	}
 	break;
+	case eFbxChar:
+	{
+		FbxChar CharValue = Property.Get<FbxChar>();
+		ValueStr = LexToString(CharValue);
+	}
+	break;
+	case eFbxUChar:
+	{
+		FbxUChar UCharValue = Property.Get<FbxUChar>();
+		ValueStr = LexToString(UCharValue);
+	}
+	break;
+	case eFbxShort:
+	{
+		FbxShort ShortValue = Property.Get<FbxShort>();
+		ValueStr = LexToString(ShortValue);
+	}
+	break;
+	case eFbxUShort:
+	{
+		FbxUShort UShortValue = Property.Get<FbxUShort>();
+		ValueStr = LexToString(UShortValue);
+	}
+	break;
 	case eFbxInt:
 	{
 		FbxInt IntValue = Property.Get<FbxInt>();
 		ValueStr = LexToString(IntValue);
 	}
 	break;
+	case eFbxUInt:
+	{
+		FbxUInt UIntValue = Property.Get<FbxUInt>();
+		ValueStr = LexToString(UIntValue);
+	}
+	break;
 	case eFbxEnum:
 	{
 		FbxEnum EnumValue = Property.Get<FbxEnum>();
diff --git a/Engine/Source/Editor/UnrealEd/Private/SkeletalMeshEdit.cpp b/Engine/Source/Editor/UnrealEd/Private/SkeletalMeshEdit.cpp
--- a/Engine/Source/Editor/UnrealEd/Private/SkeletalMeshEdit.cpp
+++ b/Engine/Source/Editor/UnrealEd/Private/SkeletalMeshEdit.cpp
@@ -1720,6 +1720,12 @@ bool UnFbx::FFbxImporter::ImportAnimation(USkeleton* Skeleton, UAnimSequence * D
 		FComponentReregisterContext ReregisterContext(*Iter);
 	}
 
+	// Import bone metadata to AnimSequence
+	for (FbxNode* SkeletonNode : SortedLinks)
+	{
+		ImportNodeCustomProperties(DestSeq, SkeletonNode, true);
+	}
+
 	return true;
 }
 
