Fix NetworkRemapPath for level script actors of streaming levels in PIE replays

diff --git a/Engine/Source/Editor/UnrealEd/Private/EditorEngine.cpp b/Engine/Source/Editor/UnrealEd/Private/EditorEngine.cpp
--- a/Engine/Source/Editor/UnrealEd/Private/EditorEngine.cpp
+++ b/Engine/Source/Editor/UnrealEd/Private/EditorEngine.cpp
@@ -6911,6 +6911,24 @@ void UEditorEngine::UpdateAutoLoadProject()
 	IFileManager::Get().Delete(*AutoLoadInProgressFilename, bRequireExists, bEvenIfReadOnly, bQuiet);
 }
 
+FString NetworkRemapPath_TestLevelScriptActor(const ALevelScriptActor* LevelScriptActor, const FString& AssetName, const FString& LevelPackageName, const FString& PathName, const FString& PrefixedPathName)
+{
+	FString ResultStr;
+
+	UClass* LSAClass = LevelScriptActor ? LevelScriptActor->GetClass() : nullptr;
+
+	if (LSAClass && LSAClass->GetName() == AssetName && LSAClass->GetOutermost()->GetName() != LevelPackageName)
+	{
+		ResultStr = PathName;
+	}
+	else
+	{
+		ResultStr = PrefixedPathName;
+	}
+
+	return ResultStr;
+}
+
 FORCEINLINE bool NetworkRemapPath_local(FWorldContext& Context, FString& Str, bool bReading, bool bIsReplay)
 {
 	if (bReading)
@@ -6948,16 +6966,7 @@ FORCEINLINE bool NetworkRemapPath_local(FWorldContext& Context, FString& Str, bo
 
 			if (WorldPackageName == PrefixedPackageName)
 			{
-				const ALevelScriptActor* LSA = World->GetLevelScriptActor();
-				if (LSA && LSA->GetClass() && LSA->GetClass()->GetName() == AssetName && LSA->GetClass()->GetOutermost()->GetName() != WorldPackageName)
-				{
-					Str = Path.ToString();
-				}
-				else
-				{
-					Str = PrefixedFullName;
-				}
-				
+				Str = NetworkRemapPath_TestLevelScriptActor(World->GetLevelScriptActor(), AssetName, WorldPackageName, Path.ToString(), PrefixedFullName);
 				return true;
 			}
 
@@ -6966,9 +6975,11 @@ FORCEINLINE bool NetworkRemapPath_local(FWorldContext& Context, FString& Str, bo
 				if (StreamingLevel != nullptr)
 				{
 					const FString StreamingLevelName = StreamingLevel->GetWorldAsset().GetLongPackageName();
+					const FString LevelPackageName = StreamingLevel->GetWorldAssetPackageName();
+
 					if (StreamingLevelName == PrefixedPackageName)
 					{
-						Str = PrefixedFullName;
+						Str = NetworkRemapPath_TestLevelScriptActor(StreamingLevel->GetLevelScriptActor(), AssetName, LevelPackageName, Path.ToString(), PrefixedFullName);
 						return true;
 					}
 				}
