Sequencer: Manual check for broken tangents on import, FBX flag is misleading.

diff --git a/Engine/Source/Editor/UnrealEd/Private/Fbx/FbxCurveDataImport.cpp b/Engine/Source/Editor/UnrealEd/Private/Fbx/FbxCurveDataImport.cpp
--- a/Engine/Source/Editor/UnrealEd/Private/Fbx/FbxCurveDataImport.cpp
+++ b/Engine/Source/Editor/UnrealEd/Private/Fbx/FbxCurveDataImport.cpp
@@ -267,17 +267,22 @@ namespace UnFbx {
 					}
 					break;
 				}
-				if (KeyTangentMode & FbxAnimCurveDef::eTangentGenericBreak)
-				{
-				 	NewTangentMode = RCTM_Break;
-				}
-				else if (KeyTangentMode &  FbxAnimCurveDef::eTangentAuto) //break and auto are exclusive
+				if (KeyTangentMode &  FbxAnimCurveDef::eTangentAuto) //tangent can be auto or broken/user (
 				{
 					NewTangentMode = RCTM_Auto;
 				}
 				else
 				{
-				 	NewTangentMode = RCTM_User;
+					//no longer use KeyTangentMode & FbxAnimCurveDef::eTangentGenericBreak to see if tangent is broken since it maya broken/unify is just a manipulation state
+					//instead we manually check to see if the tangents are the same or different, if different then broken.
+					if (FMath::IsNearlyEqual(ArriveTangent, LeaveTangent))
+					{
+						NewTangentMode = RCTM_User;
+					}
+					else
+					{
+						NewTangentMode = RCTM_Break;
+					}
 				}
 
 				switch (KeyTangentWeightMode)
@@ -318,9 +323,9 @@ namespace UnFbx {
 					}
 					break;
 				}
-				RichCurve.SetKeyInterpMode(NewKeyHandle, NewInterpMode);
-				RichCurve.SetKeyTangentMode(NewKeyHandle, NewTangentMode);
-				RichCurve.SetKeyTangentWeightMode(NewKeyHandle, NewTangentWeightMode);
+				RichCurve.SetKeyInterpMode(NewKeyHandle, NewInterpMode,false);
+				RichCurve.SetKeyTangentMode(NewKeyHandle, NewTangentMode,false);
+				RichCurve.SetKeyTangentWeightMode(NewKeyHandle, NewTangentWeightMode,false);
 
 				FRichCurveKey& NewKey = RichCurve.GetKey(NewKeyHandle);
 				NewKey.ArriveTangent = ArriveTangent;
@@ -345,6 +350,8 @@ namespace UnFbx {
 				NewKey.LeaveTangentWeight = LeaveTangentWeight;
 			}
 		}
+		//no longer set tangents till end...
+		RichCurve.AutoSetTangents();
 	}
 
 	void FFbxCurvesAPI::GetCurveData(const FFbxAnimCurveHandle &CurveHandle, FRichCurve& RichCurve, bool bNegative) const
