Blueprint actor instances get named every other number

diff --git a/Engine/Source/Editor/UnrealEd/Private/Factories/ActorFactory.cpp b/Engine/Source/Editor/UnrealEd/Private/Factories/ActorFactory.cpp
--- a/Engine/Source/Editor/UnrealEd/Private/Factories/ActorFactory.cpp
+++ b/Engine/Source/Editor/UnrealEd/Private/Factories/ActorFactory.cpp
@@ -275,11 +275,8 @@ AActor* UActorFactory::CreateActor( UObject* Asset, ULevel* InLevel, FTransform
 		{
 			PostSpawnActor(Asset, NewActor);
 
-			// Only do this if the actor wasn't already given a name.
-			// We also want to skip this if a toolchain during PostSpawn
-			// already adjusted the name.
-			if (Name == NAME_None && Asset && 
-				!NewActor->GetName().StartsWith(Asset->GetName()))
+			// Only do this if the actor wasn't already given a name
+			if (Name == NAME_None && Asset)
 			{
 				FActorLabelUtilities::SetActorLabelUnique(NewActor, Asset->GetName());
 			}
