Fixed incorrect property meta-data filtering from the UProperty -> FProperty changes

#rnx

diff --git a/Engine/Source/Editor/UnrealEd/Classes/Commandlets/GatherTextFromMetadataCommandlet.h b/Engine/Source/Editor/UnrealEd/Classes/Commandlets/GatherTextFromMetadataCommandlet.h
--- a/Engine/Source/Editor/UnrealEd/Classes/Commandlets/GatherTextFromMetadataCommandlet.h
+++ b/Engine/Source/Editor/UnrealEd/Classes/Commandlets/GatherTextFromMetadataCommandlet.h
@@ -29,16 +29,67 @@ public:
 
 private:
 	void GatherTextFromUObjects(const TArray<FString>& IncludePaths, const TArray<FString>& ExcludePaths, const FGatherParameters& Arguments);
-	void GatherTextFromUObject(UField* const Field, const FGatherParameters& Arguments, const FName InPlatformName);
-	void GatherTextFromField(FField* const Field, const FGatherParameters& Arguments, const FName InPlatformName);
-	bool ShouldGatherFromField(const UField* Field);
+
+	void GatherTextFromField(UField* Field, const FGatherParameters& Arguments, const FName InPlatformName);
+	void GatherTextFromField(FField* Field, const FGatherParameters& Arguments, const FName InPlatformName);
+	
+	template <typename FieldType>
+	bool ShouldGatherFromField(const FieldType* Field, const bool bIsEditorOnly);
+
+	template <typename FieldType>
+	void GatherTextFromFieldImpl(FieldType* Field, const FGatherParameters& Arguments, const FName InPlatformName);
+
+	template <typename FieldType>
+	void EnsureFieldDisplayNameImpl(FieldType* Field, const bool bIsBool);
 
 private:
 	bool ShouldGatherFromEditorOnlyData;
 
-	/** Array of field types (eg, UProperty, UFunction, UScriptStruct, etc) that should be included or excluded in the current gather */
-	TArray<const UClass*> FieldTypesToInclude;
-	TArray<const UClass*> FieldTypesToExclude;
+	struct FFieldClassFilter
+	{
+		explicit FFieldClassFilter(const FFieldClass* InFieldClass)
+			: FieldClass(InFieldClass)
+			, ObjectClass(nullptr)
+		{
+		}
+
+		explicit FFieldClassFilter(const UClass* InObjectClass)
+			: FieldClass(nullptr)
+			, ObjectClass(InObjectClass)
+		{
+		}
+
+		FString GetName() const
+		{
+			if (FieldClass)
+			{
+				return FieldClass->GetName();
+			}
+			if (ObjectClass)
+			{
+				return ObjectClass->GetName();
+			}
+			return FString();
+		}
+
+		bool TestClass(const FFieldClass* InFieldClass) const
+		{
+			return FieldClass && InFieldClass->IsChildOf(FieldClass);
+		}
+
+		bool TestClass(const UClass* InObjectClass) const
+		{
+			return ObjectClass && InObjectClass->IsChildOf(ObjectClass);
+		}
+
+	private:
+		const FFieldClass* FieldClass;
+		const UClass* ObjectClass;
+	};
+
+	/** Array of field types (eg, FProperty, UFunction, UScriptStruct, etc) that should be included or excluded in the current gather */
+	TArray<FFieldClassFilter> FieldTypesToInclude;
+	TArray<FFieldClassFilter> FieldTypesToExclude;
 
 	/** Array of field owner types (eg, UMyClass, FMyStruct, etc) that should have fields within them included or excluded in the current gather */
 	TArray<const UStruct*> FieldOwnerTypesToInclude;
diff --git a/Engine/Source/Editor/UnrealEd/Private/Commandlets/GatherTextFromMetadataCommandlet.cpp b/Engine/Source/Editor/UnrealEd/Private/Commandlets/GatherTextFromMetadataCommandlet.cpp
--- a/Engine/Source/Editor/UnrealEd/Private/Commandlets/GatherTextFromMetadataCommandlet.cpp
+++ b/Engine/Source/Editor/UnrealEd/Private/Commandlets/GatherTextFromMetadataCommandlet.cpp
@@ -110,40 +110,68 @@ int32 UGatherTextFromMetaDataCommandlet::Main( const FString& Params )
 
 	// FieldTypesToInclude/FieldTypesToExclude
 	{
-		auto GetFieldTypesArrayFromConfig = [this, &SectionName, &GatherTextConfigPath](const TCHAR* InConfigKey, TArray<const UClass*>& OutFieldTypes)
+		auto GetFieldTypesArrayFromConfig = [this, &SectionName, &GatherTextConfigPath](const TCHAR* InConfigKey, TArray<FFieldClassFilter>& OutFieldTypes)
 		{
 			TArray<FString> FieldTypeStrs;
 			GetStringArrayFromConfig(*SectionName, InConfigKey, FieldTypeStrs, GatherTextConfigPath);
 
-			TArray<const UClass*> AllFieldTypes;
-			AllFieldTypes.Add(UField::StaticClass());
-			GetDerivedClasses(UField::StaticClass(), (TArray<UClass*>&)AllFieldTypes);
+			if (FieldTypeStrs.Num() == 0)
+			{
+				return;
+			}
+
+			TArray<FFieldClassFilter> AllFieldTypes;
+			// FField types
+			{
+				for (const FFieldClass* FieldClass : FFieldClass::GetAllFieldClasses())
+				{
+					AllFieldTypes.Emplace(FieldClass);
+				}
+			}
+			// UField types
+			{
+				TArray<UClass*> AllUFieldClasses;
+				AllUFieldClasses.Add(UField::StaticClass());
+				GetDerivedClasses(UField::StaticClass(), AllUFieldClasses);
+				for (const UClass* FieldClass : AllUFieldClasses)
+				{
+					AllFieldTypes.Emplace(FieldClass);
+				}
+			}
 
 			for (const FString& FieldTypeStr : FieldTypeStrs)
 			{
 				const bool bIsWildcard = FieldTypeStr.GetCharArray().Contains(TEXT('*')) || FieldTypeStr.GetCharArray().Contains(TEXT('?'));
 				if (bIsWildcard)
 				{
-					for (const UClass* FieldType : AllFieldTypes)
+					for (const FFieldClassFilter& FieldTypeFilter : AllFieldTypes)
 					{
-						if (FieldType->GetName().MatchesWildcard(FieldTypeStr))
+						if (FieldTypeFilter.GetName().MatchesWildcard(FieldTypeStr))
 						{
-							OutFieldTypes.Add(FieldType);
+							OutFieldTypes.Emplace(FieldTypeFilter);
 						}
 					}
 				}
 				else
 				{
-					const UClass* FieldType = FindObject<UClass>(ANY_PACKAGE, *FieldTypeStr);
-					if (!FieldType)
+					const FFieldClass* FieldClass = FFieldClass::GetNameToFieldClassMap().FindRef(*FieldTypeStr);
+					const UClass* UFieldClass = FindObject<UClass>(ANY_PACKAGE, *FieldTypeStr);
+					if (!FieldClass && !UFieldClass)
 					{
 						UE_LOG(LogGatherTextFromMetaDataCommandlet, Warning, TEXT("Field Type %s was not found (from %s in section %s). Did you forget a ModulesToPreload entry?"), *FieldTypeStr, InConfigKey, *SectionName);
 						continue;
 					}
 
-					check(FieldType->IsChildOf<UField>());
-					OutFieldTypes.Add(FieldType);
-					GetDerivedClasses(FieldType, (TArray<UClass*>&)OutFieldTypes);
+					if (FieldClass)
+					{
+						OutFieldTypes.Emplace(FieldClass);
+					}
+
+					if (UFieldClass)
+					{
+						check(UFieldClass->IsChildOf<UField>());
+						OutFieldTypes.Emplace(UFieldClass);
+					}
 				}
 			}
 		};
@@ -232,16 +260,22 @@ void UGatherTextFromMetaDataCommandlet::GatherTextFromUObjects(const TArray<FStr
 {
 	const FFuzzyPathMatcher FuzzyPathMatcher = FFuzzyPathMatcher(IncludePaths, ExcludePaths);
 
-	for(TObjectIterator<UField> It; It; ++It)
+	for (TObjectIterator<UField> It; It; ++It)
 	{
-		// Skip fields excluded by our filter
-		if (!ShouldGatherFromField(*It))
+		UField* Field = *It;
+
+		const UPackage* FieldPackage = Field->GetOutermost();
+		const FString FieldPackageName = FieldPackage->GetName();
+		if (!FPackageName::IsScriptPackage(FieldPackageName) || FPackageName::IsMemoryPackage(FieldPackageName) || FPackageName::IsTempPackage(FieldPackageName))
 		{
 			continue;
 		}
 
 		FString SourceFilePath;
-		FSourceCodeNavigation::FindClassHeaderPath(*It, SourceFilePath);
+		if (!FSourceCodeNavigation::FindModulePath(FieldPackage, SourceFilePath))
+		{
+			continue;
+		}
 		SourceFilePath = FPaths::ConvertRelativePathToFull(SourceFilePath);
 		check(!SourceFilePath.IsEmpty());
 
@@ -252,71 +286,31 @@ void UGatherTextFromMetaDataCommandlet::GatherTextFromUObjects(const TArray<FStr
 		}
 
 		const FName MetaDataPlatformName = GetSplitPlatformNameFromPath(SourceFilePath);
-
-		UStruct* Struct = Cast<UStruct>(*It);
-		if (Struct)
-		{
-			for (TFieldIterator<FField> FieldIt(Struct); FieldIt; ++FieldIt)
-			{
-				// Skip editor-only properties if we're not gathering for editor-only data.
-				FProperty* Property = CastField<FProperty>(*FieldIt);
-				if (Property && !ShouldGatherFromEditorOnlyData && Property->HasAnyPropertyFlags(CPF_EditorOnly))
-				{
-					continue;
-				}
-				GatherTextFromField(*FieldIt, Arguments, MetaDataPlatformName);
-			}
-		}
-
-		GatherTextFromUObject(*It, Arguments, MetaDataPlatformName);
+		GatherTextFromField(Field, Arguments, MetaDataPlatformName);
 	}
 }
 
-template <typename TFieldType>
-void GatherTextFromFieldImplementation(TFieldType* const Field, const UGatherTextFromMetaDataCommandlet::FGatherParameters& Arguments, const FName InPlatformName, FLocTextHelper* GatherManifestHelper)
+void UGatherTextFromMetaDataCommandlet::GatherTextFromField(UField* Field, const FGatherParameters& Arguments, const FName InPlatformName)
 {
-	for (int32 i = 0; i < Arguments.InputKeys.Num(); ++i)
+	// For structs, also gather the new non-object field values.
+	if (UStruct* Struct = Cast<UStruct>(Field))
+	{
+		for (TFieldIterator<FField> FieldIt(Struct); FieldIt; ++FieldIt)
 		{
-			FFormatNamedArguments PatternArguments;
-		PatternArguments.Add(TEXT("FieldPath"), FText::FromString(Field->GetFullGroupName(false)));
-
-		if (Field->HasMetaData(*Arguments.InputKeys[i]))
-			{
-				const FString& MetaDataValue = Field->GetMetaData(*Arguments.InputKeys[i]);
-			if (!MetaDataValue.IsEmpty())
-				{
-				PatternArguments.Add(TEXT("MetaDataValue"), FText::FromString(MetaDataValue));
-
-				const UStruct* FieldOwnerType = Field->GetOwnerStruct();
-					const FString Namespace = Arguments.OutputNamespaces[i];
-					FLocItem LocItem(MetaDataValue);
-					FManifestContext Context;
-					Context.Key = FText::Format(Arguments.OutputKeys[i], PatternArguments).ToString();
-				Context.SourceLocation = FString::Printf(TEXT("From metadata for key %s of member %s in %s (type: %s, owner: %s)"), *Arguments.InputKeys[i], *Field->GetName(), *Field->GetFullGroupName(true), *Field->GetClass()->GetName(), FieldOwnerType ? *FieldOwnerType->GetName() : TEXT("<null>"));
-					Context.PlatformName = InPlatformName;
-					GatherManifestHelper->AddSourceText(Namespace, LocItem, Context);
-				}
-			}
+			GatherTextFromField(*FieldIt, Arguments, InPlatformName);
 		}
-}
-
+	}
 
-void UGatherTextFromMetaDataCommandlet::GatherTextFromUObject(UField* const Field, const FGatherParameters& Arguments, const FName InPlatformName)
-{
-	// Gather for object.
+	if (ShouldGatherFromField(Field, false))
 	{
-		if (!Field->HasMetaData(TEXT("DisplayName")))
+		// Gather for object.
 		{
-			Field->SetMetaData(TEXT("DisplayName"), *FName::NameToDisplayString(Field->GetName(), false));
+			EnsureFieldDisplayNameImpl(Field, false);
+			GatherTextFromFieldImpl(Field, Arguments, InPlatformName);
 		}
 
-		GatherTextFromFieldImplementation<UField>(Field, Arguments, InPlatformName, GatherManifestHelper.Get());
-	}
-
-	// For enums, also gather for enum values.
-	{
-		UEnum* Enum = Cast<UEnum>(Field);
-		if (Enum)
+		// For enums, also gather for enum values.
+		if (UEnum* Enum = Cast<UEnum>(Field))
 		{
 			const int32 ValueCount = Enum->NumEnums();
 			for (int32 i = 0; i < ValueCount; ++i)
@@ -353,7 +347,18 @@ void UGatherTextFromMetaDataCommandlet::GatherTextFromUObject(UField* const Fiel
 	}
 }
 
-bool UGatherTextFromMetaDataCommandlet::ShouldGatherFromField(const UField* Field)
+void UGatherTextFromMetaDataCommandlet::GatherTextFromField(FField* Field, const FGatherParameters& Arguments, const FName InPlatformName)
+{
+	FProperty* Property = CastField<FProperty>(Field);
+	if (ShouldGatherFromField(Field, Property && Property->HasAnyPropertyFlags(CPF_EditorOnly)))
+	{
+		EnsureFieldDisplayNameImpl(Field, Field->IsA(FBoolProperty::StaticClass()));
+		GatherTextFromFieldImpl(Field, Arguments, InPlatformName);
+	}
+}
+
+template <typename FieldType>
+bool UGatherTextFromMetaDataCommandlet::ShouldGatherFromField(const FieldType* Field, const bool bIsEditorOnly)
 {
 	auto ShouldGatherFieldByType = [this, Field]()
 	{
@@ -362,9 +367,21 @@ bool UGatherTextFromMetaDataCommandlet::ShouldGatherFromField(const UField* Fiel
 			return true;
 		}
 
-		const UClass* FieldClass = Field->GetClass();
-		return (FieldTypesToInclude.Num() == 0 || FieldTypesToInclude.Contains(FieldClass))
-			&& (FieldTypesToExclude.Num() == 0 || !FieldTypesToExclude.Contains(FieldClass));
+		const auto* FieldClass = Field->GetClass();
+		auto TestClassFilter = [FieldClass](const TArray<FFieldClassFilter>& InFieldTypeFilters)
+		{
+			for (const FFieldClassFilter& FieldTypeFilter : InFieldTypeFilters)
+			{
+				if (FieldTypeFilter.TestClass(FieldClass))
+				{
+					return true;
+				}
+			}
+			return false;
+		};
+		
+		return (FieldTypesToInclude.Num() == 0 || TestClassFilter(FieldTypesToInclude))
+			&& (FieldTypesToExclude.Num() == 0 || !TestClassFilter(FieldTypesToExclude));
 	};
 
 	auto ShouldGatherFieldByOwnerType = [this, Field]()
@@ -385,15 +402,42 @@ bool UGatherTextFromMetaDataCommandlet::ShouldGatherFromField(const UField* Fiel
 		return true;
 	};
 
-	return ShouldGatherFieldByType() && ShouldGatherFieldByOwnerType();
+	return (!bIsEditorOnly || ShouldGatherFromEditorOnlyData) && ShouldGatherFieldByType() && ShouldGatherFieldByOwnerType();
 }
 
-void UGatherTextFromMetaDataCommandlet::GatherTextFromField(FField* const Field, const FGatherParameters& Arguments, const FName InPlatformName)
+template <typename FieldType>
+void UGatherTextFromMetaDataCommandlet::GatherTextFromFieldImpl(FieldType* Field, const FGatherParameters& Arguments, const FName InPlatformName)
 {
-	if (!Field->HasMetaData(TEXT("DisplayName")))
+	for (int32 i = 0; i < Arguments.InputKeys.Num(); ++i)
 	{
-		Field->SetMetaData(TEXT("DisplayName"), *FName::NameToDisplayString(Field->GetName(), Field->IsA(FBoolProperty::StaticClass())));
+		FFormatNamedArguments PatternArguments;
+		PatternArguments.Add(TEXT("FieldPath"), FText::FromString(Field->GetFullGroupName(false)));
+
+		if (Field->HasMetaData(*Arguments.InputKeys[i]))
+		{
+			const FString& MetaDataValue = Field->GetMetaData(*Arguments.InputKeys[i]);
+			if (!MetaDataValue.IsEmpty())
+			{
+				PatternArguments.Add(TEXT("MetaDataValue"), FText::FromString(MetaDataValue));
+
+				const UStruct* FieldOwnerType = Field->GetOwnerStruct();
+				const FString Namespace = Arguments.OutputNamespaces[i];
+				FLocItem LocItem(MetaDataValue);
+				FManifestContext Context;
+				Context.Key = FText::Format(Arguments.OutputKeys[i], PatternArguments).ToString();
+				Context.SourceLocation = FString::Printf(TEXT("From metadata for key %s of member %s in %s (type: %s, owner: %s)"), *Arguments.InputKeys[i], *Field->GetName(), *Field->GetFullGroupName(true), *Field->GetClass()->GetName(), FieldOwnerType ? *FieldOwnerType->GetName() : TEXT("<null>"));
+				Context.PlatformName = InPlatformName;
+				GatherManifestHelper->AddSourceText(Namespace, LocItem, Context);
+			}
+		}
 	}
+}
 
-	GatherTextFromFieldImplementation<FField>(Field, Arguments, InPlatformName, GatherManifestHelper.Get());
+template <typename FieldType>
+void UGatherTextFromMetaDataCommandlet::EnsureFieldDisplayNameImpl(FieldType* Field, const bool bIsBool)
+{
+	if (!Field->HasMetaData(TEXT("DisplayName")))
+	{
+		Field->SetMetaData(TEXT("DisplayName"), *FName::NameToDisplayString(Field->GetName(), bIsBool));
+	}
 }
