Move GetTickableObjects() and GetPendingTickableObjects() definitions to cpp to ensure there's only a single instantiation.

diff --git a/Engine/Source/Editor/UnrealEd/Private/TickableEditorObject.cpp b/Engine/Source/Editor/UnrealEd/Private/TickableEditorObject.cpp
new file mode 100644
--- /dev/null
+++ b/Engine/Source/Editor/UnrealEd/Private/TickableEditorObject.cpp
@@ -0,0 +1,16 @@
+// Copyright Epic Games, Inc. All Rights Reserved.
+
+#include "TickableEditorObject.h"
+#include "Tickable.h"
+
+TArray<FTickableObjectBase::FTickableObjectEntry>& FTickableEditorObject::GetTickableObjects()
+{
+	static TTickableObjectsCollection TickableObjects;
+	return TickableObjects;
+}
+
+TArray<FTickableEditorObject*>& FTickableEditorObject::GetPendingTickableObjects()
+{
+	static TArray<FTickableEditorObject*> PendingTickableObjects;
+	return PendingTickableObjects;
+}
diff --git a/Engine/Source/Editor/UnrealEd/Public/TickableEditorObject.h b/Engine/Source/Editor/UnrealEd/Public/TickableEditorObject.h
--- a/Engine/Source/Editor/UnrealEd/Public/TickableEditorObject.h
+++ b/Engine/Source/Editor/UnrealEd/Public/TickableEditorObject.h
@@ -114,15 +114,7 @@ private:
 	/** True if we are in the Tick function for an editor tickable object */
 	static bool bIsInObjectsTick;
 
-	static TArray<FTickableObjectEntry>& GetTickableObjects()
-	{
-		static TTickableObjectsCollection TickableObjects;
-		return TickableObjects;
-	}
+	static TArray<FTickableObjectEntry>& GetTickableObjects();
 
-	static TArray<FTickableEditorObject*>& GetPendingTickableObjects()
-	{
-		static TArray<FTickableEditorObject*> PendingTickableObjects;
-		return PendingTickableObjects;
-	}
+	static TArray<FTickableEditorObject*>& GetPendingTickableObjects();
 };
