Restore fix for launching projects onto device not triggering build from 4.24 (CL 10262048) that was deleted during merge to Dev-Editor (CL 10384788).

diff --git a/Engine/Source/Editor/UnrealEd/Private/PlayLevelViaLauncher.cpp b/Engine/Source/Editor/UnrealEd/Private/PlayLevelViaLauncher.cpp
--- a/Engine/Source/Editor/UnrealEd/Private/PlayLevelViaLauncher.cpp
+++ b/Engine/Source/Editor/UnrealEd/Private/PlayLevelViaLauncher.cpp
@@ -126,15 +126,34 @@ void UEditorEngine::StartPlayUsingLauncherSession(FRequestPlaySessionParams& InR
 		}
 	}
 
+	// set the build/launch configuration 
+	EBuildConfiguration BuildConfiguration;
+	const ULevelEditorPlaySettings* EditorPlaySettings = PlaySessionRequest->EditorPlaySettings;
+	switch (EditorPlaySettings->LaunchConfiguration)
+	{
+	case LaunchConfig_Debug:
+		BuildConfiguration = EBuildConfiguration::Debug;
+		break;
+	case LaunchConfig_Development:
+		BuildConfiguration = EBuildConfiguration::Development;
+		break;
+	case LaunchConfig_Test:
+		BuildConfiguration = EBuildConfiguration::Test;
+		break;
+	case LaunchConfig_Shipping:
+		BuildConfiguration = EBuildConfiguration::Shipping;
+		break;
+	default:
+		// same as the running editor
+		BuildConfiguration = FApp::GetBuildConfiguration();
+		break;
+	}
+
 	// does the project have any code?
 	FGameProjectGenerationModule& GameProjectModule = FModuleManager::LoadModuleChecked<FGameProjectGenerationModule>(TEXT("GameProjectGeneration"));
 	LauncherSessionInfo->bPlayUsingLauncherHasCode = GameProjectModule.Get().ProjectHasCodeFiles();
 
-	const ULevelEditorPlaySettings* EditorPlaySettings = PlaySessionRequest->EditorPlaySettings;
-	check(EditorPlaySettings);
-
-	// Setup launch profile, keep the setting here to a minimum.
-	ILauncherProfileRef LauncherProfile = LauncherServicesModule.CreateProfile(TEXT("Launch On Device"));
+	// Figure out if we need to build anything
 	if (EditorPlaySettings->BuildGameBeforeLaunch == EPlayOnBuildMode::PlayOnBuild_Always)
 	{
 		LauncherSessionInfo->bPlayUsingLauncherBuild = true;
@@ -146,33 +165,25 @@ void UEditorEngine::StartPlayUsingLauncherSession(FRequestPlaySessionParams& InR
 	else if (EditorPlaySettings->BuildGameBeforeLaunch == EPlayOnBuildMode::PlayOnBuild_Default)
 	{
 		LauncherSessionInfo->bPlayUsingLauncherBuild = LauncherSessionInfo->bPlayUsingLauncherHasCode || !FApp::GetEngineIsPromotedBuild();
+		if (!LauncherSessionInfo->bPlayUsingLauncherBuild)
+		{
+			FText Reason;
+			if (LaunchPlatform->RequiresTempTarget(LauncherSessionInfo->bPlayUsingLauncherHasCode, BuildConfiguration, false, Reason))
+			{
+				UE_LOG(LogPlayLevel, Log, TEXT("Project requires temp target (%s)"), *Reason.ToString());
+				LauncherSessionInfo->bPlayUsingLauncherBuild = true;
+			}
+		}
 	}
 	else if (EditorPlaySettings->BuildGameBeforeLaunch == EPlayOnBuildMode::PlayOnBuild_IfEditorBuiltLocally)
 	{
 		LauncherSessionInfo->bPlayUsingLauncherBuild = !FApp::GetEngineIsPromotedBuild();
 	}
-	LauncherProfile->SetBuildGame(LauncherSessionInfo->bPlayUsingLauncherBuild);
 
-	// set the build/launch configuration 
-	switch (EditorPlaySettings->LaunchConfiguration)
-	{
-	case LaunchConfig_Debug:
-		LauncherProfile->SetBuildConfiguration(EBuildConfiguration::Debug);
-		break;
-	case LaunchConfig_Development:
-		LauncherProfile->SetBuildConfiguration(EBuildConfiguration::Development);
-		break;
-	case LaunchConfig_Test:
-		LauncherProfile->SetBuildConfiguration(EBuildConfiguration::Test);
-		break;
-	case LaunchConfig_Shipping:
-		LauncherProfile->SetBuildConfiguration(EBuildConfiguration::Shipping);
-		break;
-	default:
-		// same as the running editor
-		LauncherProfile->SetBuildConfiguration(FApp::GetBuildConfiguration());
-		break;
-	}
+	// Setup launch profile, keep the setting here to a minimum.
+	ILauncherProfileRef LauncherProfile = LauncherServicesModule.CreateProfile(TEXT("Launch On Device"));
+	LauncherProfile->SetBuildGame(LauncherSessionInfo->bPlayUsingLauncherBuild);
+	LauncherProfile->SetBuildConfiguration(BuildConfiguration);
 
 	// select the quickest cook mode based on which in editor cook mode is enabled
 	bool bIncrimentalCooking = true;
