Fix crash when creating a blueprint from selection via child actors with multiple actors selected if you undo and redo the creation.
#rnx

diff --git a/Engine/Source/Editor/UnrealEd/Private/Kismet2/Kismet2.cpp b/Engine/Source/Editor/UnrealEd/Private/Kismet2/Kismet2.cpp
--- a/Engine/Source/Editor/UnrealEd/Private/Kismet2/Kismet2.cpp
+++ b/Engine/Source/Editor/UnrealEd/Private/Kismet2/Kismet2.cpp
@@ -1623,6 +1623,13 @@ UBlueprint* FKismetEditorUtilities::CreateBlueprintFromActors(const FName Bluepr
 		}
 
 		FKismetEditorUtilities::AddComponentsToBlueprint(AssemblyProps.Blueprint, ChildActorComponents, /*bHarvesting=*/ true, AssemblyProps.RootNodeOverride);
+
+		// Since the names we create are well defined relative to the SCS but created in the transient package, we could end up reusing objects
+		// unless we rename these temporary components out of the way
+		for (UActorComponent* CAC : ChildActorComponents)
+		{
+			CAC->Rename(nullptr, nullptr, REN_DoNotDirty | REN_DontCreateRedirectors | REN_ForceNoResetLoaders);
+		}
 	};
 	
 	UBlueprint* Blueprint = nullptr;
