Correctly clear out the transform from the child actor template when converting an actor to blueprint via child actors
#rnx

diff --git a/Engine/Source/Editor/UnrealEd/Private/Kismet2/Kismet2.cpp b/Engine/Source/Editor/UnrealEd/Private/Kismet2/Kismet2.cpp
--- a/Engine/Source/Editor/UnrealEd/Private/Kismet2/Kismet2.cpp
+++ b/Engine/Source/Editor/UnrealEd/Private/Kismet2/Kismet2.cpp
@@ -1598,7 +1598,12 @@ UBlueprint* FKismetEditorUtilities::CreateBlueprintFromActors(const FName Bluepr
 			CAC->SetChildActorClass(Actor->GetClass(), Actor);
 
 			// Clear any properties that can't be on the template
-			CAC->GetChildActorTemplate()->SetActorRelativeTransform(FTransform::Identity);
+			if (USceneComponent* RootComponent = CAC->GetChildActorTemplate()->GetRootComponent())
+			{
+				RootComponent->SetRelativeLocation_Direct(FVector::ZeroVector);
+				RootComponent->SetRelativeRotation_Direct(FRotator::ZeroRotator);
+				RootComponent->SetRelativeScale3D_Direct(FVector::OneVector);
+			}
 
 			CAC->SetWorldTransform(Actor->GetTransform());
 			if (ParentComponent)
