Fix tooltip comment for ULevelEditorPlaySettings::AdditionalServerGameOptions.

diff --git a/Engine/Source/Editor/UnrealEd/Classes/Settings/LevelEditorPlaySettings.h b/Engine/Source/Editor/UnrealEd/Classes/Settings/LevelEditorPlaySettings.h
--- a/Engine/Source/Editor/UnrealEd/Classes/Settings/LevelEditorPlaySettings.h
+++ b/Engine/Source/Editor/UnrealEd/Classes/Settings/LevelEditorPlaySettings.h
@@ -356,7 +356,7 @@ private:
 	UPROPERTY(config, EditAnywhere, Category = "Multiplayer Options|Server", meta=(EditCondition = "PlayNetMode != EPlayNetMode::PIE_Standalone || bLaunchSeparateServer"))
 	FString ServerMapNameOverride;
 
-	/** Additional options that will be passed to the server as URL parameters, in the format ?bIsLanMatch=1?listen - any additional command line switches should be passed in the Command Line Arguments field below. */
+	/** Additional options that will be passed to the server as URL parameters, in the format ?bIsLanMatch=1?listen - any additional command line switches should be passed in the Additional Server Launch Parameters field below. */
 	UPROPERTY(config, EditAnywhere, Category="Multiplayer Options|Server", meta=(EditCondition = "PlayNetMode != EPlayNetMode::PIE_Standalone || bLaunchSeparateServer"))
 	FString AdditionalServerGameOptions;
 
