Make sure FbxImportUI options do not change the CDO when we close and open the FBX dialog

diff --git a/Engine/Source/Editor/UnrealEd/Classes/Factories/FbxAssetImportData.h b/Engine/Source/Editor/UnrealEd/Classes/Factories/FbxAssetImportData.h
--- a/Engine/Source/Editor/UnrealEd/Classes/Factories/FbxAssetImportData.h
+++ b/Engine/Source/Editor/UnrealEd/Classes/Factories/FbxAssetImportData.h
@@ -45,11 +45,4 @@ class UNREALED_API UFbxAssetImportData : public UAssetImportData
 	/* Use by the reimport factory to answer CanReimport, if true only factory for scene reimport will return true */
 	UPROPERTY()
 	UFbxSceneImportData* FbxSceneImportDataReference;
-
-	/* We want to change the last dialog state but not the CDO, so we cannot call UObject::SaveConfig here.
-	 * The problem is the class is use to store the import settings with the imported asset and the serialization
-	 * is serializing the diff with the CDO. Data will be lost if the CDO change.
-	*/
-	virtual void SaveOptions();
-	virtual void LoadOptions();
 };
diff --git a/Engine/Source/Editor/UnrealEd/Classes/Factories/FbxImportUI.h b/Engine/Source/Editor/UnrealEd/Classes/Factories/FbxImportUI.h
--- a/Engine/Source/Editor/UnrealEd/Classes/Factories/FbxImportUI.h
+++ b/Engine/Source/Editor/UnrealEd/Classes/Factories/FbxImportUI.h
@@ -228,6 +228,13 @@ public:
 		MeshTypeToImport = bImportAsSkeletal ? FBXIT_SkeletalMesh : FBXIT_StaticMesh;
 	}
 
+	/* We want to change the last dialog state but not the CDO, so we cannot call UObject::SaveConfig here.
+	 * The problem is the class is use to store the import settings with the imported asset and the serialization
+	 * is serializing the diff with the CDO. Data will be lost if the CDO change.
+	*/
+	static void SaveOptions(UObject* ObjectToSaveOptions);
+	static void LoadOptions(UObject* ObjectToLoadOptions);
+
 	/* Whether this UI is construct for a reimport */
 	bool bIsReimport;
 	
diff --git a/Engine/Source/Editor/UnrealEd/Private/Fbx/FbxAssetImportData.cpp b/Engine/Source/Editor/UnrealEd/Private/Fbx/FbxAssetImportData.cpp
--- a/Engine/Source/Editor/UnrealEd/Private/Fbx/FbxAssetImportData.cpp
+++ b/Engine/Source/Editor/UnrealEd/Private/Fbx/FbxAssetImportData.cpp
@@ -26,149 +26,3 @@ UFbxAssetImportData::UFbxAssetImportData(const FObjectInitializer& ObjectInitial
 {
 	
 }
-
-/** Load UI settings from ini file */
-void UFbxAssetImportData::LoadOptions()
-{
-	int32 PortFlags = 0;
-
-	for (FProperty* Property = GetClass()->PropertyLink; Property; Property = Property->PropertyLinkNext)
-	{
-		if (!Property->HasAnyPropertyFlags(CPF_Config))
-		{
-			continue;
-		}
-		FString Section = TEXT("FBX_Import_UI_Option_") + GetClass()->GetName();
-		FString Key = Property->GetName();
-
-		const bool bIsPropertyInherited = Property->GetOwnerClass() != GetClass();
-		UObject* SuperClassDefaultObject = GetClass()->GetSuperClass()->GetDefaultObject();
-
-		const FString& PropFileName = GEditorPerProjectIni;
-
-		FArrayProperty* Array = CastField<FArrayProperty>(Property);
-		if (Array)
-		{
-			FConfigSection* Sec = GConfig->GetSectionPrivate(*Section, 0, 1, *GEditorPerProjectIni);
-			if (Sec != nullptr)
-			{
-				TArray<FConfigValue> List;
-				const FName KeyName(*Key, FNAME_Find);
-				Sec->MultiFind(KeyName, List);
-
-				FScriptArrayHelper_InContainer ArrayHelper(Array, this);
-				// Only override default properties if there is something to override them with.
-				if (List.Num() > 0)
-				{
-					ArrayHelper.EmptyAndAddValues(List.Num());
-					for (int32 i = List.Num() - 1, c = 0; i >= 0; i--, c++)
-					{
-						Array->Inner->ImportText(*List[i].GetValue(), ArrayHelper.GetRawPtr(c), PortFlags, this);
-					}
-				}
-				else
-				{
-					int32 Index = 0;
-					const FConfigValue* ElementValue = nullptr;
-					do
-					{
-						// Add array index number to end of key
-						FString IndexedKey = FString::Printf(TEXT("%s[%i]"), *Key, Index);
-
-						// Try to find value of key
-						const FName IndexedName(*IndexedKey, FNAME_Find);
-						if (IndexedName == NAME_None)
-						{
-							break;
-						}
-						ElementValue = Sec->Find(IndexedName);
-
-						// If found, import the element
-						if (ElementValue != nullptr)
-						{
-							// expand the array if necessary so that Index is a valid element
-							ArrayHelper.ExpandForIndex(Index);
-							Array->Inner->ImportText(*ElementValue->GetValue(), ArrayHelper.GetRawPtr(Index), PortFlags, this);
-						}
-
-						Index++;
-					} while (ElementValue || Index < ArrayHelper.Num());
-				}
-			}
-		}
-		else
-		{
-			for (int32 i = 0; i < Property->ArrayDim; i++)
-			{
-				if (Property->ArrayDim != 1)
-				{
-					Key = FString::Printf(TEXT("%s[%i]"), *Property->GetName(), i);
-				}
-
-				FString Value;
-				bool bFoundValue = GConfig->GetString(*Section, *Key, Value, *GEditorPerProjectIni);
-
-				if (bFoundValue)
-				{
-					if (Property->ImportText(*Value, Property->ContainerPtrToValuePtr<uint8>(this, i), PortFlags, this) == NULL)
-					{
-						// this should be an error as the properties from the .ini / .int file are not correctly being read in and probably are affecting things in subtle ways
-						UE_LOG(LogFbx, Error, TEXT("FBX Options LoadOptions (%s): import failed for %s in: %s"), *GetPathName(), *Property->GetName(), *Value);
-					}
-				}
-			}
-		}
-	}
-}
-
-/** Save UI settings to ini file */
-void UFbxAssetImportData::SaveOptions()
-{
-	int32 PortFlags = 0;
-
-	for (FProperty* Property = GetClass()->PropertyLink; Property; Property = Property->PropertyLinkNext)
-	{
-		if (!Property->HasAnyPropertyFlags(CPF_Config))
-		{
-			continue;
-		}
-		FString Section = TEXT("FBX_Import_UI_Option_") + GetClass()->GetName();
-		FString Key = Property->GetName();
-		
-		const bool bIsPropertyInherited = Property->GetOwnerClass() != GetClass();
-		UObject* SuperClassDefaultObject = GetClass()->GetSuperClass()->GetDefaultObject();
-
-		FArrayProperty* Array = CastField<FArrayProperty>(Property);
-		if (Array)
-		{
-			FConfigSection* Sec = GConfig->GetSectionPrivate(*Section, 1, 0, *GEditorPerProjectIni);
-			check(Sec);
-			Sec->Remove(*Key);
-
-			FScriptArrayHelper_InContainer ArrayHelper(Array, this);
-			for (int32 i = 0; i < ArrayHelper.Num(); i++)
-			{
-				FString	Buffer;
-				Array->Inner->ExportTextItem(Buffer, ArrayHelper.GetRawPtr(i), ArrayHelper.GetRawPtr(i), this, PortFlags);
-				Sec->Add(*Key, *Buffer);
-			}
-		}
-		else
-		{
-			TCHAR TempKey[MAX_SPRINTF] = TEXT("");
-			for (int32 Index = 0; Index < Property->ArrayDim; Index++)
-			{
-				if (Property->ArrayDim != 1)
-				{
-					FCString::Sprintf(TempKey, TEXT("%s[%i]"), *Property->GetName(), Index);
-					Key = TempKey;
-				}
-
-				FString	Value;
-				Property->ExportText_InContainer(Index, Value, this, this, this, PortFlags);
-				GConfig->SetString(*Section, *Key, *Value, *GEditorPerProjectIni);
-			}
-		}
-	}
-	GConfig->Flush(0);
-}
diff --git a/Engine/Source/Editor/UnrealEd/Private/Fbx/FbxFactory.cpp b/Engine/Source/Editor/UnrealEd/Private/Fbx/FbxFactory.cpp
--- a/Engine/Source/Editor/UnrealEd/Private/Fbx/FbxFactory.cpp
+++ b/Engine/Source/Editor/UnrealEd/Private/Fbx/FbxFactory.cpp
@@ -54,6 +54,7 @@ void UFbxFactory::PostInitProperties()
 	bText = false;
 
 	ImportUI = NewObject<UFbxImportUI>(this, NAME_None, RF_NoFlags);
+	UFbxImportUI::LoadOptions(ImportUI);
 }
 
 
@@ -987,19 +988,19 @@ UFbxImportUI::UFbxImportUI(const FObjectInitializer& ObjectInitializer)
 	
 	StaticMeshImportData = CreateDefaultSubobject<UFbxStaticMeshImportData>(TEXT("StaticMeshImportData"));
 	StaticMeshImportData->SetFlags(RF_Transactional);
-	StaticMeshImportData->LoadOptions();
+	UFbxImportUI::LoadOptions(StaticMeshImportData);
 	
 	SkeletalMeshImportData = CreateDefaultSubobject<UFbxSkeletalMeshImportData>(TEXT("SkeletalMeshImportData"));
 	SkeletalMeshImportData->SetFlags(RF_Transactional);
-	SkeletalMeshImportData->LoadOptions();
+	UFbxImportUI::LoadOptions(SkeletalMeshImportData);
 	
 	AnimSequenceImportData = CreateDefaultSubobject<UFbxAnimSequenceImportData>(TEXT("AnimSequenceImportData"));
 	AnimSequenceImportData->SetFlags(RF_Transactional);
-	AnimSequenceImportData->LoadOptions();
+	UFbxImportUI::LoadOptions(AnimSequenceImportData);
 	
 	TextureImportData = CreateDefaultSubobject<UFbxTextureImportData>(TEXT("TextureImportData"));
 	TextureImportData->SetFlags(RF_Transactional);
-	TextureImportData->LoadOptions();
+	UFbxImportUI::LoadOptions(TextureImportData);
 }
 
 
@@ -1625,6 +1626,152 @@ void UFbxImportUI::UpdateCompareData(UnFbx::FFbxImporter* FbxImporter)
 	FbxImporter->PartialCleanUp();
 }
 
+/** Load UI settings from ini file */
+void UFbxImportUI::LoadOptions(UObject* ObjectToLoadOptions)
+{
+	check(ObjectToLoadOptions);
+	int32 PortFlags = 0;
+	UClass* Class = ObjectToLoadOptions->GetClass();
+	for (FProperty* Property = Class->PropertyLink; Property; Property = Property->PropertyLinkNext)
+	{
+		if (!Property->HasAnyPropertyFlags(CPF_Config))
+		{
+			continue;
+		}
+		FString Section = TEXT("FBX_Import_UI_Option_") + Class->GetName();
+		FString Key = Property->GetName();
+
+		const bool bIsPropertyInherited = Property->GetOwnerClass() != Class;
+		UObject* SuperClassDefaultObject = Class->GetSuperClass()->GetDefaultObject();
+
+		const FString& PropFileName = GEditorPerProjectIni;
+
+		FArrayProperty* Array = CastField<FArrayProperty>(Property);
+		if (Array)
+		{
+			FConfigSection* Sec = GConfig->GetSectionPrivate(*Section, 0, 1, *GEditorPerProjectIni);
+			if (Sec != nullptr)
+			{
+				TArray<FConfigValue> List;
+				const FName KeyName(*Key, FNAME_Find);
+				Sec->MultiFind(KeyName, List);
+
+				FScriptArrayHelper_InContainer ArrayHelper(Array, ObjectToLoadOptions);
+				// Only override default properties if there is something to override them with.
+				if (List.Num() > 0)
+				{
+					ArrayHelper.EmptyAndAddValues(List.Num());
+					for (int32 i = List.Num() - 1, c = 0; i >= 0; i--, c++)
+					{
+						Array->Inner->ImportText(*List[i].GetValue(), ArrayHelper.GetRawPtr(c), PortFlags, ObjectToLoadOptions);
+					}
+				}
+				else
+				{
+					int32 Index = 0;
+					const FConfigValue* ElementValue = nullptr;
+					do
+					{
+						// Add array index number to end of key
+						FString IndexedKey = FString::Printf(TEXT("%s[%i]"), *Key, Index);
+
+						// Try to find value of key
+						const FName IndexedName(*IndexedKey, FNAME_Find);
+						if (IndexedName == NAME_None)
+						{
+							break;
+						}
+						ElementValue = Sec->Find(IndexedName);
+
+						// If found, import the element
+						if (ElementValue != nullptr)
+						{
+							// expand the array if necessary so that Index is a valid element
+							ArrayHelper.ExpandForIndex(Index);
+							Array->Inner->ImportText(*ElementValue->GetValue(), ArrayHelper.GetRawPtr(Index), PortFlags, ObjectToLoadOptions);
+						}
+
+						Index++;
+					} while (ElementValue || Index < ArrayHelper.Num());
+				}
+			}
+		}
+		else
+		{
+			for (int32 i = 0; i < Property->ArrayDim; i++)
+			{
+				if (Property->ArrayDim != 1)
+				{
+					Key = FString::Printf(TEXT("%s[%i]"), *Property->GetName(), i);
+				}
 
+				FString Value;
+				bool bFoundValue = GConfig->GetString(*Section, *Key, Value, *GEditorPerProjectIni);
+
+				if (bFoundValue)
+				{
+					if (Property->ImportText(*Value, Property->ContainerPtrToValuePtr<uint8>(ObjectToLoadOptions, i), PortFlags, ObjectToLoadOptions) == NULL)
+					{
+						// this should be an error as the properties from the .ini / .int file are not correctly being read in and probably are affecting things in subtle ways
+						UE_LOG(LogFbx, Error, TEXT("FBX Options LoadOptions (%s): import failed for %s in: %s"), *ObjectToLoadOptions->GetPathName(), *Property->GetName(), *Value);
+					}
+				}
+			}
+		}
+	}
+}
+
+/** Save UI settings to ini file */
+void UFbxImportUI::SaveOptions(UObject* ObjectToSaveOptions)
+{
+	check(ObjectToSaveOptions);
+	int32 PortFlags = 0;
+	UClass* Class = ObjectToSaveOptions->GetClass();
+	for (FProperty* Property = Class->PropertyLink; Property; Property = Property->PropertyLinkNext)
+	{
+		if (!Property->HasAnyPropertyFlags(CPF_Config))
+		{
+			continue;
+		}
+		FString Section = TEXT("FBX_Import_UI_Option_") + Class->GetName();
+		FString Key = Property->GetName();
+
+		const bool bIsPropertyInherited = Property->GetOwnerClass() != Class;
+		UObject* SuperClassDefaultObject = Class->GetSuperClass()->GetDefaultObject();
+
+		FArrayProperty* Array = CastField<FArrayProperty>(Property);
+		if (Array)
+		{
+			FConfigSection* Sec = GConfig->GetSectionPrivate(*Section, 1, 0, *GEditorPerProjectIni);
+			check(Sec);
+			Sec->Remove(*Key);
+
+			FScriptArrayHelper_InContainer ArrayHelper(Array, ObjectToSaveOptions);
+			for (int32 i = 0; i < ArrayHelper.Num(); i++)
+			{
+				FString	Buffer;
+				Array->Inner->ExportTextItem(Buffer, ArrayHelper.GetRawPtr(i), ArrayHelper.GetRawPtr(i), ObjectToSaveOptions, PortFlags);
+				Sec->Add(*Key, *Buffer);
+			}
+		}
+		else
+		{
+			TCHAR TempKey[MAX_SPRINTF] = TEXT("");
+			for (int32 Index = 0; Index < Property->ArrayDim; Index++)
+			{
+				if (Property->ArrayDim != 1)
+				{
+					FCString::Sprintf(TempKey, TEXT("%s[%i]"), *Property->GetName(), Index);
+					Key = TempKey;
+				}
+
+				FString	Value;
+				Property->ExportText_InContainer(Index, Value, ObjectToSaveOptions, ObjectToSaveOptions, ObjectToSaveOptions, PortFlags);
+				GConfig->SetString(*Section, *Key, *Value, *GEditorPerProjectIni);
+			}
+		}
+	}
+	GConfig->Flush(0);
+}
 
 #undef LOCTEXT_NAMESPACE
diff --git a/Engine/Source/Editor/UnrealEd/Private/Fbx/FbxMainImport.cpp b/Engine/Source/Editor/UnrealEd/Private/Fbx/FbxMainImport.cpp
--- a/Engine/Source/Editor/UnrealEd/Private/Fbx/FbxMainImport.cpp
+++ b/Engine/Source/Editor/UnrealEd/Private/Fbx/FbxMainImport.cpp
@@ -294,26 +294,26 @@ FBXImportOptions* GetImportOptions( UnFbx::FFbxImporter* FbxImporter, UFbxImport
 			}
 		}
 
-		ImportUI->SaveConfig();
+		UFbxImportUI::SaveOptions(ImportUI);
 
 		if( ImportUI->StaticMeshImportData )
 		{
-			ImportUI->StaticMeshImportData->SaveOptions();
+			UFbxImportUI::SaveOptions(ImportUI->StaticMeshImportData);
 		}
 
 		if( ImportUI->SkeletalMeshImportData )
 		{
-			ImportUI->SkeletalMeshImportData->SaveOptions();
+			UFbxImportUI::SaveOptions(ImportUI->SkeletalMeshImportData);
 		}
 
 		if( ImportUI->AnimSequenceImportData )
 		{
-			ImportUI->AnimSequenceImportData->SaveOptions();
+			UFbxImportUI::SaveOptions(ImportUI->AnimSequenceImportData);
 		}
 
 		if( ImportUI->TextureImportData )
 		{
-			ImportUI->TextureImportData->SaveOptions();
+			UFbxImportUI::SaveOptions(ImportUI->TextureImportData);
 		}
 
 		if (FbxOptionWindow->ShouldImport())
