Revert change from 11577116 so we get the path to the field rather than its owner module when filtering files

#rnx

diff --git a/Engine/Source/Editor/UnrealEd/Private/Commandlets/GatherTextFromMetadataCommandlet.cpp b/Engine/Source/Editor/UnrealEd/Private/Commandlets/GatherTextFromMetadataCommandlet.cpp
--- a/Engine/Source/Editor/UnrealEd/Private/Commandlets/GatherTextFromMetadataCommandlet.cpp
+++ b/Engine/Source/Editor/UnrealEd/Private/Commandlets/GatherTextFromMetadataCommandlet.cpp
@@ -273,7 +273,7 @@ void UGatherTextFromMetaDataCommandlet::GatherTextFromUObjects(const TArray<FStr
 		}
 
 		FString SourceFilePath;
-		if (!FSourceCodeNavigation::FindModulePath(FieldPackage, SourceFilePath))
+		if (!FSourceCodeNavigation::FindClassHeaderPath(Field, SourceFilePath))
 		{
 			continue;
 		}
