Source Control: Unattended package save now handles adding new files to source control in more situations.
- CheckoutPackages now handles adding new files
- InternalCheckoutAndSavePackages now handles adding new files

diff --git a/Engine/Source/Editor/UnrealEd/Private/FileHelpers.cpp b/Engine/Source/Editor/UnrealEd/Private/FileHelpers.cpp
--- a/Engine/Source/Editor/UnrealEd/Private/FileHelpers.cpp
+++ b/Engine/Source/Editor/UnrealEd/Private/FileHelpers.cpp
@@ -1811,8 +1811,9 @@ ECommandResult::Type FEditorFileUtils::CheckoutPackages(const TArray<UPackage*>&
 				}
 				else if (SourceControlState->CanAdd())
 				{
+					// Cannot add unsaved packages to source control
 					FString Filename;
-					if (FPackageName::DoesPackageExist(PackageToCheckOut->GetName(), nullptr, &Filename))
+					if (!PackageToCheckOut->HasAnyPackageFlags(PKG_NewlyCreated) || FPackageName::DoesPackageExist(PackageToCheckOut->GetName(), nullptr, &Filename))
 					{
 						bShowCheckoutError = false;
 						FinalPackageMarkForAddList.Add(PackageToCheckOut);
@@ -4314,6 +4315,7 @@ static bool InternalCheckoutAndSavePackages(const TArray<UPackage*>& PackagesToS
 			for (UPackage* Package : PackagesToSave)
 			{
 				// List unsaved packages that were not checked out
+				if (Package->HasAnyPackageFlags(PKG_NewlyCreated) && !PackagesCheckedOut.Contains(Package))
 				if (!PackagesCheckedOut.Contains(Package))
 				{
 					PackagesToMarkForAdd.Add(Package);
