Fixed canceling the request to save the map blocking you from entering PIE on subsequent attempts to play.

#trivial

diff --git a/Engine/Source/Editor/UnrealEd/Private/PlayLevel.cpp b/Engine/Source/Editor/UnrealEd/Private/PlayLevel.cpp
--- a/Engine/Source/Editor/UnrealEd/Private/PlayLevel.cpp
+++ b/Engine/Source/Editor/UnrealEd/Private/PlayLevel.cpp
@@ -1063,6 +1063,8 @@ void UEditorEngine::StartQueuedPlaySessionRequestImpl()
 		UE_LOG(LogPlayLevel, Warning, TEXT("%s"), *ErrorMsg.ToString());
 		FMessageLog(NAME_CategoryPIE).Warning(ErrorMsg);
 		FMessageLog(NAME_CategoryPIE).Open();
+		
+		CancelRequestPlaySession();
 		return;
 	}
 
