Fixed a potential TOptional access after being unset if PIE startup is canceled mid-way through.

#rnx

diff --git a/Engine/Source/Editor/UnrealEd/Private/PlayLevel.cpp b/Engine/Source/Editor/UnrealEd/Private/PlayLevel.cpp
--- a/Engine/Source/Editor/UnrealEd/Private/PlayLevel.cpp
+++ b/Engine/Source/Editor/UnrealEd/Private/PlayLevel.cpp
@@ -1709,7 +1709,7 @@ void UEditorEngine::CreateNewPlayInEditorInstance(FRequestPlaySessionParams &InR
 		}
 
 		// Only count non-dedicated instances as clients so that our indexes line up for log-ins.
-		if (!GameInstancePIEParameters.bRunAsDedicated)
+		if (!GameInstancePIEParameters.bRunAsDedicated && PlayInEditorSessionInfo.IsSet())
 		{
 			PlayInEditorSessionInfo->NumClientInstancesCreated++;
 		}
