Swap order of checks when reloading packages

diff --git a/Engine/Source/Editor/UnrealEd/Private/PackageTools.cpp b/Engine/Source/Editor/UnrealEd/Private/PackageTools.cpp
--- a/Engine/Source/Editor/UnrealEd/Private/PackageTools.cpp
+++ b/Engine/Source/Editor/UnrealEd/Private/PackageTools.cpp
@@ -510,13 +510,13 @@ UPackageTools::UPackageTools(const FObjectInitializer& ObjectInitializer)
 					// Get outermost packages, in case groups were selected.
 					UPackage* RealPackage = TopLevelPackage->GetOutermost() ? TopLevelPackage->GetOutermost() : TopLevelPackage;
 
-					if (RealPackage->IsDirty())
+					if (RealPackage->HasAnyPackageFlags(PKG_InMemoryOnly))
 					{
-						DirtyPackages.AddUnique(RealPackage);
+						InMemoryPackages.AddUnique(RealPackage);
 					}
-					else if (RealPackage->HasAnyPackageFlags(PKG_InMemoryOnly))
+					else if (RealPackage->IsDirty())
 					{
-						InMemoryPackages.AddUnique(RealPackage);
+						DirtyPackages.AddUnique(RealPackage);
 					}
 					else
 					{
