Performance and memory optimizations for the expand command of ShaderPipelineCacheToolsCommandlet

diff --git a/Engine/Source/Editor/UnrealEd/Private/Commandlets/ShaderPipelineCacheToolsCommandlet.cpp b/Engine/Source/Editor/UnrealEd/Private/Commandlets/ShaderPipelineCacheToolsCommandlet.cpp
--- a/Engine/Source/Editor/UnrealEd/Private/Commandlets/ShaderPipelineCacheToolsCommandlet.cpp
+++ b/Engine/Source/Editor/UnrealEd/Private/Commandlets/ShaderPipelineCacheToolsCommandlet.cpp
@@ -106,6 +106,25 @@ static void LoadStableSCLs(TMultiMap<FStableShaderKeyAndValue, FSHAHash>& Stable
 	}
 }
 
+// Version optimized for ExpandPSOSC
+static void LoadStableSCLs(TMultiMap<int32, FSHAHash>& StableMap, TArray<FStableShaderKeyAndValue>& StableShaderKeyIndexTable, TArrayView<const FStringView> FileNames)
+{
+	TArray<TArray<FStableShaderKeyAndValue>> StableArrays;
+	StableArrays.AddDefaulted(FileNames.Num());
+	ParallelFor(FileNames.Num(), [&StableArrays, &FileNames](int32 Index) { LoadStableSCL(StableArrays[Index], FileNames[Index]); });
+
+	const int32 StableArrayCount = Algo::TransformAccumulate(StableArrays, &TArray<FStableShaderKeyAndValue>::Num, 0);
+	StableMap.Reserve(StableMap.Num() + StableArrayCount);
+	for (const TArray<FStableShaderKeyAndValue>& StableArray : StableArrays)
+	{
+		for (const FStableShaderKeyAndValue& Item : StableArray)
+		{
+			int32 ItemIndex = StableShaderKeyIndexTable.Add(Item);
+			StableMap.AddUnique(ItemIndex, Item.OutputHash);
+		}
+	}
+}
+
 static bool LoadAndDecompressStableCSV(const FString& Filename, TArray<uint8>& UncompressedData)
 {
 	bool bResult = false;
@@ -339,7 +358,7 @@ int32 DumpPSOSC(FString& Token)
 	return 0;
 }
 
-static void PrintShaders(const TMap<FSHAHash, TArray<FStableShaderKeyAndValue>>& InverseMap, const FSHAHash& Shader, const TCHAR *Label)
+static void PrintShaders(const TMap<FSHAHash, TArray<int32>>& InverseMap, TArray<FStableShaderKeyAndValue>& StableArray, const FSHAHash& Shader, const TCHAR *Label)
 {
 	UE_LOG(LogShaderPipelineCacheTools, Display, TEXT(" -- %s"), Label);
 
@@ -348,25 +367,25 @@ static void PrintShaders(const TMap<FSHAHash, TArray<FStableShaderKeyAndValue>>&
 		UE_LOG(LogShaderPipelineCacheTools, Display, TEXT("    null"));
 		return;
 	}
-	const TArray<FStableShaderKeyAndValue>* Out = InverseMap.Find(Shader);
+	const TArray<int32>* Out = InverseMap.Find(Shader);
 	if (!Out)
 	{
 		UE_LOG(LogShaderPipelineCacheTools, Display, TEXT("    No shaders found with hash %s"), *Shader.ToString());
 		return;
 	}
-	for (const FStableShaderKeyAndValue& Item : *Out)
+	for (const int32& Item : *Out)
 	{
-		UE_LOG(LogShaderPipelineCacheTools, Display, TEXT("    %s"), *Item.ToString());
+		UE_LOG(LogShaderPipelineCacheTools, Display, TEXT("    %s"), *StableArray[Item].ToString());
 	}
 }
 
-static bool GetStableShadersAndZeroHash(const TMap<FSHAHash, TArray<FStableShaderKeyAndValue>>& InverseMap, const FSHAHash& Shader, TArray<FStableShaderKeyAndValue>& StableShaders, bool& bOutAnyActiveButMissing)
+static bool GetStableShaders(const TMap<FSHAHash, TArray<int32>>& InverseMap, TArray<FStableShaderKeyAndValue>& StableArray, const FSHAHash& Shader, TArray<int32>& StableShaders, bool& bOutAnyActiveButMissing)
 {
 	if (Shader == FSHAHash())
 	{
 		return false;
 	}
-	const TArray<FStableShaderKeyAndValue>* Out = InverseMap.Find(Shader);
+	const TArray<int32>* Out = InverseMap.Find(Shader);
 	if (!Out)
 	{
 		UE_LOG(LogShaderPipelineCacheTools, Warning, TEXT("No shaders found with hash %s"), *Shader.ToString());
@@ -375,17 +394,15 @@ static bool GetStableShadersAndZeroHash(const TMap<FSHAHash, TArray<FStableShade
 		return false;
 	}
 	StableShaders.Reserve(Out->Num());
-	for (const FStableShaderKeyAndValue& Item : *Out)
+	for (const int32& Item : *Out)
 	{
-		FStableShaderKeyAndValue Temp = Item;
-		Temp.OutputHash = FSHAHash();
-		if (StableShaders.Contains(Temp))
+		if (StableShaders.Contains(Item))
 		{
 			UE_LOG(LogShaderPipelineCacheTools, Error, TEXT("Duplicate stable shader. This is bad because it means our stable key is not exhaustive."));
-			UE_LOG(LogShaderPipelineCacheTools, Error, TEXT(" %s"), *Item.ToString());
+			UE_LOG(LogShaderPipelineCacheTools, Error, TEXT(" %s"), *StableArray[Item].ToString());
 			continue;
 		}
-		StableShaders.Add(Temp);
+		StableShaders.Add(Item);
 	}
 	return true;
 }
@@ -408,6 +425,25 @@ static void StableShadersSerializationSelfTest(const TMultiMap<FStableShaderKeyA
 	}
 }
 
+// Version optimized for ExpandPSOSC
+static void StableShadersSerializationSelfTest(const TMultiMap<int32, FSHAHash>& StableMap, const TArray<FStableShaderKeyAndValue>& StableArray)
+{
+	TAnsiStringBuilder<384> TestString;
+	for (const auto& Pair : StableMap)
+	{
+		TestString.Reset();
+		FStableShaderKeyAndValue Item(StableArray[Pair.Key]);
+		Item.OutputHash = Pair.Value;
+		check(Pair.Value != FSHAHash());
+		Item.AppendString(TestString);
+		FStableShaderKeyAndValue TestItem;
+		TestItem.ParseFromString(UTF8_TO_TCHAR(TestString.ToString()));
+		check(Item == TestItem);
+		check(GetTypeHash(Item) == GetTypeHash(TestItem));
+		check(Item.OutputHash == TestItem.OutputHash);
+	}
+}
+
 // return true if these two shaders could be part of the same stable PSO
 // for example, if they come from two different vertex factories, we return false because that situation cannot occur
 bool CouldBeUsedTogether(const FStableShaderKeyAndValue& A, const FStableShaderKeyAndValue& B)
@@ -491,12 +527,12 @@ void IntersectSets(TSet<FCompactFullName>& Intersect, const TSet<FCompactFullNam
 	}
 }
 
-struct FPermuation
+struct FPermutation
 {
-	FStableShaderKeyAndValue Slots[SF_NumFrequencies];
+	int32 Slots[SF_NumFrequencies];
 };
 
-void GeneratePermuations(TArray<FPermuation>& Permutations, FPermuation& WorkingPerm, int32 SlotIndex , const TArray<FStableShaderKeyAndValue> StableShadersPerSlot[SF_NumFrequencies], const bool ActivePerSlot[SF_NumFrequencies])
+void GeneratePermuations(TArray<FPermutation>& Permutations, FPermutation& WorkingPerm, int32 SlotIndex , const TArray<int32> StableShadersPerSlot[SF_NumFrequencies], const TArray<FStableShaderKeyAndValue>& StableArray, const bool ActivePerSlot[SF_NumFrequencies])
 {
 	check(SlotIndex >= 0 && SlotIndex <= SF_NumFrequencies);
 	while (SlotIndex < SF_NumFrequencies && !ActivePerSlot[SlotIndex])
@@ -519,7 +555,7 @@ void GeneratePermuations(TArray<FPermuation>& Permutations, FPermuation& Working
 				continue;
 			}
 			check(SlotIndex != SF_Compute && SlotIndexInner != SF_Compute); // there is never any matching with compute shaders
-			if (!CouldBeUsedTogether(StableShadersPerSlot[SlotIndex][StableIndex], WorkingPerm.Slots[SlotIndexInner]))
+			if (!CouldBeUsedTogether(StableArray[StableShadersPerSlot[SlotIndex][StableIndex]], StableArray[WorkingPerm.Slots[SlotIndexInner]]))
 			{
 				bKeep = false;
 				break;
@@ -530,7 +566,7 @@ void GeneratePermuations(TArray<FPermuation>& Permutations, FPermuation& Working
 			continue;
 		}
 		WorkingPerm.Slots[SlotIndex] = StableShadersPerSlot[SlotIndex][StableIndex];
-		GeneratePermuations(Permutations, WorkingPerm, SlotIndex + 1, StableShadersPerSlot, ActivePerSlot);
+		GeneratePermuations(Permutations, WorkingPerm, SlotIndex + 1, StableShadersPerSlot, StableArray, ActivePerSlot);
 	}
 }
 
@@ -547,8 +583,11 @@ int32 ExpandPSOSC(const TArray<FString>& Tokens)
 		}
 	}
 
-	TMultiMap<FStableShaderKeyAndValue, FSHAHash> StableMap;
-	LoadStableSCLs(StableMap, StableCSVs);
+	// To save memory and make operations on the stable map faster, all the stable shader keys are stored in StableShaderKeyIndexTable array and shader map keys
+	// and permutation slots use indices to this array instead of storing their own copies of FStableShaderKeyAndValue objects
+	TArray<FStableShaderKeyAndValue> StableShaderKeyIndexTable;
+	TMultiMap<int32, FSHAHash> StableMap;
+	LoadStableSCLs(StableMap, StableShaderKeyIndexTable, StableCSVs);
 	if (!StableMap.Num())
 	{
 		UE_LOG(LogShaderPipelineCacheTools, Warning, TEXT("No .scl.csv found or they were all empty. Nothing to do."));
@@ -559,11 +598,11 @@ int32 ExpandPSOSC(const TArray<FString>& Tokens)
 		UE_LOG(LogShaderPipelineCacheTools, Verbose, TEXT("    %s"), *FStableShaderKeyAndValue::HeaderLine());
 		for (const auto& Pair : StableMap)
 		{
-			FStableShaderKeyAndValue Temp(Pair.Key);
+			FStableShaderKeyAndValue Temp(StableShaderKeyIndexTable[Pair.Key]);
 			Temp.OutputHash = Pair.Value;
 			UE_LOG(LogShaderPipelineCacheTools, Verbose, TEXT("    %s"), *Temp.ToString());
 		}
-		StableShadersSerializationSelfTest(StableMap);
+		StableShadersSerializationSelfTest(StableMap, StableShaderKeyIndexTable);
 	}
 	UE_LOG(LogShaderPipelineCacheTools, Display, TEXT("Loaded %d unique shader info lines total."), StableMap.Num());
 
@@ -636,7 +675,7 @@ int32 ExpandPSOSC(const TArray<FString>& Tokens)
 
 		for (const auto& Pair : StableMap)
 		{
-			FStableShaderKeyAndValue Temp(Pair.Key);
+			FStableShaderKeyAndValue Temp(StableShaderKeyIndexTable[Pair.Key]);
 			Temp.OutputHash = Pair.Value;
 			InverseMap.FindOrAdd(Pair.Value).Add(Temp.ToString());
 		}
@@ -663,13 +702,11 @@ int32 ExpandPSOSC(const TArray<FString>& Tokens)
 			}
 		}
 	}
-	TMap<FSHAHash, TArray<FStableShaderKeyAndValue>> InverseMap;
+	TMap<FSHAHash, TArray<int32>> InverseMap;
 
 	for (const auto& Pair : StableMap)
 	{
-		FStableShaderKeyAndValue Item(Pair.Key);
-		Item.OutputHash = Pair.Value;
-		InverseMap.FindOrAdd(Item.OutputHash).AddUnique(Item);
+		InverseMap.FindOrAdd(Pair.Value).AddUnique(Pair.Key);
 	}
 
 	int32 TotalStablePSOs = 0;
@@ -678,7 +715,7 @@ int32 ExpandPSOSC(const TArray<FString>& Tokens)
 	{
 		const FPipelineCacheFileFormatPSO* PSO;
 		bool ActivePerSlot[SF_NumFrequencies];
-		TArray<FPermuation> Permutations;
+		TArray<FPermutation> Permutations;
 
 		FPermsPerPSO()
 			: PSO(nullptr)
@@ -699,22 +736,22 @@ int32 ExpandPSOSC(const TArray<FString>& Tokens)
 	{ 
 		NumExamined++;
 		check(SF_Vertex == 0 && SF_Compute == 5);
-		TArray<FStableShaderKeyAndValue> StableShadersPerSlot[SF_NumFrequencies];
+		TArray<int32> StableShadersPerSlot[SF_NumFrequencies];
 		bool ActivePerSlot[SF_NumFrequencies] = { false };
 
 		bool OutAnyActiveButMissing = false;
 
 		if (Item.Type == FPipelineCacheFileFormatPSO::DescriptorType::Compute)
 		{
-			ActivePerSlot[SF_Compute] = GetStableShadersAndZeroHash(InverseMap, Item.ComputeDesc.ComputeShader, StableShadersPerSlot[SF_Compute], OutAnyActiveButMissing);
+			ActivePerSlot[SF_Compute] = GetStableShaders(InverseMap, StableShaderKeyIndexTable, Item.ComputeDesc.ComputeShader, StableShadersPerSlot[SF_Compute], OutAnyActiveButMissing);
 		}
 		else
 		{
-			ActivePerSlot[SF_Vertex] = GetStableShadersAndZeroHash(InverseMap, Item.GraphicsDesc.VertexShader, StableShadersPerSlot[SF_Vertex], OutAnyActiveButMissing);
-			ActivePerSlot[SF_Pixel] = GetStableShadersAndZeroHash(InverseMap, Item.GraphicsDesc.FragmentShader, StableShadersPerSlot[SF_Pixel], OutAnyActiveButMissing);
-			ActivePerSlot[SF_Geometry] = GetStableShadersAndZeroHash(InverseMap, Item.GraphicsDesc.GeometryShader, StableShadersPerSlot[SF_Geometry], OutAnyActiveButMissing);
-			ActivePerSlot[SF_Hull] = GetStableShadersAndZeroHash(InverseMap, Item.GraphicsDesc.HullShader, StableShadersPerSlot[SF_Hull], OutAnyActiveButMissing);
-			ActivePerSlot[SF_Domain] = GetStableShadersAndZeroHash(InverseMap, Item.GraphicsDesc.DomainShader, StableShadersPerSlot[SF_Domain], OutAnyActiveButMissing);
+			ActivePerSlot[SF_Vertex] = GetStableShaders(InverseMap, StableShaderKeyIndexTable, Item.GraphicsDesc.VertexShader, StableShadersPerSlot[SF_Vertex], OutAnyActiveButMissing);
+			ActivePerSlot[SF_Pixel] = GetStableShaders(InverseMap, StableShaderKeyIndexTable, Item.GraphicsDesc.FragmentShader, StableShadersPerSlot[SF_Pixel], OutAnyActiveButMissing);
+			ActivePerSlot[SF_Geometry] = GetStableShaders(InverseMap, StableShaderKeyIndexTable, Item.GraphicsDesc.GeometryShader, StableShadersPerSlot[SF_Geometry], OutAnyActiveButMissing);
+			ActivePerSlot[SF_Hull] = GetStableShaders(InverseMap, StableShaderKeyIndexTable, Item.GraphicsDesc.HullShader, StableShadersPerSlot[SF_Hull], OutAnyActiveButMissing);
+			ActivePerSlot[SF_Domain] = GetStableShaders(InverseMap, StableShaderKeyIndexTable, Item.GraphicsDesc.DomainShader, StableShadersPerSlot[SF_Domain], OutAnyActiveButMissing);
 		}
 
 
@@ -723,16 +760,16 @@ int32 ExpandPSOSC(const TArray<FString>& Tokens)
 			UE_LOG(LogShaderPipelineCacheTools, Display, TEXT("PSO had an active shader slot that did not match any current shaders, ignored."));
 			if (Item.Type == FPipelineCacheFileFormatPSO::DescriptorType::Compute)
 			{
-				PrintShaders(InverseMap, Item.ComputeDesc.ComputeShader, TEXT("ComputeShader"));
+				PrintShaders(InverseMap, StableShaderKeyIndexTable, Item.ComputeDesc.ComputeShader, TEXT("ComputeShader"));
 			}
 			else
 			{
 				UE_LOG(LogShaderPipelineCacheTools, Display, TEXT("   %s"), *Item.GraphicsDesc.StateToString());
-				PrintShaders(InverseMap, Item.GraphicsDesc.VertexShader, TEXT("VertexShader"));
-				PrintShaders(InverseMap, Item.GraphicsDesc.FragmentShader, TEXT("FragmentShader"));
-				PrintShaders(InverseMap, Item.GraphicsDesc.GeometryShader, TEXT("GeometryShader"));
-				PrintShaders(InverseMap, Item.GraphicsDesc.HullShader, TEXT("HullShader"));
-				PrintShaders(InverseMap, Item.GraphicsDesc.DomainShader, TEXT("DomainShader"));
+				PrintShaders(InverseMap, StableShaderKeyIndexTable, Item.GraphicsDesc.VertexShader, TEXT("VertexShader"));
+				PrintShaders(InverseMap, StableShaderKeyIndexTable, Item.GraphicsDesc.FragmentShader, TEXT("FragmentShader"));
+				PrintShaders(InverseMap, StableShaderKeyIndexTable, Item.GraphicsDesc.GeometryShader, TEXT("GeometryShader"));
+				PrintShaders(InverseMap, StableShaderKeyIndexTable, Item.GraphicsDesc.HullShader, TEXT("HullShader"));
+				PrintShaders(InverseMap, StableShaderKeyIndexTable, Item.GraphicsDesc.DomainShader, TEXT("DomainShader"));
 			}
 			continue;
 		}
@@ -762,7 +799,7 @@ int32 ExpandPSOSC(const TArray<FString>& Tokens)
 						bool bFoundCompat = false;
 						for (int32 StableIndexInner = 0; StableIndexInner < StableShadersPerSlot[SlotIndexInner].Num(); StableIndexInner++)
 						{
-							if (CouldBeUsedTogether(StableShadersPerSlot[SlotIndex][StableIndex], StableShadersPerSlot[SlotIndexInner][StableIndexInner]))
+							if (CouldBeUsedTogether(StableShaderKeyIndexTable[StableShadersPerSlot[SlotIndex][StableIndex]], StableShaderKeyIndexTable[StableShadersPerSlot[SlotIndexInner][StableIndexInner]]))
 							{
 								bFoundCompat = true;
 								break;
@@ -796,11 +833,11 @@ int32 ExpandPSOSC(const TArray<FString>& Tokens)
 				UE_LOG(LogShaderPipelineCacheTools, Warning, TEXT("PSO did not create any stable PSOs! (no cross shader slot compatibility)"));
 				UE_LOG(LogShaderPipelineCacheTools, Warning, TEXT("   %s"), *Item.GraphicsDesc.StateToString());
 
-				PrintShaders(InverseMap, Item.GraphicsDesc.VertexShader, TEXT("VertexShader"));
-				PrintShaders(InverseMap, Item.GraphicsDesc.FragmentShader, TEXT("FragmentShader"));
-				PrintShaders(InverseMap, Item.GraphicsDesc.GeometryShader, TEXT("GeometryShader"));
-				PrintShaders(InverseMap, Item.GraphicsDesc.HullShader, TEXT("HullShader"));
-				PrintShaders(InverseMap, Item.GraphicsDesc.DomainShader, TEXT("DomainShader"));
+				PrintShaders(InverseMap, StableShaderKeyIndexTable, Item.GraphicsDesc.VertexShader, TEXT("VertexShader"));
+				PrintShaders(InverseMap, StableShaderKeyIndexTable, Item.GraphicsDesc.FragmentShader, TEXT("FragmentShader"));
+				PrintShaders(InverseMap, StableShaderKeyIndexTable, Item.GraphicsDesc.GeometryShader, TEXT("GeometryShader"));
+				PrintShaders(InverseMap, StableShaderKeyIndexTable, Item.GraphicsDesc.HullShader, TEXT("HullShader"));
+				PrintShaders(InverseMap, StableShaderKeyIndexTable, Item.GraphicsDesc.DomainShader, TEXT("DomainShader"));
 
 				continue;
 			}
@@ -816,9 +853,9 @@ int32 ExpandPSOSC(const TArray<FString>& Tokens)
 			Current.ActivePerSlot[Index] = ActivePerSlot[Index];
 		}
 
-		TArray<FPermuation>& Permutations(Current.Permutations);
-		FPermuation WorkingPerm;
-		GeneratePermuations(Permutations, WorkingPerm, 0, StableShadersPerSlot, ActivePerSlot);
+		TArray<FPermutation>& Permutations(Current.Permutations);
+		FPermutation WorkingPerm;
+		GeneratePermuations(Permutations, WorkingPerm, 0, StableShadersPerSlot, StableShaderKeyIndexTable, ActivePerSlot);
 		if (!Permutations.Num())
 		{
 			UE_LOG(LogShaderPipelineCacheTools, Error, TEXT("PSO did not create any stable PSOs! (somehow)"));
@@ -868,7 +905,7 @@ int32 ExpandPSOSC(const TArray<FString>& Tokens)
 				UE_LOG(LogShaderPipelineCacheTools, Verbose, TEXT(" %s"), *Item.PSO->GraphicsDesc.StateToString());
 			}
 			int32 PermIndex = 0;
-			for (const FPermuation& Perm : Item.Permutations)
+			for (const FPermutation& Perm : Item.Permutations)
 			{
 				UE_LOG(LogShaderPipelineCacheTools, Verbose, TEXT("  ----- perm %d"), PermIndex);
 				for (int32 SlotIndex = 0; SlotIndex < SF_NumFrequencies; SlotIndex++)
@@ -877,14 +914,16 @@ int32 ExpandPSOSC(const TArray<FString>& Tokens)
 					{
 						continue;
 					}
-					UE_LOG(LogShaderPipelineCacheTools, Verbose, TEXT("   %s"), *Perm.Slots[SlotIndex].ToString());
+					FStableShaderKeyAndValue ShaderKeyAndValue = StableShaderKeyIndexTable[Perm.Slots[SlotIndex]];
+					ShaderKeyAndValue.OutputHash = FSHAHash(); // Saved output hash needs to be zeroed so that BuildPSOSC can use this entry even if shaders code changes in future builds
+					UE_LOG(LogShaderPipelineCacheTools, Verbose, TEXT("   %s"), *ShaderKeyAndValue.ToString());
 				}
 				PermIndex++;
 			}
 
 			UE_LOG(LogShaderPipelineCacheTools, Verbose, TEXT("-----"));
 		}
-		for (const FPermuation& Perm : Item.Permutations)
+		for (const FPermutation& Perm : Item.Permutations)
 		{
 			// because it is a CSV, and for backward compat, compute shaders will just be a zeroed graphics desc with the shader in the hull shader slot.
 			FString PSOLine = Item.PSO->CommonToString();
@@ -899,7 +938,9 @@ int32 ExpandPSOSC(const TArray<FString>& Tokens)
 					check(!Item.ActivePerSlot[SlotIndex]); // none of these should be active for a compute shader
 					if (SlotIndex == SF_Hull)
 					{
-						PSOLine += FString::Printf(TEXT(",\"%s\""), *Perm.Slots[SF_Compute].ToString());
+						FStableShaderKeyAndValue ShaderKeyAndValue = StableShaderKeyIndexTable[Perm.Slots[SF_Compute]];
+						ShaderKeyAndValue.OutputHash = FSHAHash(); // Saved output hash needs to be zeroed so that BuildPSOSC can use this entry even if shaders code changes in future builds
+						PSOLine += FString::Printf(TEXT(",\"%s\""), *ShaderKeyAndValue.ToString());
 					}
 					else
 					{
@@ -917,7 +958,9 @@ int32 ExpandPSOSC(const TArray<FString>& Tokens)
 						PSOLine += FString::Printf(TEXT(",\"\""));
 						continue;
 					}
-					PSOLine += FString::Printf(TEXT(",\"%s\""), *Perm.Slots[SlotIndex].ToString());
+					FStableShaderKeyAndValue ShaderKeyAndValue = StableShaderKeyIndexTable[Perm.Slots[SlotIndex]];
+					ShaderKeyAndValue.OutputHash = FSHAHash(); // Saved output hash needs to be zeroed so that BuildPSOSC can use this entry even if shaders code changes in future builds
+					PSOLine += FString::Printf(TEXT(",\"%s\""), *ShaderKeyAndValue.ToString());
 				}
 			}
 
