Convert some transaction debugging to be controlled by a cvar instead of having to modify a #if 0
#rnx

diff --git a/Engine/Source/Editor/UnrealEd/Private/EditorTransaction.cpp b/Engine/Source/Editor/UnrealEd/Private/EditorTransaction.cpp
--- a/Engine/Source/Editor/UnrealEd/Private/EditorTransaction.cpp
+++ b/Engine/Source/Editor/UnrealEd/Private/EditorTransaction.cpp
@@ -1165,6 +1165,10 @@ int32 UTransBuffer::Begin( const TCHAR* SessionContext, const FText& Description
 	return BeginInternal<FTransaction>(SessionContext, Description);
 }
 
+namespace TransBuffer
+{
+	static FAutoConsoleVariable DumpTransBufferObjectMap(TEXT("TransBuffer.DumpObjectMap"), false, TEXT("Whether to dump the object map each time a transaction is written for debugging purposes."));
+}
 
 int32 UTransBuffer::End()
 {
@@ -1176,16 +1180,15 @@ int32 UTransBuffer::End()
 	{
 		if( --ActiveCount==0 )
 		{
-#if 0 // @todo DB: please don't remove this code -- thanks! :)
-			// End the current transaction.
-			if ( GUndo && GLog )
-			{
-				// @todo DB: Fix this potentially unsafe downcast.
-				static_cast<FTransaction*>(GUndo)->DumpObjectMap( *GLog );
-			}
-#endif
 			if (GUndo)
 			{
+				if (GLog && TransBuffer::DumpTransBufferObjectMap->GetBool())
+				{
+					// @todo DB: Fix this potentially unsafe downcast.
+					static_cast<FTransaction*>(GUndo)->DumpObjectMap( *GLog );
+				}
+
+				// End the current transaction.
 				GUndo->Finalize();
 				TransactionStateChangedDelegate.Broadcast(GUndo->GetContext(), ETransactionStateEventType::TransactionFinalized);
 				GUndo->EndOperation();
