Fix custom LOD workflow, the insert LOD into the base mesh function was not updated to work with the skeletalmesh refactor

diff --git a/Engine/Source/Developer/MeshBuilder/Private/SkeletalMeshBuilder.cpp b/Engine/Source/Developer/MeshBuilder/Private/SkeletalMeshBuilder.cpp
--- a/Engine/Source/Developer/MeshBuilder/Private/SkeletalMeshBuilder.cpp
+++ b/Engine/Source/Developer/MeshBuilder/Private/SkeletalMeshBuilder.cpp
@@ -53,9 +53,10 @@ struct FSkeletalMeshVertInstanceIDAndZ
 //or simply add parameter to FSkeletalMeshImportData::CopyLODImportData to do the job inside the function.
 namespace SkeletalMeshBuilderHelperNS
 {
-	void FixFaceMaterial(USkeletalMesh* SkeletalMesh, TArray<SkeletalMeshImportData::FMaterial>& RawMeshMaterials, TArray<SkeletalMeshImportData::FMeshFace>& LODFaces)
+	//Base LOD (index 0), which is not using LODMaterialMap, we want to adjust the source face material index data to fit with the SkeletalMesh material list
+	void AdjustSourceFaceMaterialIndex(USkeletalMesh* SkeletalMesh, TArray<SkeletalMeshImportData::FMaterial>& RawMeshMaterials, TArray<SkeletalMeshImportData::FMeshFace>& LODFaces, int32 LODIndex)
 	{
-		if (RawMeshMaterials.Num() <= 1)
+		if (RawMeshMaterials.Num() <= 1 || LODIndex > 0)
 		{
 			//Nothing to fix if we have 1 or less material
 			return;
@@ -75,10 +76,7 @@ namespace SkeletalMeshBuilderHelperNS
 				FName MeshMaterialName = SkeletalMesh->Materials[MeshMaterialIndex].ImportedMaterialSlotName;
 				if (MaterialImportName == MeshMaterialName)
 				{
-					if (MaterialRemap[MaterialIndex] != MeshMaterialIndex)
-					{
-						bNeedRemapping = true;
-					}
+					bNeedRemapping |= (MaterialRemap[MaterialIndex] != MeshMaterialIndex);
 					MaterialRemap[MaterialIndex] = MeshMaterialIndex;
 					break;
 				}
@@ -151,8 +149,12 @@ bool FSkeletalMeshBuilder::Build(USkeletalMesh* SkeletalMesh, const int32 LODInd
 
 		//Use the max because we need to have at least one texture coordinnate
 		NumTextCoord = FMath::Max<int32>(NumTextCoord, SkeletalMeshImportData.NumTexCoords);
-
-		SkeletalMeshBuilderHelperNS::FixFaceMaterial(SkeletalMesh, SkeletalMeshImportData.Materials, LODFaces);
+		
+		if (LODIndex == 0)
+		{
+			//BaseLOD should not use LODMaterialMap (except if user has change the material), so we have to make sure the source data fit with the skeletalmesh materials array
+			SkeletalMeshBuilderHelperNS::AdjustSourceFaceMaterialIndex(SkeletalMesh, SkeletalMeshImportData.Materials, LODFaces, LODIndex);
+		}
 
 		//Build the skeletalmesh using mesh utilities module
 		IMeshUtilities::MeshBuildOptions Options;
diff --git a/Engine/Source/Editor/Persona/Private/PersonaMeshDetails.cpp b/Engine/Source/Editor/Persona/Private/PersonaMeshDetails.cpp
--- a/Engine/Source/Editor/Persona/Private/PersonaMeshDetails.cpp
+++ b/Engine/Source/Editor/Persona/Private/PersonaMeshDetails.cpp
@@ -116,6 +116,11 @@ static int32 GetDefaultMaterialIndex(const USkeletalMesh* SkeletalMesh, int32 LO
 	}
 
 	const FSkeletalMeshLODModel& LODModel = SkeletalMesh->GetImportedModel()->LODModels[LODIndex];
+	if (LODIndex > 0)
+	{
+		//Only base LOD is using the default/modified feature
+		return INDEX_NONE;
+	}
 
 	if (LODModel.Sections.IsValidIndex(SectionIndex))
 	{
@@ -4092,8 +4097,8 @@ void FPersonaMeshDetails::OnDeleteMaterialSlot(int32 MaterialIndex)
 	{
 		FScopedSkeletalMeshPostEditChange ScopedPostEditChange(SkeletalMeshPtr.Get());
 		//When we delete a material slot we must invalidate the DDC because material index is not part of the DDC key by design
-	SkeletalMeshPtr->Materials.RemoveAt(MaterialIndex);
-	FSkeletalMeshModel* Model = SkeletalMeshPtr->GetImportedModel();
+		SkeletalMeshPtr->Materials.RemoveAt(MaterialIndex);
+		FSkeletalMeshModel* Model = SkeletalMeshPtr->GetImportedModel();
 
 		int32 NumLODInfos = SkeletalMeshPtr->GetLODNum();
 
@@ -4109,13 +4114,13 @@ void FPersonaMeshDetails::OnDeleteMaterialSlot(int32 MaterialIndex)
 					SectionMaterialIndex = LODMaterialMap[SectionIndex];
 				}
 				if (SectionMaterialIndex > MaterialIndex)
-	{
+				{
 					SectionMaterialIndex--;
 				}
 				if (SectionMaterialIndex != Model->LODModels[LODInfoIdx].Sections[SectionIndex].MaterialIndex)
-		{
+				{
 					while(!LODMaterialMap.IsValidIndex(SectionIndex))
-			{
+					{
 						LODMaterialMap.Add(INDEX_NONE);
 					}
 					LODMaterialMap[SectionIndex] = SectionMaterialIndex;
diff --git a/Engine/Source/Editor/PropertyEditor/Private/PropertyCustomizationHelpers.cpp b/Engine/Source/Editor/PropertyEditor/Private/PropertyCustomizationHelpers.cpp
--- a/Engine/Source/Editor/PropertyEditor/Private/PropertyCustomizationHelpers.cpp
+++ b/Engine/Source/Editor/PropertyEditor/Private/PropertyCustomizationHelpers.cpp
@@ -1089,7 +1089,7 @@ private:
 		FString MaterialSlotDisplayName;
 		SectionItem.MaterialSlotName.ToString(MaterialSlotDisplayName);
 		FString MaterialSlotRemapString = TEXT("");
-		if (SectionItem.DefaultMaterialIndex != SectionItem.MaterialSlotIndex)
+		if (SectionItem.DefaultMaterialIndex != INDEX_NONE && SectionItem.DefaultMaterialIndex != SectionItem.MaterialSlotIndex)
 		{
 			MaterialSlotRemapString = TEXT(" (Modified)");
 		}
diff --git a/Engine/Source/Editor/UnrealEd/Private/Fbx/FbxSkeletalMeshImport.cpp b/Engine/Source/Editor/UnrealEd/Private/Fbx/FbxSkeletalMeshImport.cpp
--- a/Engine/Source/Editor/UnrealEd/Private/Fbx/FbxSkeletalMeshImport.cpp
+++ b/Engine/Source/Editor/UnrealEd/Private/Fbx/FbxSkeletalMeshImport.cpp
@@ -3745,31 +3745,66 @@ void UnFbx::FFbxImporter::InsertNewLODToBaseSkeletalMesh(USkeletalMesh* InSkelet
 				{
 					break;
 				}
+
 				//Get the current skelmesh section slot import name
 				int32 ExistRemapMaterialIndex = ExistingSkelMeshLodInfo.LODMaterialMap[SectionIndex];
 				if (!BaseSkeletalMesh->Materials.IsValidIndex(ExistRemapMaterialIndex))
 				{
-					//LODMaterialMap can have INDEX_NONE value or an invalid index in the material array
-					continue;
+					//If we have an INDEX_NONE it mean we have to use the section material index. If we have a bad value we will use the section material index if possible
+					if (DestImportedResource->LODModels.IsValidIndex(DesiredLOD) && DestImportedResource->LODModels[DesiredLOD].Sections.IsValidIndex(SectionIndex))
+					{
+						ExistRemapMaterialIndex = DestImportedResource->LODModels[DesiredLOD].Sections[SectionIndex].MaterialIndex;
+					}
+					else
+					{
+						//If the material is wrong log and pick the zero index
+						UE_LOG(LogFbx, Display, TEXT("Insert LOD, invalid material rematch: LOD [%d] Section[%d]- Skeletalmesh:[%s]"), DesiredLOD, SectionIndex, *BaseSkeletalMesh->GetFullName());
+						ExistRemapMaterialIndex = 0;
+					}
 				}
+				check(BaseSkeletalMesh->Materials.IsValidIndex(ExistRemapMaterialIndex));
 
 				ExistingMeshSectionSlotNames.Add(BaseSkeletalMesh->Materials[ExistRemapMaterialIndex].ImportedMaterialSlotName);
 
 				//Get the Last imported skelmesh section slot import name
 				OriginalImportMeshSectionSlotNames.Add(OriginalImportMeshLodSectionsData.SectionOriginalMaterialName[SectionIndex]);
-				//We have some old LodModel data to restore
-				if (DestImportedResource->LODModels.IsValidIndex(DesiredLOD)
-					&& DestImportedResource->LODModels[DesiredLOD].Sections.IsValidIndex(SectionIndex)
-					&& NewLODModel.Sections.IsValidIndex(SectionIndex))
+			}
+
+			if (DestImportedResource->LODModels.IsValidIndex(DesiredLOD))
+			{
+				const FSkeletalMeshLODModel& ExistLODModel = DestImportedResource->LODModels[DesiredLOD];
+				for (int32 NewSectionIndex = 0; NewSectionIndex < NewLODModel.Sections.Num(); NewSectionIndex++)
 				{
-					NewLODModel.Sections[SectionIndex].bCastShadow = DestImportedResource->LODModels[DesiredLOD].Sections[SectionIndex].bCastShadow;
-					NewLODModel.Sections[SectionIndex].bRecomputeTangent = DestImportedResource->LODModels[DesiredLOD].Sections[SectionIndex].bRecomputeTangent;
-					NewLODModel.Sections[SectionIndex].bDisabled = DestImportedResource->LODModels[DesiredLOD].Sections[SectionIndex].bDisabled;
-					NewLODModel.Sections[SectionIndex].GenerateUpToLodIndex = DestImportedResource->LODModels[DesiredLOD].Sections[SectionIndex].GenerateUpToLodIndex;
+					FSkelMeshSection& NewSection = NewLODModel.Sections[NewSectionIndex];
+					int32 MatIdx = NewSection.MaterialIndex;
+					for (int32 ExistSectionIndex = 0; ExistSectionIndex < ExistLODModel.Sections.Num(); ExistSectionIndex++)
+					{
+						if (OriginalImportMeshSectionSlotNames.IsValidIndex(ExistSectionIndex) && OriginalImportMeshSectionSlotNames[ExistSectionIndex] == InSkeletalMesh->Materials[MatIdx].ImportedMaterialSlotName)
+						{
+							const FSkelMeshSection& ExistSection = ExistLODModel.Sections[ExistSectionIndex];
+							//We found a match, restore the data
+							NewSection.bCastShadow = ExistSection.bCastShadow;
+							NewSection.bRecomputeTangent = ExistSection.bRecomputeTangent;
+							NewSection.bDisabled = ExistSection.bDisabled;
+							NewSection.GenerateUpToLodIndex = ExistSection.GenerateUpToLodIndex;
+							bool bBoneChunkedSection = NewSection.ChunkedParentSectionIndex >= 0;
+							int32 ParentOriginalSectionIndex = NewSection.OriginalDataSectionIndex;
+							if (!bBoneChunkedSection)
+							{
+								//Set the new Parent Index
+								FSkelMeshSourceSectionUserData& UserSectionData = NewLODModel.UserSectionsData.FindOrAdd(ParentOriginalSectionIndex);
+								UserSectionData.bDisabled = NewSection.bDisabled;
+								UserSectionData.bCastShadow = NewSection.bCastShadow;
+								UserSectionData.bRecomputeTangent = NewSection.bRecomputeTangent;
+								UserSectionData.GenerateUpToLodIndex = NewSection.GenerateUpToLodIndex;
+								//The cloth will be rebind later after the reimport is done
+							}
+							break;
+						}
+					}
 				}
 			}
 		}
-		
 	}
 
 	// If we want to add this as a new LOD to this mesh - add to LODModels/LODInfo array.
@@ -3794,8 +3829,9 @@ void UnFbx::FFbxImporter::InsertNewLODToBaseSkeletalMesh(USkeletalMesh* InSkelet
 	
 	LODInfo.LODMaterialMap.Empty();
 	// Now set up the material mapping array.
-	for (int32 MatIdx = 0; MatIdx < InSkeletalMesh->Materials.Num(); MatIdx++)
+	for (int32 SectionIndex = 0; SectionIndex < NewLODModel.Sections.Num(); SectionIndex++)
 	{
+		int32 MatIdx = NewLODModel.Sections[SectionIndex].MaterialIndex;
 		// Try and find the auto-assigned material in the array.
 		int32 LODMatIndex = INDEX_NONE;
 		//First try to match by name
@@ -3809,35 +3845,15 @@ void UnFbx::FFbxImporter::InsertNewLODToBaseSkeletalMesh(USkeletalMesh* InSkelet
 			}
 		}
 
-		//Then try to match by material
-		if (LODMatIndex == INDEX_NONE && InSkeletalMesh->Materials[MatIdx].MaterialInterface != NULL)
-		{
-			LODMatIndex = BaseSkeletalMesh->Materials.Find(InSkeletalMesh->Materials[MatIdx]);
-			
-			//If the material does not exist in the base skeletal mesh but is used in other LODs, add it.
-			if (LODMatIndex == INDEX_NONE)
-			{
-				LODMatIndex = BaseSkeletalMesh->Materials.Add(InSkeletalMesh->Materials[MatIdx]);
-			}
-		}
-
-		// If we didn't just use the index - but make sure its within range of the Materials array.
+		// If we dont have a match, add a new entry to the material list.
 		if (LODMatIndex == INDEX_NONE)
 		{
-			LODMatIndex = FMath::Clamp(MatIdx, 0, BaseSkeletalMesh->Materials.Num() - 1);
+			LODMatIndex = BaseSkeletalMesh->Materials.Add(InSkeletalMesh->Materials[MatIdx]);
 		}
 
 		LODInfo.LODMaterialMap.Add(LODMatIndex);
 	}
 
-	// if new LOD has more material slot, add the extra to main skeletal
-	if (BaseSkeletalMesh->Materials.Num() < InSkeletalMesh->Materials.Num())
-	{
-		BaseSkeletalMesh->Materials.AddZeroed(InSkeletalMesh->Materials.Num() - BaseSkeletalMesh->Materials.Num());
-	}
-
-	// same from here as FbxImporter
-
 	// Release all resources before replacing the model
 	BaseSkeletalMesh->PreEditChange(NULL);
 
@@ -4099,7 +4115,7 @@ bool UnFbx::FFbxImporter::ImportSkeletalMeshLOD(USkeletalMesh* InSkeletalMesh, U
 			OutMaterialIndex = Section.MaterialIndex;
 			if (Info.LODMaterialMap.IsValidIndex(SectionIndex) && Info.LODMaterialMap[SectionIndex] != INDEX_NONE)
 			{
-				OutMaterialIndex = Info.LODMaterialMap[Section.MaterialIndex];
+				OutMaterialIndex = Info.LODMaterialMap[SectionIndex];
 			}
 			FName ImportedMaterialSlotName = NAME_None;
 			if (SkelMesh->Materials.IsValidIndex(OutMaterialIndex))
