Adding the scalability button/menu as an optional menu in the CommonEditorViewportToolbar, and turning it on for Niagara.
#rnx

diff --git a/Engine/Plugins/FX/Niagara/Source/NiagaraEditor/Private/Widgets/SNiagaraSystemViewportToolBar.h b/Engine/Plugins/FX/Niagara/Source/NiagaraEditor/Private/Widgets/SNiagaraSystemViewportToolBar.h
--- a/Engine/Plugins/FX/Niagara/Source/NiagaraEditor/Private/Widgets/SNiagaraSystemViewportToolBar.h
+++ b/Engine/Plugins/FX/Niagara/Source/NiagaraEditor/Private/Widgets/SNiagaraSystemViewportToolBar.h
@@ -21,7 +21,10 @@ public:
 	// SCommonEditorViewportToolbarBase interface
 	virtual TSharedRef<SWidget> GenerateShowMenu() const override;
 	virtual void ExtendOptionsMenu(FMenuBuilder& OptionsMenuBuilder) const override;
-
+	virtual bool GetShowScalabilityMenu() const override
+	{
+		return true;
+	}
 	// End of SCommonEditorViewportToolbarBase
 
 	virtual bool IsViewModeSupported(EViewModeIndex ViewModeIndex) const override; 
diff --git a/Engine/Source/Editor/UnrealEd/Private/SCommonEditorViewportToolbarBase.cpp b/Engine/Source/Editor/UnrealEd/Private/SCommonEditorViewportToolbarBase.cpp
--- a/Engine/Source/Editor/UnrealEd/Private/SCommonEditorViewportToolbarBase.cpp
+++ b/Engine/Source/Editor/UnrealEd/Private/SCommonEditorViewportToolbarBase.cpp
@@ -14,6 +14,9 @@
 #include "EditorViewportCommands.h"
 #include "SEditorViewportToolBarMenu.h"
 #include "SEditorViewportViewMenu.h"
+#include "Editor/EditorPerformanceSettings.h"
+#include "Scalability.h"
+#include "SScalabilitySettings.h"
 #include "AssetEditorViewportLayout.h"
 #include "SAssetEditorViewport.h"
 
@@ -107,6 +110,21 @@ void SCommonEditorViewportToolbarBase::Construct(const FArguments& InArgs, TShar
 			.OnGetMenuContent(this, &SCommonEditorViewportToolbarBase::GenerateViewModeOptionsMenu)
 		];
 
+	MainBoxPtr->AddSlot()
+		.AutoWidth()
+		.Padding(ToolbarSlotPadding)
+		[
+			// Button to show scalability warnings
+			SNew(SEditorViewportToolbarMenu)
+			.ParentToolBar(SharedThis(this))
+			.Cursor(EMouseCursor::Default)
+			.Label(this, &SCommonEditorViewportToolbarBase::GetScalabilityWarningLabel)
+			.MenuStyle(FEditorStyle::Get(), "EditorViewportToolBar.MenuButtonWarning")
+			.OnGetMenuContent(this, &SCommonEditorViewportToolbarBase::GetScalabilityWarningMenuContent)
+			.Visibility(this, &SCommonEditorViewportToolbarBase::GetScalabilityWarningVisibility)
+			.ToolTipText(LOCTEXT("ScalabilityWarning_ToolTip", "Non-default scalability settings could be affecting what is shown in this viewport.\nFor example you may experience lower visual quality, reduced particle counts, and other artifacts that don't match what the scene would look like when running outside of the editor. Click to make changes."))
+			];
+
 	// Add optional toolbar slots to be added by child classes inherited from this common viewport toolbar
 	ExtendLeftAlignedToolbarSlots(MainBoxPtr, SharedThis(this));
 
@@ -489,4 +507,31 @@ TSharedRef<SEditorViewportViewMenu> SCommonEditorViewportToolbarBase::MakeViewMe
 		.MenuExtenders(GetViewMenuExtender());
 }
 
+FText SCommonEditorViewportToolbarBase::GetScalabilityWarningLabel() const
+{
+	const int32 QualityLevel = Scalability::GetQualityLevels().GetMinQualityLevel();
+	if (QualityLevel >= 0)
+	{
+		return FText::Format(LOCTEXT("ScalabilityWarning", "Scalability: {0}"), Scalability::GetScalabilityNameFromQualityLevel(QualityLevel));
+	}
+
+	return FText::GetEmpty();
+}
+
+EVisibility SCommonEditorViewportToolbarBase::GetScalabilityWarningVisibility() const
+{
+	//This method returns magic numbers. 3 means epic
+	return GetDefault<UEditorPerformanceSettings>()->bEnableScalabilityWarningIndicator && GetShowScalabilityMenu() && Scalability::GetQualityLevels().GetMinQualityLevel() < 3 ? EVisibility::Visible : EVisibility::Collapsed;
+}
+
+TSharedRef<SWidget> SCommonEditorViewportToolbarBase::GetScalabilityWarningMenuContent() const
+{
+	return
+		SNew(SBorder)
+		.BorderImage(FEditorStyle::GetBrush("Menu.Background"))
+		[
+			SNew(SScalabilitySettings)
+		];
+}
+
 #undef LOCTEXT_NAMESPACE
diff --git a/Engine/Source/Editor/UnrealEd/Public/SCommonEditorViewportToolbarBase.h b/Engine/Source/Editor/UnrealEd/Public/SCommonEditorViewportToolbarBase.h
--- a/Engine/Source/Editor/UnrealEd/Public/SCommonEditorViewportToolbarBase.h
+++ b/Engine/Source/Editor/UnrealEd/Public/SCommonEditorViewportToolbarBase.h
@@ -1,4 +1,4 @@
-﻿// Copyright Epic Games, Inc. All Rights Reserved.
+// Copyright Epic Games, Inc. All Rights Reserved.
 
 #pragma once
 
@@ -197,6 +197,13 @@ protected:
 	// Creates the view menu widget (override point for children)
 	virtual TSharedRef<class SEditorViewportViewMenu> MakeViewMenu();
 
+	FText GetScalabilityWarningLabel() const;
+	EVisibility GetScalabilityWarningVisibility() const;
+	TSharedRef<SWidget> GetScalabilityWarningMenuContent() const;
+	virtual bool GetShowScalabilityMenu() const
+	{
+		return false;
+	}
 	/** Called when the FOV slider is adjusted in the perspective viewport */
 	virtual void OnFOVValueChanged(float NewValue) const;
 
