Fixed a bug in ShaderPipelineCacheCommandlet that was making it generate too many PSO permutations on D3D12

diff --git a/Engine/Source/Editor/UnrealEd/Private/Commandlets/ShaderPipelineCacheToolsCommandlet.cpp b/Engine/Source/Editor/UnrealEd/Private/Commandlets/ShaderPipelineCacheToolsCommandlet.cpp
--- a/Engine/Source/Editor/UnrealEd/Private/Commandlets/ShaderPipelineCacheToolsCommandlet.cpp
+++ b/Engine/Source/Editor/UnrealEd/Private/Commandlets/ShaderPipelineCacheToolsCommandlet.cpp
@@ -448,15 +448,13 @@ static void StableShadersSerializationSelfTest(const TMultiMap<int32, FSHAHash>&
 // for example, if they come from two different vertex factories, we return false because that situation cannot occur
 bool CouldBeUsedTogether(const FStableShaderKeyAndValue& A, const FStableShaderKeyAndValue& B)
 {
-	// if the shaders belong to the same FShaderPipeline, they can be used to together, otherwise they never match
+	// if the shaders don't belong to the same FShaderPipeline, they cannot be used together
 	if ((A.PipelineHash != FSHAHash()) || (B.PipelineHash != FSHAHash()))
 	{
 		if (A.PipelineHash != B.PipelineHash)
 		{
 			return false;
 		}
-
-		return true;
 	}
 
 	static FName NAME_FDeferredDecalVS("FDeferredDecalVS");
