Prevent the distance-scaled camera speed from zooming out indefinitely.

diff --git a/Engine/Source/Editor/UnrealEd/Private/EditorViewportClient.cpp b/Engine/Source/Editor/UnrealEd/Private/EditorViewportClient.cpp
--- a/Engine/Source/Editor/UnrealEd/Private/EditorViewportClient.cpp
+++ b/Engine/Source/Editor/UnrealEd/Private/EditorViewportClient.cpp
@@ -2378,7 +2378,7 @@ void FEditorViewportClient::InputAxisForOrbit(FViewport* InViewport, const FVect
 		Drag.X = DragDelta.X;
 
 		FViewportCameraTransform& ViewTransform = GetViewTransform();
-		const float CameraSpeedDistanceScale = ShouldScaleCameraSpeedByDistance() ? FVector::Dist( GetViewLocation(), GetLookAtLocation() ) / 1000.f : 1.0f;
+		const float CameraSpeedDistanceScale = ShouldScaleCameraSpeedByDistance() ? FMath::Min(FVector::Dist( GetViewLocation(), GetLookAtLocation() ) / 1000.f, 1000.f) : 1.0f;
 
 		if ( IsOrbitRotationMode( InViewport ) )
 		{
@@ -5055,7 +5055,7 @@ void FEditorViewportClient::MoveViewportPerspectiveCamera( const FVector& InDrag
 	}
 
 	const float DistanceToCurrentLookAt = FVector::Dist( GetViewLocation(), GetLookAtLocation() );
-	const float CameraSpeedDistanceScale = ShouldScaleCameraSpeedByDistance() ? DistanceToCurrentLookAt / 1000.f : 1.f;
+	const float CameraSpeedDistanceScale = ShouldScaleCameraSpeedByDistance() ? FMath::Min(DistanceToCurrentLookAt / 1000.f, 1000.f) : 1.f;
 
 	// Update camera Location
 	ViewLocation += InDrag * CameraSpeedDistanceScale;
