Change the verbosity of detailed material info printed to the log in ShaderPipelineCacheToolsCommandlet to not spam the log unnecessarily

diff --git a/Engine/Source/Editor/UnrealEd/Private/Commandlets/ShaderPipelineCacheToolsCommandlet.cpp b/Engine/Source/Editor/UnrealEd/Private/Commandlets/ShaderPipelineCacheToolsCommandlet.cpp
--- a/Engine/Source/Editor/UnrealEd/Private/Commandlets/ShaderPipelineCacheToolsCommandlet.cpp
+++ b/Engine/Source/Editor/UnrealEd/Private/Commandlets/ShaderPipelineCacheToolsCommandlet.cpp
@@ -287,7 +287,7 @@ static void PrintShaders(const TMap<FSHAHash, TArray<FString>>& InverseMap, cons
 
 	for (const FString& Item : *Out)
 	{
-		UE_LOG(LogShaderPipelineCacheTools, Display, TEXT("    %s"), *Item);
+		UE_LOG(LogShaderPipelineCacheTools, Verbose, TEXT("    %s"), *Item);
 	}
 }
 
@@ -375,7 +375,7 @@ static void PrintShaders(const TMap<FSHAHash, TArray<int32>>& InverseMap, TArray
 	}
 	for (const int32& Item : *Out)
 	{
-		UE_LOG(LogShaderPipelineCacheTools, Display, TEXT("    %s"), *StableArray[Item].ToString());
+		UE_LOG(LogShaderPipelineCacheTools, Verbose, TEXT("    %s"), *StableArray[Item].ToString());
 	}
 }
 
