Fixing worldtime in realtime editor viewports.
This was a long standing change from 2007 that passed editor viewports a time other than world time so as to keep panning textures functional.
I can only assume at this time, world time was not upating correctly in realtime viewports. This is no longer the case and panning textures work just fine without this.
In addition, now that the render thread and game thread are using the same time value, FX visibiliy culling also works correctly.

diff --git a/Engine/Source/Editor/UnrealEd/Private/EditorViewportClient.cpp b/Engine/Source/Editor/UnrealEd/Private/EditorViewportClient.cpp
--- a/Engine/Source/Editor/UnrealEd/Private/EditorViewportClient.cpp
+++ b/Engine/Source/Editor/UnrealEd/Private/EditorViewportClient.cpp
@@ -3698,8 +3698,7 @@ void FEditorViewportClient::Draw(FViewport* InViewport, FCanvas* Canvas)
 	float DeltaTimeSeconds;
 
 	UWorld* World = GetWorld();
-	// During Simulation blueprints are directly using the World time, causing a mismatch with Material's Frame time
-	if (!World || ( GetScene() != World->Scene) || (IsRealtime() && !IsSimulateInEditorViewport()))
+	if (!World || (GetScene() != World->Scene))
 	{
 		// Use time relative to start time to avoid issues with float vs double
 		TimeSeconds = FApp::GetCurrentTime() - GStartTime;
