Source gather improvements

 - Per-file-type parsables.
 - Per-file-type comments (now ignores INI comments).
 - Early out when parsing line comment.
 - Correct improperly terminated literals, and report a warning.

diff --git a/Engine/Source/Editor/UnrealEd/Classes/Commandlets/GatherTextFromSourceCommandlet.h b/Engine/Source/Editor/UnrealEd/Classes/Commandlets/GatherTextFromSourceCommandlet.h
--- a/Engine/Source/Editor/UnrealEd/Classes/Commandlets/GatherTextFromSourceCommandlet.h
+++ b/Engine/Source/Editor/UnrealEd/Classes/Commandlets/GatherTextFromSourceCommandlet.h
@@ -11,6 +11,14 @@
 
 class Error;
 
+enum EGatherTextSourceFileTypes : uint8
+{
+	None = 0,
+	Cpp = 1 << 0,
+	Ini = 1 << 1,
+};
+ENUM_CLASS_FLAGS(EGatherTextSourceFileTypes);
+
 /**
  *	UGatherTextFromSourceCommandlet: Localization commandlet that collects all text to be localized from the source code.
  */
@@ -101,6 +109,7 @@ private:
 		void AddStringTableEntryMetaData( const FName InTableId, const FString& InKey, const FName InMetaDataId, const FString& InMetaData );
 
 		//Working data
+		EGatherTextSourceFileTypes FileTypes;
 		FString Filename;
 		int32 LineNumber;
 		FName FilePlatformName;
@@ -112,7 +121,7 @@ private:
 		bool WithinLineComment;
 		bool WithinStringLiteral;
 		bool WithinNamespaceDefine;
-		FString WithinStartingLine;
+		const TCHAR* WithinStartingLine;
 
 		//Should editor-only data be included in this gather?
 		bool ShouldGatherFromEditorOnlyData;
@@ -121,7 +130,8 @@ private:
 		TMap<FName, FParsedStringTable> ParsedStringTables;
 
 		FSourceFileParseContext(UGatherTextFromSourceCommandlet* InOwnerCommandlet)
-			: Filename()
+			: FileTypes(EGatherTextSourceFileTypes::None)
+			, Filename()
 			, LineNumber(0)
 			, FilePlatformName()
 			, LineText()
@@ -132,7 +142,7 @@ private:
 			, WithinLineComment(false)
 			, WithinStringLiteral(false)
 			, WithinNamespaceDefine(false)
-			, WithinStartingLine()
+			, WithinStartingLine(nullptr)
 			, ShouldGatherFromEditorOnlyData(false)
 			, MacroBlockStack()
 			, CachedEditorOnlyDefineState()
@@ -156,21 +166,27 @@ private:
 	class FParsableDescriptor
 	{
 	public:
-		FParsableDescriptor():bOverridesLongerTokens(false){}
-		FParsableDescriptor(bool bOverride):bOverridesLongerTokens(bOverride){}
-		virtual ~FParsableDescriptor(){}
+		virtual ~FParsableDescriptor() = default;
 		virtual const FString& GetToken() const = 0;
 		virtual void TryParse(const FString& Text, FSourceFileParseContext& Context) const = 0;
 
-		bool OverridesLongerTokens(){ return bOverridesLongerTokens; }
+		bool MatchesFileTypes(const EGatherTextSourceFileTypes InFileTypes) { return EnumHasAnyFlags(ApplicableFileTypes, InFileTypes); }
+		bool OverridesLongerTokens() { return bOverridesLongerTokens; }
+
 	protected:
-		bool bOverridesLongerTokens;
+		EGatherTextSourceFileTypes ApplicableFileTypes = EGatherTextSourceFileTypes::None;
+		bool bOverridesLongerTokens = false;
 	};
 
 	class FPreProcessorDescriptor : public FParsableDescriptor
 	{
 	public:
-		FPreProcessorDescriptor():FParsableDescriptor(true){}
+		FPreProcessorDescriptor()
+		{
+			ApplicableFileTypes = EGatherTextSourceFileTypes::Cpp;
+			bOverridesLongerTokens = true;
+		}
+
 	protected:
 		static const FString DefineString;
 		static const FString UndefString;
@@ -237,7 +253,11 @@ private:
 	public:
 		static const FString TextMacroString;
 
-		FMacroDescriptor(FString InName) : Name(MoveTemp(InName)) {}
+		FMacroDescriptor(FString InName)
+			: Name(MoveTemp(InName))
+		{
+			ApplicableFileTypes = EGatherTextSourceFileTypes::Cpp;
+		}
 
 		virtual const FString& GetToken() const override { return Name; }
 
@@ -253,12 +273,18 @@ private:
 	class FUICommandMacroDescriptor : public FMacroDescriptor
 	{
 	public:
-		FUICommandMacroDescriptor() : FMacroDescriptor(TEXT("UI_COMMAND")) {}
+		FUICommandMacroDescriptor()
+			: FMacroDescriptor(TEXT("UI_COMMAND"))
+		{
+		}
 
 		virtual void TryParse(const FString& Text, FSourceFileParseContext& Context) const override;
 
 	protected:
-		FUICommandMacroDescriptor(FString InName) : FMacroDescriptor(MoveTemp(InName)) {}
+		FUICommandMacroDescriptor(FString InName)
+			: FMacroDescriptor(MoveTemp(InName))
+		{
+		}
 
 		void TryParseArgs(const FString& Text, FSourceFileParseContext& Context, const TArray<FString>& Arguments, const int32 ArgIndexOffset) const;
 	};
@@ -266,7 +292,10 @@ private:
 	class FUICommandExtMacroDescriptor : public FUICommandMacroDescriptor
 	{
 	public:
-		FUICommandExtMacroDescriptor() : FUICommandMacroDescriptor(TEXT("UI_COMMAND_EXT")) {}
+		FUICommandExtMacroDescriptor()
+			: FUICommandMacroDescriptor(TEXT("UI_COMMAND_EXT"))
+		{
+		}
 
 		virtual void TryParse(const FString& Text, FSourceFileParseContext& Context) const override;
 	};
@@ -291,6 +320,7 @@ private:
 
 		FStringMacroDescriptor(FString InName, FMacroArg Arg0, FMacroArg Arg1, FMacroArg Arg2) : FMacroDescriptor(InName)
 		{
+			ApplicableFileTypes = EGatherTextSourceFileTypes::Cpp | EGatherTextSourceFileTypes::Ini;
 			Arguments.Add(Arg0);
 			Arguments.Add(Arg1);
 			Arguments.Add(Arg2);
@@ -298,12 +328,14 @@ private:
 
 		FStringMacroDescriptor(FString InName, FMacroArg Arg0, FMacroArg Arg1) : FMacroDescriptor(InName)
 		{
+			ApplicableFileTypes = EGatherTextSourceFileTypes::Cpp | EGatherTextSourceFileTypes::Ini;
 			Arguments.Add(Arg0);
 			Arguments.Add(Arg1);
 		}
 
 		FStringMacroDescriptor(FString InName, FMacroArg Arg0) : FMacroDescriptor(InName)
 		{
+			ApplicableFileTypes = EGatherTextSourceFileTypes::Cpp | EGatherTextSourceFileTypes::Ini;
 			Arguments.Add(Arg0);
 		}
 
@@ -316,7 +348,10 @@ private:
 	class FStringTableMacroDescriptor : public FMacroDescriptor
 	{
 	public:
-		FStringTableMacroDescriptor() : FMacroDescriptor(TEXT("LOCTABLE_NEW")) {}
+		FStringTableMacroDescriptor()
+			: FMacroDescriptor(TEXT("LOCTABLE_NEW"))
+		{
+		}
 
 		virtual void TryParse(const FString& Text, FSourceFileParseContext& Context) const override;
 	};
@@ -335,7 +370,10 @@ private:
 	class FStringTableEntryMacroDescriptor : public FMacroDescriptor
 	{
 	public:
-		FStringTableEntryMacroDescriptor() : FMacroDescriptor(TEXT("LOCTABLE_SETSTRING")) {}
+		FStringTableEntryMacroDescriptor()
+			: FMacroDescriptor(TEXT("LOCTABLE_SETSTRING"))
+		{
+		}
 
 		virtual void TryParse(const FString& Text, FSourceFileParseContext& Context) const override;
 	};
@@ -343,7 +381,10 @@ private:
 	class FStringTableEntryMetaDataMacroDescriptor : public FMacroDescriptor
 	{
 	public:
-		FStringTableEntryMetaDataMacroDescriptor() : FMacroDescriptor(TEXT("LOCTABLE_SETMETA")) {}
+		FStringTableEntryMetaDataMacroDescriptor()
+			: FMacroDescriptor(TEXT("LOCTABLE_SETMETA"))
+		{
+		}
 
 		virtual void TryParse(const FString& Text, FSourceFileParseContext& Context) const override;
 	};
@@ -351,6 +392,11 @@ private:
 	class FIniNamespaceDescriptor : public FPreProcessorDescriptor
 	{
 	public:
+		FIniNamespaceDescriptor()
+		{
+			ApplicableFileTypes = EGatherTextSourceFileTypes::Ini;
+		}
+
 		virtual const FString& GetToken() const override { return FPreProcessorDescriptor::IniNamespaceString; }
 		virtual void TryParse(const FString& Text, FSourceFileParseContext& Context) const override;
 	};
diff --git a/Engine/Source/Editor/UnrealEd/Private/Commandlets/GatherTextFromSourceCommandlet.cpp b/Engine/Source/Editor/UnrealEd/Private/Commandlets/GatherTextFromSourceCommandlet.cpp
--- a/Engine/Source/Editor/UnrealEd/Private/Commandlets/GatherTextFromSourceCommandlet.cpp
+++ b/Engine/Source/Editor/UnrealEd/Private/Commandlets/GatherTextFromSourceCommandlet.cpp
@@ -284,6 +284,7 @@ int32 UGatherTextFromSourceCommandlet::Main( const FString& Params )
 		}
 
 		ParseCtxt.Filename = SourceFile;
+		ParseCtxt.FileTypes = ParseCtxt.Filename.EndsWith(TEXT(".ini")) ? EGatherTextSourceFileTypes::Ini : EGatherTextSourceFileTypes::Cpp;
 		FPaths::MakePathRelativeTo(ParseCtxt.Filename, *ProjectBasePath);
 		ParseCtxt.LineNumber = 0;
 		ParseCtxt.FilePlatformName = GetSplitPlatformNameFromPath(ParseCtxt.Filename);
@@ -294,7 +295,7 @@ int32 UGatherTextFromSourceCommandlet::Main( const FString& Params )
 		ParseCtxt.WithinLineComment = false;
 		ParseCtxt.WithinStringLiteral = false;
 		ParseCtxt.WithinNamespaceDefine = false;
-		ParseCtxt.WithinStartingLine.Empty();
+		ParseCtxt.WithinStartingLine = nullptr;
 		ParseCtxt.FlushMacroStack();
 
 		FString SourceFileText;
@@ -685,20 +686,48 @@ FString UGatherTextFromSourceCommandlet::StripCommentsFromToken(const FString& I
 	return StrippedToken.TrimStartAndEnd();
 }
 
-bool UGatherTextFromSourceCommandlet::ParseSourceText(const FString& Text, const TArray<UGatherTextFromSourceCommandlet::FParsableDescriptor*>& Parsables, FSourceFileParseContext& ParseCtxt)
+bool UGatherTextFromSourceCommandlet::ParseSourceText(const FString& Text, const TArray<FParsableDescriptor*>& Parsables, FSourceFileParseContext& ParseCtxt)
 {
-	// Create array of ints, one for each parsable we're looking for.
-	TArray<int32> ParsableMatchCounters;
-	ParsableMatchCounters.AddZeroed(Parsables.Num());
+	// Cache array of parsables and tokens valid for this filetype
+	TArray< FParsableDescriptor*> ParsablesForFile;
+	TArray<FString> ParsableTokensForFile;
+	for (FParsableDescriptor* Parsable : Parsables)
+	{
+		if (Parsable->MatchesFileTypes(ParseCtxt.FileTypes))
+		{
+			ParsablesForFile.Add(Parsable);
+			ParsableTokensForFile.Add(Parsable->GetToken());
+		}
+	}
+	check(ParsablesForFile.Num() == ParsableTokensForFile.Num());
 
-	// Cache array of tokens
-	TArray<FString> ParsableTokens;
-	for (int32 ParIdx=0; ParIdx<Parsables.Num(); ParIdx++)
+	// Anything to parse for this filetype?
+	if (ParsablesForFile.Num() == 0)
 	{
-		ParsableTokens.Add(Parsables[ParIdx]->GetToken());
+		return true;
 	}
 
-	// Split the file into lines of 
+	// Create array of ints, one for each parsable we're looking for.
+	TArray<int32> ParsableMatchCountersForFile;
+	ParsableMatchCountersForFile.AddZeroed(ParsablesForFile.Num());
+
+	// Use the file extension to work out what comments look like for this file
+	// We default to C++-style comments
+	const TCHAR* LineComment = TEXT("//");
+	const TCHAR* BlockCommentStart = TEXT("/*");
+	const TCHAR* BlockCommentEnd = TEXT("*/");
+	if (EnumHasAnyFlags(ParseCtxt.FileTypes, EGatherTextSourceFileTypes::Ini))
+	{
+		LineComment = TEXT(";");
+		BlockCommentStart = nullptr;
+		BlockCommentEnd = nullptr;
+	}
+	const int32 LineCommentLen = LineComment ? FCString::Strlen(LineComment) : 0;
+	const int32 BlockCommentStartLen = BlockCommentStart ? FCString::Strlen(BlockCommentStart) : 0;
+	const int32 BlockCommentEndLen = BlockCommentEnd ? FCString::Strlen(BlockCommentEnd) : 0;
+	checkf((BlockCommentStartLen == 0 && BlockCommentEndLen == 0) || (BlockCommentStartLen > 0 && BlockCommentEndLen > 0), TEXT("Block comments require both a start and an end marker!"));
+
+	// Split the file into lines 
 	TArray<FString> TextLines;
 	Text.ParseIntoArrayLines(TextLines, false);
 
@@ -711,39 +740,38 @@ bool UGatherTextFromSourceCommandlet::ParseSourceText(const FString& Text, const
 			continue;
 
 		// Use these pending vars to defer parsing a token hit until longer tokens can't hit too
-		int32 PendingParseIdx = -1;
+		int32 PendingParseIdx = INDEX_NONE;
 		const TCHAR* ParsePoint = NULL;
-		ParsableMatchCounters.Empty(Parsables.Num());
-		ParsableMatchCounters.AddZeroed(Parsables.Num());
+		for (int32& Element : ParsableMatchCountersForFile)
+		{
+			Element = 0;
+		}
 		ParseCtxt.LineNumber = LineIdx + 1;
 		ParseCtxt.LineText = Line;
+		ParseCtxt.WithinLineComment = false;
 		ParseCtxt.EndParsingCurrentLine = false;
 
 		const TCHAR* Cursor = *Line;
-		bool EndOfLine = false;
-		while (!EndOfLine && !ParseCtxt.EndParsingCurrentLine)
+		while (*Cursor && !ParseCtxt.EndParsingCurrentLine)
 		{
-			// Check if we're starting comments or string literals. Begins *at* "//" or "/*".
+			// Check if we're starting comments or string literals
+
 			if (!ParseCtxt.WithinLineComment && !ParseCtxt.WithinBlockComment && !ParseCtxt.WithinStringLiteral)
 			{
-				if(*Cursor == TEXT('/'))
+				if (LineCommentLen > 0 && FCString::Strncmp(Cursor, LineComment, LineCommentLen) == 0)
 				{
-					const TCHAR* const ForwardCursor = Cursor + 1;
-					switch(*ForwardCursor)
-					{
-					case TEXT('/'):
-						{
-							ParseCtxt.WithinLineComment = true;
-							ParseCtxt.WithinStartingLine = ParseCtxt.LineText;
-						}
-						break;
-					case TEXT('*'):
-						{
-							ParseCtxt.WithinBlockComment = true;
-							ParseCtxt.WithinStartingLine = ParseCtxt.LineText;
-						}
-						break;
-					}
+					ParseCtxt.WithinLineComment = true;
+					ParseCtxt.WithinStartingLine = *ParseCtxt.LineText;
+					ParseCtxt.EndParsingCurrentLine = true;
+					Cursor += LineCommentLen;
+					continue;
+				}
+				else if (BlockCommentStartLen > 0 && FCString::Strncmp(Cursor, BlockCommentStart, BlockCommentStartLen) == 0)
+				{
+					ParseCtxt.WithinBlockComment = true;
+					ParseCtxt.WithinStartingLine = *ParseCtxt.LineText;
+					Cursor += BlockCommentStartLen;
+					continue;
 				}
 			}
 
@@ -754,7 +782,9 @@ bool UGatherTextFromSourceCommandlet::ParseSourceText(const FString& Text, const
 					if (Cursor == *Line)
 					{
 						ParseCtxt.WithinStringLiteral = true;
-						ParseCtxt.WithinStartingLine = ParseCtxt.LineText;
+						ParseCtxt.WithinStartingLine = *ParseCtxt.LineText;
+						++Cursor;
+						continue;
 					}
 					else if (Cursor > *Line)
 					{
@@ -762,7 +792,9 @@ bool UGatherTextFromSourceCommandlet::ParseSourceText(const FString& Text, const
 						if (*ReverseCursor != TEXT('\\') && *ReverseCursor != TEXT('\''))
 						{
 							ParseCtxt.WithinStringLiteral = true;
-							ParseCtxt.WithinStartingLine = ParseCtxt.LineText;
+							ParseCtxt.WithinStartingLine = *ParseCtxt.LineText;
+							++Cursor;
+							continue;
 						}
 						else 
 						{
@@ -779,7 +811,9 @@ bool UGatherTextFromSourceCommandlet::ParseSourceText(const FString& Text, const
 							if (IsEscaped)
 							{
 								ParseCtxt.WithinStringLiteral = true;
-								ParseCtxt.WithinStartingLine = ParseCtxt.LineText;
+								ParseCtxt.WithinStartingLine = *ParseCtxt.LineText;
+								++Cursor;
+								continue;
 							}
 							else
 							{
@@ -788,7 +822,9 @@ bool UGatherTextFromSourceCommandlet::ParseSourceText(const FString& Text, const
 								if (*ReverseCursor == TEXT('\'') && *ForwardCursor != TEXT('\''))
 								{
 									ParseCtxt.WithinStringLiteral = true;
-									ParseCtxt.WithinStartingLine = ParseCtxt.LineText;
+									ParseCtxt.WithinStartingLine = *ParseCtxt.LineText;
+									++Cursor;
+									continue;
 								}
 							}
 						}
@@ -802,6 +838,8 @@ bool UGatherTextFromSourceCommandlet::ParseSourceText(const FString& Text, const
 					if (Cursor == *Line)
 					{
 						ParseCtxt.WithinStringLiteral = false;
+						++Cursor;
+						continue;
 					}
 					else if (Cursor > *Line)
 					{
@@ -809,6 +847,8 @@ bool UGatherTextFromSourceCommandlet::ParseSourceText(const FString& Text, const
 						if (*ReverseCursor != TEXT('\\') && *ReverseCursor != TEXT('\''))
 						{
 							ParseCtxt.WithinStringLiteral = false;
+							++Cursor;
+							continue;
 						}
 						else
 						{
@@ -825,6 +865,8 @@ bool UGatherTextFromSourceCommandlet::ParseSourceText(const FString& Text, const
 							if (IsEscaped)
 							{
 								ParseCtxt.WithinStringLiteral = false;
+								++Cursor;
+								continue;
 							}
 							else
 							{
@@ -833,6 +875,8 @@ bool UGatherTextFromSourceCommandlet::ParseSourceText(const FString& Text, const
 								if (*ReverseCursor == TEXT('\'') && *ForwardCursor != TEXT('\''))
 								{
 									ParseCtxt.WithinStringLiteral = false;
+									++Cursor;
+									continue;
 								}
 							}
 						}
@@ -840,33 +884,29 @@ bool UGatherTextFromSourceCommandlet::ParseSourceText(const FString& Text, const
 				}
 			}
 
-			// Check if we're ending comments. Ends *after* "*/".
-			if(ParseCtxt.WithinBlockComment)
+			// Check if we're ending comments
+			if (ParseCtxt.WithinBlockComment && BlockCommentEndLen > 0 && FCString::Strncmp(Cursor, BlockCommentEnd, BlockCommentEndLen) == 0)
 			{
-				if(*Cursor == TEXT('/') && Cursor > *Line)
-				{
-					const TCHAR* const ReverseCursor = Cursor - 1;
-					if (*ReverseCursor == TEXT('*'))
-					{
-						ParseCtxt.WithinBlockComment = false;
-					}
-				}
+				ParseCtxt.WithinBlockComment = false;
+				Cursor += BlockCommentEndLen;
+				continue;
 			}
 
-			for (int32 ParIdx=0; ParIdx<Parsables.Num(); ParIdx++)
+			for (int32 ParIdx = 0; ParIdx < ParsablesForFile.Num(); ++ParIdx)
 			{
-				FString& Token = ParsableTokens[ParIdx];
+				FParsableDescriptor* Parsable = ParsablesForFile[ParIdx];
+				const FString& Token = ParsableTokensForFile[ParIdx];
 
-				if (Token.Len() == ParsableMatchCounters[ParIdx])
+				if (Token.Len() == ParsableMatchCountersForFile[ParIdx])
 				{
 					// already seen this entire token and are looking for longer matches - skip it
 					continue;
 				}
 
-				if (*Cursor == Token[ParsableMatchCounters[ParIdx]])
+				if (*Cursor == Token[ParsableMatchCountersForFile[ParIdx]])
 				{
 					// Char at cursor matches the next char in the parsable's identifying token
-					if (Token.Len() == ++(ParsableMatchCounters[ParIdx]))
+					if (Token.Len() == ++(ParsableMatchCountersForFile[ParIdx]))
 					{
 						// don't immediately parse - this parsable has seen its entire token but a longer one could be about to hit too
 						const TCHAR* TokenStart = Cursor + 1 - Token.Len();
@@ -881,25 +921,24 @@ bool UGatherTextFromSourceCommandlet::ParseSourceText(const FString& Text, const
 				{
 					// Char at cursor doesn't match the next char in the parsable's identifying token
 					// Reset the counter to start of the token
-					ParsableMatchCounters[ParIdx] = 0;
+					ParsableMatchCountersForFile[ParIdx] = 0;
 				}
 			}
 
 			// Now check PendingParse and only run it if there are no better candidates
-			if (0 <= PendingParseIdx)
+			if (PendingParseIdx != INDEX_NONE)
 			{
+				FParsableDescriptor* PendingParsable = ParsablesForFile[PendingParseIdx];
+
 				bool MustDefer = false; // pending will be deferred if another parsable has a equal and greater number of matched chars
-				if( !Parsables[PendingParseIdx]->OverridesLongerTokens() )
+				if (!PendingParsable->OverridesLongerTokens())
 				{
-					for (int32 ParIdx=0; ParIdx<Parsables.Num(); ParIdx++)
+					for (int32 ParIdx = 0; ParIdx < ParsablesForFile.Num(); ++ParIdx)
 					{
-						if (PendingParseIdx != ParIdx)
+						if (PendingParseIdx != ParIdx && ParsableMatchCountersForFile[ParIdx] >= ParsableTokensForFile[PendingParseIdx].Len())
 						{
-							if (ParsableMatchCounters[ParIdx] >= ParsableTokens[PendingParseIdx].Len())
-							{
-								// a longer token is matching so defer
-								MustDefer = true;
-							}
+							// a longer token is matching so defer
+							MustDefer = true;
 						}
 					}
 				}
@@ -907,21 +946,36 @@ bool UGatherTextFromSourceCommandlet::ParseSourceText(const FString& Text, const
 				if (!MustDefer)
 				{
 					// Do the parse now
-					Parsables[PendingParseIdx]->TryParse(FString(ParsePoint), ParseCtxt);
-					for(auto& Element : ParsableMatchCounters)
+					// TODO: Would be nice if TryParse returned what it consumed, and operated on const TCHAR*
+					PendingParsable->TryParse(FString(ParsePoint), ParseCtxt);
+					for (int32& Element : ParsableMatchCountersForFile)
 					{
 						Element = 0;
 					}
-					PendingParseIdx = -1;
+					PendingParseIdx = INDEX_NONE;
 					ParsePoint = NULL;
 				}
 			}
 
-			EndOfLine = ('\0' == *(++Cursor)) ? true : false;
-			if(EndOfLine)
+			// Advance cursor
+			++Cursor;
+		}
+
+		// Handle a string literal that went beyond a single line
+		if (ParseCtxt.WithinStringLiteral)
+		{
+			if (EnumHasAnyFlags(ParseCtxt.FileTypes, EGatherTextSourceFileTypes::Ini))
 			{
-				ParseCtxt.WithinLineComment = false;
+				// INI files don't support multi-line literals; always terminate them after ending a line
+				ParseCtxt.WithinStringLiteral = false;
 			}
+			else if (Cursor > *Line)
+			{
+				// C++ only allows multi-line literals if they're escaped with a trailing slash
+				ParseCtxt.WithinStringLiteral = *(Cursor - 1) == TEXT('\\');
+			}
+
+			UE_CLOG(!ParseCtxt.WithinStringLiteral, LogGatherTextFromSourceCommandlet, Warning, TEXT("A string literal was not correctly terminated. File %s at line %d"), *ParseCtxt.Filename, ParseCtxt.LineNumber);
 		}
 	}
 
@@ -1819,8 +1873,7 @@ void UGatherTextFromSourceCommandlet::FIniNamespaceDescriptor::TryParse(const FS
 	// [<config section name>]
 	if (!Context.ExcludedRegion)
 	{
-		if( FCString::Stricmp( *FPaths::GetExtension( Context.Filename, false ), TEXT("ini") ) == 0 &&
-			Context.LineText[ 0 ] == '[' )
+		if( Context.LineText[ 0 ] == '[' )
 		{
 			int32 ClosingBracket;
 			if( Text.FindChar( ']', ClosingBracket ) && ClosingBracket > 1 )
