Clang complains about other enums shadowing this one.

diff --git a/Engine/Source/Editor/UnrealEd/Classes/Commandlets/GatherTextFromSourceCommandlet.h b/Engine/Source/Editor/UnrealEd/Classes/Commandlets/GatherTextFromSourceCommandlet.h
--- a/Engine/Source/Editor/UnrealEd/Classes/Commandlets/GatherTextFromSourceCommandlet.h
+++ b/Engine/Source/Editor/UnrealEd/Classes/Commandlets/GatherTextFromSourceCommandlet.h
@@ -11,7 +11,7 @@
 
 class Error;
 
-enum EGatherTextSourceFileTypes : uint8
+enum class EGatherTextSourceFileTypes : uint8
 {
 	None = 0,
 	Cpp = 1 << 0,
