Repairs to the WorldNavigationData commandlet
#rnx

diff --git a/Engine/Source/Editor/UnrealEd/Private/EditorEngine.cpp b/Engine/Source/Editor/UnrealEd/Private/EditorEngine.cpp
--- a/Engine/Source/Editor/UnrealEd/Private/EditorEngine.cpp
+++ b/Engine/Source/Editor/UnrealEd/Private/EditorEngine.cpp
@@ -6239,6 +6239,11 @@ bool UEditorEngine::IsAnyViewportRealtime()
 
 bool UEditorEngine::ShouldThrottleCPUUsage() const
 {
+	if (IsRunningCommandlet())
+	{
+		return false;
+	}
+
 	bool bShouldThrottle = false;
 
 	const bool bRunningCommandlet = IsRunningCommandlet();
@@ -6251,26 +6256,29 @@ bool UEditorEngine::ShouldThrottleCPUUsage() const
 		bShouldThrottle = Settings->bThrottleCPUWhenNotForeground;
 
 		// Check if we should throttle due to all windows being minimized
-		if ( !bShouldThrottle )
-		{
-			 bShouldThrottle = AreAllWindowsHidden();
-		}
-
-		// Do not throttle during drag and drop
-		if (bShouldThrottle && FSlateApplication::Get().IsDragDropping())
+		if (FSlateApplication::IsInitialized())
 		{
-			bShouldThrottle = false;
-		}
+			if (!bShouldThrottle)
+			{
+				bShouldThrottle = AreAllWindowsHidden();
+			}
 
-		if (bShouldThrottle)
-		{
-			static const FName AssetRegistryName(TEXT("AssetRegistry"));
-			FAssetRegistryModule* AssetRegistryModule = FModuleManager::GetModulePtr<FAssetRegistryModule>(AssetRegistryName);
-			// Don't throttle during amortized export, greatly increases export time
-			if (IsLightingBuildCurrentlyExporting() || GShaderCompilingManager->IsCompiling() || (AssetRegistryModule && AssetRegistryModule->Get().IsLoadingAssets()))
+			// Do not throttle during drag and drop
+			if (bShouldThrottle && FSlateApplication::Get().IsDragDropping())
 			{
 				bShouldThrottle = false;
 			}
+
+			if (bShouldThrottle)
+			{
+				static const FName AssetRegistryName(TEXT("AssetRegistry"));
+				FAssetRegistryModule* AssetRegistryModule = FModuleManager::GetModulePtr<FAssetRegistryModule>(AssetRegistryName);
+				// Don't throttle during amortized export, greatly increases export time
+				if (IsLightingBuildCurrentlyExporting() || GShaderCompilingManager->IsCompiling() || (AssetRegistryModule && AssetRegistryModule->Get().IsLoadingAssets()))
+				{
+					bShouldThrottle = false;
+				}
+			}
 		}
 	}
 
