Check container types when determining if there are zombie members of type User Defined struct

diff --git a/Engine/Source/Editor/UnrealEd/Private/Kismet2/StructureEditorUtils.cpp b/Engine/Source/Editor/UnrealEd/Private/Kismet2/StructureEditorUtils.cpp
--- a/Engine/Source/Editor/UnrealEd/Private/Kismet2/StructureEditorUtils.cpp
+++ b/Engine/Source/Editor/UnrealEd/Private/Kismet2/StructureEditorUtils.cpp
@@ -569,11 +569,18 @@ void FStructureEditorUtils::RemoveInvalidStructureMemberVariableFromBlueprint(UB
 		for (int32 VarIndex = 0; VarIndex < Blueprint->NewVariables.Num(); ++VarIndex)
 		{
 			const FBPVariableDescription& Var = Blueprint->NewVariables[VarIndex];
-			if (Var.VarType.PinCategory == UEdGraphSchema_K2::PC_Struct)
+			if (Var.VarType.PinCategory == UEdGraphSchema_K2::PC_Struct || Var.VarType.IsContainer())
 			{
-				const UScriptStruct* ScriptStruct = Cast<const UScriptStruct>(Var.VarType.PinSubCategoryObject.Get());
-				const bool bInvalidStruct = (NULL == ScriptStruct) || (FallbackStruct == ScriptStruct);
-				if (bInvalidStruct)
+				UScriptStruct* ScriptStruct = Cast<UScriptStruct>(Var.VarType.PinSubCategoryObject.Get());
+				
+				// If this script is a container then we need to check the value type's validity
+				if (!ScriptStruct)
+				{
+					ScriptStruct = Cast<UScriptStruct>(Var.VarType.PinValueType.TerminalSubCategoryObject.Get());
+				}
+				
+				// If there is no underlying struct object or it is marked to use the fallback struct then this variable is now invalid
+				if (!ScriptStruct || (FallbackStruct == ScriptStruct))
 				{
 					DislpayList += Var.FriendlyName.IsEmpty() ? Var.VarName.ToString() : Var.FriendlyName;
 					DislpayList += TEXT("\n");
