Fix crash when trying to restore package at startup that no longer exists

diff --git a/Engine/Source/Editor/UnrealEd/Private/PackageRestore.cpp b/Engine/Source/Editor/UnrealEd/Private/PackageRestore.cpp
--- a/Engine/Source/Editor/UnrealEd/Private/PackageRestore.cpp
+++ b/Engine/Source/Editor/UnrealEd/Private/PackageRestore.cpp
@@ -553,11 +553,13 @@ FEditorFileUtils::EPromptReturnCode PackageRestore::PromptToRestorePackages(cons
 		else
 		{
 			// A package may not exist on disk if it was for a newly added or imported asset, which hasn't yet had SaveDirtyPackages called for it
-			PackageFilename = FPackageName::LongPackageNameToFilename(PackageFullPath); // no extension yet
-			PackageFilename += FPaths::GetExtension(AutoSavePath, true/*bIncludeDot*/);
+			if (FPackageName::TryConvertLongPackageNameToFilename(PackageFullPath, PackageFilename)) // no extension yet
+			{
+				PackageFilename += FPaths::GetExtension(AutoSavePath, true/*bIncludeDot*/);
 
-			FPackageRestoreItemPtr PackageItemPtr = MakeShareable(new FPackageRestoreItem(PackageFullPath, PackageFilename, AutoSaveDir / AutoSavePath, false/*bIsExistingPackage*/));
-			PackageRestoreItems.Add(PackageItemPtr);
+				FPackageRestoreItemPtr PackageItemPtr = MakeShareable(new FPackageRestoreItem(PackageFullPath, PackageFilename, AutoSaveDir / AutoSavePath, false/*bIsExistingPackage*/));
+				PackageRestoreItems.Add(PackageItemPtr);
+			}
 		}
 	}
 
