#UE4 Passing the current level's world into the asset reference filter when dragdropping an actor into a level, since it is the one that will have the reference when the actor is placed.

diff --git a/Engine/Source/Editor/UnrealEd/Private/LevelEditorViewport.cpp b/Engine/Source/Editor/UnrealEd/Private/LevelEditorViewport.cpp
--- a/Engine/Source/Editor/UnrealEd/Private/LevelEditorViewport.cpp
+++ b/Engine/Source/Editor/UnrealEd/Private/LevelEditorViewport.cpp
@@ -1310,8 +1310,11 @@ FDropQuery FLevelEditorViewportClient::CanDropObjectsAtCoordinates(int32 MouseX,
 
 	if (CurrentWorld)
 	{
+		ULevel* CurrentLevel = CurrentWorld->GetCurrentLevel();
+		UWorld* CurrentLevelOuterWorld = CurrentLevel ? Cast<UWorld>(CurrentLevel->GetOuter()) : nullptr;  
+		UWorld* ReferencingWorld = CurrentLevelOuterWorld ? CurrentLevelOuterWorld : CurrentWorld;
 		FAssetReferenceFilterContext AssetReferenceFilterContext;
-		AssetReferenceFilterContext.ReferencingAssets.Add(FAssetData(CurrentWorld));
+		AssetReferenceFilterContext.ReferencingAssets.Add(FAssetData(ReferencingWorld));
 
 		TSharedPtr<IAssetReferenceFilter> AssetReferenceFilter = GEditor->MakeAssetReferenceFilter(AssetReferenceFilterContext);
 		if (AssetReferenceFilter.IsValid())
