Removed Editor Modes are also removed from the set of default modes for a mode manager.

# Speculative fix for Activate Mode attempting to activate an invalid editor mode. When modes are destroyed, they may not be removed out of the default modes for the manager, which may cause an Activate after a deactivate to fail.

diff --git a/Engine/Source/Editor/UnrealEd/Private/EditorModeManager.cpp b/Engine/Source/Editor/UnrealEd/Private/EditorModeManager.cpp
--- a/Engine/Source/Editor/UnrealEd/Private/EditorModeManager.cpp
+++ b/Engine/Source/Editor/UnrealEd/Private/EditorModeManager.cpp
@@ -903,6 +903,15 @@ FEdMode* FEditorModeTools::FindMode(FEditorModeID InID)
 
 void FEditorModeTools::DestroyMode( FEditorModeID InID )
 {
+	// Since deactivating the last active mode will cause the default modes to be activated, make sure this mode is removed from defaults.
+	RemoveDefaultMode( InID );
+	
+	// Add back the default default mode if we just removed the last valid default.
+	if ( DefaultModeIDs.Num() == 0 )
+	{
+		AddDefaultMode( FBuiltinEditorModes::EM_Default );
+	}
+
 	// Find the mode from the ID and exit it.
 	for( int32 Index = ActiveModes.Num() - 1; Index >= 0; --Index )
 	{
