Fix invalid warnings being thrown for containers of object reference types

#rnx

diff --git a/Engine/Source/Editor/UnrealEd/Private/Kismet2/StructureEditorUtils.cpp b/Engine/Source/Editor/UnrealEd/Private/Kismet2/StructureEditorUtils.cpp
--- a/Engine/Source/Editor/UnrealEd/Private/Kismet2/StructureEditorUtils.cpp
+++ b/Engine/Source/Editor/UnrealEd/Private/Kismet2/StructureEditorUtils.cpp
@@ -569,23 +569,28 @@ void FStructureEditorUtils::RemoveInvalidStructureMemberVariableFromBlueprint(UB
 		for (int32 VarIndex = 0; VarIndex < Blueprint->NewVariables.Num(); ++VarIndex)
 		{
 			const FBPVariableDescription& Var = Blueprint->NewVariables[VarIndex];
-			if (Var.VarType.PinCategory == UEdGraphSchema_K2::PC_Struct || Var.VarType.IsContainer())
+
+			bool bIsInvalid = false;
+
+			if (Var.VarType.PinCategory == UEdGraphSchema_K2::PC_Struct)
 			{
+				// The variable is invalid if the struct object is null, or it points to the fallback struct
 				UScriptStruct* ScriptStruct = Cast<UScriptStruct>(Var.VarType.PinSubCategoryObject.Get());
-				
-				// If this script is a container then we need to check the value type's validity
-				if (!ScriptStruct)
-				{
-					ScriptStruct = Cast<UScriptStruct>(Var.VarType.PinValueType.TerminalSubCategoryObject.Get());
-				}
-				
-				// If there is no underlying struct object or it is marked to use the fallback struct then this variable is now invalid
-				if (!ScriptStruct || (FallbackStruct == ScriptStruct))
-				{
-					DislpayList += Var.FriendlyName.IsEmpty() ? Var.VarName.ToString() : Var.FriendlyName;
-					DislpayList += TEXT("\n");
-					ZombieMemberNames.Add(Var.VarName);
-				}
+
+				bIsInvalid = (!ScriptStruct || (FallbackStruct == ScriptStruct));
+			}
+			else if (Var.VarType.IsMap() && Var.VarType.PinValueType.TerminalCategory == UEdGraphSchema_K2::PC_Struct)
+			{
+				// If there is no ValueType object then the variable is invalid
+				bIsInvalid = (!Var.VarType.PinValueType.TerminalSubCategoryObject.Get());
+			}
+
+			// If this variable is invalid then display a warning
+			if (bIsInvalid)
+			{
+				DislpayList += Var.FriendlyName.IsEmpty() ? Var.VarName.ToString() : Var.FriendlyName;
+				DislpayList += TEXT("\n");
+				ZombieMemberNames.Add(Var.VarName);
 			}
 		}
 
