Added FGA exporter for UVectorFieldStatic assets

diff --git a/Engine/Source/Editor/UnrealEd/Classes/Exporters/VectorFieldExporter.h b/Engine/Source/Editor/UnrealEd/Classes/Exporters/VectorFieldExporter.h
new file mode 100644
--- /dev/null
+++ b/Engine/Source/Editor/UnrealEd/Classes/Exporters/VectorFieldExporter.h
@@ -0,0 +1,27 @@
+// Copyright Epic Games, Inc. All Rights Reserved.
+
+//~=============================================================================
+// VectorFieldExporter
+//~=============================================================================
+
+#pragma once
+
+#include "CoreMinimal.h"
+#include "UObject/ObjectMacros.h"
+#include "Exporters/Exporter.h"
+#include "VectorFieldExporter.generated.h"
+
+UCLASS()
+class UNREALED_API UVectorFieldExporter : public UExporter
+{
+	GENERATED_BODY()
+
+	UVectorFieldExporter(const FObjectInitializer& ObjectInitializer = FObjectInitializer());
+
+	//~ Begin UExporter Interface
+	virtual bool ExportText(const FExportObjectInnerContext* Context, UObject* Object, const TCHAR* Type, FOutputDevice& Ar, FFeedbackContext* Warn, uint32 PortFlags = 0) override;
+	//~ End UExporter Interface
+};
+
+
+
diff --git a/Engine/Source/Editor/UnrealEd/Private/Factories/VectorFieldFactory.cpp b/Engine/Source/Editor/UnrealEd/Private/Factories/VectorFieldFactory.cpp
--- a/Engine/Source/Editor/UnrealEd/Private/Factories/VectorFieldFactory.cpp
+++ b/Engine/Source/Editor/UnrealEd/Private/Factories/VectorFieldFactory.cpp
@@ -14,6 +14,7 @@
 #include "EditorFramework/AssetImportData.h"
 #include "Particles/ParticleSystemComponent.h"
 #include "Components/VectorFieldComponent.h"
+#include "Exporters/VectorFieldExporter.h"
 #include "Editor.h"
 
 #include "UObject/UObjectHash.h"
@@ -356,4 +357,48 @@ int32 UReimportVectorFieldStaticFactory::GetPriority() const
 	return ImportPriority;
 }
 
+/*-----------------------------------------------------------------------------
+	UVectorFieldExporter
+-----------------------------------------------------------------------------*/
+UVectorFieldExporter::UVectorFieldExporter(const FObjectInitializer& ObjectInitializer)
+	: Super(ObjectInitializer)
+{
+	SupportedClass = UVectorFieldStatic::StaticClass();
+	bText = true;
+	PreferredFormatIndex = 0;
+	FormatExtension.Add(TEXT("FGA"));
+	FormatDescription.Add(TEXT("Vector Field"));
+}
+
+bool UVectorFieldExporter::ExportText(const FExportObjectInnerContext* Context, UObject* Object, const TCHAR* Type, FOutputDevice& Ar, FFeedbackContext* Warn, uint32 PortFlags)
+{
+	UVectorFieldStatic* LocalVectorField = CastChecked<UVectorFieldStatic>(StaticDuplicateObject(Object, GetTransientPackage()));
+
+	LocalVectorField->SetCPUAccessEnabled();
+
+	Ar.Logf(TEXT("%d, %d, %d"), LocalVectorField->SizeX, LocalVectorField->SizeY, LocalVectorField->SizeZ);
+	Ar.Logf(TEXT(", "));
+	Ar.Logf(TEXT("%f, %f, %f, %f, %f, %f"),
+		LocalVectorField->Bounds.Min.X,
+		LocalVectorField->Bounds.Min.Y,
+		LocalVectorField->Bounds.Min.Z,
+		LocalVectorField->Bounds.Max.X,
+		LocalVectorField->Bounds.Max.Y,
+		LocalVectorField->Bounds.Max.Z);
+
+	const int32 SampleCount = LocalVectorField->SizeX * LocalVectorField->SizeY * LocalVectorField->SizeZ;
+
+	if (LocalVectorField->CPUData.Num() == SampleCount)
+	{
+		for (const FVector4& SampleData : LocalVectorField->CPUData)
+		{
+			Ar.Logf(TEXT(", %f, %f, %f"), SampleData.X, SampleData.Y, SampleData.Z);
+		}
+		return true;
+	}
+
+	UE_LOG(LogVectorFieldFactory, Warning, TEXT("Failed to export [%s] unabled to acquire bulk data"), *Object->GetFullName());
+	return false;
+}
+
 #undef LOCTEXT_NAMESPACE
