Editor: Fixed Blueprint Breakpoints not working in BeginPlay if you had a specific object set in the Debug Filter. The TransferBlueprintDebugReferences call had been moved to OnAllPIEInstancesStarted to avoid it getting called for each PIE client started, however BeginPlay is called on each world as it is created so if a specific object had been set as the debug filter the reference would not have been transfered to the PIE world by the time BeginPlay was called on that object, thus skipping the breakpoint.

This moves the TransferBlueprintDebugReferences back to before BeginPlay but only calls it on the first PIE world created to avoid calling it multiple times.

#rnx

diff --git a/Engine/Source/Editor/UnrealEd/Private/PlayLevel.cpp b/Engine/Source/Editor/UnrealEd/Private/PlayLevel.cpp
--- a/Engine/Source/Editor/UnrealEd/Private/PlayLevel.cpp
+++ b/Engine/Source/Editor/UnrealEd/Private/PlayLevel.cpp
@@ -1528,12 +1528,6 @@ void UEditorEngine::OnLoginPIEComplete_Deferred(int32 LocalUserNum, bool bWasSuc
 
 void UEditorEngine::OnAllPIEInstancesStarted()
 {
-	// Transfer the Blueprint Debug references to the last client that was created (since that's all we track).
-	if (PlayWorld)
-	{
-		EditorWorld->TransferBlueprintDebugReferences(PlayWorld);
-	}
-
 	GiveFocusToLastClientPIEViewport();
 
 	// Print out a log message stating the overall startup time.
@@ -2653,6 +2647,7 @@ UGameInstance* UEditorEngine::CreateInnerProcessPIEGameInstance(FRequestPlaySess
 	GameInstance->AddToRoot();
 
 	// Attempt to initialize the GameInstance. This will construct the world.
+	const bool bFirstWorld = !PlayWorld;
 	const FGameInstancePIEResult InitializeResult = GameInstance->InitializeForPlayInEditor(InPIEInstanceIndex, InPIEParameters);
 	if (!InitializeResult.IsSuccess())
 	{
@@ -2785,6 +2780,13 @@ UGameInstance* UEditorEngine::CreateInnerProcessPIEGameInstance(FRequestPlaySess
 		GameInstance->GetWorld()->GetWorldSettings()->DefaultGameMode = InParams.GameModeOverride;
 	}
 
+	// Transfer the Blueprint Debug references to the first client world that is created. This needs to be called before 
+	// GameInstance->StartPlayInEditorGameInstance so that references are transfered by the time BeginPlay is called.
+	if (bFirstWorld && PlayWorld)
+	{
+		EditorWorld->TransferBlueprintDebugReferences(PlayWorld);
+	}
+
 	FGameInstancePIEResult StartResult = FGameInstancePIEResult::Success();
 	{
 		FTemporaryPlayInEditorIDOverride OverrideIDHelper(InPIEInstanceIndex);
