Fixed time in cascade viewport.

diff --git a/Engine/Source/Editor/UnrealEd/Private/EditorViewportClient.cpp b/Engine/Source/Editor/UnrealEd/Private/EditorViewportClient.cpp
--- a/Engine/Source/Editor/UnrealEd/Private/EditorViewportClient.cpp
+++ b/Engine/Source/Editor/UnrealEd/Private/EditorViewportClient.cpp
@@ -3698,7 +3698,7 @@ void FEditorViewportClient::Draw(FViewport* InViewport, FCanvas* Canvas)
 	float DeltaTimeSeconds;
 
 	UWorld* World = GetWorld();
-	if (!World || (GetScene() != World->Scene))
+	if (!World || (GetScene() != World->Scene) || UseAppTime()) 
 	{
 		// Use time relative to start time to avoid issues with float vs double
 		TimeSeconds = FApp::GetCurrentTime() - GStartTime;
diff --git a/Engine/Source/Editor/UnrealEd/Public/EditorViewportClient.h b/Engine/Source/Editor/UnrealEd/Public/EditorViewportClient.h
--- a/Engine/Source/Editor/UnrealEd/Public/EditorViewportClient.h
+++ b/Engine/Source/Editor/UnrealEd/Public/EditorViewportClient.h
@@ -794,6 +794,8 @@ public:
 	 */
 	virtual bool WantsDrawWhenAppIsHidden() const { return false; }
 
+	/** Should this viewport use app time instead of world time. */
+	virtual bool UseAppTime() const { return IsRealtime() && !IsSimulateInEditorViewport(); }
 public:
 	/** True if the window is maximized or floating */
 	bool IsVisible() const;
diff --git a/Engine/Source/Editor/UnrealEd/Public/LevelEditorViewport.h b/Engine/Source/Editor/UnrealEd/Public/LevelEditorViewport.h
--- a/Engine/Source/Editor/UnrealEd/Public/LevelEditorViewport.h
+++ b/Engine/Source/Editor/UnrealEd/Public/LevelEditorViewport.h
@@ -539,6 +539,8 @@ public:
 	 */
 	static UObject* GetOrCreateMaterialFromTexture( UTexture* UnrealTexture );
 
+	virtual bool UseAppTime() const override { return false; }
+
 protected:
 	/**
 	* Sets the state of creating a preview actor in the viewport.
