Add graph to BP pin validation check error

#rnx

diff --git a/Engine/Source/Editor/UnrealEd/Private/Kismet2/BlueprintEditorUtils.cpp b/Engine/Source/Editor/UnrealEd/Private/Kismet2/BlueprintEditorUtils.cpp
--- a/Engine/Source/Editor/UnrealEd/Private/Kismet2/BlueprintEditorUtils.cpp
+++ b/Engine/Source/Editor/UnrealEd/Private/Kismet2/BlueprintEditorUtils.cpp
@@ -4369,8 +4369,8 @@ void FBlueprintEditorUtils::ValidatePinConnections(const UEdGraphNode* Node, FCo
 				{
 					if (Link && Link != Pin && Schema->CanCreateConnection(Pin, Link).Response == CONNECT_RESPONSE_DISALLOW)
 					{
-						FText const ErrorFormat = LOCTEXT("BadConnection", "Invalid pin connection from '@@' to '@@'. You may have changed the type after the connections were made.");
-						MessageLog.Error(*ErrorFormat.ToString(), Pin, Link);
+						FText const ErrorFormat = LOCTEXT("BadConnection", "Invalid pin connection in graph '@@' from '@@' to '@@'. You may have changed the type after the connections were made.");
+						MessageLog.Error(*ErrorFormat.ToString(), Node->GetGraph(), Pin, Link);
 					}
 				}
 			}
