#UE4 Skip over levels added by transient LevelStreamings when building HLODs. These levels are not saved with the map and should not be contributing to the HLOD mesh (they are often "preview" levels that are temporary)

diff --git a/Engine/Source/Editor/UnrealEd/Private/HierarchicalLOD.cpp b/Engine/Source/Editor/UnrealEd/Private/HierarchicalLOD.cpp
--- a/Engine/Source/Editor/UnrealEd/Private/HierarchicalLOD.cpp
+++ b/Engine/Source/Editor/UnrealEd/Private/HierarchicalLOD.cpp
@@ -93,14 +93,17 @@ void FHierarchicalLODBuilder::PreviewBuild()
 	const TArray<ULevel*>& Levels = World->GetLevels();
 	for (ULevel* LevelIter : Levels)
 	{
-		// Only build clusters for levels that are visible, and throw warning if any are hidden
-		if (LevelIter->bIsVisible)
+		if (ShouldBuildHLODForLevel(World, LevelIter))
 		{
-			BuildClusters(LevelIter);
-		}
-		else
-		{
-			bVisibleLevelsWarning |= LevelIter->GetWorldSettings()->bEnableHierarchicalLODSystem;
+			// Only build clusters for levels that are visible, and throw warning if any are hidden
+			if (LevelIter->bIsVisible)
+			{
+				BuildClusters(LevelIter);
+			}
+			else
+			{
+				bVisibleLevelsWarning |= LevelIter->GetWorldSettings()->bEnableHierarchicalLODSystem;
+			}
 		}
 	}
 
@@ -630,6 +633,25 @@ void FHierarchicalLODBuilder::HandleHLODVolumes(ULevel* InLevel)
 	}
 }
 
+bool FHierarchicalLODBuilder::ShouldBuildHLODForLevel(const UWorld* InWorld, const ULevel* InLevel) const
+{
+	check(InWorld);
+	if (!InLevel)
+	{
+		return false;
+	}
+
+	ULevelStreaming* const* SourceLevelStreamingPtr = InWorld->GetStreamingLevels().FindByPredicate([InLevel](ULevelStreaming* LS) { return LS && LS->GetLoadedLevel() == InLevel; });
+	ULevelStreaming* SourceLevelStreaming = SourceLevelStreamingPtr ? *SourceLevelStreamingPtr : nullptr;
+	if (SourceLevelStreaming && SourceLevelStreaming->HasAnyFlags(RF_Transient))
+	{
+		// Skip over levels from transient ULevelStreamings. These are levels that are not saved in the map and should not contribute to the HLOD
+		return false;
+	}
+
+	return true;
+}
+
 bool FHierarchicalLODBuilder::ShouldGenerateCluster(AActor* Actor, const int32 HLODLevelIndex)
 {
 	if (!Actor)
@@ -717,13 +739,16 @@ void FHierarchicalLODBuilder::ClearHLODs()
 
 	for (ULevel* Level : World->GetLevels())
 	{
-		if (Level->bIsVisible)
+		if (ShouldBuildHLODForLevel(World, Level))
 		{
-			DeleteLODActors(Level);
-		}
-		else
-		{
-			bVisibleLevelsWarning |= Level->GetWorldSettings()->bEnableHierarchicalLODSystem;
+			if (Level->bIsVisible)
+			{
+				DeleteLODActors(Level);
+			}
+			else
+			{
+				bVisibleLevelsWarning |= Level->GetWorldSettings()->bEnableHierarchicalLODSystem;
+			}
 		}
 	}
 
@@ -750,6 +775,11 @@ void FHierarchicalLODBuilder::BuildMeshesForLODActors(bool bForceAll)
 	const TArray<ULevel*>& Levels = World->GetLevels();
 	for (const ULevel* LevelIter : Levels)
 	{
+		if (!ShouldBuildHLODForLevel(World, LevelIter))
+		{
+			continue;
+		}
+
 		// Only meshes for clusters that are in a visible level
 		if (!LevelIter->bIsVisible)
 		{
diff --git a/Engine/Source/Editor/UnrealEd/Public/HierarchicalLOD.h b/Engine/Source/Editor/UnrealEd/Public/HierarchicalLOD.h
--- a/Engine/Source/Editor/UnrealEd/Public/HierarchicalLOD.h
+++ b/Engine/Source/Editor/UnrealEd/Public/HierarchicalLOD.h
@@ -156,6 +156,15 @@ private:
 	*/
 	void HandleHLODVolumes(ULevel* InLevel);
 
+	/**
+	* Determine whether or not this level should have HLODs built for it in the specified world
+	*
+	* @param World - The world the level is part of
+	* @param Level - The level to test
+	* @return bool
+	*/
+	bool ShouldBuildHLODForLevel(const UWorld* World, const ULevel* Level) const;
+
 	/**
 	* Determine whether or not this actor is eligble for HLOD creation
 	*
