Call PostEditChange() on materials created by dropping textures onto viewport, make sure cached data is up to date

diff --git a/Engine/Source/Editor/UnrealEd/Private/LevelEditorViewport.cpp b/Engine/Source/Editor/UnrealEd/Private/LevelEditorViewport.cpp
--- a/Engine/Source/Editor/UnrealEd/Private/LevelEditorViewport.cpp
+++ b/Engine/Source/Editor/UnrealEd/Private/LevelEditorViewport.cpp
@@ -543,6 +543,7 @@ static bool TryAndCreateMaterialInput( UMaterial* UnrealMaterial, EMaterialKind:
 		UnrealMaterial->Normal.Expression = UnrealTextureExpression;
 	}
 
+
 	return true;
 }
 
@@ -635,6 +636,9 @@ UObject* FLevelEditorViewportClient::GetOrCreateMaterialFromTexture( UTexture* U
 		TryAndCreateMaterialInput( UnrealMaterial, EMaterialKind::Normal, NormalTexture, UnrealMaterial->Normal, HSpace, VSpace * 2 );
 	}
 
+	UnrealMaterial->PreEditChange(nullptr);
+	UnrealMaterial->PostEditChange();
+
 	// Notify the asset registry
 	FAssetRegistryModule::AssetCreated( UnrealMaterial );
 
