#UE4 IWYU fix

diff --git a/Engine/Source/Editor/UnrealEd/Private/HierarchicalLOD.cpp b/Engine/Source/Editor/UnrealEd/Private/HierarchicalLOD.cpp
--- a/Engine/Source/Editor/UnrealEd/Private/HierarchicalLOD.cpp
+++ b/Engine/Source/Editor/UnrealEd/Private/HierarchicalLOD.cpp
@@ -37,6 +37,7 @@
 #include "Serialization/ArchiveObjectCrc32.h"
 #include "IMeshReductionManagerModule.h"
 #include "Engine/HLODProxy.h"
+#include "Engine/LevelStreaming.h"
 
 DEFINE_LOG_CATEGORY_STATIC(LogLODGenerator, Log, All);
 
