Integrating CLs 11310084, 11341910, 11736184, 11751405 from //UE4/Dev-Editor

PIE sessions which launch a new standalone window no longer double render debug safe zones and all PIE options are more consistent with handling r.DebugSafeZone.XXXX cvars.

diff --git a/Engine/Source/Editor/UMGEditor/Private/Designer/SDesignerView.cpp b/Engine/Source/Editor/UMGEditor/Private/Designer/SDesignerView.cpp
--- a/Engine/Source/Editor/UMGEditor/Private/Designer/SDesignerView.cpp
+++ b/Engine/Source/Editor/UMGEditor/Private/Designer/SDesignerView.cpp
@@ -3680,8 +3680,7 @@ FReply SDesignerView::HandleSwapAspectRatioClicked()
 
 	if (!PreviewOverrideName.IsEmpty())
 	{
-		ULevelEditorPlaySettings* PlaySettings = GetMutableDefault<ULevelEditorPlaySettings>();
-		DesignerSafeZoneOverride = PlaySettings->CalculateCustomUnsafeZones(CustomSafeZoneStarts, CustomSafeZoneDimensions, PreviewOverrideName, FVector2D(PreviewWidth, PreviewHeight));
+		DesignerSafeZoneOverride = PlayInSettings->CalculateCustomUnsafeZones(CustomSafeZoneStarts, CustomSafeZoneDimensions, PreviewOverrideName, FVector2D(PreviewWidth, PreviewHeight));
 	}
 	else
 	{
diff --git a/Engine/Source/Editor/UnrealEd/Classes/Settings/LevelEditorPlaySettings.h b/Engine/Source/Editor/UnrealEd/Classes/Settings/LevelEditorPlaySettings.h
--- a/Engine/Source/Editor/UnrealEd/Classes/Settings/LevelEditorPlaySettings.h
+++ b/Engine/Source/Editor/UnrealEd/Classes/Settings/LevelEditorPlaySettings.h
@@ -543,7 +543,8 @@ public:
 	// End of UObject interface
 
 #if WITH_EDITOR
-	void SwapSafeZoneTypes();
+	// Recalculates and broadcasts safe zone size changes based on device to emulate and r.DebugSafeZone.TitleRatio values.
+	void UpdateCustomSafeZones();
 #endif
 
 	FMargin CalculateCustomUnsafeZones(TArray<FVector2D>& CustomSafeZoneStarts, TArray<FVector2D>& CustomSafeZoneDimensions, FString& DeviceType, FVector2D PreviewSize);
diff --git a/Engine/Source/Editor/UnrealEd/Private/PlayLevelNewProcess.cpp b/Engine/Source/Editor/UnrealEd/Private/PlayLevelNewProcess.cpp
--- a/Engine/Source/Editor/UnrealEd/Private/PlayLevelNewProcess.cpp
+++ b/Engine/Source/Editor/UnrealEd/Private/PlayLevelNewProcess.cpp
@@ -154,22 +154,19 @@ void UEditorEngine::LaunchNewProcess(const FRequestPlaySessionParams& InParams,
 		// Ensure the executable writes out a differently named config file to avoid multiple instances overwriting each other.
 		// ToDo: Should this be on all multi-client launches?
 		CommandLine += TEXT(" -MultiprocessSaveConfig");
+	}
 
-		// In order for the mobile previewer to adjust its safe zone according to the device profile specified in the editor play settings,
-		// we need to pass the PIESafeZoneOverride's values as command line variables to the new process that we are about to launch.
-		FMargin PIESafeZoneOverride = InParams.EditorPlaySettings->PIESafeZoneOverride;
+	// In order for the previewer to adjust its safe zone according to the device profile specified in the editor play settings,
+	// we need to pass the PIESafeZoneOverride's values as command line variables to the new process that we are about to launch.
+	FMargin PIESafeZoneOverride = InParams.EditorPlaySettings->PIESafeZoneOverride;
+	if (!PIESafeZoneOverride.GetDesiredSize().IsZero())
+	{
 		CommandLine += FString::Printf(TEXT(" -SafeZonePaddingLeft=%f -SafeZonePaddingRight=%f -SafeZonePaddingTop=%f -SafeZonePaddingBottom=%f"),
 			PIESafeZoneOverride.Left,
 			PIESafeZoneOverride.Right,
 			PIESafeZoneOverride.Top,
 			PIESafeZoneOverride.Bottom
 		);
-
-		if (!PIESafeZoneOverride.GetDesiredSize().IsZero())
-		{
-			// Send -DrawUnSafeZones so that the red "unsafe" zones show up in the mobile previewer.
-			CommandLine += TEXT(" -DrawUnSafeZones");
-		}
 	}
 
 	if (InParams.SessionPreviewTypeOverride.Get(EPlaySessionPreviewType::NoPreview) == EPlaySessionPreviewType::VulkanPreview)
diff --git a/Engine/Source/Editor/UnrealEd/Private/Settings/LevelEditorPlaySettingsCustomization.cpp b/Engine/Source/Editor/UnrealEd/Private/Settings/LevelEditorPlaySettingsCustomization.cpp
new file mode 100644
--- /dev/null
+++ b/Engine/Source/Editor/UnrealEd/Private/Settings/LevelEditorPlaySettingsCustomization.cpp
@@ -0,0 +1,789 @@
+// Copyright Epic Games, Inc. All Rights Reserved.
+
+#include "LevelEditorPlaySettingsCustomization.h"
+
+#include "Misc/Attribute.h"
+#include "Layout/Margin.h"
+#include "Layout/Visibility.h"
+#include "Widgets/DeclarativeSyntaxSupport.h"
+#include "Widgets/SWidget.h"
+#include "Widgets/SCompoundWidget.h"
+#include "Widgets/SBoxPanel.h"
+#include "Editor.h"
+#include "Textures/SlateIcon.h"
+#include "Framework/Commands/UIAction.h"
+#include "PropertyHandle.h"
+#include "DetailLayoutBuilder.h"
+#include "DetailCategoryBuilder.h"
+#include "Framework/MultiBox/MultiBoxBuilder.h"
+#include "Widgets/Text/STextBlock.h"
+#include "Widgets/Images/SImage.h"
+#include "Widgets/Input/SButton.h"
+#include "Widgets/Input/SComboButton.h"
+#include "IDetailPropertyRow.h"
+#include "DetailWidgetRow.h"
+#include "Widgets/Input/SComboBox.h"
+#include "Runtime/Engine/Classes/Sound/AudioSettings.h"
+#include "DeviceProfiles/DeviceProfileManager.h"
+#include "DeviceProfiles/DeviceProfile.h"
+
+#define LOCTEXT_NAMESPACE "FLevelEditorPlaySettingsCustomization"
+
+void SScreenPositionCustomization::Construct( const FArguments& InArgs, IDetailLayoutBuilder* LayoutBuilder, const TSharedRef<IPropertyHandle>& InWindowPositionProperty, const TSharedRef<IPropertyHandle>& InCenterWindowProperty )
+{
+	check( LayoutBuilder != NULL );
+
+	CenterWindowProperty = InCenterWindowProperty;
+
+	ChildSlot
+		[
+			SNew( SVerticalBox )
+			+ SVerticalBox::Slot()
+		.AutoHeight()
+		[
+			SNew( SHorizontalBox )
+			+ SHorizontalBox::Slot()
+		[
+			SNew( SVerticalBox )
+			.IsEnabled( this, &SScreenPositionCustomization::HandleNewWindowPositionPropertyIsEnabled )
+		+ SVerticalBox::Slot()
+		.AutoHeight()
+		[
+			InWindowPositionProperty->CreatePropertyNameWidget( LOCTEXT( "WindowPosXLabel", "Left Position" ) )
+		]
+	+ SVerticalBox::Slot()
+		.AutoHeight()
+		[
+			InWindowPositionProperty->GetChildHandle( 0 )->CreatePropertyValueWidget()
+		]
+		]
+
+	+ SHorizontalBox::Slot()
+		.Padding( 8.0f, 0.0f, 0.0f, 0.0f )
+		[
+			SNew( SVerticalBox )
+			.IsEnabled( this, &SScreenPositionCustomization::HandleNewWindowPositionPropertyIsEnabled )
+
+		+ SVerticalBox::Slot()
+		.AutoHeight()
+		[
+			InWindowPositionProperty->CreatePropertyNameWidget( LOCTEXT( "TopPositionLabel", "Top Position" ) )
+		]
+
+	+ SVerticalBox::Slot()
+		.AutoHeight()
+		[
+			InWindowPositionProperty->GetChildHandle( 1 )->CreatePropertyValueWidget()
+		]
+		]
+		]
+	+ SVerticalBox::Slot()
+		.Padding( 0.0f, 2.0f )
+		.AutoHeight()
+		[
+			SNew( SHorizontalBox )
+			+ SHorizontalBox::Slot()
+		.VAlign( VAlign_Center )
+		.AutoWidth()
+		[
+			InCenterWindowProperty->CreatePropertyValueWidget()
+		]
+	+ SHorizontalBox::Slot()
+		.AutoWidth()
+		.Padding( 4.0f, 0.0f, 0.0f, 0.0f )
+		.VAlign( VAlign_Bottom )
+		[
+			InWindowPositionProperty->CreatePropertyNameWidget( LOCTEXT( "CenterWindowLabel", "Always center window to screen" ) )
+		]
+		]
+		];
+}
+
+bool SScreenPositionCustomization::HandleNewWindowPositionPropertyIsEnabled() const
+{
+	bool CenterNewWindow;
+	CenterWindowProperty->GetValue( CenterNewWindow );
+
+	return !CenterNewWindow;
+
+}
+
+void SScreenResolutionCustomization::Construct( const FArguments& InArgs, IDetailLayoutBuilder* LayoutBuilder, const TSharedRef<IPropertyHandle>& InWindowHeightProperty, const TSharedRef<IPropertyHandle>& InWindowWidthProperty )
+{
+	check( LayoutBuilder != NULL );
+
+	WindowHeightProperty = InWindowHeightProperty;
+	WindowWidthProperty = InWindowWidthProperty;
+
+	FSimpleDelegate SizeChangeDelegate = FSimpleDelegate::CreateSP( this, &SScreenResolutionCustomization::OnSizeChanged );
+	WindowHeightProperty->SetOnPropertyValueChanged( SizeChangeDelegate );
+	WindowWidthProperty->SetOnPropertyValueChanged( SizeChangeDelegate );
+	ChildSlot
+		[
+			SNew( SVerticalBox )
+			+ SVerticalBox::Slot()
+		.VAlign( VAlign_Bottom )
+		[
+			SNew( SHorizontalBox )
+			+ SHorizontalBox::Slot()
+		.AutoWidth()
+		.VAlign( VAlign_Center )
+		[
+			SNew( SComboButton )
+			.VAlign( VAlign_Center )
+		.ButtonContent()
+		[
+			SNew( STextBlock )
+			.Font( LayoutBuilder->GetDetailFont() )
+		.Text( LOCTEXT( "CommonResolutionsButtonText", "Common Resolutions" ) )
+		]
+	.ContentPadding( FMargin( 6, 2 ) )
+		.MenuContent()
+		[
+			MakeCommonResolutionsMenu()
+		]
+	.ToolTipText( LOCTEXT( "CommonResolutionsButtonTooltip", "Pick from a list of common screen resolutions" ) )
+		]
+	+ SHorizontalBox::Slot()
+		.Padding( 0, 0, 6, 0 )
+		.AutoWidth()
+		.VAlign( VAlign_Center )
+		[
+			SNew( SButton )
+			.OnClicked( this, &SScreenResolutionCustomization::HandleSwapAspectRatioClicked )
+		.ContentPadding( FMargin( 3, 0, 3, 1 ) )
+		.Content()
+		[
+			SNew( SImage )
+			.Image( this, &SScreenResolutionCustomization::GetAspectRatioSwitchImage )
+		]
+	.ToolTipText( LOCTEXT( "SwapAspectRatioTooltip", "Swap between portrait and landscape orientation." ) )
+		]
+		]
+	+ SVerticalBox::Slot()
+		.AutoHeight()
+		[
+			SNew( SHorizontalBox )
+			+ SHorizontalBox::Slot()
+		[
+			SNew( SVerticalBox )
+			+ SVerticalBox::Slot()
+		.AutoHeight()
+		[
+			WindowWidthProperty->CreatePropertyNameWidget( LOCTEXT( "ViewportWidthLabel", "Viewport Width" ) )
+		]
+	+ SVerticalBox::Slot()
+		[
+			WindowWidthProperty->CreatePropertyValueWidget()
+		]
+		]
+	+ SHorizontalBox::Slot()
+		.Padding( 8, 0, 0, 0 )
+		[
+			SNew( SVerticalBox )
+			+ SVerticalBox::Slot()
+		.AutoHeight()
+		[
+			WindowHeightProperty->CreatePropertyNameWidget( LOCTEXT( "ViewportHeightLabel", "Viewport Height" ) )
+		]
+	+ SVerticalBox::Slot()
+		.AutoHeight()
+		[
+			WindowHeightProperty->CreatePropertyValueWidget()
+		]
+		]
+		]
+		];
+}
+
+void SScreenResolutionCustomization::AddCommonResolutionEntry( FMenuBuilder& MenuBuilder, int32 Width, int32 Height, const FString& AspectRatio, const FText& Description )
+{
+}
+
+void SScreenResolutionCustomization::AddScreenResolutionSection( FMenuBuilder& MenuBuilder, const TArray<FPlayScreenResolution> Resolutions, const FText SectionName )
+{
+	MenuBuilder.BeginSection( NAME_None, SectionName );
+	{
+		for ( auto Iter = Resolutions.CreateConstIterator(); Iter; ++Iter )
+		{
+			FUIAction Action( FExecuteAction::CreateRaw( this, &SScreenResolutionCustomization::HandleCommonResolutionSelected, *Iter ) );
+
+			FInternationalization& I18N = FInternationalization::Get();
+
+			FFormatNamedArguments Args;
+			Args.Add( TEXT( "Width" ), FText::AsNumber( Iter->Width, NULL, I18N.GetInvariantCulture() ) );
+			Args.Add( TEXT( "Height" ), FText::AsNumber( Iter->Height, NULL, I18N.GetInvariantCulture() ) );
+			Args.Add( TEXT( "AspectRatio" ), FText::FromString( Iter->AspectRatio ) );
+
+			MenuBuilder.AddMenuEntry( FText::FromString( Iter->Description ), FText::Format( LOCTEXT( "CommonResolutionFormat", "{Width} x {Height} ({AspectRatio})" ), Args ), FSlateIcon(), Action );
+		}
+	}
+	MenuBuilder.EndSection();
+}
+
+TSharedRef<SWidget> SScreenResolutionCustomization::MakeCommonResolutionsMenu()
+{
+	const ULevelEditorPlaySettings* PlaySettings = GetDefault<ULevelEditorPlaySettings>();
+	FMenuBuilder MenuBuilder( true, NULL );
+
+
+	FText PhoneTitle = LOCTEXT( "CommonPhonesSectionHeader", "Phones" );
+	FText TabletTitle = LOCTEXT( "CommonTabletsSectionHeader", "Tablets" );
+	FText LaptopTitle = LOCTEXT( "CommonLaptopsSectionHeader", "Laptops" );
+	FText MonitorTitle = LOCTEXT( "CommonMonitorsSectionHeader", "Monitors" );
+	FText TelevisionTitle = LOCTEXT( "CommonTelevesionsSectionHeader", "Televisions" );
+	MenuBuilder.AddSubMenu(
+		PhoneTitle,
+		FText(),
+		FNewMenuDelegate::CreateRaw( this, &SScreenResolutionCustomization::AddScreenResolutionSection, (PlaySettings->PhoneScreenResolutions), PhoneTitle ),
+		false,
+		FSlateIcon() );
+	MenuBuilder.AddSubMenu(
+		TabletTitle,
+		FText(),
+		FNewMenuDelegate::CreateRaw( this, &SScreenResolutionCustomization::AddScreenResolutionSection, (PlaySettings->TabletScreenResolutions), TabletTitle ),
+		false,
+		FSlateIcon() );
+	MenuBuilder.AddSubMenu(
+		LaptopTitle,
+		FText(),
+		FNewMenuDelegate::CreateRaw( this, &SScreenResolutionCustomization::AddScreenResolutionSection, (PlaySettings->LaptopScreenResolutions), LaptopTitle ),
+		false,
+		FSlateIcon() );
+	MenuBuilder.AddSubMenu(
+		MonitorTitle,
+		FText(),
+		FNewMenuDelegate::CreateRaw( this, &SScreenResolutionCustomization::AddScreenResolutionSection, (PlaySettings->MonitorScreenResolutions), MonitorTitle ),
+		false,
+		FSlateIcon() );
+	MenuBuilder.AddSubMenu(
+		TelevisionTitle,
+		FText(),
+		FNewMenuDelegate::CreateRaw( this, &SScreenResolutionCustomization::AddScreenResolutionSection, (PlaySettings->TelevisionScreenResolutions), TelevisionTitle ),
+		false,
+		FSlateIcon() );
+
+	return MenuBuilder.MakeWidget();
+}
+
+FReply SScreenResolutionCustomization::HandleSwapAspectRatioClicked()
+{
+	FString HeightString;
+	WindowHeightProperty->GetValueAsDisplayString( HeightString );
+	FString WidthString;
+	WindowWidthProperty->GetValueAsDisplayString( WidthString );
+	int32 NewHeight = FCString::Atoi( *WidthString );
+	int32 NewWidth = FCString::Atoi( *HeightString );
+
+	ULevelEditorPlaySettings* PlayInSettings = GetMutableDefault<ULevelEditorPlaySettings>();
+
+	const UDeviceProfile* DeviceProfile = UDeviceProfileManager::Get().FindProfile(PlayInSettings->DeviceToEmulate, false);
+	if ( DeviceProfile )
+	{
+		// Rescale the swapped sizes if we are on Android
+		if ( DeviceProfile && DeviceProfile->DeviceType == TEXT( "Android" ) )
+		{
+			float ScaleFactor = 1.0f;
+			PlayInSettings->RescaleForMobilePreview( DeviceProfile, NewWidth, NewHeight, ScaleFactor );
+		}
+	}
+
+	WindowHeightProperty->SetValue(NewHeight);
+	WindowWidthProperty->SetValue(NewWidth);
+
+	return FReply::Handled();
+}
+
+void SScreenResolutionCustomization::HandleCommonResolutionSelected( const FPlayScreenResolution Resolution )
+{
+	int32 Width = Resolution.Width;
+	int32 Height = Resolution.Height;
+	float ScaleFactor;
+	ULevelEditorPlaySettings* PlayInSettings = GetMutableDefault<ULevelEditorPlaySettings>();
+	// Maintain previous orientation (i.e., swap Width and Height if required)
+	int32 PreviousWidth = -1;
+	int32 PreviousHeight = -1;
+	if ( WindowWidthProperty->GetValue( PreviousWidth ) == FPropertyAccess::Success && WindowHeightProperty->GetValue( PreviousHeight ) == FPropertyAccess::Success )
+	{
+		const bool bIsOrientationPreserved = (PreviousWidth < PreviousHeight) != (Width < Height);
+		if ( bIsOrientationPreserved )
+		{
+			Width = Resolution.Height;
+			Height = Resolution.Width;
+		}
+	}
+
+	UDeviceProfile* DeviceProfile = UDeviceProfileManager::Get().FindProfile( Resolution.ProfileName, false );
+	if ( DeviceProfile )
+	{
+		PlayInSettings->DeviceToEmulate = Resolution.ProfileName;
+		PlayInSettings->RescaleForMobilePreview( DeviceProfile, Width, Height, ScaleFactor );
+	}
+	else
+	{
+		PlayInSettings->DeviceToEmulate = FString();
+	}
+
+	if ( (PreviousWidth == Width) && (PreviousHeight == Height) )
+	{
+		// Since the values are identical, the SetValue on the property handle will not actually change the value and trigger OnPropertyValueChanged, so call our handler directly
+		OnSizeChanged();
+	}
+	else
+	{
+		WindowHeightProperty->SetValue(Height);
+		WindowWidthProperty->SetValue(Width);
+	}
+}
+
+const FSlateBrush* SScreenResolutionCustomization::GetAspectRatioSwitchImage() const
+{
+	FString HeightString;
+	WindowHeightProperty->GetValueAsDisplayString( HeightString );
+	int32 Height = FCString::Atoi( *HeightString );
+	FString WidthString;
+	WindowWidthProperty->GetValueAsDisplayString( WidthString );
+	int32 Width = FCString::Atoi( *WidthString );
+	if ( Height > Width )
+	{
+		return FEditorStyle::Get().GetBrush( "UMGEditor.OrientPortrait" );
+	}
+	return FEditorStyle::Get().GetBrush( "UMGEditor.OrientLandscape" );
+}
+
+void SScreenResolutionCustomization::OnSizeChanged()
+{
+	GetMutableDefault<ULevelEditorPlaySettings>()->UpdateCustomSafeZones();
+}
+
+/** Virtual destructor. */
+FLevelEditorPlaySettingsCustomization::~FLevelEditorPlaySettingsCustomization()
+{
+}
+
+void FLevelEditorPlaySettingsCustomization::CustomizeDetails( IDetailLayoutBuilder& LayoutBuilder )
+{
+	const float MaxPropertyWidth = 400.0f;
+
+	// play in editor settings
+	IDetailCategoryBuilder& PlayInEditorCategory = LayoutBuilder.EditCategory( "PlayInEditor" );
+	{
+		TArray<TSharedRef<IPropertyHandle>> PIECategoryProperties;
+		PlayInEditorCategory.GetDefaultProperties( PIECategoryProperties, true, false );
+
+		TSharedPtr<IPropertyHandle> PIEEnableSoundHandle = LayoutBuilder.GetProperty( GET_MEMBER_NAME_CHECKED( ULevelEditorPlaySettings, EnableGameSound ) );
+		PIESoundQualityLevelHandle = LayoutBuilder.GetProperty( GET_MEMBER_NAME_CHECKED( ULevelEditorPlaySettings, PlayInEditorSoundQualityLevel ) );
+		PIESoundQualityLevelHandle->MarkHiddenByCustomization();
+
+		for ( TSharedRef<IPropertyHandle>& PropertyHandle : PIECategoryProperties )
+		{
+			if ( PropertyHandle->GetProperty() != PIESoundQualityLevelHandle->GetProperty() )
+			{
+				PlayInEditorCategory.AddProperty( PropertyHandle );
+			}
+
+			if ( PropertyHandle->GetProperty() == PIEEnableSoundHandle->GetProperty() )
+			{
+				PlayInEditorCategory.AddCustomRow( PIESoundQualityLevelHandle->GetPropertyDisplayName(), false )
+					.NameContent()
+					[
+						PIESoundQualityLevelHandle->CreatePropertyNameWidget()
+					]
+				.ValueContent()
+					.MaxDesiredWidth( MaxPropertyWidth )
+					[
+						SAssignNew( QualityLevelComboBox, SComboBox<TSharedPtr<FString>> )
+						.OptionsSource( &AvailableQualityLevels )
+					.OnComboBoxOpening( this, &FLevelEditorPlaySettingsCustomization::HandleQualityLevelComboBoxOpening )
+					.OnGenerateWidget( this, &FLevelEditorPlaySettingsCustomization::HandleQualityLevelComboBoxGenerateWidget )
+					.OnSelectionChanged( this, &FLevelEditorPlaySettingsCustomization::HandleQualityLevelSelectionChanged )
+					[
+						SNew( STextBlock )
+						.Text( this, &FLevelEditorPlaySettingsCustomization::GetSelectedQualityLevelName )
+					]
+					];
+			}
+		}
+
+
+	}
+	IDetailCategoryBuilder& GameViewportSettings = LayoutBuilder.EditCategory( "GameViewportSettings" );
+	{
+		// new window resolution
+		TSharedRef<IPropertyHandle> WindowHeightHandle = LayoutBuilder.GetProperty( GET_MEMBER_NAME_CHECKED( ULevelEditorPlaySettings, NewWindowHeight ) );
+		TSharedRef<IPropertyHandle> WindowWidthHandle = LayoutBuilder.GetProperty( GET_MEMBER_NAME_CHECKED( ULevelEditorPlaySettings, NewWindowWidth ) );
+		TSharedRef<IPropertyHandle> WindowPositionHandle = LayoutBuilder.GetProperty( GET_MEMBER_NAME_CHECKED( ULevelEditorPlaySettings, NewWindowPosition ) );
+		TSharedRef<IPropertyHandle> CenterNewWindowHandle = LayoutBuilder.GetProperty( GET_MEMBER_NAME_CHECKED( ULevelEditorPlaySettings, CenterNewWindow ) );
+		TSharedRef<IPropertyHandle> EmulatedDeviceHandle = LayoutBuilder.GetProperty( GET_MEMBER_NAME_CHECKED( ULevelEditorPlaySettings, DeviceToEmulate ) );
+
+		WindowHeightHandle->MarkHiddenByCustomization();
+		WindowWidthHandle->MarkHiddenByCustomization();
+		WindowPositionHandle->MarkHiddenByCustomization();
+		CenterNewWindowHandle->MarkHiddenByCustomization();
+		EmulatedDeviceHandle->MarkHiddenByCustomization();
+
+		GameViewportSettings.AddCustomRow( LOCTEXT( "NewViewportResolutionRow", "New Viewport Resolution" ), false )
+			.NameContent()
+			[
+				SNew( STextBlock )
+				.Font( LayoutBuilder.GetDetailFont() )
+			.Text( LOCTEXT( "NewViewportResolutionName", "New Viewport Resolution" ) )
+			.ToolTipText( LOCTEXT( "NewWindowSizeTooltip", "Sets the width and height of floating PIE windows (in pixels)" ) )
+			]
+		.ValueContent()
+			.MaxDesiredWidth( MaxPropertyWidth )
+			[
+				SNew( SScreenResolutionCustomization, &LayoutBuilder, WindowHeightHandle, WindowWidthHandle )
+			];
+
+		GameViewportSettings.AddCustomRow( LOCTEXT( "NewWindowPositionRow", "New Window Position" ), false )
+			.NameContent()
+			[
+				SNew( STextBlock )
+				.Font( LayoutBuilder.GetDetailFont() )
+			.Text( LOCTEXT( "NewWindowPositionName", "New Window Position" ) )
+			.ToolTipText( LOCTEXT( "NewWindowPositionTooltip", "Sets the screen coordinates for the top-left corner of floating PIE windows (in pixels)" ) )
+			]
+		.ValueContent()
+			.MaxDesiredWidth( MaxPropertyWidth )
+			[
+				SNew( SScreenPositionCustomization, &LayoutBuilder, WindowPositionHandle, CenterNewWindowHandle )
+			];
+
+		GameViewportSettings.AddCustomRow( LOCTEXT( "SafeZonePreviewName", "Safe Zone Preview" ), false )
+			.NameContent()
+			[
+				SNew( STextBlock )
+				.Font( LayoutBuilder.GetDetailFont() )
+			.Text( LOCTEXT( "SafeZonePreviewName", "Safe Zone Preview" ) )
+			]
+		.ValueContent()
+			[
+				SNew( STextBlock )
+				.Font( LayoutBuilder.GetDetailFont() )
+				.Text( this, &FLevelEditorPlaySettingsCustomization::GetPreviewText )
+				.ToolTipText(this, &FLevelEditorPlaySettingsCustomization::GetPreviewTextToolTipText)
+			];
+	}
+
+	// play in new window settings
+	IDetailCategoryBuilder& PlayInNewWindowCategory = LayoutBuilder.EditCategory( "PlayInNewWindow" );
+	{
+		// Mac does not support parenting, do not show
+#if PLATFORM_MAC
+		PlayInNewWindowCategory.AddProperty( "PIEAlwaysOnTop" )
+			.DisplayName( LOCTEXT( "PIEAlwaysOnTop", "Always On Top" ) )
+			.IsEnabled( false );
+#else
+		PlayInNewWindowCategory.AddProperty( "PIEAlwaysOnTop" )
+			.DisplayName( LOCTEXT( "PIEAlwaysOnTop", "Always On Top" ) );
+#endif
+	}
+
+
+	// play in standalone game settings
+	IDetailCategoryBuilder& PlayInStandaloneCategory = LayoutBuilder.EditCategory( "PlayInStandaloneGame" );
+	{
+		// command line options
+		TSharedPtr<IPropertyHandle> DisableStandaloneSoundProperty = LayoutBuilder.GetProperty( GET_MEMBER_NAME_CHECKED( ULevelEditorPlaySettings, DisableStandaloneSound ) );
+
+		DisableStandaloneSoundProperty->MarkHiddenByCustomization();
+
+		PlayInStandaloneCategory.AddCustomRow( LOCTEXT( "AdditionalStandaloneDetails", "Additional Options" ), false )
+			.NameContent()
+			[
+				SNew( STextBlock )
+				.Font( LayoutBuilder.GetDetailFont() )
+			.Text( LOCTEXT( "ClientCmdLineName", "Command Line Options" ) )
+			.ToolTipText( LOCTEXT( "ClientCmdLineTooltip", "Generates a command line for additional settings that will be passed to the game clients." ) )
+			]
+		.ValueContent()
+			.MaxDesiredWidth( MaxPropertyWidth )
+			[
+				SNew( SHorizontalBox )
+
+				+ SHorizontalBox::Slot()
+			.AutoWidth()
+			[
+				DisableStandaloneSoundProperty->CreatePropertyValueWidget()
+			]
+
+		+ SHorizontalBox::Slot()
+			.Padding( 0.0f, 2.5f )
+			.VAlign( VAlign_Center )
+			.AutoWidth()
+			[
+				DisableStandaloneSoundProperty->CreatePropertyNameWidget( LOCTEXT( "DisableStandaloneSoundLabel", "Disable Sound (-nosound)" ) )
+			]
+
+			];
+	}
+
+	// multi-player options
+	IDetailCategoryBuilder& NetworkCategory = LayoutBuilder.EditCategory( "Multiplayer Options" );
+	TArray<TSharedRef<IPropertyHandle>> AllProperties;
+	NetworkCategory.GetDefaultProperties( AllProperties );
+
+	{
+		// Add all the default properties in before we add custom rows at the bottom.
+		for ( TSharedRef<IPropertyHandle> DefaultProperty : AllProperties )
+		{
+			NetworkCategory.AddProperty( DefaultProperty );
+		}
+
+
+		// client window size
+		TSharedRef<IPropertyHandle> WindowHeightHandle = LayoutBuilder.GetProperty( "ClientWindowHeight" );
+		TSharedRef<IPropertyHandle> WindowWidthHandle = LayoutBuilder.GetProperty( "ClientWindowWidth" );
+
+		WindowHeightHandle->MarkHiddenByCustomization();
+		WindowWidthHandle->MarkHiddenByCustomization();
+
+
+
+		NetworkCategory.AddCustomRow( LOCTEXT( "PlayInNetworkViewportSize", "Viewport Size\n(Additional Clients)" ), false )
+			.NameContent()
+			[
+				WindowHeightHandle->CreatePropertyNameWidget( LOCTEXT( "ClientViewportSizeName", "Multiplayer Viewport Size (in pixels)" ), LOCTEXT( "ClientWindowSizeTooltip", "Width and Height to use when spawning additional clients. Useful when you need multiple clients connected but only interact with one window." ) )
+			]
+		.ValueContent()
+			.MaxDesiredWidth( MaxPropertyWidth )
+			[
+				SNew( SScreenResolutionCustomization, &LayoutBuilder, WindowHeightHandle, WindowWidthHandle )
+			]
+		.IsEnabled( TAttribute<bool>( this, &FLevelEditorPlaySettingsCustomization::HandleClientWindowSizePropertyIsEnabled ) );
+
+		NetworkCategory.AddCustomRow( LOCTEXT( "AdditionalMultiplayerDetails", "Additional Options" ), true )
+			.NameContent()
+			[
+				SNew( STextBlock )
+				.Font( LayoutBuilder.GetDetailFont() )
+			.Text( LOCTEXT( "PlainTextName", "Play In Editor Description" ) )
+			.ToolTipText( LOCTEXT( "PlainTextToolTip", "A brief description of the multiplayer settings and what to expect if you play with them in the editor." ) )
+			]
+		.ValueContent()
+			.MaxDesiredWidth( MaxPropertyWidth )
+			[
+				SNew( STextBlock )
+				.Font( LayoutBuilder.GetDetailFont() )
+			.Text( this, &FLevelEditorPlaySettingsCustomization::HandleMultiplayerOptionsDescription )
+			.WrapTextAt( MaxPropertyWidth )
+			];
+	}
+}
+
+TSharedRef<IDetailCustomization> FLevelEditorPlaySettingsCustomization::MakeInstance()
+{
+	return MakeShareable( new FLevelEditorPlaySettingsCustomization() );
+}
+
+FText FLevelEditorPlaySettingsCustomization::HandleMultiplayerOptionsDescription() const
+{
+	const ULevelEditorPlaySettings* PlayInSettings = GetDefault<ULevelEditorPlaySettings>();
+	const bool bLaunchSeparateServer = PlayInSettings->bLaunchSeparateServer;
+	const bool CanRunUnderOneProcess = [&PlayInSettings] { bool RunUnderOneProcess( false ); return (PlayInSettings->GetRunUnderOneProcess( RunUnderOneProcess ) && RunUnderOneProcess); }();
+	const int32 PlayNumberOfClients = [&PlayInSettings] { int32 NumClients( 0 ); return (PlayInSettings->GetPlayNumberOfClients( NumClients ) ? NumClients : 0); }();
+	const EPlayNetMode PlayNetMode = [&PlayInSettings] { EPlayNetMode NetMode( PIE_Standalone ); return (PlayInSettings->GetPlayNetMode( NetMode ) ? NetMode : PIE_Standalone); }();
+	FString Desc;
+	if ( CanRunUnderOneProcess )
+	{
+		Desc += LOCTEXT( "MultiplayerDescription_OneProcess", "The following will all run under one UE4 instance:\n" ).ToString();
+		if ( PlayNetMode == EPlayNetMode::PIE_Client )
+		{
+			Desc += LOCTEXT( "MultiplayerDescription_DedicatedServerHidden", "A hidden dedicated server instance will run in editor. " ).ToString();
+			if ( PlayNumberOfClients == 1 )
+			{
+				Desc += LOCTEXT( "MultiplayerDescription_EditorClient", "The editor will connect as a client. " ).ToString();
+			}
+			else
+			{
+				Desc += FText::Format( LOCTEXT( "MultiplayerDescription_EditorAndClients", "The editor will connect as a client and {0} additional client window(s) will also connect. " ), FText::AsNumber( PlayNumberOfClients - 1 ) ).ToString();
+			}
+		}
+		else if ( PlayNetMode == EPlayNetMode::PIE_ListenServer )
+		{
+			if ( PlayNumberOfClients == 1 )
+			{
+				Desc += LOCTEXT( "MultiplayerDescription_EditorListenServer", "The editor will run as a listen server. " ).ToString();
+			}
+			else
+			{
+				Desc += FText::Format( LOCTEXT( "MultiplayerDescription_EditorListenServerAndClients", "The editor will run as a listen server and {0} additional client window(s) will also connect to it. " ), FText::AsNumber( PlayNumberOfClients - 1 ) ).ToString();
+			}
+		}
+		else
+		{
+			if ( PlayNumberOfClients == 1 )
+			{
+				Desc += LOCTEXT( "MultiplayerDescription_EditorStandalone", "The editor will run in offline mode. " ).ToString();
+			}
+			else
+			{
+				Desc += FText::Format( LOCTEXT( "MultiplayerDescription_StandaloneAndClients", "The editor will run offline and {0} additional offline mode window(s) will also open. " ), FText::AsNumber( PlayNumberOfClients - 1 ) ).ToString();
+			}
+
+			if ( bLaunchSeparateServer )
+			{
+				Desc += LOCTEXT( "MultiplayerDescription_StandaloneSeparateServer", "\nAn additional server instance will be launched but not connected to. Use \"open 127.0.0.1:<port>\" to connect. " ).ToString();
+			}
+		}
+	}
+	else
+	{
+		Desc += LOCTEXT( "MultiplayerDescription_MultiProcess", "The following will run with multiple UE4 instances:\n" ).ToString();
+		if ( PlayNetMode == PIE_Standalone )
+		{
+			if ( PlayNumberOfClients == 1 )
+			{
+				Desc += LOCTEXT( "MultiplayerDescription_EditorStandalone", "The editor will run in offline mode. " ).ToString();
+			}
+			else
+			{
+				Desc += FText::Format( LOCTEXT( "MultiplayerDescription_StandaloneAndClients", "The editor will run offline and {0} additional offline mode window(s) will also open. " ), FText::AsNumber( PlayNumberOfClients - 1 ) ).ToString();
+			}
+
+			if ( bLaunchSeparateServer )
+			{
+				Desc += LOCTEXT( "MultiplayerDescription_StandaloneSeparateServer", "\nAn additional server instance will be launched but not connected to. Use \"open 127.0.0.1:<port>\" to connect. " ).ToString();
+			}
+		}
+		else if ( PlayNetMode == PIE_ListenServer )
+		{
+			if ( PlayNumberOfClients == 1 )
+			{
+				Desc += LOCTEXT( "MultiplayerDescription_EditorListenServer", "The editor will run as a listen server. " ).ToString();
+			}
+			else
+			{
+				Desc += FText::Format( LOCTEXT( "MultiplayerDescription_EditorListenServerAndClients", "The editor will run as a listen server and {0} additional client window(s) will also connect to it. " ), FText::AsNumber( PlayNumberOfClients - 1 ) ).ToString();
+			}
+		}
+		else
+		{
+			// Client requires additional dedicated server instance
+			Desc += LOCTEXT( "MultiplayerDescription_DedicatedServerNewWindow", "A dedicated server will open in a new window. " ).ToString();
+			if ( PlayNumberOfClients == 1 )
+			{
+				Desc += LOCTEXT( "MultiplayerDescription_EditorClient", "The editor will connect as a client. " ).ToString();
+			}
+			else
+			{
+				Desc += FText::Format( LOCTEXT( "MultiplayerDescription_EditorAndClients", "The editor will connect as a client and {0} additional client window(s) will also connect. " ), FText::AsNumber( PlayNumberOfClients - 1 ) ).ToString();
+			}
+
+		}
+	}
+	return FText::FromString( Desc );
+}
+
+bool FLevelEditorPlaySettingsCustomization::HandleClientWindowSizePropertyIsEnabled() const
+{
+	return GetDefault<ULevelEditorPlaySettings>()->IsClientWindowSizeActive();
+}
+
+bool FLevelEditorPlaySettingsCustomization::HandleGameOptionsIsEnabled() const
+{
+	return GetDefault<ULevelEditorPlaySettings>()->IsAdditionalServerGameOptionsActive();
+}
+
+bool FLevelEditorPlaySettingsCustomization::HandleRerouteInputToSecondWindowEnabled() const
+{
+	return GetDefault<ULevelEditorPlaySettings>()->IsRouteGamepadToSecondWindowActive();
+}
+
+EVisibility FLevelEditorPlaySettingsCustomization::HandleRerouteInputToSecondWindowVisibility() const
+{
+	return GetDefault<ULevelEditorPlaySettings>()->GetRouteGamepadToSecondWindowVisibility();
+}
+
+void FLevelEditorPlaySettingsCustomization::HandleQualityLevelComboBoxOpening()
+{
+	const UAudioSettings* AudioSettings = GetDefault<UAudioSettings>();
+	AvailableQualityLevels.Empty( AudioSettings->QualityLevels.Num() );
+	for ( const FAudioQualitySettings& AQSettings : AudioSettings->QualityLevels )
+	{
+		AvailableQualityLevels.Add( MakeShareable( new FString( AQSettings.DisplayName.ToString() ) ) );
+	}
+	QualityLevelComboBox->RefreshOptions();
+}
+
+TSharedRef<SWidget> FLevelEditorPlaySettingsCustomization::HandleQualityLevelComboBoxGenerateWidget( TSharedPtr<FString> InItem )
+{
+	return SNew( STextBlock )
+		.Text( FText::FromString( *InItem ) );
+}
+
+void FLevelEditorPlaySettingsCustomization::HandleQualityLevelSelectionChanged( TSharedPtr<FString> InSelection, ESelectInfo::Type SelectInfo )
+{
+	if ( InSelection.IsValid() )
+	{
+		const UAudioSettings* AudioSettings = GetDefault<UAudioSettings>();
+		for ( int32 QualityLevel = 0; QualityLevel < AudioSettings->QualityLevels.Num(); ++QualityLevel )
+		{
+			if ( AudioSettings->QualityLevels[QualityLevel].DisplayName.ToString() == *InSelection )
+			{
+				PIESoundQualityLevelHandle->SetValue( QualityLevel );
+				break;
+			}
+		}
+	}
+}
+
+FText FLevelEditorPlaySettingsCustomization::GetSelectedQualityLevelName() const
+{
+	int32 QualityLevel = 0;
+	PIESoundQualityLevelHandle->GetValue( QualityLevel );
+	const UAudioSettings* AudioSettings = GetDefault<UAudioSettings>();
+	return (QualityLevel >= 0 && QualityLevel < AudioSettings->QualityLevels.Num() ? AudioSettings->QualityLevels[QualityLevel].DisplayName : FText::GetEmpty());
+}
+
+FText FLevelEditorPlaySettingsCustomization::GetPreviewText() const
+{
+	const ULevelEditorPlaySettings* LevelEditorPlaySettings = GetDefault<ULevelEditorPlaySettings>();
+	float SafeZone = FDisplayMetrics::GetDebugTitleSafeZoneRatio();
+	if ( FMath::IsNearlyEqual( SafeZone, 1.0f ) )
+	{
+		if ( LevelEditorPlaySettings->DeviceToEmulate.IsEmpty() || LevelEditorPlaySettings->PIESafeZoneOverride.GetDesiredSize().IsZero() )
+		{
+			return LOCTEXT( "NoSafeZoneSet", "No Device Safe Zone Set" );
+		}
+		else
+		{
+			return FText::FromString( LevelEditorPlaySettings->DeviceToEmulate );
+		}
+	}
+	else
+	{
+		return FText::Format( LOCTEXT( "UniformSafeZone", "Uniform Safe Zone: {0}" ), FText::AsNumber( SafeZone ) );
+	}
+}
+
+FText FLevelEditorPlaySettingsCustomization::GetPreviewTextToolTipText() const
+{
+	if (FDisplayMetrics::GetDebugTitleSafeZoneRatio() < 1.0f)
+	{
+		return LOCTEXT("SafeZoneSetFromRatioCvarToolTip", "Uniform Safe Zone was set from cvar (r.DebugSafeZone.TitleRatio).");
+	}
+
+	const ULevelEditorPlaySettings* LevelEditorPlaySettings = GetDefault<ULevelEditorPlaySettings>();
+	if (!LevelEditorPlaySettings->DeviceToEmulate.IsEmpty() && !LevelEditorPlaySettings->PIESafeZoneOverride.GetDesiredSize().IsZero() )
+	{
+		FMargin PIESafeZoneOverride = LevelEditorPlaySettings->PIESafeZoneOverride;
+		FNumberFormattingOptions FormattingOptions;
+		FormattingOptions.SetMaximumFractionalDigits(2);
+		FormattingOptions.SetMinimumFractionalDigits(2);
+
+		return FText::Format(
+			LOCTEXT("CustomSafeZoneSetFromDevice", "PIE safe zone override set from device {0} profile.\nLeft: {1}, Top: {2}, Right: {3}, Bottom: {4}"),
+			FText::FromString(LevelEditorPlaySettings->DeviceToEmulate),
+			FText::AsNumber(PIESafeZoneOverride.Left, &FormattingOptions),
+			FText::AsNumber(PIESafeZoneOverride.Top, &FormattingOptions),
+			FText::AsNumber(PIESafeZoneOverride.Right, &FormattingOptions),
+			FText::AsNumber(PIESafeZoneOverride.Bottom, &FormattingOptions)
+		);
+	}
+
+	return FText();
+}
+
+#undef LOCTEXT_NAMESPACE
\ No newline at end of file
diff --git a/Engine/Source/Editor/UnrealEd/Private/Settings/LevelEditorPlaySettingsCustomization.h b/Engine/Source/Editor/UnrealEd/Private/Settings/LevelEditorPlaySettingsCustomization.h
--- a/Engine/Source/Editor/UnrealEd/Private/Settings/LevelEditorPlaySettingsCustomization.h
+++ b/Engine/Source/Editor/UnrealEd/Private/Settings/LevelEditorPlaySettingsCustomization.h
@@ -3,34 +3,19 @@
 #pragma once
 
 #include "CoreMinimal.h"
-#include "Misc/Attribute.h"
-#include "Layout/Margin.h"
-#include "Layout/Visibility.h"
+#include "IDetailCustomization.h"
+#include "Settings/LevelEditorPlaySettings.h"
 #include "Widgets/DeclarativeSyntaxSupport.h"
-#include "Widgets/SWidget.h"
 #include "Widgets/SCompoundWidget.h"
-#include "Widgets/SBoxPanel.h"
-#include "Settings/LevelEditorPlaySettings.h"
-#include "Editor.h"
-#include "Textures/SlateIcon.h"
-#include "Framework/Commands/UIAction.h"
-#include "IDetailCustomization.h"
-#include "PropertyHandle.h"
-#include "DetailLayoutBuilder.h"
-#include "DetailCategoryBuilder.h"
-#include "Framework/MultiBox/MultiBoxBuilder.h"
-#include "Widgets/Text/STextBlock.h"
-#include "Widgets/Images/SImage.h"
-#include "Widgets/Input/SButton.h"
-#include "Widgets/Input/SComboButton.h"
-#include "IDetailPropertyRow.h"
-#include "DetailWidgetRow.h"
-#include "Widgets/Input/SComboBox.h"
-#include "Runtime/Engine/Classes/Sound/AudioSettings.h"
-#include "DeviceProfiles/DeviceProfileManager.h"
-#include "DeviceProfiles/DeviceProfile.h"
-
-#define LOCTEXT_NAMESPACE "FLevelEditorPlaySettingsCustomization"
+
+class FMenuBuilder;
+struct FSlateBrush;
+class IDetailCustomization;
+class IDetailLayoutBuilder;
+class IPropertyHandle;
+template <typename OptionType = TSharedPtr<FString>>
+class SComboBox;
+class SWidget;
 
 class SScreenPositionCustomization
 	: public SCompoundWidget
@@ -50,87 +35,12 @@ public:
 	 * @param InWindowPositionProperty The handle to the window position property.
 	 * @param InCenterWindowProperty The handle to the center window property.
 	 */
-	void Construct( const FArguments& InArgs, IDetailLayoutBuilder* LayoutBuilder, const TSharedRef<IPropertyHandle>& InWindowPositionProperty, const TSharedRef<IPropertyHandle>& InCenterWindowProperty )
-	{
-		check(LayoutBuilder != NULL);
-
-		CenterWindowProperty = InCenterWindowProperty;
-
-		ChildSlot
-		[
-			SNew(SVerticalBox)
-			+SVerticalBox::Slot()
-			.AutoHeight()
-			[
-			SNew(SHorizontalBox)
-			+ SHorizontalBox::Slot()
-			[
-				SNew(SVerticalBox)
-				.IsEnabled(this, &SScreenPositionCustomization::HandleNewWindowPositionPropertyIsEnabled)
-				+ SVerticalBox::Slot()
-				.AutoHeight()
-				[
-					InWindowPositionProperty->CreatePropertyNameWidget(LOCTEXT("WindowPosXLabel", "Left Position"))
-				]
-				+ SVerticalBox::Slot()
-				.AutoHeight()
-				[
-					InWindowPositionProperty->GetChildHandle(0)->CreatePropertyValueWidget()
-				]
-			]
-
-			+ SHorizontalBox::Slot()
-				.Padding(8.0f, 0.0f, 0.0f, 0.0f)
-				[
-					SNew(SVerticalBox)
-						.IsEnabled(this, &SScreenPositionCustomization::HandleNewWindowPositionPropertyIsEnabled)
-
-					+ SVerticalBox::Slot()
-						.AutoHeight()
-						[
-							InWindowPositionProperty->CreatePropertyNameWidget(LOCTEXT("TopPositionLabel", "Top Position"))
-						]
-
-					+ SVerticalBox::Slot()
-						.AutoHeight()
-					[
-						InWindowPositionProperty->GetChildHandle(1)->CreatePropertyValueWidget()
-					]
-				]
-			]
-			+SVerticalBox::Slot()
-			.Padding(0.0f, 2.0f)
-			.AutoHeight()
-				[
-				SNew(SHorizontalBox)
-				+ SHorizontalBox::Slot()
-				.VAlign(VAlign_Center)
-				.AutoWidth()
-				[
-					InCenterWindowProperty->CreatePropertyValueWidget()
-				]
-				+ SHorizontalBox::Slot()
-					.AutoWidth()
-					.Padding(4.0f, 0.0f, 0.0f, 0.0f)
-					.VAlign(VAlign_Bottom)
-					[
-						InWindowPositionProperty->CreatePropertyNameWidget(LOCTEXT("CenterWindowLabel", "Always center window to screen"))
-					]
-				]
-		];
-	}
+	void Construct( const FArguments& InArgs, IDetailLayoutBuilder* LayoutBuilder, const TSharedRef<IPropertyHandle>& InWindowPositionProperty, const TSharedRef<IPropertyHandle>& InCenterWindowProperty );
 
 private:
 
 	// Callback for checking whether the window position properties are enabled.
-	bool HandleNewWindowPositionPropertyIsEnabled( ) const
-	{
-		bool CenterNewWindow;
-		CenterWindowProperty->GetValue(CenterNewWindow);
-
-		return !CenterNewWindow;
-
-	}
+	bool HandleNewWindowPositionPropertyIsEnabled() const;
 
 private:
 
@@ -160,102 +70,14 @@ public:
 	 * @param InWindowHeightProperty The handle to the window height property.
 	 * @param InWindowWidthProperty The handle to the window width property.
 	 */
-	void Construct( const FArguments& InArgs, IDetailLayoutBuilder* LayoutBuilder, const TSharedRef<IPropertyHandle>& InWindowHeightProperty, const TSharedRef<IPropertyHandle>& InWindowWidthProperty )
-	{
-		check(LayoutBuilder != NULL);
-
-		WindowHeightProperty = InWindowHeightProperty;
-		WindowWidthProperty = InWindowWidthProperty;
-		
-		FSimpleDelegate SizeChangeDelegate = FSimpleDelegate::CreateSP(this, &SScreenResolutionCustomization::OnSizeChanged);
-		WindowHeightProperty->SetOnPropertyValueChanged(SizeChangeDelegate);
-		WindowWidthProperty->SetOnPropertyValueChanged(SizeChangeDelegate);
-		ChildSlot
-		[
-			SNew(SVerticalBox)
-			+SVerticalBox::Slot()
-			.VAlign(VAlign_Bottom)
-			[
-				SNew(SHorizontalBox)
-				+ SHorizontalBox::Slot()
-				.AutoWidth()
-				.VAlign(VAlign_Center)
-				[
-					SNew(SComboButton)
-					.VAlign(VAlign_Center)
-					.ButtonContent()
-					[
-						SNew(STextBlock)
-						.Font(LayoutBuilder->GetDetailFont())
-						.Text(LOCTEXT("CommonResolutionsButtonText", "Common Resolutions"))
-					]
-					.ContentPadding(FMargin(6,2))
-					.MenuContent()
-					[
-						MakeCommonResolutionsMenu()
-					]
-					.ToolTipText(LOCTEXT("CommonResolutionsButtonTooltip", "Pick from a list of common screen resolutions"))
-				]
-				+ SHorizontalBox::Slot()
-				.Padding(0,0,6,0)
-				.AutoWidth()
-				.VAlign(VAlign_Center)
-				[
-					SNew(SButton)
-					.OnClicked(this, &SScreenResolutionCustomization::HandleSwapAspectRatioClicked)
-					.ContentPadding(FMargin(3,0,3,1))
-					.Content()
-					[
-						SNew(SImage)
-						.Image(this, &SScreenResolutionCustomization::GetAspectRatioSwitchImage)
-					]
-					.ToolTipText(LOCTEXT("SwapAspectRatioTooltip", "Swap between portrait and landscape orientation."))
-				]
-			]
-			+ SVerticalBox::Slot()
-			.AutoHeight()
-			[
-				SNew(SHorizontalBox)
-				+SHorizontalBox::Slot()
-				[
-					SNew(SVerticalBox)
-					+ SVerticalBox::Slot()
-					.AutoHeight()
-					[
-						WindowWidthProperty->CreatePropertyNameWidget(LOCTEXT("ViewportWidthLabel", "Viewport Width"))
-					]
-					+ SVerticalBox::Slot()
-					[
-						WindowWidthProperty->CreatePropertyValueWidget()
-					]
-				]
-				+ SHorizontalBox::Slot()
-				.Padding(8,0,0,0)
-				[
-					SNew(SVerticalBox)
-					+ SVerticalBox::Slot()
-					.AutoHeight()
-					[
-						WindowHeightProperty->CreatePropertyNameWidget(LOCTEXT("ViewportHeightLabel", "Viewport Height"))
-					]
-					+ SVerticalBox::Slot()
-					.AutoHeight()
-					[
-						WindowHeightProperty->CreatePropertyValueWidget()
-					]
-				]
-			]
-		];
-	}
+	void Construct( const FArguments& InArgs, IDetailLayoutBuilder* LayoutBuilder, const TSharedRef<IPropertyHandle>& InWindowHeightProperty, const TSharedRef<IPropertyHandle>& InWindowWidthProperty );
 
 protected:
 
 	/**
 	 * Adds a menu entry to the common screen resolutions menu.
 	 */
-	void AddCommonResolutionEntry( FMenuBuilder& MenuBuilder, int32 Width, int32 Height, const FString& AspectRatio, const FText& Description )
-	{
-	}
+	void AddCommonResolutionEntry( FMenuBuilder& MenuBuilder, int32 Width, int32 Height, const FString& AspectRatio, const FText& Description );
 
 	/**
 	 * Adds a section to the screen resolution menu.
@@ -264,204 +86,34 @@ protected:
 	 * @param Resolutions The collection of screen resolutions to add.
 	 * @param SectionName The name of the section to add.
 	 */
-	void AddScreenResolutionSection( FMenuBuilder& MenuBuilder, const TArray<FPlayScreenResolution> Resolutions, const FText SectionName )
-	{
-		MenuBuilder.BeginSection(NAME_None, SectionName);
-		{
-			for (auto Iter = Resolutions.CreateConstIterator(); Iter; ++Iter)
-			{
-				FUIAction Action(FExecuteAction::CreateRaw(this, &SScreenResolutionCustomization::HandleCommonResolutionSelected, *Iter));
-
-				FInternationalization& I18N = FInternationalization::Get();
-
-				FFormatNamedArguments Args;
-				Args.Add(TEXT("Width"), FText::AsNumber(Iter->Width, NULL, I18N.GetInvariantCulture()));
-				Args.Add(TEXT("Height"), FText::AsNumber(Iter->Height, NULL, I18N.GetInvariantCulture()));
-				Args.Add(TEXT("AspectRatio"), FText::FromString(Iter->AspectRatio));
-
-				MenuBuilder.AddMenuEntry(FText::FromString(Iter->Description), FText::Format(LOCTEXT("CommonResolutionFormat", "{Width} x {Height} ({AspectRatio})"), Args), FSlateIcon(), Action);
-			}
-		}
-		MenuBuilder.EndSection();
-	}
+	void AddScreenResolutionSection( FMenuBuilder& MenuBuilder, const TArray<FPlayScreenResolution> Resolutions, const FText SectionName );
 
 	/**
 	 * Creates a widget for the resolution picker.
 	 *
 	 * @return The widget.
 	 */
-	TSharedRef<SWidget> MakeCommonResolutionsMenu( )
-	{
-		const ULevelEditorPlaySettings* PlaySettings = GetDefault<ULevelEditorPlaySettings>();
-		FMenuBuilder MenuBuilder(true, NULL);
-
-
-		FText PhoneTitle = LOCTEXT("CommonPhonesSectionHeader", "Phones");
-		FText TabletTitle = LOCTEXT("CommonTabletsSectionHeader", "Tablets");
-		FText LaptopTitle = LOCTEXT("CommonLaptopsSectionHeader", "Laptops");
-		FText MonitorTitle = LOCTEXT("CommonMonitorsSectionHeader", "Monitors");
-		FText TelevisionTitle = LOCTEXT("CommonTelevesionsSectionHeader", "Televisions");
-		MenuBuilder.AddSubMenu(
-			PhoneTitle,
-			FText(),
-			FNewMenuDelegate::CreateRaw(this, &SScreenResolutionCustomization::AddScreenResolutionSection, (PlaySettings->PhoneScreenResolutions), PhoneTitle),
-			false,
-			FSlateIcon());
-		MenuBuilder.AddSubMenu(
-			TabletTitle,
-			FText(),
-			FNewMenuDelegate::CreateRaw(this, &SScreenResolutionCustomization::AddScreenResolutionSection, (PlaySettings->TabletScreenResolutions), TabletTitle),
-			false,
-			FSlateIcon());
-		MenuBuilder.AddSubMenu(
-			LaptopTitle,
-			FText(),
-			FNewMenuDelegate::CreateRaw(this, &SScreenResolutionCustomization::AddScreenResolutionSection, (PlaySettings->LaptopScreenResolutions), LaptopTitle),
-			false,
-			FSlateIcon());
-		MenuBuilder.AddSubMenu(
-			MonitorTitle,
-			FText(),
-			FNewMenuDelegate::CreateRaw(this, &SScreenResolutionCustomization::AddScreenResolutionSection, (PlaySettings->MonitorScreenResolutions), MonitorTitle),
-			false,
-			FSlateIcon());
-		MenuBuilder.AddSubMenu(
-			TelevisionTitle,
-			FText(),
-			FNewMenuDelegate::CreateRaw(this, &SScreenResolutionCustomization::AddScreenResolutionSection, (PlaySettings->TelevisionScreenResolutions), TelevisionTitle),
-			false,
-			FSlateIcon());
-
-		return MenuBuilder.MakeWidget();
-	}
-
-	FReply HandleSwapAspectRatioClicked()
-	{
-		FString HeightString;
-		WindowHeightProperty->GetValueAsDisplayString(HeightString);
-		FString WidthString;
-		WindowWidthProperty->GetValueAsDisplayString(WidthString);
-		int32 NewHeight = FCString::Atoi(*WidthString);
-		int32 NewWidth = FCString::Atoi(*HeightString);
-
-		ULevelEditorPlaySettings* PlayInSettings = GetMutableDefault<ULevelEditorPlaySettings>();
-
-		if (!PlayInSettings->DeviceToEmulate.IsEmpty() && FDisplayMetrics::GetDebugTitleSafeZoneRatio() == 1.0f)
-		{
-			const UDeviceProfile* DeviceProfile = UDeviceProfileManager::Get().FindProfile(PlayInSettings->DeviceToEmulate, false);
-			// Rescale the swapped sizes if we are on Android
-			if (DeviceProfile && DeviceProfile->DeviceType == TEXT("Android"))
-			{
-				float ScaleFactor = 1.0f;
-				PlayInSettings->RescaleForMobilePreview(DeviceProfile, NewWidth, NewHeight, ScaleFactor);
-			}
-			PlayInSettings->PIESafeZoneOverride = PlayInSettings->CalculateCustomUnsafeZones(PlayInSettings->CustomUnsafeZoneStarts, PlayInSettings->CustomUnsafeZoneDimensions, PlayInSettings->DeviceToEmulate, FVector2D(NewWidth, NewHeight));
-		}
-		else
-		{
-			FSlateApplication::Get().ResetCustomSafeZone();
-			FSlateApplication::Get().GetSafeZoneSize(PlayInSettings->PIESafeZoneOverride, FVector2D(NewWidth, NewHeight));
-		}
-		FMargin SafeZoneRatio = PlayInSettings->PIESafeZoneOverride;
-		SafeZoneRatio.Left /= (NewWidth / 2.0f);
-		SafeZoneRatio.Right /= (NewWidth / 2.0f);
-		SafeZoneRatio.Bottom /= (NewHeight / 2.0f);
-		SafeZoneRatio.Top /= (NewHeight / 2.0f);
-		FSlateApplication::Get().OnDebugSafeZoneChanged.Broadcast(SafeZoneRatio, true);
-
-
-		bSetFromMenu = true;
-		WindowHeightProperty->SetValue(NewHeight);
-		bSetFromMenu = true;
-		WindowWidthProperty->SetValue(NewWidth);
-
-		return FReply::Handled();
-	}
+	TSharedRef<SWidget> MakeCommonResolutionsMenu();
+	
+	// Resets LevelEditorPlaySettings that relate to the title safe zone debug draw.
+	void ResetCustomTitleSafeZoneSettings(ULevelEditorPlaySettings* PlayInSettings, const int32 Width, const int32 Height);
+
+	// Called when the PIE safe zone override is changed. Performs final safe zone calculations and broadcasts the result.
+	void BroadcastSafeZoneChanged(const FMargin& SafeZoneRatio, const int32 Width, const int32 Height);
+
+	// Handles swapping the current aspect ratio.
+	// Also will set a custom safe zone matching the device profile, if r.DebugSafeZone.TitleRatio is set to 1.0
+	FReply HandleSwapAspectRatioClicked();
+
 private:
 
 	// Handles selecting a common screen resolution.
-	void HandleCommonResolutionSelected(const FPlayScreenResolution Resolution)
-	{
-		int32 Width = Resolution.Width;
-		int32 Height = Resolution.Height;
-		float ScaleFactor;
-		ULevelEditorPlaySettings* PlayInSettings = GetMutableDefault<ULevelEditorPlaySettings>();
-		// Maintain previous orientation (i.e., swap Width and Height if required)
-		int32 PreviousWidth;
-		int32 PreviousHeight;
-		if (WindowWidthProperty->GetValue(PreviousWidth) == FPropertyAccess::Success && WindowHeightProperty->GetValue(PreviousHeight) == FPropertyAccess::Success)
-		{
-			const bool bIsOrientationPreserved = (PreviousWidth < PreviousHeight) != (Width < Height);
-			if (bIsOrientationPreserved)
-			{
-				Width = Resolution.Height;
-				Height = Resolution.Width;
-			}
-		}
-
-		UDeviceProfile* DeviceProfile = UDeviceProfileManager::Get().FindProfile(Resolution.ProfileName, false);
-		if (DeviceProfile)
-		{
-			PlayInSettings->DeviceToEmulate = Resolution.ProfileName;
-			PlayInSettings->RescaleForMobilePreview(DeviceProfile, Width, Height, ScaleFactor);
-			PlayInSettings->PIESafeZoneOverride = PlayInSettings->CalculateCustomUnsafeZones(PlayInSettings->CustomUnsafeZoneStarts, PlayInSettings->CustomUnsafeZoneDimensions, PlayInSettings->DeviceToEmulate, FVector2D(Width, Height));
-		}
-		else
-		{
-			PlayInSettings->DeviceToEmulate = FString();
-			FSlateApplication::Get().ResetCustomSafeZone();
-			FSlateApplication::Get().GetSafeZoneSize(PlayInSettings->PIESafeZoneOverride, FVector2D(Width, Height));
-		}
-		FMargin SafeZoneRatio = PlayInSettings->PIESafeZoneOverride;
-		SafeZoneRatio.Left /= (Width / 2.0f);
-		SafeZoneRatio.Right /= (Width / 2.0f);
-		SafeZoneRatio.Bottom /= (Height / 2.0f);
-		SafeZoneRatio.Top /= (Height / 2.0f);
-		bSetFromMenu = true;
-		WindowHeightProperty->SetValue(Height);
-		bSetFromMenu = true;
-		WindowWidthProperty->SetValue(Width);
-		FSlateApplication::Get().OnDebugSafeZoneChanged.Broadcast(SafeZoneRatio, true);
-	}
-
-	const FSlateBrush* GetAspectRatioSwitchImage() const
-	{
-		FString HeightString;
-		WindowHeightProperty->GetValueAsDisplayString(HeightString);
-		int32 Height = FCString::Atoi(*HeightString);
-		FString WidthString;
-		WindowWidthProperty->GetValueAsDisplayString(WidthString);
-		int32 Width = FCString::Atoi(*WidthString);
-		if(Height > Width)
-		{
-			return FEditorStyle::Get().GetBrush("UMGEditor.OrientPortrait");
-		}
-		return FEditorStyle::Get().GetBrush("UMGEditor.OrientLandscape");
-	}
-
-	void OnSizeChanged()
-	{
-		if (!bSetFromMenu)
-		{
-			FString HeightString;
-			WindowHeightProperty->GetValueAsDisplayString(HeightString);
-			int32 Height = FCString::Atoi(*HeightString);
-			FString WidthString;
-			WindowWidthProperty->GetValueAsDisplayString(WidthString);
-			int32 Width = FCString::Atoi(*WidthString);
-			ULevelEditorPlaySettings* PlayInSettings = GetMutableDefault<ULevelEditorPlaySettings>();
-			PlayInSettings->DeviceToEmulate = FString();
-			FSlateApplication::Get().ResetCustomSafeZone();
-			FSlateApplication::Get().GetSafeZoneSize(PlayInSettings->PIESafeZoneOverride, FVector2D(Width, Height));
-			FMargin SafeZoneRatio = PlayInSettings->PIESafeZoneOverride;
-			SafeZoneRatio.Left /= (Width / 2.0f);
-			SafeZoneRatio.Right /= (Width / 2.0f);
-			SafeZoneRatio.Bottom /= (Height / 2.0f);
-			SafeZoneRatio.Top /= (Height / 2.0f);
-			FSlateApplication::Get().OnDebugSafeZoneChanged.Broadcast(SafeZoneRatio, true);
-		}
-		bSetFromMenu = false;
-	}
+	// Also will set a custom safe zone matching the device profile, if r.DebugSafeZone.TitleRatio is set to 1.0
+	void HandleCommonResolutionSelected( const FPlayScreenResolution Resolution );
+
+	const FSlateBrush* GetAspectRatioSwitchImage() const;
+
+	void OnSizeChanged();
 
 private:
 
@@ -485,222 +137,12 @@ class FLevelEditorPlaySettingsCustomization
 public:
 
 	/** Virtual destructor. */
-	virtual ~FLevelEditorPlaySettingsCustomization( ) { }
+	virtual ~FLevelEditorPlaySettingsCustomization();
 
 public:
 
 	// IDetailCustomization interface
-
-	virtual void CustomizeDetails( IDetailLayoutBuilder& LayoutBuilder ) override
-	{
-		const float MaxPropertyWidth = 400.0f;
-
-		// play in editor settings
-		IDetailCategoryBuilder& PlayInEditorCategory = LayoutBuilder.EditCategory("PlayInEditor");
-		{
-			TArray<TSharedRef<IPropertyHandle>> PIECategoryProperties;
-			PlayInEditorCategory.GetDefaultProperties(PIECategoryProperties, true, false);
-
-			TSharedPtr<IPropertyHandle> PIEEnableSoundHandle = LayoutBuilder.GetProperty(GET_MEMBER_NAME_CHECKED(ULevelEditorPlaySettings, EnableGameSound));
-			PIESoundQualityLevelHandle = LayoutBuilder.GetProperty(GET_MEMBER_NAME_CHECKED(ULevelEditorPlaySettings, PlayInEditorSoundQualityLevel));
-			PIESoundQualityLevelHandle->MarkHiddenByCustomization();
-
-			for (TSharedRef<IPropertyHandle>& PropertyHandle : PIECategoryProperties)
-			{
-				if (PropertyHandle->GetProperty() != PIESoundQualityLevelHandle->GetProperty())
-				{
-					PlayInEditorCategory.AddProperty(PropertyHandle);
-				}
-
-				if (PropertyHandle->GetProperty() == PIEEnableSoundHandle->GetProperty())
-				{
-					PlayInEditorCategory.AddCustomRow(PIESoundQualityLevelHandle->GetPropertyDisplayName(), false)
-						.NameContent()
-						[
-							PIESoundQualityLevelHandle->CreatePropertyNameWidget()
-						]
-						.ValueContent()
-						.MaxDesiredWidth(MaxPropertyWidth)
-						[
-							SAssignNew(QualityLevelComboBox, SComboBox<TSharedPtr<FString>>)
-							.OptionsSource(&AvailableQualityLevels)
-							.OnComboBoxOpening(this, &FLevelEditorPlaySettingsCustomization::HandleQualityLevelComboBoxOpening)
-							.OnGenerateWidget(this, &FLevelEditorPlaySettingsCustomization::HandleQualityLevelComboBoxGenerateWidget)
-							.OnSelectionChanged(this, &FLevelEditorPlaySettingsCustomization::HandleQualityLevelSelectionChanged)
-							[
-								SNew(STextBlock)
-								.Text(this, &FLevelEditorPlaySettingsCustomization::GetSelectedQualityLevelName)
-							]
-						];
-				}
-			}
-
-
-		}
-		IDetailCategoryBuilder& GameViewportSettings = LayoutBuilder.EditCategory("GameViewportSettings");
-		{
-		// new window resolution
-		TSharedRef<IPropertyHandle> WindowHeightHandle = LayoutBuilder.GetProperty(GET_MEMBER_NAME_CHECKED(ULevelEditorPlaySettings, NewWindowHeight));
-		TSharedRef<IPropertyHandle> WindowWidthHandle = LayoutBuilder.GetProperty(GET_MEMBER_NAME_CHECKED(ULevelEditorPlaySettings, NewWindowWidth));
-		TSharedRef<IPropertyHandle> WindowPositionHandle = LayoutBuilder.GetProperty(GET_MEMBER_NAME_CHECKED(ULevelEditorPlaySettings, NewWindowPosition));
-		TSharedRef<IPropertyHandle> CenterNewWindowHandle = LayoutBuilder.GetProperty(GET_MEMBER_NAME_CHECKED(ULevelEditorPlaySettings, CenterNewWindow));
-		TSharedRef<IPropertyHandle> EmulatedDeviceHandle = LayoutBuilder.GetProperty(GET_MEMBER_NAME_CHECKED(ULevelEditorPlaySettings, DeviceToEmulate));
-
-		WindowHeightHandle->MarkHiddenByCustomization();
-		WindowWidthHandle->MarkHiddenByCustomization();
-		WindowPositionHandle->MarkHiddenByCustomization();
-		CenterNewWindowHandle->MarkHiddenByCustomization();
-		EmulatedDeviceHandle->MarkHiddenByCustomization();
-
-		GameViewportSettings.AddCustomRow(LOCTEXT("NewViewportResolutionRow", "New Viewport Resolution"), false)
-			.NameContent()
-			[
-				SNew(STextBlock)
-				.Font(LayoutBuilder.GetDetailFont())
-				.Text(LOCTEXT("NewViewportResolutionName", "New Viewport Resolution"))
-				.ToolTipText(LOCTEXT("NewWindowSizeTooltip", "Sets the width and height of floating PIE windows (in pixels)"))
-			]
-			.ValueContent()
-			.MaxDesiredWidth(MaxPropertyWidth)
-			[
-				SNew(SScreenResolutionCustomization, &LayoutBuilder, WindowHeightHandle, WindowWidthHandle)
-			];
-
-		GameViewportSettings.AddCustomRow(LOCTEXT("NewWindowPositionRow", "New Window Position"), false)
-			.NameContent()
-			[
-				SNew(STextBlock)
-				.Font(LayoutBuilder.GetDetailFont())
-				.Text(LOCTEXT("NewWindowPositionName", "New Window Position"))
-				.ToolTipText(LOCTEXT("NewWindowPositionTooltip", "Sets the screen coordinates for the top-left corner of floating PIE windows (in pixels)"))
-			]
-			.ValueContent()
-			.MaxDesiredWidth(MaxPropertyWidth)
-			[
-				SNew(SScreenPositionCustomization, &LayoutBuilder, WindowPositionHandle, CenterNewWindowHandle)
-			];
-
-		GameViewportSettings.AddCustomRow(LOCTEXT("SafeZonePreviewName", "Safe Zone Preview"), false)
-			.NameContent()
-			[
-				SNew(STextBlock)
-				.Font(LayoutBuilder.GetDetailFont())
-				.Text(LOCTEXT("SafeZonePreviewName", "Safe Zone Preview"))
-			]
-			.ValueContent()
-			[
-				SNew(STextBlock)
-				.Font(LayoutBuilder.GetDetailFont())
-				.Text(this, &FLevelEditorPlaySettingsCustomization::GetPreviewText)
-			];
-		}
-
-		// play in new window settings
-		IDetailCategoryBuilder& PlayInNewWindowCategory = LayoutBuilder.EditCategory("PlayInNewWindow");
-		{
-			// Mac does not support parenting, do not show
-#if PLATFORM_MAC
-			PlayInNewWindowCategory.AddProperty("PIEAlwaysOnTop")
-				.DisplayName(LOCTEXT("PIEAlwaysOnTop", "Always On Top"))
-				.IsEnabled(false);
-#else
-			PlayInNewWindowCategory.AddProperty("PIEAlwaysOnTop")
-				.DisplayName(LOCTEXT("PIEAlwaysOnTop", "Always On Top"));
-#endif
-		}
-			
-
-		// play in standalone game settings
-		IDetailCategoryBuilder& PlayInStandaloneCategory = LayoutBuilder.EditCategory("PlayInStandaloneGame");
-		{
-			// command line options
-			TSharedPtr<IPropertyHandle> DisableStandaloneSoundProperty = LayoutBuilder.GetProperty(GET_MEMBER_NAME_CHECKED(ULevelEditorPlaySettings, DisableStandaloneSound));
-
-			DisableStandaloneSoundProperty->MarkHiddenByCustomization();
-
-			PlayInStandaloneCategory.AddCustomRow(LOCTEXT("AdditionalStandaloneDetails", "Additional Options"), false)
-				.NameContent()
-				[
-					SNew(STextBlock)
-						.Font(LayoutBuilder.GetDetailFont())
-						.Text(LOCTEXT("ClientCmdLineName", "Command Line Options"))
-						.ToolTipText(LOCTEXT("ClientCmdLineTooltip", "Generates a command line for additional settings that will be passed to the game clients."))
-				]
-				.ValueContent()
-				.MaxDesiredWidth(MaxPropertyWidth)
-				[
-					SNew(SHorizontalBox)
-
-					+ SHorizontalBox::Slot()
-						.AutoWidth()
-						[
-							DisableStandaloneSoundProperty->CreatePropertyValueWidget()
-						]
-
-					+ SHorizontalBox::Slot()
-						.Padding(0.0f, 2.5f)
-						.VAlign(VAlign_Center)
-						.AutoWidth()
-						[
-							DisableStandaloneSoundProperty->CreatePropertyNameWidget(LOCTEXT("DisableStandaloneSoundLabel", "Disable Sound (-nosound)"))
-						]
-		
-				];
-		}
-
-		// multi-player options
-		IDetailCategoryBuilder& NetworkCategory = LayoutBuilder.EditCategory("Multiplayer Options");
-		TArray<TSharedRef<IPropertyHandle>> AllProperties;
-		NetworkCategory.GetDefaultProperties(AllProperties);
-
-		{
-			// Add all the default properties in before we add custom rows at the bottom.
-			for (TSharedRef<IPropertyHandle> DefaultProperty : AllProperties)
-			{
-				NetworkCategory.AddProperty(DefaultProperty);
-			}
-			
-
-			// client window size
-			TSharedRef<IPropertyHandle> WindowHeightHandle = LayoutBuilder.GetProperty("ClientWindowHeight");
-			TSharedRef<IPropertyHandle> WindowWidthHandle = LayoutBuilder.GetProperty("ClientWindowWidth");
-
-			WindowHeightHandle->MarkHiddenByCustomization();
-			WindowWidthHandle->MarkHiddenByCustomization();
-
-			
-
-			NetworkCategory.AddCustomRow(LOCTEXT("PlayInNetworkViewportSize", "Viewport Size\n(Additional Clients)"), false)
-				.NameContent()
-				[
-					WindowHeightHandle->CreatePropertyNameWidget(LOCTEXT("ClientViewportSizeName", "Multiplayer Viewport Size (in pixels)"), LOCTEXT("ClientWindowSizeTooltip", "Width and Height to use when spawning additional clients. Useful when you need multiple clients connected but only interact with one window."))
-				]
-				.ValueContent()
-				.MaxDesiredWidth(MaxPropertyWidth)
-				[
-					SNew(SScreenResolutionCustomization, &LayoutBuilder, WindowHeightHandle, WindowWidthHandle)
-				]
-				.IsEnabled(TAttribute<bool>(this, &FLevelEditorPlaySettingsCustomization::HandleClientWindowSizePropertyIsEnabled));
-				
-			NetworkCategory.AddCustomRow(LOCTEXT("AdditionalMultiplayerDetails", "Additional Options"), true)
-				.NameContent()
-				[
-					SNew(STextBlock)
-					.Font(LayoutBuilder.GetDetailFont())
-					.Text(LOCTEXT("PlainTextName", "Play In Editor Description"))
-					.ToolTipText(LOCTEXT("PlainTextToolTip", "A brief description of the multiplayer settings and what to expect if you play with them in the editor."))
-				]
-				.ValueContent()
-				.MaxDesiredWidth(MaxPropertyWidth)
-				[
-					SNew(STextBlock)
-					.Font(LayoutBuilder.GetDetailFont())
-					.Text(this, &FLevelEditorPlaySettingsCustomization::HandleMultiplayerOptionsDescription)
-					.WrapTextAt(MaxPropertyWidth)
-				];
-		}
-	}
-
+	virtual void CustomizeDetails( IDetailLayoutBuilder& LayoutBuilder ) override;
 	// End IDetailCustomization interface
 
 public:
@@ -710,191 +152,36 @@ public:
 	 *
 	 * @return A new struct customization for play-in settings.
 	 */
-	static TSharedRef<IDetailCustomization> MakeInstance( )
-	{
-		return MakeShareable(new FLevelEditorPlaySettingsCustomization());
-	}
+	static TSharedRef<IDetailCustomization> MakeInstance();
 
 private:
 
 	// Callback for getting the description of the settings
-	FText HandleMultiplayerOptionsDescription( ) const
-	{
-		const ULevelEditorPlaySettings* PlayInSettings = GetDefault<ULevelEditorPlaySettings>();
-		const bool bLaunchSeparateServer = PlayInSettings->bLaunchSeparateServer;
-		const bool CanRunUnderOneProcess = [&PlayInSettings]{ bool RunUnderOneProcess(false); return (PlayInSettings->GetRunUnderOneProcess(RunUnderOneProcess) && RunUnderOneProcess); }();
-		const int32 PlayNumberOfClients = [&PlayInSettings] { int32 NumClients(0); return (PlayInSettings->GetPlayNumberOfClients(NumClients) ? NumClients : 0); }();
-		const EPlayNetMode PlayNetMode = [&PlayInSettings]{ EPlayNetMode NetMode(PIE_Standalone); return (PlayInSettings->GetPlayNetMode(NetMode) ? NetMode : PIE_Standalone); }();
-		FString Desc;
-		if (CanRunUnderOneProcess)
-		{
-			Desc += LOCTEXT("MultiplayerDescription_OneProcess", "The following will all run under one UE4 instance:\n").ToString();
-			if (PlayNetMode == EPlayNetMode::PIE_Client)
-			{
-				Desc += LOCTEXT("MultiplayerDescription_DedicatedServerHidden", "A hidden dedicated server instance will run in editor. ").ToString();
-				if (PlayNumberOfClients == 1)
-				{
-					Desc += LOCTEXT("MultiplayerDescription_EditorClient", "The editor will connect as a client. ").ToString();
-				}
-				else
-				{
-					Desc += FText::Format(LOCTEXT("MultiplayerDescription_EditorAndClients", "The editor will connect as a client and {0} additional client window(s) will also connect. "), FText::AsNumber(PlayNumberOfClients-1)).ToString();
-				}
-			}
-			else if(PlayNetMode == EPlayNetMode::PIE_ListenServer)
-			{
-				if (PlayNumberOfClients == 1)
-				{
-					Desc += LOCTEXT("MultiplayerDescription_EditorListenServer", "The editor will run as a listen server. ").ToString();
-				}
-				else
-				{
-					Desc += FText::Format(LOCTEXT("MultiplayerDescription_EditorListenServerAndClients", "The editor will run as a listen server and {0} additional client window(s) will also connect to it. "), FText::AsNumber(PlayNumberOfClients-1)).ToString();
-				}
-			}
-			else
-			{
-				if (PlayNumberOfClients == 1)
-				{
-					Desc += LOCTEXT("MultiplayerDescription_EditorStandalone", "The editor will run in offline mode. ").ToString();
-				}
-				else
-				{
-					Desc += FText::Format(LOCTEXT("MultiplayerDescription_StandaloneAndClients", "The editor will run offline and {0} additional offline mode window(s) will also open. "), FText::AsNumber(PlayNumberOfClients - 1)).ToString();
-				}
-
-				if (bLaunchSeparateServer)
-				{
-					Desc += LOCTEXT("MultiplayerDescription_StandaloneSeparateServer", "\nAn additional server instance will be launched but not connected to. Use \"open 127.0.0.1:<port>\" to connect. ").ToString();
-				}
-			}
-		}
-		else
-		{
-			Desc += LOCTEXT("MultiplayerDescription_MultiProcess", "The following will run with multiple UE4 instances:\n").ToString();
-			if (PlayNetMode == PIE_Standalone)
-			{
-				if (PlayNumberOfClients == 1)
-				{
-					Desc += LOCTEXT("MultiplayerDescription_EditorStandalone", "The editor will run in offline mode. ").ToString();
-				}
-				else
-				{
-					Desc += FText::Format(LOCTEXT("MultiplayerDescription_StandaloneAndClients", "The editor will run offline and {0} additional offline mode window(s) will also open. "), FText::AsNumber(PlayNumberOfClients - 1)).ToString();
-				}
-
-				if (bLaunchSeparateServer)
-				{
-					Desc += LOCTEXT("MultiplayerDescription_StandaloneSeparateServer", "\nAn additional server instance will be launched but not connected to. Use \"open 127.0.0.1:<port>\" to connect. ").ToString();
-				}
-			}
-			else if (PlayNetMode == PIE_ListenServer)
-			{
-				if (PlayNumberOfClients == 1)
-				{
-					Desc += LOCTEXT("MultiplayerDescription_EditorListenServer", "The editor will run as a listen server. ").ToString();
-				}
-				else
-				{
-					Desc += FText::Format(LOCTEXT("MultiplayerDescription_EditorListenServerAndClients", "The editor will run as a listen server and {0} additional client window(s) will also connect to it. "), FText::AsNumber(PlayNumberOfClients-1)).ToString();
-				}	
-			}
-			else
-			{
-				// Client requires additional dedicated server instance
-				Desc += LOCTEXT("MultiplayerDescription_DedicatedServerNewWindow", "A dedicated server will open in a new window. ").ToString();
-				if (PlayNumberOfClients == 1)
-				{
-					Desc += LOCTEXT("MultiplayerDescription_EditorClient", "The editor will connect as a client. ").ToString();
-				}
-				else
-				{
-					Desc += FText::Format(LOCTEXT("MultiplayerDescription_EditorAndClients", "The editor will connect as a client and {0} additional client window(s) will also connect. "), FText::AsNumber(PlayNumberOfClients-1)).ToString();
-				}
-
-			}
-		}
-		return FText::FromString(Desc);
-	}
+	FText HandleMultiplayerOptionsDescription() const;
 
 	// Callback for checking whether the ClientWindowHeight and ClientWindowWidth properties are enabled.
-	bool HandleClientWindowSizePropertyIsEnabled( ) const
-	{
-		return GetDefault<ULevelEditorPlaySettings>()->IsClientWindowSizeActive();
-	}
-
-	// Callback for checking whether the AdditionalServerGameOptions is enabled.
-	bool HandleGameOptionsIsEnabled( ) const
-	{
-		return GetDefault<ULevelEditorPlaySettings>()->IsAdditionalServerGameOptionsActive();
-	}
-
-	// Callback for getting the enabled state of the RerouteInputToSecondWindow property.
-	bool HandleRerouteInputToSecondWindowEnabled( ) const
-	{
-		return GetDefault<ULevelEditorPlaySettings>()->IsRouteGamepadToSecondWindowActive();
-	}
+	bool HandleClientWindowSizePropertyIsEnabled() const;
+
+	// Callback for checking whether the AdditionalServerGameOptions is enabled.;
+	bool HandleGameOptionsIsEnabled( ) const;
+
+	// Callback for getting the enabled state of the RerouteInputToSecondWindow property.;
+	bool HandleRerouteInputToSecondWindowEnabled( ) const;
 	
 	// Callback for getting the visibility of the RerouteInputToSecondWindow property.
-	EVisibility HandleRerouteInputToSecondWindowVisibility( ) const
-	{
-		return GetDefault<ULevelEditorPlaySettings>()->GetRouteGamepadToSecondWindowVisibility();
-	}
-
-	void HandleQualityLevelComboBoxOpening()
-	{
-		const UAudioSettings* AudioSettings = GetDefault<UAudioSettings>();
-		AvailableQualityLevels.Empty(AudioSettings->QualityLevels.Num());
-		for (const FAudioQualitySettings& AQSettings : AudioSettings->QualityLevels)
-		{
-			AvailableQualityLevels.Add(MakeShareable(new FString(AQSettings.DisplayName.ToString())));
-		}
-		QualityLevelComboBox->RefreshOptions();
-	}
-
-	TSharedRef<SWidget> HandleQualityLevelComboBoxGenerateWidget(TSharedPtr<FString> InItem)
-	{
-		return SNew(STextBlock)
-				.Text(FText::FromString(*InItem));
-	}
-
-	void HandleQualityLevelSelectionChanged(TSharedPtr<FString> InSelection, ESelectInfo::Type SelectInfo)
-	{
-		if (InSelection.IsValid())
-		{
-			const UAudioSettings* AudioSettings = GetDefault<UAudioSettings>();
-			for (int32 QualityLevel = 0; QualityLevel < AudioSettings->QualityLevels.Num(); ++QualityLevel)
-			{
-				if (AudioSettings->QualityLevels[QualityLevel].DisplayName.ToString() == *InSelection)
-				{
-					PIESoundQualityLevelHandle->SetValue(QualityLevel);
-					break;
-				}
-			}
-		}
-	}
-
-	FText GetSelectedQualityLevelName() const
-	{
-		int32 QualityLevel = 0;
-		PIESoundQualityLevelHandle->GetValue(QualityLevel);
-		const UAudioSettings* AudioSettings = GetDefault<UAudioSettings>();
-		return (QualityLevel >= 0 && QualityLevel < AudioSettings->QualityLevels.Num() ? AudioSettings->QualityLevels[QualityLevel].DisplayName : FText::GetEmpty());
-	}
-
-	FText GetPreviewText() const 
-	{
-		if (GetDefault<ULevelEditorPlaySettings>()->DeviceToEmulate.IsEmpty())
-		{
-			float SafeZone = FDisplayMetrics::GetDebugTitleSafeZoneRatio();
-			if (FMath::IsNearlyEqual(SafeZone, 1.0f))
-			{
-				return LOCTEXT("NoSafeZoneSet", "No Device Safe Zone Set");
-			}
-			return FText::Format(LOCTEXT("UniformSafeZone", "Uniform Safe Zone: {0}"), FText::AsNumber(SafeZone));
-		}
-		return FText::FromString(GetDefault<ULevelEditorPlaySettings>()->DeviceToEmulate);
-	}
+	EVisibility HandleRerouteInputToSecondWindowVisibility() const;
+
+	void HandleQualityLevelComboBoxOpening();
+
+	TSharedRef<SWidget> HandleQualityLevelComboBoxGenerateWidget( TSharedPtr<FString> InItem );
+
+	void HandleQualityLevelSelectionChanged( TSharedPtr<FString> InSelection, ESelectInfo::Type SelectInfo );
+
+	FText GetSelectedQualityLevelName() const;
+
+	FText GetPreviewText() const;
+
+	FText GetPreviewTextToolTipText() const;
 private:
 
 	/** Collection of possible quality levels we can use as a parent for this profile */
@@ -903,6 +190,3 @@ private:
 	TSharedPtr<SComboBox<TSharedPtr<FString>>> QualityLevelComboBox;
 
 };
-
-
-#undef LOCTEXT_NAMESPACE
\ No newline at end of file
diff --git a/Engine/Source/Editor/UnrealEd/Private/Settings/SettingsClasses.cpp b/Engine/Source/Editor/UnrealEd/Private/Settings/SettingsClasses.cpp
--- a/Engine/Source/Editor/UnrealEd/Private/Settings/SettingsClasses.cpp
+++ b/Engine/Source/Editor/UnrealEd/Private/Settings/SettingsClasses.cpp
@@ -513,7 +513,7 @@ void ULevelEditorPlaySettings::PostInitProperties()
 	NetworkEmulationSettings.OnPostInitProperties();
 
 #if WITH_EDITOR
-	FCoreDelegates::OnSafeFrameChangedEvent.AddUObject(this, &ULevelEditorPlaySettings::SwapSafeZoneTypes);
+	FCoreDelegates::OnSafeFrameChangedEvent.AddUObject(this, &ULevelEditorPlaySettings::UpdateCustomSafeZones);
 #endif
 }
 
@@ -531,12 +531,25 @@ bool ULevelEditorPlaySettings::CanEditChange(const FProperty* InProperty) const
 }
 
 #if WITH_EDITOR
-void ULevelEditorPlaySettings::SwapSafeZoneTypes()
+void ULevelEditorPlaySettings::UpdateCustomSafeZones()
 {
+	// Prefer to use r.DebugSafeZone.TitleRatio if it is set
 	if (FDisplayMetrics::GetDebugTitleSafeZoneRatio() < 1.f)
 	{
-		DeviceToEmulate = FString();
+		FSlateApplication::Get().ResetCustomSafeZone();
+		PIESafeZoneOverride = FMargin();
 	}
+	else
+	{
+		PIESafeZoneOverride = CalculateCustomUnsafeZones(CustomUnsafeZoneStarts, CustomUnsafeZoneDimensions, DeviceToEmulate, FVector2D(NewWindowWidth, NewWindowHeight));
+	}
+
+	FMargin SafeZoneRatio = PIESafeZoneOverride;
+	SafeZoneRatio.Left /= (NewWindowWidth / 2.0f);
+	SafeZoneRatio.Right /= (NewWindowWidth / 2.0f);
+	SafeZoneRatio.Bottom /= (NewWindowHeight / 2.0f);
+	SafeZoneRatio.Top /= (NewWindowHeight / 2.0f);
+	FSlateApplication::Get().OnDebugSafeZoneChanged.Broadcast(SafeZoneRatio, true);
 }
 #endif
 
diff --git a/Engine/Source/Runtime/Engine/Classes/Engine/GameViewportClient.h b/Engine/Source/Runtime/Engine/Classes/Engine/GameViewportClient.h
--- a/Engine/Source/Runtime/Engine/Classes/Engine/GameViewportClient.h
+++ b/Engine/Source/Runtime/Engine/Classes/Engine/GameViewportClient.h
@@ -77,7 +77,10 @@ public:
 	UPROPERTY(Config)
 	int32 MaxSplitscreenPlayers = 4;
 
-	/** if true then the title safe border is drawn */
+	/** if true then the title safe border is drawn
+	  * @deprecated - Use the cvar "r.DebugSafeZone.Mode=1".
+	  */
+	UE_DEPRECATED(4.26, "Use the cvar \"r.DebugSafeZone.Mode=1\".")
 	uint32 bShowTitleSafeZone:1;
 
 	/** If true, this viewport is a play in editor viewport */
@@ -115,8 +118,10 @@ public:
 	UFUNCTION(exec)
 	virtual void SSSwapControllers();
 
-	/** Exec for toggling the display of the title safe area */
-	UFUNCTION(exec)
+	/** Exec for toggling the display of the title safe area
+	  * @deprecated Use the cvar "r.DebugSafeZone.Mode=1".
+	  */
+	UFUNCTION(exec, meta = (DeprecatedFunction, DeprecationMessage = "Use the cvar \"r.DebugSafeZone.Mode=1.\""))
 	virtual void ShowTitleSafeArea();
 
 	/** Sets the player which console commands will be executed in the context of. */
@@ -438,7 +443,7 @@ public:
 	bool CalculateDeadZoneForAllSides( ULocalPlayer* LPlayer, UCanvas* Canvas, float& fTopSafeZone, float& fBottomSafeZone, float& fLeftSafeZone, float& fRightSafeZone, bool bUseMaxPercent = false );
 
 	/**  
-	 * Draw the safe area using the current TitleSafeZone settings. 
+	 * Draws the safe area using the current r.DebugSafeZone.Mode=1 when there is not a valid PlayerController HUD.
 	 * 
 	 * @param Canvas	Canvas on which to draw
 	 */
diff --git a/Engine/Source/Runtime/Engine/Private/GameViewportClient.cpp b/Engine/Source/Runtime/Engine/Private/GameViewportClient.cpp
--- a/Engine/Source/Runtime/Engine/Private/GameViewportClient.cpp
+++ b/Engine/Source/Runtime/Engine/Private/GameViewportClient.cpp
@@ -208,9 +208,6 @@ void UGameViewportClient::UpdateCsvCameraStats(const TMap<ULocalPlayer*, FSceneV
 
 UGameViewportClient::UGameViewportClient(const FObjectInitializer& ObjectInitializer)
 	: Super(ObjectInitializer)
-#if WITH_EDITOR
-	, bShowTitleSafeZone(true)
-#endif
 	, EngineShowFlags(ESFIM_Game)
 	, CurrentBufferVisualizationMode(NAME_None)
 	, HighResScreenshotDialog(NULL)
@@ -295,9 +292,6 @@ UGameViewportClient::UGameViewportClient(const FObjectInitializer& ObjectInitial
 
 UGameViewportClient::UGameViewportClient(FVTableHelper& Helper)
 	: Super(Helper)
-#if WITH_EDITOR
-	, bShowTitleSafeZone(true)
-#endif
 	, EngineShowFlags(ESFIM_Game)
 	, CurrentBufferVisualizationMode(NAME_None)
 	, HighResScreenshotDialog(NULL)
@@ -2033,12 +2027,8 @@ bool UGameViewportClient::IsOrtho() const
 void UGameViewportClient::PostRender(UCanvas* Canvas)
 {
 #if WITH_EDITOR
-	if(bShowTitleSafeZone)
-	{
-		DrawTitleSafeArea(Canvas);
-	}
+	DrawTitleSafeArea(Canvas);
 #endif
-
 	// Draw the transition screen.
 	DrawTransition(Canvas);
 }
@@ -2080,9 +2070,19 @@ void UGameViewportClient::SSSwapControllers()
 
 void UGameViewportClient::ShowTitleSafeArea()
 {
-#if !UE_BUILD_SHIPPING
-	bShowTitleSafeZone = !bShowTitleSafeZone;
-#endif
+	static IConsoleVariable* DebugSafeZoneModeCvar = IConsoleManager::Get().FindConsoleVariable(TEXT("r.DebugSafeZone.Mode"));
+	if (DebugSafeZoneModeCvar)
+	{
+		const int32 DebugSafeZoneMode = DebugSafeZoneModeCvar->GetInt();
+		if (DebugSafeZoneMode != 1)
+		{
+			DebugSafeZoneModeCvar->Set(1);
+		}
+		else
+		{
+			DebugSafeZoneModeCvar->Set(0);
+		}
+	}
 }
 
 void UGameViewportClient::SetConsoleTarget(int32 PlayerIndex)
@@ -2507,19 +2507,36 @@ bool UGameViewportClient::CalculateDeadZoneForAllSides( ULocalPlayer* LPlayer, U
 void UGameViewportClient::DrawTitleSafeArea( UCanvas* Canvas )
 {
 #if WITH_EDITOR
+	// If we have a valid player hud, then the title area has already rendered.
+	APlayerController* FirstPlayerController = GetWorld()->GetFirstPlayerController();
+	if (FirstPlayerController && FirstPlayerController->GetHUD())
+	{
+		return;
+	}
+
+	// If r.DebugSafeZone.Mode isn't set to draw title area, don't draw it.
+	static IConsoleVariable* SafeZoneModeCvar = IConsoleManager::Get().FindConsoleVariable(TEXT("r.DebugSafeZone.Mode"));
+	if (SafeZoneModeCvar && (SafeZoneModeCvar->GetInt() != 1))
+	{
+		return;
+	}
+
 	FMargin SafeZone;
 	const ULevelEditorPlaySettings* PlayInSettings = GetDefault<ULevelEditorPlaySettings>();
 
 	float Width, Height;
 	GetPixelSizeOfScreen(Width, Height, Canvas, 0);
 
-	const FLinearColor UnsafeZoneColor(1.0f, 0.0f, 0.0f, 0.25f);
+	FLinearColor UnsafeZoneColor(1.0f, 0.0f, 0.0f, 0.25f);
+	static IConsoleVariable* AlphaCvar = IConsoleManager::Get().FindConsoleVariable(TEXT("r.DebugSafeZone.OverlayAlpha"));
+	if (AlphaCvar)
+	{
+		UnsafeZoneColor.A = AlphaCvar->GetFloat();
+	}
+
 	FCanvasTileItem TileItem(FVector2D::ZeroVector, GWhiteTexture, UnsafeZoneColor);
 	TileItem.BlendMode = SE_BLEND_Translucent;
 
-	// Command line override used by mobile PIE.
-	static bool bDrawUnSafeZones = FParse::Param(FCommandLine::Get(), TEXT("DrawUnSafeZones"));
-
 	// CalculateSafeZoneValues() can be slow, so we only want to run it if we have boundaries to draw
 	if (FDisplayMetrics::GetDebugTitleSafeZoneRatio() < 1.f)
 	{
@@ -2545,7 +2562,7 @@ void UGameViewportClient::DrawTitleSafeArea( UCanvas* Canvas )
 		TileItem.Size = FVector2D(SafeZone.Right, HeightOfSides);
 		Canvas->DrawItem(TileItem);
 	}
-	else if (!FSlateApplication::Get().GetCustomSafeZone().GetDesiredSize().IsZero() || bDrawUnSafeZones)
+	else if (!FSlateApplication::Get().GetCustomSafeZone().GetDesiredSize().IsZero())
 	{
 		ULevelEditorPlaySettings* PlaySettings = GetMutableDefault<ULevelEditorPlaySettings>();
 		PlaySettings->CalculateCustomUnsafeZones(PlaySettings->CustomUnsafeZoneStarts, PlaySettings->CustomUnsafeZoneDimensions, PlaySettings->DeviceToEmulate, FVector2D(Width, Height));
diff --git a/Engine/Source/Runtime/Engine/Private/Slate/SceneViewport.cpp b/Engine/Source/Runtime/Engine/Private/Slate/SceneViewport.cpp
--- a/Engine/Source/Runtime/Engine/Private/Slate/SceneViewport.cpp
+++ b/Engine/Source/Runtime/Engine/Private/Slate/SceneViewport.cpp
@@ -1465,19 +1465,6 @@ void FSceneViewport::ResizeViewport(uint32 NewSizeX, uint32 NewSizeY, EWindowMod
 		bIsResizing = true;
 
 		UpdateViewportRHI(false, NewSizeX, NewSizeY, NewWindowMode, PF_Unknown);
-
-#if WITH_EDITOR
-		FMargin SafeMargin;
-		if (FDisplayMetrics::GetDebugTitleSafeZoneRatio() < 1.f || !FSlateApplication::Get().GetCustomSafeZone().GetDesiredSize().IsZero())
-		{
-			FSlateApplication::Get().GetSafeZoneSize(SafeMargin, FVector2D(NewSizeX, NewSizeY));
-			SafeMargin.Left /= (NewSizeX / 2.0f);
-			SafeMargin.Right /= (NewSizeX / 2.0f);
-			SafeMargin.Bottom /= (NewSizeY / 2.0f);
-			SafeMargin.Top /= (NewSizeY / 2.0f);
-		}
-		FSlateApplication::Get().OnDebugSafeZoneChanged.Broadcast(SafeMargin, false);
-#endif
 		FCoreDelegates::OnSafeFrameChangedEvent.Broadcast();
 
 		if (ViewportClient)
diff --git a/Engine/Source/Runtime/SlateCore/Public/Application/SlateApplicationBase.h b/Engine/Source/Runtime/SlateCore/Public/Application/SlateApplicationBase.h
--- a/Engine/Source/Runtime/SlateCore/Public/Application/SlateApplicationBase.h
+++ b/Engine/Source/Runtime/SlateCore/Public/Application/SlateApplicationBase.h
@@ -587,8 +587,12 @@ protected:
 	{
 		FDisplayMetrics DisplayMetrics;
 		GetDisplayMetrics(DisplayMetrics);
-		CustomSafeZoneRatio = FMargin();
-		OnDebugSafeZoneChanged.Broadcast(FMargin(), false);
+
+		if (FDisplayMetrics::GetDebugTitleSafeZoneRatio() < 1.0f)
+		{
+			CustomSafeZoneRatio = FMargin();
+			OnDebugSafeZoneChanged.Broadcast(FMargin(), false);
+		}
 	}
 #endif
 
