Added support for C++11 raw strings to the localization source code gatherer

It will now detect and parse over these, including when they span multiple lines.

Also fixed a case where a #define containing a string literal spanning multiple lines was confusing the parser, and excluded ThirdParty code from any plugins that may be included in the gather.

diff --git a/Engine/Config/Localization/Editor.ini b/Engine/Config/Localization/Editor.ini
--- a/Engine/Config/Localization/Editor.ini
+++ b/Engine/Config/Localization/Editor.ini
@@ -23,6 +23,7 @@ ExcludePathFilters=Plugins/Online/*
 ExcludePathFilters=Plugins/Runtime/HTTPChunkInstaller/*
 ExcludePathFilters=*/NoRedist/*
 ExcludePathFilters=*/NotForLicensees/*
+ExcludePathFilters=*/ThirdParty/*
 FileNameFilters=*.cpp
 FileNameFilters=*.h
 FileNameFilters=*.c
diff --git a/Engine/Config/Localization/Engine.ini b/Engine/Config/Localization/Engine.ini
--- a/Engine/Config/Localization/Engine.ini
+++ b/Engine/Config/Localization/Engine.ini
@@ -28,6 +28,7 @@ ExcludePathFilters=Plugins/Online/*
 ExcludePathFilters=Plugins/Runtime/HTTPChunkInstaller/*
 ExcludePathFilters=*/NoRedist/*
 ExcludePathFilters=*/NotForLicensees/*
+ExcludePathFilters=*/ThirdParty/*
 FileNameFilters=*.cpp
 FileNameFilters=*.h
 FileNameFilters=*.c
diff --git a/Engine/Source/Editor/UnrealEd/Classes/Commandlets/GatherTextFromSourceCommandlet.h b/Engine/Source/Editor/UnrealEd/Classes/Commandlets/GatherTextFromSourceCommandlet.h
--- a/Engine/Source/Editor/UnrealEd/Classes/Commandlets/GatherTextFromSourceCommandlet.h
+++ b/Engine/Source/Editor/UnrealEd/Classes/Commandlets/GatherTextFromSourceCommandlet.h
@@ -115,6 +115,7 @@ private:
 		FName FilePlatformName;
 		FString LineText;
 		FString Namespace;
+		FString RawStringLiteralClosingDelim;
 		bool ExcludedRegion;
 		bool EndParsingCurrentLine;
 		bool WithinBlockComment;
diff --git a/Engine/Source/Editor/UnrealEd/Private/Commandlets/GatherTextFromSourceCommandlet.cpp b/Engine/Source/Editor/UnrealEd/Private/Commandlets/GatherTextFromSourceCommandlet.cpp
--- a/Engine/Source/Editor/UnrealEd/Private/Commandlets/GatherTextFromSourceCommandlet.cpp
+++ b/Engine/Source/Editor/UnrealEd/Private/Commandlets/GatherTextFromSourceCommandlet.cpp
@@ -288,8 +288,9 @@ int32 UGatherTextFromSourceCommandlet::Main( const FString& Params )
 		FPaths::MakePathRelativeTo(ParseCtxt.Filename, *ProjectBasePath);
 		ParseCtxt.LineNumber = 0;
 		ParseCtxt.FilePlatformName = GetSplitPlatformNameFromPath(ParseCtxt.Filename);
-		ParseCtxt.LineText.Empty();
-		ParseCtxt.Namespace.Empty();
+		ParseCtxt.LineText.Reset();
+		ParseCtxt.Namespace.Reset();
+		ParseCtxt.RawStringLiteralClosingDelim.Reset();
 		ParseCtxt.ExcludedRegion = false;
 		ParseCtxt.WithinBlockComment = false;
 		ParseCtxt.WithinLineComment = false;
@@ -761,7 +762,7 @@ bool UGatherTextFromSourceCommandlet::ParseSourceText(const FString& Text, const
 				if (LineCommentLen > 0 && FCString::Strncmp(Cursor, LineComment, LineCommentLen) == 0)
 				{
 					ParseCtxt.WithinLineComment = true;
-					ParseCtxt.WithinStartingLine = *ParseCtxt.LineText;
+					ParseCtxt.WithinStartingLine = *Line;
 					ParseCtxt.EndParsingCurrentLine = true;
 					Cursor += LineCommentLen;
 					continue;
@@ -769,7 +770,7 @@ bool UGatherTextFromSourceCommandlet::ParseSourceText(const FString& Text, const
 				else if (BlockCommentStartLen > 0 && FCString::Strncmp(Cursor, BlockCommentStart, BlockCommentStartLen) == 0)
 				{
 					ParseCtxt.WithinBlockComment = true;
-					ParseCtxt.WithinStartingLine = *ParseCtxt.LineText;
+					ParseCtxt.WithinStartingLine = *Line;
 					Cursor += BlockCommentStartLen;
 					continue;
 				}
@@ -782,17 +783,58 @@ bool UGatherTextFromSourceCommandlet::ParseSourceText(const FString& Text, const
 					if (Cursor == *Line)
 					{
 						ParseCtxt.WithinStringLiteral = true;
-						ParseCtxt.WithinStartingLine = *ParseCtxt.LineText;
+						ParseCtxt.WithinStartingLine = *Line;
 						++Cursor;
 						continue;
 					}
 					else if (Cursor > *Line)
 					{
 						const TCHAR* const ReverseCursor = Cursor - 1;
+						if (EnumHasAnyFlags(ParseCtxt.FileTypes, EGatherTextSourceFileTypes::Cpp) && *ReverseCursor == TEXT('R'))
+						{
+							// Potentially a C++11 raw string literal, so walk forwards and validate that this looks legit
+							// While doing this we can parse out its optional user defined delimiter so we can find when the string closes
+							//   eg) For 'R"Delim(string)Delim"', ')Delim' would be the closing delimiter.
+							//   eg) For 'R"(string)"', ')' would be the closing delimiter.
+							ParseCtxt.RawStringLiteralClosingDelim = TEXT(")");
+							{
+								bool bIsValid = true;
+
+								const TCHAR* ForwardCursor = Cursor + 1;
+								for (;;)
+								{
+									const TCHAR DelimChar = *ForwardCursor++;
+									if (DelimChar == TEXT('('))
+									{
+										break;
+									}
+									if (DelimChar == 0 || !FChar::IsAlnum(DelimChar))
+									{
+										bIsValid = false;
+										break;
+									}
+									ParseCtxt.RawStringLiteralClosingDelim += DelimChar;
+								}
+
+								if (bIsValid)
+								{
+									ParseCtxt.WithinStringLiteral = true;
+									ParseCtxt.WithinStartingLine = *Line;
+									Cursor = ForwardCursor;
+									continue;
+								}
+								else
+								{
+									ParseCtxt.RawStringLiteralClosingDelim.Reset();
+									// Fall through to the quoted string parsing below
+								}
+							}
+						}
+						
 						if (*ReverseCursor != TEXT('\\') && *ReverseCursor != TEXT('\''))
 						{
 							ParseCtxt.WithinStringLiteral = true;
-							ParseCtxt.WithinStartingLine = *ParseCtxt.LineText;
+							ParseCtxt.WithinStartingLine = *Line;
 							++Cursor;
 							continue;
 						}
@@ -811,7 +853,7 @@ bool UGatherTextFromSourceCommandlet::ParseSourceText(const FString& Text, const
 							if (IsEscaped)
 							{
 								ParseCtxt.WithinStringLiteral = true;
-								ParseCtxt.WithinStartingLine = *ParseCtxt.LineText;
+								ParseCtxt.WithinStartingLine = *Line;
 								++Cursor;
 								continue;
 							}
@@ -822,7 +864,7 @@ bool UGatherTextFromSourceCommandlet::ParseSourceText(const FString& Text, const
 								if (*ReverseCursor == TEXT('\'') && *ForwardCursor != TEXT('\''))
 								{
 									ParseCtxt.WithinStringLiteral = true;
-									ParseCtxt.WithinStartingLine = *ParseCtxt.LineText;
+									ParseCtxt.WithinStartingLine = *Line;
 									++Cursor;
 									continue;
 								}
@@ -835,7 +877,7 @@ bool UGatherTextFromSourceCommandlet::ParseSourceText(const FString& Text, const
 			{
 				if (*Cursor == TEXT('\"'))
 				{
-					if (Cursor == *Line)
+					if (Cursor == *Line && ParseCtxt.RawStringLiteralClosingDelim.IsEmpty())
 					{
 						ParseCtxt.WithinStringLiteral = false;
 						++Cursor;
@@ -843,6 +885,19 @@ bool UGatherTextFromSourceCommandlet::ParseSourceText(const FString& Text, const
 					}
 					else if (Cursor > *Line)
 					{
+						// Is this ending a C++11 raw string literal?
+						if (!ParseCtxt.RawStringLiteralClosingDelim.IsEmpty())
+						{
+							const TCHAR* EndDelimCursor = Cursor - ParseCtxt.RawStringLiteralClosingDelim.Len();
+							if (EndDelimCursor >= *Line && FCString::Strncmp(EndDelimCursor, *ParseCtxt.RawStringLiteralClosingDelim, ParseCtxt.RawStringLiteralClosingDelim.Len()) == 0)
+							{
+								ParseCtxt.RawStringLiteralClosingDelim.Reset();
+								ParseCtxt.WithinStringLiteral = false;
+							}
+							++Cursor;
+							continue;
+						}
+
 						const TCHAR* const ReverseCursor = Cursor - 1;
 						if (*ReverseCursor != TEXT('\\') && *ReverseCursor != TEXT('\''))
 						{
@@ -969,16 +1024,22 @@ bool UGatherTextFromSourceCommandlet::ParseSourceText(const FString& Text, const
 				// INI files don't support multi-line literals; always terminate them after ending a line
 				ParseCtxt.WithinStringLiteral = false;
 			}
-			else if (Cursor > *Line)
+			else if (Cursor > *Line && ParseCtxt.RawStringLiteralClosingDelim.IsEmpty())
 			{
-				// C++ only allows multi-line literals if they're escaped with a trailing slash
+				// C++ only allows multi-line literals if they're escaped with a trailing slash or within a C++11 raw string literal
 				ParseCtxt.WithinStringLiteral = *(Cursor - 1) == TEXT('\\');
 			}
 
-			UE_CLOG(!ParseCtxt.WithinStringLiteral, LogGatherTextFromSourceCommandlet, Warning, TEXT("A string literal was not correctly terminated. File %s at line %d"), *ParseCtxt.Filename, ParseCtxt.LineNumber);
+			UE_CLOG(!ParseCtxt.WithinStringLiteral, LogGatherTextFromSourceCommandlet, Warning, TEXT("A string literal was not correctly terminated. File %s at line %d, starting line: %s"), *ParseCtxt.Filename, ParseCtxt.LineNumber, ParseCtxt.WithinStartingLine);
 		}
 	}
 
+	// Handle a string C++11 raw string literal that was never closed as this is likely a false positive that needs to be fixed in the parser
+	if (ParseCtxt.WithinStringLiteral && !ParseCtxt.RawStringLiteralClosingDelim.IsEmpty())
+	{
+		UE_LOG(LogGatherTextFromSourceCommandlet, Warning, TEXT("A C++11 raw string literal was not correctly terminated. File %s, starting line: %s"), *ParseCtxt.Filename, ParseCtxt.WithinStartingLine);
+	}
+
 	return true;
 }
 
diff --git a/Engine/Source/Runtime/Engine/Private/OnlineReplStructs.cpp b/Engine/Source/Runtime/Engine/Private/OnlineReplStructs.cpp
--- a/Engine/Source/Runtime/Engine/Private/OnlineReplStructs.cpp
+++ b/Engine/Source/Runtime/Engine/Private/OnlineReplStructs.cpp
@@ -522,7 +522,8 @@ void TestUniqueIdRepl(UWorld* InWorld)
 	FUniqueNetIdRepl UpperCaseStringIdIn(UOnlineEngineInterface::Get()->CreateUniquePlayerId(TEXT("abcDEF")));
 
 #if 1
-#define WAYTOOLONG TEXT("deadbeefba5eba11deadbeefba5eba11 \
+#define WAYTOOLONG TEXT(\
+		"deadbeefba5eba11deadbeefba5eba11 \
 		deadbeefba5eba11deadbeefba5eba11deadbeefba5eba11deadbeefba5eba11deadbeefba5eba11deadbeefba5eba11 \
 		deadbeefba5eba11deadbeefba5eba11deadbeefba5eba11deadbeefba5eba11deadbeefba5eba11deadbeefba5eba11 \
 		deadbeefba5eba11deadbeefba5eba11deadbeefba5eba11deadbeefba5eba11deadbeefba5eba11deadbeefba5eba11 \
