clamp the camera view to slightly less than the full -90 to 90 degree range, to avoid a fully ambiguous conversion when converting to/from orbit camera

#rnx

diff --git a/Engine/Source/Editor/UnrealEd/Private/EditorViewportClient.cpp b/Engine/Source/Editor/UnrealEd/Private/EditorViewportClient.cpp
--- a/Engine/Source/Editor/UnrealEd/Private/EditorViewportClient.cpp
+++ b/Engine/Source/Editor/UnrealEd/Private/EditorViewportClient.cpp
@@ -5032,8 +5032,8 @@ void FEditorViewportClient::MoveViewportPerspectiveCamera( const FVector& InDrag
 
 		// normalize to -180 to 180
 		ViewRotation.Pitch = FRotator::NormalizeAxis(ViewRotation.Pitch);
-		// Make sure its withing  +/- 90 degrees.
-		ViewRotation.Pitch = FMath::Clamp( ViewRotation.Pitch, -90.f, 90.f );
+		// Make sure its withing  +/- 90 degrees (minus a small tolerance to avoid numerical issues w/ camera orientation conversions later on).
+		ViewRotation.Pitch = FMath::Clamp( ViewRotation.Pitch, -90.f+KINDA_SMALL_NUMBER, 90.f-KINDA_SMALL_NUMBER );
 	}
 	else
 	{
