Replace references now considers subclasses of DataTable to be compatible including CompositeDataTable

diff --git a/Engine/Source/Editor/UnrealEd/Private/ConsolidateWindow.cpp b/Engine/Source/Editor/UnrealEd/Private/ConsolidateWindow.cpp
--- a/Engine/Source/Editor/UnrealEd/Private/ConsolidateWindow.cpp
+++ b/Engine/Source/Editor/UnrealEd/Private/ConsolidateWindow.cpp
@@ -28,6 +28,7 @@
 #include "AssetSelection.h"
 #include "ObjectTools.h"
 #include "Interfaces/IMainFrameModule.h"
+#include "Engine/DataTable.h"
 
 
 #define LOCTEXT_NAMESPACE "SConsolidateWindow"
@@ -611,7 +612,9 @@ bool SConsolidateToolWidget::DetermineAssetCompatibility( const TArray<UObject*>
 				const UClass* NearestCommonBase = CurProposedObj->FindNearestCommonBaseClass( ComparisonClass );
 
 				// If the proposed object doesn't share a common class or a common base that is allowed as an exception, it is not a compatible object
-				if ( !( NearestCommonBase->IsChildOf( UTexture::StaticClass() ) )  && ! ( NearestCommonBase->IsChildOf( UMaterialInterface::StaticClass() ) ) )
+				if ( !( NearestCommonBase->IsChildOf( UTexture::StaticClass() ) )  &&
+					 !( NearestCommonBase->IsChildOf( UMaterialInterface::StaticClass() ) ) &&
+					 !( NearestCommonBase->IsChildOf( UDataTable::StaticClass() ) ) )
 				{
 					bAllAssetsValid = false;
 					continue;
