Update type of expected delegate when using "Assign selected Actor" context menu option on an event node in a level blueprint. Also hardened crash site.

diff --git a/Engine/Source/Editor/Kismet/Private/BlueprintEditor.cpp b/Engine/Source/Editor/Kismet/Private/BlueprintEditor.cpp
--- a/Engine/Source/Editor/Kismet/Private/BlueprintEditor.cpp
+++ b/Engine/Source/Editor/Kismet/Private/BlueprintEditor.cpp
@@ -6669,6 +6669,10 @@ void FBlueprintEditor::OnAssignReferencedActor()
 						// Store the node's current state and replace the referenced actor
 						CurrentEvent->Modify();
 						CurrentEvent->EventOwner = SelectedActor;
+						if (!SelectedActor->IsA(CurrentEvent->DelegateOwnerClass))
+						{
+							CurrentEvent->DelegateOwnerClass = SelectedActor->GetClass();
+						}
 						CurrentEvent->ReconstructNode();
 					}
 					FBlueprintEditorUtils::MarkBlueprintAsModified(GetBlueprintObj());
diff --git a/Engine/Source/Editor/UnrealEd/Private/Kismet2/BlueprintEditorUtils.cpp b/Engine/Source/Editor/UnrealEd/Private/Kismet2/BlueprintEditorUtils.cpp
--- a/Engine/Source/Editor/UnrealEd/Private/Kismet2/BlueprintEditorUtils.cpp
+++ b/Engine/Source/Editor/UnrealEd/Private/Kismet2/BlueprintEditorUtils.cpp
@@ -7707,7 +7707,9 @@ void FBlueprintEditorUtils::FixLevelScriptActorBindings(ALevelScriptActor* Level
 				{
 					// Grab the MC delegate we need to add to
 					FMulticastDelegateProperty* TargetDelegate = EventNode->GetTargetDelegateProperty();
-					if( TargetDelegate != nullptr )
+					if( TargetDelegate != nullptr && 
+						TargetDelegate->GetOwnerClass() &&
+						TargetDelegate->GetOwnerClass()->IsChildOf(EventNode->EventOwner->GetClass()))
 					{
 						// Create the delegate, and add it if it doesn't already exist
 						FScriptDelegate Delegate;
