Fix incorrect 'child of' test
#rnx

diff --git a/Engine/Source/Editor/UnrealEd/Private/Kismet2/BlueprintEditorUtils.cpp b/Engine/Source/Editor/UnrealEd/Private/Kismet2/BlueprintEditorUtils.cpp
--- a/Engine/Source/Editor/UnrealEd/Private/Kismet2/BlueprintEditorUtils.cpp
+++ b/Engine/Source/Editor/UnrealEd/Private/Kismet2/BlueprintEditorUtils.cpp
@@ -7709,7 +7709,7 @@ void FBlueprintEditorUtils::FixLevelScriptActorBindings(ALevelScriptActor* Level
 					FMulticastDelegateProperty* TargetDelegate = EventNode->GetTargetDelegateProperty();
 					if( TargetDelegate != nullptr && 
 						TargetDelegate->GetOwnerClass() &&
-						TargetDelegate->GetOwnerClass()->IsChildOf(EventNode->EventOwner->GetClass()))
+						EventNode->EventOwner->GetClass()->IsChildOf(TargetDelegate->GetOwnerClass()))
 					{
 						// Create the delegate, and add it if it doesn't already exist
 						FScriptDelegate Delegate;
