- Improve the error message for when macro graphs fail to compiler with the resoning in the compiler logs.
- Only check non-orphaned pins to remove some duplicate errors

#rnx

diff --git a/Engine/Source/Editor/UnrealEd/Private/Kismet2/BlueprintEditorUtils.cpp b/Engine/Source/Editor/UnrealEd/Private/Kismet2/BlueprintEditorUtils.cpp
--- a/Engine/Source/Editor/UnrealEd/Private/Kismet2/BlueprintEditorUtils.cpp
+++ b/Engine/Source/Editor/UnrealEd/Private/Kismet2/BlueprintEditorUtils.cpp
@@ -4363,14 +4363,15 @@ void FBlueprintEditorUtils::ValidatePinConnections(const UEdGraphNode* Node, FCo
 		// This is necessary because the user could change the type of the pin 
 		for (UEdGraphPin* Pin : Node->Pins)
 		{
-			if (Pin)
+			if (Pin && !Pin->bOrphanedPin)
 			{
 				for (UEdGraphPin* Link : Pin->LinkedTo)
 				{
-					if (Link && Link != Pin && Schema->CanCreateConnection(Pin, Link).Response == CONNECT_RESPONSE_DISALLOW)
+					const FPinConnectionResponse ConnectionResponse = Schema->CanCreateConnection(Pin, Link);
+					if (Link && !Link->bOrphanedPin && Link != Pin && ConnectionResponse.Response == CONNECT_RESPONSE_DISALLOW)
 					{
-						FText const ErrorFormat = LOCTEXT("BadConnection", "Invalid pin connection in graph '@@' from '@@' to '@@'. You may have changed the type after the connections were made.");
-						MessageLog.Error(*ErrorFormat.ToString(), Node->GetGraph(), Pin, Link);
+						const FString ErrorMessage = FText::Format(LOCTEXT("BadConnection_ErrorFmt", "Invalid pin connection from @@ and @@: {0} (You may have changed the type after the connections were made)"), ConnectionResponse.Message).ToString();
+						MessageLog.Error(*ErrorMessage, Pin, Link);
 					}
 				}
 			}
