Adding ResetToFbx option for automated import

diff --git a/Engine/Plugins/Runtime/ApexDestruction/Source/ApexDestruction/Private/ApexDestructibleAssetImport.cpp b/Engine/Plugins/Runtime/ApexDestruction/Source/ApexDestruction/Private/ApexDestructibleAssetImport.cpp
--- a/Engine/Plugins/Runtime/ApexDestruction/Source/ApexDestruction/Private/ApexDestructibleAssetImport.cpp
+++ b/Engine/Plugins/Runtime/ApexDestruction/Source/ApexDestruction/Private/ApexDestructibleAssetImport.cpp
@@ -256,7 +256,7 @@ static void RestoreExistingDestMeshData(ExistingDestMeshData* MeshData, UDestruc
 		// Restore old settings, but resize arrays to make sense with the new NxDestructibleAsset
 		if (MeshData->SkelMeshData != NULL)
 		{
-			RestoreExistingSkelMeshData(MeshData->SkelMeshData, DestructibleMesh, INDEX_NONE, false, false);
+			RestoreExistingSkelMeshData(MeshData->SkelMeshData, DestructibleMesh, INDEX_NONE, false, false, false);
 		}
 		DestructibleMesh->BodySetup =  MeshData->BodySetup;
 		DestructibleMesh->FractureEffects = MeshData->FractureEffects;
diff --git a/Engine/Source/Editor/UnrealEd/Classes/Factories/FbxImportUI.h b/Engine/Source/Editor/UnrealEd/Classes/Factories/FbxImportUI.h
--- a/Engine/Source/Editor/UnrealEd/Classes/Factories/FbxImportUI.h
+++ b/Engine/Source/Editor/UnrealEd/Classes/Factories/FbxImportUI.h
@@ -193,6 +193,10 @@ public:
 	UPROPERTY(EditAnywhere, BlueprintReadWrite, config, Category=Material, meta = (ImportType = "GeoOnly", EditCondition="!bImportMaterials"))
 	uint32 bImportTextures:1;
 
+	/** If true, the imported material sections will automatically be reset to the imported data in case of a reimport conflict. */
+	UPROPERTY(BlueprintReadWrite, Transient, Category = Material)
+	uint32 bResetToFbxOnMaterialConflict:1;
+
 	/** Import data used when importing static meshes */
 	UPROPERTY(EditAnywhere, BlueprintReadWrite, Transient, Instanced, Category = Mesh, meta=(ImportType = "StaticMesh"))
 	class UFbxStaticMeshImportData* StaticMeshImportData;
diff --git a/Engine/Source/Editor/UnrealEd/Classes/Factories/FbxSkeletalMeshImportData.h b/Engine/Source/Editor/UnrealEd/Classes/Factories/FbxSkeletalMeshImportData.h
--- a/Engine/Source/Editor/UnrealEd/Classes/Factories/FbxSkeletalMeshImportData.h
+++ b/Engine/Source/Editor/UnrealEd/Classes/Factories/FbxSkeletalMeshImportData.h
@@ -110,7 +110,7 @@ namespace SkeletalMeshImportData
 }
 
 extern UNREALED_API ExistingSkelMeshData* SaveExistingSkelMeshData(USkeletalMesh* ExistingSkelMesh, bool bSaveMaterials, int32 ReimportLODIndex);
-extern UNREALED_API void RestoreExistingSkelMeshData(ExistingSkelMeshData* MeshData, USkeletalMesh* SkeletalMesh, int32 ReimportLODIndex, bool bCanShowDialog, bool bImportSkinningOnly);
+extern UNREALED_API void RestoreExistingSkelMeshData(ExistingSkelMeshData* MeshData, USkeletalMesh* SkeletalMesh, int32 ReimportLODIndex, bool bCanShowDialog, bool bImportSkinningOnly, bool bForceMaterialReset);
 extern UNREALED_API void ProcessImportMeshInfluences(FSkeletalMeshImportData& ImportData);
 extern UNREALED_API void ProcessImportMeshMaterials(TArray<FSkeletalMaterial>& Materials, FSkeletalMeshImportData& ImportData);
 extern UNREALED_API bool ProcessImportMeshSkeleton(const USkeleton* SkeletonAsset, FReferenceSkeleton& RefSkeleton, int32& SkeletalDepth, FSkeletalMeshImportData& ImportData);namespace SkeletalMeshHelper
diff --git a/Engine/Source/Editor/UnrealEd/Private/Factories/SkeletalMeshImport.cpp b/Engine/Source/Editor/UnrealEd/Private/Factories/SkeletalMeshImport.cpp
--- a/Engine/Source/Editor/UnrealEd/Private/Factories/SkeletalMeshImport.cpp
+++ b/Engine/Source/Editor/UnrealEd/Private/Factories/SkeletalMeshImport.cpp
@@ -634,7 +634,7 @@ namespace SkeletalMeshHelper
 	}
 } //namespace SkeletalMeshHelper
 
-void RestoreExistingSkelMeshData(ExistingSkelMeshData* MeshData, USkeletalMesh* SkeletalMesh, int32 ReimportLODIndex, bool bCanShowDialog, bool bImportSkinningOnly)
+void RestoreExistingSkelMeshData(ExistingSkelMeshData* MeshData, USkeletalMesh* SkeletalMesh, int32 ReimportLODIndex, bool bCanShowDialog, bool bImportSkinningOnly, bool bForceMaterialReset)
 {
 	if (!MeshData || !SkeletalMesh)
 	{
@@ -669,7 +669,7 @@ void RestoreExistingSkelMeshData(ExistingSkelMeshData* MeshData, USkeletalMesh*
 	{
 		UnFbx::EFBXReimportDialogReturnOption ReturnOption;
 		//Ask the user to match the materials conflict
-		UnFbx::FFbxImporter::PrepareAndShowMaterialConflictDialog<FSkeletalMaterial>(MeshData->ExistingMaterials, SkeletalMesh->Materials, RemapMaterial, RemapMaterialName, bCanShowDialog, false, ReturnOption);
+		UnFbx::FFbxImporter::PrepareAndShowMaterialConflictDialog<FSkeletalMaterial>(MeshData->ExistingMaterials, SkeletalMesh->Materials, RemapMaterial, RemapMaterialName, bCanShowDialog, false, bForceMaterialReset, ReturnOption);
 
 		if (ReturnOption != UnFbx::EFBXReimportDialogReturnOption::FBXRDRO_ResetToFbx)
 		{
diff --git a/Engine/Source/Editor/UnrealEd/Private/Fbx/FbxFactory.cpp b/Engine/Source/Editor/UnrealEd/Private/Fbx/FbxFactory.cpp
--- a/Engine/Source/Editor/UnrealEd/Private/Fbx/FbxFactory.cpp
+++ b/Engine/Source/Editor/UnrealEd/Private/Fbx/FbxFactory.cpp
@@ -985,6 +985,7 @@ UFbxImportUI::UFbxImportUI(const FObjectInitializer& ObjectInitializer)
 	ReimportMesh = nullptr;
 	bAllowContentTypeImport = false;
 	bAutomatedImportShouldDetectType = true;
+	bResetToFbxOnMaterialConflict = false;
 	//Make sure we are transactional to allow undo redo
 	this->SetFlags(RF_Transactional);
 	
diff --git a/Engine/Source/Editor/UnrealEd/Private/Fbx/FbxMainImport.cpp b/Engine/Source/Editor/UnrealEd/Private/Fbx/FbxMainImport.cpp
--- a/Engine/Source/Editor/UnrealEd/Private/Fbx/FbxMainImport.cpp
+++ b/Engine/Source/Editor/UnrealEd/Private/Fbx/FbxMainImport.cpp
@@ -71,7 +71,7 @@ void PrepareAndShowMaterialConflictPreviewDialog(UFbxImportUI* ImportUI)
 		ResultMaterial[Materialindex].ImportedMaterialSlotName = ImportUI->MaterialCompareData.ResultAsset[Materialindex].ImportedMaterialSlotName;
 	}
 	UnFbx::EFBXReimportDialogReturnOption OutReturnOption;
-	UnFbx::FFbxImporter::PrepareAndShowMaterialConflictDialog<TMaterialType>(CurrentMaterial, ResultMaterial, RemapMaterial, RemapMaterialName, true, true, OutReturnOption);
+	UnFbx::FFbxImporter::PrepareAndShowMaterialConflictDialog<TMaterialType>(CurrentMaterial, ResultMaterial, RemapMaterial, RemapMaterialName, true, true, false, OutReturnOption);
 }
 
 void PrepareAndShowSkeletonConflictPreviewDialog(UFbxImportUI* ImportUI)
@@ -358,6 +358,7 @@ void ApplyImportUIToImportOptions(UFbxImportUI* ImportUI, FBXImportOptions& InOu
 		InOutImportOptions.bUsedAsFullName			= ImportUI->bOverrideFullName;
 		InOutImportOptions.ImportType				= ImportUI->MeshTypeToImport;
 
+		InOutImportOptions.bResetToFbxOnMaterialConflict = ImportUI->bResetToFbxOnMaterialConflict;
 		InOutImportOptions.bAutoComputeLodDistances	= ImportUI->bAutoComputeLodDistances;
 		InOutImportOptions.LodDistances.Empty(8);
 		InOutImportOptions.LodDistances.Add(ImportUI->LodDistance0);
diff --git a/Engine/Source/Editor/UnrealEd/Private/Fbx/FbxReimportDialogs.cpp b/Engine/Source/Editor/UnrealEd/Private/Fbx/FbxReimportDialogs.cpp
--- a/Engine/Source/Editor/UnrealEd/Private/Fbx/FbxReimportDialogs.cpp
+++ b/Engine/Source/Editor/UnrealEd/Private/Fbx/FbxReimportDialogs.cpp
@@ -227,7 +227,7 @@ void ResetMaterialSlot(const TArray<TMaterialType>& CurrentMaterial, TArray<TMat
 }
 
 template<typename TMaterialType>
-void FFbxImporter::PrepareAndShowMaterialConflictDialog(const TArray<TMaterialType>& CurrentMaterial, TArray<TMaterialType>& ResultMaterial, TArray<int32>& RemapMaterial, TArray<FName>& RemapMaterialName, bool bCanShowDialog, bool bIsPreviewDialog, EFBXReimportDialogReturnOption& OutReturnOption)
+void FFbxImporter::PrepareAndShowMaterialConflictDialog(const TArray<TMaterialType>& CurrentMaterial, TArray<TMaterialType>& ResultMaterial, TArray<int32>& RemapMaterial, TArray<FName>& RemapMaterialName, bool bCanShowDialog, bool bIsPreviewDialog, bool bForceResetOnConflict, EFBXReimportDialogReturnOption& OutReturnOption)
 {
 	OutReturnOption = EFBXReimportDialogReturnOption::FBXRDRO_Ok;
 	bool bHasSomeUnmatchedMaterial = false;
@@ -287,20 +287,27 @@ void FFbxImporter::PrepareAndShowMaterialConflictDialog(const TArray<TMaterialTy
 				AutoRemapMaterials[BestMaterialIndex] = true;
 			}
 		}
-		if (bCanShowDialog)
+		
+		if (bForceResetOnConflict)
+		{
+			OutReturnOption = EFBXReimportDialogReturnOption::FBXRDRO_ResetToFbx;
+		}
+		else if (bCanShowDialog)
 		{
 			ShowFbxMaterialConflictWindow<TMaterialType>(CurrentMaterial, ResultMaterial, RemapMaterial, AutoRemapMaterials, OutReturnOption, bIsPreviewDialog);
-			if (OutReturnOption == EFBXReimportDialogReturnOption::FBXRDRO_ResetToFbx)
+		}
+
+		if (OutReturnOption == EFBXReimportDialogReturnOption::FBXRDRO_ResetToFbx)
+		{
+			//Make identity remap because we reset to ResultMaterial
+			for (int32 MaterialIndex = 0; MaterialIndex < ResultMaterial.Num(); ++MaterialIndex)
 			{
-				//Make identity remap because we reset to ResultMaterial
-				for (int32 MaterialIndex = 0; MaterialIndex < ResultMaterial.Num(); ++MaterialIndex)
-				{
-					RemapMaterial[MaterialIndex] = MaterialIndex;
-					RemapMaterialName[MaterialIndex] = ResultMaterial[MaterialIndex].ImportedMaterialSlotName;
-				}
-				ResetMaterialSlot(CurrentMaterial, ResultMaterial);
+				RemapMaterial[MaterialIndex] = MaterialIndex;
+				RemapMaterialName[MaterialIndex] = ResultMaterial[MaterialIndex].ImportedMaterialSlotName;
 			}
+			ResetMaterialSlot(CurrentMaterial, ResultMaterial);
 		}
+		
 	}
 }
 
@@ -362,7 +369,7 @@ void FFbxImporter::ShowFbxMaterialConflictWindow(const TArray<TMaterialType>& In
 template void FFbxImporter::ShowFbxMaterialConflictWindow<FStaticMaterial>(const TArray<FStaticMaterial>& InSourceMaterials, const TArray<FStaticMaterial>& InResultMaterials, TArray<int32>& RemapMaterials, TArray<bool>& AutoRemapMaterials, EFBXReimportDialogReturnOption& OutReturnOption, bool bIsPreviewConflict);
 template void FFbxImporter::ShowFbxMaterialConflictWindow<FSkeletalMaterial>(const TArray<FSkeletalMaterial>& InSourceMaterials, const TArray<FSkeletalMaterial>& InResultMaterials, TArray<int32>& RemapMaterials, TArray<bool>& AutoRemapMaterials, EFBXReimportDialogReturnOption& OutReturnOption, bool bIsPreviewConflict);
 
-template void FFbxImporter::PrepareAndShowMaterialConflictDialog<FStaticMaterial>(const TArray<FStaticMaterial>& CurrentMaterial, TArray<FStaticMaterial>& ResultMaterial, TArray<int32>& RemapMaterial, TArray<FName>& RemapMaterialName, bool bCanShowDialog, bool bIsPreviewDialog, EFBXReimportDialogReturnOption& OutReturnOption);
-template void FFbxImporter::PrepareAndShowMaterialConflictDialog<FSkeletalMaterial>(const TArray<FSkeletalMaterial>& CurrentMaterial, TArray<FSkeletalMaterial>& ResultMaterial, TArray<int32>& RemapMaterial, TArray<FName>& RemapMaterialName, bool bCanShowDialog, bool bIsPreviewDialog, EFBXReimportDialogReturnOption& OutReturnOption);
+template void FFbxImporter::PrepareAndShowMaterialConflictDialog<FStaticMaterial>(const TArray<FStaticMaterial>& CurrentMaterial, TArray<FStaticMaterial>& ResultMaterial, TArray<int32>& RemapMaterial, TArray<FName>& RemapMaterialName, bool bCanShowDialog, bool bIsPreviewDialog, bool bForceResetOnConflict, EFBXReimportDialogReturnOption& OutReturnOption);
+template void FFbxImporter::PrepareAndShowMaterialConflictDialog<FSkeletalMaterial>(const TArray<FSkeletalMaterial>& CurrentMaterial, TArray<FSkeletalMaterial>& ResultMaterial, TArray<int32>& RemapMaterial, TArray<FName>& RemapMaterialName, bool bCanShowDialog, bool bIsPreviewDialog, bool bForceResetOnConflict, EFBXReimportDialogReturnOption& OutReturnOption);
 
 #undef LOCTEXT_NAMESPACE
diff --git a/Engine/Source/Editor/UnrealEd/Private/Fbx/FbxSkeletalMeshImport.cpp b/Engine/Source/Editor/UnrealEd/Private/Fbx/FbxSkeletalMeshImport.cpp
--- a/Engine/Source/Editor/UnrealEd/Private/Fbx/FbxSkeletalMeshImport.cpp
+++ b/Engine/Source/Editor/UnrealEd/Private/Fbx/FbxSkeletalMeshImport.cpp
@@ -1938,7 +1938,7 @@ USkeletalMesh* UnFbx::FFbxImporter::ImportSkeletalMesh(FImportSkeletalMeshArgs &
 
 		if (ExistSkelMeshDataPtr)
 		{
-			RestoreExistingSkelMeshData(ExistSkelMeshDataPtr, SkeletalMesh, ImportSkeletalMeshArgs.LodIndex, ImportOptions->bCanShowDialog, ImportOptions->bImportAsSkeletalSkinning);
+			RestoreExistingSkelMeshData(ExistSkelMeshDataPtr, SkeletalMesh, ImportSkeletalMeshArgs.LodIndex, ImportOptions->bCanShowDialog, ImportOptions->bImportAsSkeletalSkinning, ImportOptions->bResetToFbxOnMaterialConflict);
 		}
 
 		SkeletalMesh->CalculateInvRefMatrices();
diff --git a/Engine/Source/Editor/UnrealEd/Private/Fbx/FbxStaticMeshImport.cpp b/Engine/Source/Editor/UnrealEd/Private/Fbx/FbxStaticMeshImport.cpp
--- a/Engine/Source/Editor/UnrealEd/Private/Fbx/FbxStaticMeshImport.cpp
+++ b/Engine/Source/Editor/UnrealEd/Private/Fbx/FbxStaticMeshImport.cpp
@@ -45,7 +45,7 @@ using namespace UnFbx;
 struct ExistingStaticMeshData;
 extern ExistingStaticMeshData* SaveExistingStaticMeshData(UStaticMesh* ExistingMesh, FBXImportOptions* ImportOptions, int32 LodIndex);
 extern void RestoreExistingMeshSettings(struct ExistingStaticMeshData* ExistingMesh, UStaticMesh* NewMesh, int32 LODIndex);
-extern void RestoreExistingMeshData(struct ExistingStaticMeshData* ExistingMeshDataPtr, UStaticMesh* NewMesh, int32 LodLevel, bool bCanShowDialog);
+extern void RestoreExistingMeshData(struct ExistingStaticMeshData* ExistingMeshDataPtr, UStaticMesh* NewMesh, int32 LodLevel, bool bCanShowDialog, bool bForceConflictingMaterialReset);
 extern void UpdateSomeLodsImportMeshData(UStaticMesh* NewMesh, TArray<int32> *ReimportLodList);
 
 static FbxString GetNodeNameWithoutNamespace( FbxNode* Node )
@@ -1073,7 +1073,7 @@ UStaticMesh* UnFbx::FFbxImporter::ReimportSceneStaticMesh(uint64 FbxNodeUniqueId
 		}
 	}
 	//Don't restore materials when reimporting scene
-	RestoreExistingMeshData(ExistMeshDataPtr, FirstBaseMesh, INDEX_NONE, false);
+	RestoreExistingMeshData(ExistMeshDataPtr, FirstBaseMesh, INDEX_NONE, false, ImportOptions->bResetToFbxOnMaterialConflict);
 	return FirstBaseMesh;
 }
 
@@ -1340,7 +1340,7 @@ UStaticMesh* UnFbx::FFbxImporter::ReimportStaticMesh(UStaticMesh* Mesh, UFbxStat
 	if (NewMesh != nullptr)
 	{
 		UpdateSomeLodsImportMeshData(NewMesh, &ReimportLodList);
-		RestoreExistingMeshData(ExistMeshDataPtr, NewMesh, INDEX_NONE, ImportOptions->bCanShowDialog);
+		RestoreExistingMeshData(ExistMeshDataPtr, NewMesh, INDEX_NONE, ImportOptions->bCanShowDialog, ImportOptions->bResetToFbxOnMaterialConflict);
 	}
 	return NewMesh;
 }
diff --git a/Engine/Source/Editor/UnrealEd/Private/FbxMeshUtils.cpp b/Engine/Source/Editor/UnrealEd/Private/FbxMeshUtils.cpp
--- a/Engine/Source/Editor/UnrealEd/Private/FbxMeshUtils.cpp
+++ b/Engine/Source/Editor/UnrealEd/Private/FbxMeshUtils.cpp
@@ -44,7 +44,7 @@ DEFINE_LOG_CATEGORY_STATIC(LogExportMeshUtils, Log, All);
 struct ExistingStaticMeshData;
 extern ExistingStaticMeshData* SaveExistingStaticMeshData(UStaticMesh* ExistingMesh, UnFbx::FBXImportOptions* ImportOptions, int32 LodIndex);
 extern void RestoreExistingMeshSettings(struct ExistingStaticMeshData* ExistingMesh, UStaticMesh* NewMesh, int32 LODIndex);
-extern void RestoreExistingMeshData(struct ExistingStaticMeshData* ExistingMeshDataPtr, UStaticMesh* NewMesh, int32 LodLevel, bool bCanShowDialog);
+extern void RestoreExistingMeshData(struct ExistingStaticMeshData* ExistingMeshDataPtr, UStaticMesh* NewMesh, int32 LodLevel, bool bCanShowDialog, bool bForceConflictingMaterialReset);
 extern void UpdateSomeLodsImportMeshData(UStaticMesh* NewMesh, TArray<int32> *ReimportLodList);
 
 
@@ -200,7 +200,7 @@ namespace FbxMeshUtils
 					UpdateSomeLodsImportMeshData(BaseStaticMesh, &ReimportLodList);
 					if(IsReimport)
 					{
-						RestoreExistingMeshData(ExistMeshDataPtr, BaseStaticMesh, LODLevel, false);
+						RestoreExistingMeshData(ExistMeshDataPtr, BaseStaticMesh, LODLevel, false, ImportOptions->bResetToFbxOnMaterialConflict);
 					}
 
 					// Update mesh component
@@ -524,7 +524,7 @@ namespace FbxMeshUtils
 
 					if (SkelMeshDataPtr != nullptr)
 					{
-						RestoreExistingSkelMeshData(SkelMeshDataPtr, SelectedSkelMesh, SelectedLOD, false, ImportOptions->bImportAsSkeletalSkinning);
+						RestoreExistingSkelMeshData(SkelMeshDataPtr, SelectedSkelMesh, SelectedLOD, false, ImportOptions->bImportAsSkeletalSkinning, ImportOptions->bResetToFbxOnMaterialConflict);
 					}
 
 					if (ImportOptions->bImportMorph)
diff --git a/Engine/Source/Editor/UnrealEd/Private/StaticMeshEdit.cpp b/Engine/Source/Editor/UnrealEd/Private/StaticMeshEdit.cpp
--- a/Engine/Source/Editor/UnrealEd/Private/StaticMeshEdit.cpp
+++ b/Engine/Source/Editor/UnrealEd/Private/StaticMeshEdit.cpp
@@ -1413,7 +1413,7 @@ void UpdateSomeLodsImportMeshData(UStaticMesh* NewMesh, TArray<int32> *ReimportL
 	}
 }
 
-void RestoreExistingMeshData(ExistingStaticMeshData* ExistingMeshDataPtr, UStaticMesh* NewMesh, int32 LodLevel, bool bCanShowDialog)
+void RestoreExistingMeshData(ExistingStaticMeshData* ExistingMeshDataPtr, UStaticMesh* NewMesh, int32 LodLevel, bool bCanShowDialog, bool bForceConflictingMaterialReset)
 {
 	if (!ExistingMeshDataPtr || !NewMesh)
 	{
@@ -1438,7 +1438,7 @@ void RestoreExistingMeshData(ExistingStaticMeshData* ExistingMeshDataPtr, UStati
 	//If user is attended, ask him to verify the match is good
 	UnFbx::EFBXReimportDialogReturnOption ReturnOption;
 	//Ask the user to match the materials conflict
-	UnFbx::FFbxImporter::PrepareAndShowMaterialConflictDialog<FStaticMaterial>(ExistingMeshDataPtr->ExistingMaterials, NewMesh->StaticMaterials, RemapMaterial, RemapMaterialName, bCanShowDialog, false, ReturnOption);
+	UnFbx::FFbxImporter::PrepareAndShowMaterialConflictDialog<FStaticMaterial>(ExistingMeshDataPtr->ExistingMaterials, NewMesh->StaticMaterials, RemapMaterial, RemapMaterialName, bCanShowDialog, false, bForceConflictingMaterialReset, ReturnOption);
 	
 	if (ReturnOption != UnFbx::EFBXReimportDialogReturnOption::FBXRDRO_ResetToFbx)
 	{
diff --git a/Engine/Source/Editor/UnrealEd/Public/FbxImporter.h b/Engine/Source/Editor/UnrealEd/Public/FbxImporter.h
--- a/Engine/Source/Editor/UnrealEd/Public/FbxImporter.h
+++ b/Engine/Source/Editor/UnrealEd/Public/FbxImporter.h
@@ -197,6 +197,7 @@ struct FBXImportOptions
 	bool	bSetMaterialDriveParameterOnCustomAttribute;
 	bool	bRemoveRedundantKeys;
 	bool	bDoNotImportCurveWithZero;
+	bool	bResetToFbxOnMaterialConflict;
 	TArray<FString> MaterialCurveSuffixes;
 
 	/** This allow to add a prefix to the material name when unreal material get created.	
@@ -1047,7 +1048,7 @@ public:
 	static void ShowFbxSkeletonConflictWindow(USkeletalMesh *SkeletalMesh, USkeleton* Skeleton, ImportCompareHelper::FSkeletonCompareData& SkeletonCompareData);
 
 	template<typename TMaterialType>
-	static void PrepareAndShowMaterialConflictDialog(const TArray<TMaterialType>& CurrentMaterial, TArray<TMaterialType>& ResultMaterial, TArray<int32>& RemapMaterial, TArray<FName>& RemapMaterialName, bool bCanShowDialog, bool bIsPreviewDialog, EFBXReimportDialogReturnOption& OutReturnOption);
+	static void PrepareAndShowMaterialConflictDialog(const TArray<TMaterialType>& CurrentMaterial, TArray<TMaterialType>& ResultMaterial, TArray<int32>& RemapMaterial, TArray<FName>& RemapMaterialName, bool bCanShowDialog, bool bIsPreviewDialog, bool bForceResetOnConflict, EFBXReimportDialogReturnOption& OutReturnOption);
 	/*
 	* This function show a dialog to let the user resolve the material conflict that arise when re-importing a mesh
 	*/
