Add a check when debugging in blueprint for output params when traversing the stack frame to avoid a crash

diff --git a/Engine/Source/Editor/UnrealEd/Private/Kismet2/KismetDebugUtilities.cpp b/Engine/Source/Editor/UnrealEd/Private/Kismet2/KismetDebugUtilities.cpp
--- a/Engine/Source/Editor/UnrealEd/Private/Kismet2/KismetDebugUtilities.cpp
+++ b/Engine/Source/Editor/UnrealEd/Private/Kismet2/KismetDebugUtilities.cpp
@@ -38,6 +38,7 @@
 #include "Logging/MessageLog.h"
 #include "Misc/UObjectToken.h"
 #include "AnimGraphNode_Base.h"
+#include "UObject/UnrealType.h"
 
 #define LOCTEXT_NAMESPACE "BlueprintDebugging"
 
@@ -1122,12 +1123,22 @@ FKismetDebugUtilities::EWatchTextResult FKismetDebugUtilities::GetWatchText(FStr
 	void* DataPtr = nullptr;
 	void* DeltaPtr = nullptr;
 	UObject* ParentObj = nullptr;
+	bool bIsInDebuggingFrame = false;
 	TArray<UObject*> SeenObjects;
-	FKismetDebugUtilities::EWatchTextResult Result = FindDebuggingData(Blueprint, ActiveObject, WatchPin, PropertyToDebug, DataPtr, DeltaPtr, ParentObj, SeenObjects);
+	FKismetDebugUtilities::EWatchTextResult Result = FindDebuggingData(Blueprint, ActiveObject, WatchPin, PropertyToDebug, DataPtr, DeltaPtr, ParentObj, SeenObjects, &bIsInDebuggingFrame);
 
 	if (Result == FKismetDebugUtilities::EWatchTextResult::EWTR_Valid)
 	{
-		PropertyToDebug->ExportText_InContainer(/*ArrayElement=*/ 0, /*inout*/ OutWatchText, DataPtr, DeltaPtr, /*Parent=*/ ParentObj, PPF_PropertyWindow | PPF_BlueprintDebugView);
+		// If this came from an array property we need to avoid using ExportText_InContainer in order to properly 
+		// calculate the internal offset
+		if(bIsInDebuggingFrame && PropertyToDebug->IsA<FArrayProperty>())
+		{
+			PropertyToDebug->ExportText_Direct(/*inout*/ OutWatchText, DataPtr, DeltaPtr, ParentObj, PPF_PropertyWindow | PPF_BlueprintDebugView);
+		}
+		else
+		{
+			PropertyToDebug->ExportText_InContainer(/*ArrayElement=*/ 0, /*inout*/ OutWatchText, DataPtr, DeltaPtr, /*Parent=*/ ParentObj, PPF_PropertyWindow | PPF_BlueprintDebugView);
+		}
 	}
 
 	return Result;
@@ -1150,7 +1161,7 @@ FKismetDebugUtilities::EWatchTextResult FKismetDebugUtilities::GetDebugInfo(FDeb
 	return Result;
 }
 
-FKismetDebugUtilities::EWatchTextResult FKismetDebugUtilities::FindDebuggingData(UBlueprint* Blueprint, UObject* ActiveObject, const UEdGraphPin* WatchPin, FProperty*& OutProperty, void*& OutData, void*& OutDelta, UObject*& OutParent, TArray<UObject*>& SeenObjects)
+FKismetDebugUtilities::EWatchTextResult FKismetDebugUtilities::FindDebuggingData(UBlueprint* Blueprint, UObject* ActiveObject, const UEdGraphPin* WatchPin, FProperty*& OutProperty, void*& OutData, void*& OutDelta, UObject*& OutParent, TArray<UObject*>& SeenObjects, bool* OutbIsInDebuggingFrame /* = nullptr */)
 {
 	FKismetDebugUtilitiesData& Data = FKismetDebugUtilitiesData::Get();
 
@@ -1194,17 +1205,33 @@ FKismetDebugUtilities::EWatchTextResult FKismetDebugUtilities::FindDebuggingData
 					{
 						if (OutParmRec->Property == Property)
 						{
-							// try to use the output pin we're linked to
-							// otherwise the output param won't show any data since the return node hasn't executed when we stop here
-							if (WatchPin->Direction == EEdGraphPinDirection::EGPD_Input && WatchPin->LinkedTo.Num() == 1)
+							if (WatchPin->Direction == EEdGraphPinDirection::EGPD_Input)
 							{
-								return FindDebuggingData(Blueprint, ActiveObject, WatchPin->LinkedTo[0], OutProperty, OutData, OutDelta, OutParent, SeenObjects);
+								// try to use the output pin we're linked to
+								// otherwise the output param won't show any data since the return node hasn't executed when we stop here
+								if (WatchPin->LinkedTo.Num() == 1)
+								{
+									return FindDebuggingData(Blueprint, ActiveObject, WatchPin->LinkedTo[0], OutProperty, OutData, OutDelta, OutParent, SeenObjects, OutbIsInDebuggingFrame);
+								}
+								else if (!WatchPin->LinkedTo.Num())
+								{
+									// If this is an output pin with no links, then we have no debug data
+									// so fallback to the local stack frame
+									PropertyBase = TestFrame->Locals;
+								}
 							}
 
-							PropertyBase = OutParmRec->PropAddr;
+							// If this is an out property then we need to let the caller know
+							if (PropertyBase == nullptr && OutbIsInDebuggingFrame)
+							{
+								*OutbIsInDebuggingFrame = true;
+								PropertyBase = OutParmRec->PropAddr;
+							}
 							break;
 						}
 					}
+
+					// Fallback to the local variables if we couldn't find one
 					if (PropertyBase == nullptr)
 					{
 						PropertyBase = TestFrame->Locals;
diff --git a/Engine/Source/Editor/UnrealEd/Public/Kismet2/KismetDebugUtilities.h b/Engine/Source/Editor/UnrealEd/Public/Kismet2/KismetDebugUtilities.h
--- a/Engine/Source/Editor/UnrealEd/Public/Kismet2/KismetDebugUtilities.h
+++ b/Engine/Source/Editor/UnrealEd/Public/Kismet2/KismetDebugUtilities.h
@@ -249,9 +249,21 @@ protected:
 	static void GetDebugInfo_InContainer(int32 Index, FDebugInfo& DebugInfo, FProperty* Property, const void* Data);
 	static void GetDebugInfoInternal(FDebugInfo& DebugInfo, FProperty* Property, const void* PropertyValue);
 
-	// helper function for converting between blueprint and debuggable data
-	// output params are only valid if the return result is EWatchTextResult::EWTR_Valid
-	static EWatchTextResult FindDebuggingData(UBlueprint* Blueprint, UObject* ActiveObject, const UEdGraphPin* WatchPin, FProperty*& OutProperty, void*& OutData, void*& OutDelta, UObject*& OutParent, TArray<UObject*>& SeenObjects);
+	/**
+	* @brief	Helper function for converting between blueprint and debuggable data
+	*			output params are only valid if the return result is EWatchTextResult::EWTR_Valid
+	* 
+	* @param 	Blueprint		Active blueprint that is being debugged
+	* @param 	ActiveObject	Instance of the object that is being debugged
+	* @param 	WatchPin		The pin where this debug breakpoint is from
+	* @param 	OutProperty		Property of interest
+	* @param 	OutData			Populated with the property address of interest
+	* @param 	OutDelta		Populated with the same thing as OutData
+	* @param 	OutParent		Populated with the active object
+	* @param 	SeenObjects		Used to track what objects have been traversed to find the OutProperty address
+	* @return	EWTR_Valid if the debug data could be found, otherwise an appropriate error code
+	*/
+	static EWatchTextResult FindDebuggingData(UBlueprint* Blueprint, UObject* ActiveObject, const UEdGraphPin* WatchPin, FProperty*& OutProperty, void*& OutData, void*& OutDelta, UObject*& OutParent, TArray<UObject*>& SeenObjects, bool* OutbIsInDebuggingFrame = nullptr);
 
 private:
 	FKismetDebugUtilities() {}
