Rename out paramaters of FindDebuggingData to bShouldUseContainerOffset

#rnx

diff --git a/Engine/Source/Editor/UnrealEd/Private/Kismet2/KismetDebugUtilities.cpp b/Engine/Source/Editor/UnrealEd/Private/Kismet2/KismetDebugUtilities.cpp
--- a/Engine/Source/Editor/UnrealEd/Private/Kismet2/KismetDebugUtilities.cpp
+++ b/Engine/Source/Editor/UnrealEd/Private/Kismet2/KismetDebugUtilities.cpp
@@ -1123,15 +1123,15 @@ FKismetDebugUtilities::EWatchTextResult FKismetDebugUtilities::GetWatchText(FStr
 	void* DataPtr = nullptr;
 	void* DeltaPtr = nullptr;
 	UObject* ParentObj = nullptr;
-	bool bIsInDebuggingFrame = false;
+	bool bShouldUseContainerOffset = false;
 	TArray<UObject*> SeenObjects;
-	FKismetDebugUtilities::EWatchTextResult Result = FindDebuggingData(Blueprint, ActiveObject, WatchPin, PropertyToDebug, DataPtr, DeltaPtr, ParentObj, SeenObjects, &bIsInDebuggingFrame);
+	FKismetDebugUtilities::EWatchTextResult Result = FindDebuggingData(Blueprint, ActiveObject, WatchPin, PropertyToDebug, DataPtr, DeltaPtr, ParentObj, SeenObjects, &bShouldUseContainerOffset);
 
 	if (Result == FKismetDebugUtilities::EWatchTextResult::EWTR_Valid)
 	{
 		// If this came from an array property we need to avoid using ExportText_InContainer in order to properly 
 		// calculate the internal offset
-		if(bIsInDebuggingFrame && PropertyToDebug->IsA<FArrayProperty>())
+		if(bShouldUseContainerOffset)
 		{
 			PropertyToDebug->ExportText_Direct(/*inout*/ OutWatchText, DataPtr, DeltaPtr, ParentObj, PPF_PropertyWindow | PPF_BlueprintDebugView);
 		}
@@ -1161,7 +1161,7 @@ FKismetDebugUtilities::EWatchTextResult FKismetDebugUtilities::GetDebugInfo(FDeb
 	return Result;
 }
 
-FKismetDebugUtilities::EWatchTextResult FKismetDebugUtilities::FindDebuggingData(UBlueprint* Blueprint, UObject* ActiveObject, const UEdGraphPin* WatchPin, FProperty*& OutProperty, void*& OutData, void*& OutDelta, UObject*& OutParent, TArray<UObject*>& SeenObjects, bool* OutbIsInDebuggingFrame /* = nullptr */)
+FKismetDebugUtilities::EWatchTextResult FKismetDebugUtilities::FindDebuggingData(UBlueprint* Blueprint, UObject* ActiveObject, const UEdGraphPin* WatchPin, FProperty*& OutProperty, void*& OutData, void*& OutDelta, UObject*& OutParent, TArray<UObject*>& SeenObjects, bool* OutbShouldUseContainerOffset /* = nullptr */)
 {
 	FKismetDebugUtilitiesData& Data = FKismetDebugUtilitiesData::Get();
 
@@ -1211,7 +1211,7 @@ FKismetDebugUtilities::EWatchTextResult FKismetDebugUtilities::FindDebuggingData
 								// otherwise the output param won't show any data since the return node hasn't executed when we stop here
 								if (WatchPin->LinkedTo.Num() == 1)
 								{
-									return FindDebuggingData(Blueprint, ActiveObject, WatchPin->LinkedTo[0], OutProperty, OutData, OutDelta, OutParent, SeenObjects, OutbIsInDebuggingFrame);
+									return FindDebuggingData(Blueprint, ActiveObject, WatchPin->LinkedTo[0], OutProperty, OutData, OutDelta, OutParent, SeenObjects, OutbShouldUseContainerOffset);
 								}
 								else if (!WatchPin->LinkedTo.Num())
 								{
@@ -1222,9 +1222,9 @@ FKismetDebugUtilities::EWatchTextResult FKismetDebugUtilities::FindDebuggingData
 							}
 
 							// If this is an out property then we need to let the caller know
-							if (PropertyBase == nullptr && OutbIsInDebuggingFrame)
+							if (PropertyBase == nullptr && OutbShouldUseContainerOffset && OutParmRec->Property->IsA<FArrayProperty>())
 							{
-								*OutbIsInDebuggingFrame = true;
+								*OutbShouldUseContainerOffset = true;
 								PropertyBase = OutParmRec->PropAddr;
 							}
 							break;
diff --git a/Engine/Source/Editor/UnrealEd/Public/Kismet2/KismetDebugUtilities.h b/Engine/Source/Editor/UnrealEd/Public/Kismet2/KismetDebugUtilities.h
--- a/Engine/Source/Editor/UnrealEd/Public/Kismet2/KismetDebugUtilities.h
+++ b/Engine/Source/Editor/UnrealEd/Public/Kismet2/KismetDebugUtilities.h
@@ -263,7 +263,7 @@ protected:
 	* @param 	SeenObjects		Used to track what objects have been traversed to find the OutProperty address
 	* @return	EWTR_Valid if the debug data could be found, otherwise an appropriate error code
 	*/
-	static EWatchTextResult FindDebuggingData(UBlueprint* Blueprint, UObject* ActiveObject, const UEdGraphPin* WatchPin, FProperty*& OutProperty, void*& OutData, void*& OutDelta, UObject*& OutParent, TArray<UObject*>& SeenObjects, bool* OutbIsInDebuggingFrame = nullptr);
+	static EWatchTextResult FindDebuggingData(UBlueprint* Blueprint, UObject* ActiveObject, const UEdGraphPin* WatchPin, FProperty*& OutProperty, void*& OutData, void*& OutDelta, UObject*& OutParent, TArray<UObject*>& SeenObjects, bool* OutbShouldUseContainerOffset = nullptr);
 
 private:
 	FKismetDebugUtilities() {}
