Add additional checks for array, map and set types in the BP debugger utils. Add additional comments

#rnx

diff --git a/Engine/Source/Editor/UnrealEd/Private/Kismet2/KismetDebugUtilities.cpp b/Engine/Source/Editor/UnrealEd/Private/Kismet2/KismetDebugUtilities.cpp
--- a/Engine/Source/Editor/UnrealEd/Private/Kismet2/KismetDebugUtilities.cpp
+++ b/Engine/Source/Editor/UnrealEd/Private/Kismet2/KismetDebugUtilities.cpp
@@ -1221,8 +1221,11 @@ FKismetDebugUtilities::EWatchTextResult FKismetDebugUtilities::FindDebuggingData
 								}
 							}
 
-							// If this is an out property then we need to let the caller know
-							if (PropertyBase == nullptr && OutbShouldUseContainerOffset && OutParmRec->Property->IsA<FArrayProperty>())
+							// If this is an out container property then a different offset must be used when exporting this property 
+							// to text. Only container properties are effected by this because ExportText_InContainer adds an extra 
+							// 8 byte offset, which  would point to the container's first element, not the container itself. 
+							const bool bIsContainer = OutParmRec->Property->IsA<FArrayProperty>() || OutParmRec->Property->IsA<FSetProperty>() || OutParmRec->Property->IsA<FMapProperty>();
+							if (PropertyBase == nullptr && OutbShouldUseContainerOffset && bIsContainer)
 							{
 								*OutbShouldUseContainerOffset = true;
 								PropertyBase = OutParmRec->PropAddr;
