Don't include UClass*'s in the reinsancer's replacement map - we know CPFUO won't create UClasses, so this little kludge prevents us from overwriting UClass*'s unintentionally

diff --git a/Engine/Source/Editor/UnrealEd/Private/Kismet2/KismetReinstanceUtilities.cpp b/Engine/Source/Editor/UnrealEd/Private/Kismet2/KismetReinstanceUtilities.cpp
--- a/Engine/Source/Editor/UnrealEd/Private/Kismet2/KismetReinstanceUtilities.cpp
+++ b/Engine/Source/Editor/UnrealEd/Private/Kismet2/KismetReinstanceUtilities.cpp
@@ -1926,7 +1926,20 @@ void FBlueprintCompileReinstancer::ReplaceInstancesOfClass_Inner(TMap<UClass*, U
 	{
 		void OnObjectsReplaced(const TMap<UObject*, UObject*>& InReplacedObjects)
 		{
-			ReplacedObjects.Append(InReplacedObjects);
+			for (const TPair<UObject*, UObject*>& Pair : InReplacedObjects)
+			{
+				// CPFUO is going to tell us that the old class
+				// has been replaced with the new class, but we created
+				// the old class and we don't want to blindly replace
+				// references to the old class. This could cause, for example,
+				// the compilation manager to replace its references to the
+				// old class with references to the new class:
+				if (!Pair.Key->IsA<UClass>() &&
+					!Pair.Value->IsA<UClass>())
+				{
+					ReplacedObjects.Add(Pair);
+				}
+			}
 		}
 
 		TMap<UObject*, UObject*> ReplacedObjects;
