Allow nullptr mappings to behave as before
#rnx

diff --git a/Engine/Source/Editor/UnrealEd/Private/Kismet2/KismetReinstanceUtilities.cpp b/Engine/Source/Editor/UnrealEd/Private/Kismet2/KismetReinstanceUtilities.cpp
--- a/Engine/Source/Editor/UnrealEd/Private/Kismet2/KismetReinstanceUtilities.cpp
+++ b/Engine/Source/Editor/UnrealEd/Private/Kismet2/KismetReinstanceUtilities.cpp
@@ -1934,8 +1934,10 @@ void FBlueprintCompileReinstancer::ReplaceInstancesOfClass_Inner(TMap<UClass*, U
 				// references to the old class. This could cause, for example,
 				// the compilation manager to replace its references to the
 				// old class with references to the new class:
-				if (!Pair.Key->IsA<UClass>() &&
-					!Pair.Value->IsA<UClass>())
+				if (Pair.Key == nullptr || 
+					Pair.Value == nullptr ||
+					(	!Pair.Key->IsA<UClass>() &&
+						!Pair.Value->IsA<UClass>()) )
 				{
 					ReplacedObjects.Add(Pair);
 				}
