- Optionally show a realtime disabled warning

diff --git a/Engine/Plugins/FX/Niagara/Source/NiagaraEditor/Private/Widgets/SNiagaraSystemViewportToolBar.cpp b/Engine/Plugins/FX/Niagara/Source/NiagaraEditor/Private/Widgets/SNiagaraSystemViewportToolBar.cpp
--- a/Engine/Plugins/FX/Niagara/Source/NiagaraEditor/Private/Widgets/SNiagaraSystemViewportToolBar.cpp
+++ b/Engine/Plugins/FX/Niagara/Source/NiagaraEditor/Private/Widgets/SNiagaraSystemViewportToolBar.cpp
@@ -14,7 +14,7 @@
 
 void SNiagaraSystemViewportToolBar::Construct(const FArguments& InArgs, TSharedPtr<class SNiagaraSystemViewport> InViewport)
 {
-	SCommonEditorViewportToolbarBase::Construct(SCommonEditorViewportToolbarBase::FArguments(), InViewport);
+	SCommonEditorViewportToolbarBase::Construct(SCommonEditorViewportToolbarBase::FArguments().AddRealtimeButton(true), InViewport);
 }
 
 TSharedRef<SWidget> SNiagaraSystemViewportToolBar::GenerateShowMenu() const
diff --git a/Engine/Source/Editor/UnrealEd/Private/SCommonEditorViewportToolbarBase.cpp b/Engine/Source/Editor/UnrealEd/Private/SCommonEditorViewportToolbarBase.cpp
--- a/Engine/Source/Editor/UnrealEd/Private/SCommonEditorViewportToolbarBase.cpp
+++ b/Engine/Source/Editor/UnrealEd/Private/SCommonEditorViewportToolbarBase.cpp
@@ -13,6 +13,7 @@
 #include "SEditorViewport.h"
 #include "EditorViewportCommands.h"
 #include "SEditorViewportToolBarMenu.h"
+#include "SEditorViewportToolBarButton.h"
 #include "SEditorViewportViewMenu.h"
 #include "Editor/EditorPerformanceSettings.h"
 #include "Scalability.h"
@@ -98,6 +99,30 @@ void SCommonEditorViewportToolbarBase::Construct(const FArguments& InArgs, TShar
 			.OnGetMenuContent(this, &SCommonEditorViewportToolbarBase::GenerateShowMenu)
 		];
 
+	// Realtime button
+	if (InArgs._AddRealtimeButton)
+	{
+		MainBoxPtr->AddSlot()
+			.AutoWidth()
+			.Padding(ToolbarSlotPadding)
+			[
+				SNew(SEditorViewportToolBarButton)
+				.Cursor(EMouseCursor::Default)
+				.ButtonType(EUserInterfaceActionType::Button)
+				.ButtonStyle(&FEditorStyle::Get().GetWidgetStyle<FButtonStyle>("EditorViewportToolBar.MenuButtonWarning"))
+				.OnClicked(this, &SCommonEditorViewportToolbarBase::OnRealtimeWarningClicked)
+				.Visibility(this, &SCommonEditorViewportToolbarBase::GetRealtimeWarningVisibility)
+				.ToolTipText(LOCTEXT("RealtimeOff_ToolTip", "This viewport is not updating in realtime.  Click to turn on realtime mode."))
+				.Content()
+				[
+					SNew(STextBlock)
+					.Font(FEditorStyle::GetFontStyle("EditorViewportToolBar.Font"))
+					.Text(LOCTEXT("RealtimeOff", "Realtime: Off"))
+					.ColorAndOpacity(FLinearColor::Black)
+				]
+			];
+	}
+
 	MainBoxPtr->AddSlot()
 		.AutoWidth()
 		.Padding(ToolbarSlotPadding)
@@ -444,6 +469,21 @@ void SCommonEditorViewportToolbarBase::OnFarViewPlaneValueChanged(float NewValue
 	GetViewportClient().OverrideFarClipPlane(NewValue);
 }
 
+FReply SCommonEditorViewportToolbarBase::OnRealtimeWarningClicked()
+{
+	FEditorViewportClient& ViewportClient = GetViewportClient();
+	ViewportClient.SetRealtime(true);
+
+	return FReply::Handled();
+}
+
+EVisibility SCommonEditorViewportToolbarBase::GetRealtimeWarningVisibility() const
+{
+	FEditorViewportClient& ViewportClient = GetViewportClient();
+	// If the viewport is not realtime and there is no override then realtime is off
+	return !ViewportClient.IsRealtime() && !ViewportClient.IsRealtimeOverrideSet() ? EVisibility::Visible : EVisibility::Collapsed;
+}
+
 TSharedPtr<FExtender> SCommonEditorViewportToolbarBase::GetCombinedExtenderList(TSharedRef<FExtender> MenuExtender) const
 {
 	TSharedPtr<FExtender> HostEditorExtenders = GetInfoProvider().GetExtenders();
diff --git a/Engine/Source/Editor/UnrealEd/Public/SCommonEditorViewportToolbarBase.h b/Engine/Source/Editor/UnrealEd/Public/SCommonEditorViewportToolbarBase.h
--- a/Engine/Source/Editor/UnrealEd/Public/SCommonEditorViewportToolbarBase.h
+++ b/Engine/Source/Editor/UnrealEd/Public/SCommonEditorViewportToolbarBase.h
@@ -65,7 +65,11 @@ namespace CommonEditorViewportUtils
 class UNREALED_API SCommonEditorViewportToolbarBase : public SViewportToolBar
 {
 public:
-	SLATE_BEGIN_ARGS(SCommonEditorViewportToolbarBase){}
+	SLATE_BEGIN_ARGS(SCommonEditorViewportToolbarBase)
+		: _AddRealtimeButton(false)
+		{}
+
+		SLATE_ARGUMENT(bool, AddRealtimeButton)
 	SLATE_END_ARGS()
 
 	void Construct(const FArguments& InArgs, TSharedPtr<class ICommonEditorViewportToolbarInfoProvider> InInfoProvider);
@@ -171,6 +175,11 @@ private:
 	/** Called when the far view plane slider is adjusted in the perspective viewport */
 	void OnFarViewPlaneValueChanged( float NewValue );
 
+	/** Called when we click the realtime warning */
+	FReply OnRealtimeWarningClicked();
+	/** Called to determine if we should show the realtime warning */
+	EVisibility GetRealtimeWarningVisibility() const;
+
 protected:
 	// Merges the extender list from the host with the specified extender and returns the results
 	TSharedPtr<FExtender> GetCombinedExtenderList(TSharedRef<FExtender> MenuExtender) const;
