Fixes excessive logs when switching InteractiveTools

diff --git a/Engine/Source/Editor/UnrealEd/Private/Tools/UEdMode.cpp b/Engine/Source/Editor/UnrealEd/Private/Tools/UEdMode.cpp
--- a/Engine/Source/Editor/UnrealEd/Private/Tools/UEdMode.cpp
+++ b/Engine/Source/Editor/UnrealEd/Private/Tools/UEdMode.cpp
@@ -289,7 +289,7 @@ public:
 		{
 			ToolsContext->PostToolNotificationMessage(Message);
 		}
-		if (Level == EToolMessageLevel::UserWarning)
+		else if (Level == EToolMessageLevel::UserWarning)
 		{
 			ToolsContext->PostToolWarningMessage(Message);
 		}
