Allow the Actor name to be included when harvesting components

diff --git a/Engine/Plugins/Enterprise/DatasmithImporter/Source/DatasmithImporter/Private/DatasmithImporter.cpp b/Engine/Plugins/Enterprise/DatasmithImporter/Source/DatasmithImporter/Private/DatasmithImporter.cpp
--- a/Engine/Plugins/Enterprise/DatasmithImporter/Source/DatasmithImporter/Private/DatasmithImporter.cpp
+++ b/Engine/Plugins/Enterprise/DatasmithImporter/Source/DatasmithImporter/Private/DatasmithImporter.cpp
@@ -662,7 +662,7 @@ void FDatasmithImporter::ImportActors( FDatasmithImportContext& ImportContext )
 		if (ImportContext.Options->HierarchyHandling == EDatasmithImportHierarchy::UseOneBlueprint && ImportContext.RootBlueprint != nullptr)
 		{
 			// Reparent all scene components attached to root actor toward blueprint root
-			FKismetEditorUtilities::AddComponentsToBlueprint(ImportContext.RootBlueprint, ImportSceneActor->GetInstanceComponents(), false, (USCS_Node*)nullptr, true);
+			FKismetEditorUtilities::AddComponentsToBlueprint(ImportContext.RootBlueprint, ImportSceneActor->GetInstanceComponents(), FKismetEditorUtilities::EAddComponentToBPHarvestMode::None, (USCS_Node*)nullptr, true);
 		}
 
 		// After all actors were imported, perform a post import step so that any dependencies can be resolved
diff --git a/Engine/Source/Editor/UnrealEd/Private/Fbx/ReimportFbxSceneFactory.cpp b/Engine/Source/Editor/UnrealEd/Private/Fbx/ReimportFbxSceneFactory.cpp
--- a/Engine/Source/Editor/UnrealEd/Private/Fbx/ReimportFbxSceneFactory.cpp
+++ b/Engine/Source/Editor/UnrealEd/Private/Fbx/ReimportFbxSceneFactory.cpp
@@ -1003,7 +1003,7 @@ UBlueprint *UReimportFbxSceneFactory::UpdateOriginalBluePrint(FString &BluePrint
 		CollectGarbage(GARBAGE_COLLECTION_KEEPFLAGS);
 
 		//Create the new nodes from the hierarchy actor
-		FKismetEditorUtilities::AddComponentsToBlueprint(BluePrint, HierarchyActor->GetInstanceComponents(), false, nullptr, true);
+		FKismetEditorUtilities::AddComponentsToBlueprint(BluePrint, HierarchyActor->GetInstanceComponents(), FKismetEditorUtilities::EAddComponentToBPHarvestMode::None, nullptr, true);
 		
 		UWorld* World = HierarchyActor->GetWorld();
 		World->DestroyActor(HierarchyActor);
diff --git a/Engine/Source/Editor/UnrealEd/Private/Kismet2/Kismet2.cpp b/Engine/Source/Editor/UnrealEd/Private/Kismet2/Kismet2.cpp
--- a/Engine/Source/Editor/UnrealEd/Private/Kismet2/Kismet2.cpp
+++ b/Engine/Source/Editor/UnrealEd/Private/Kismet2/Kismet2.cpp
@@ -1062,7 +1062,7 @@ UBlueprint* FKismetEditorUtilities::CreateBlueprintFromActors(const FString& Pat
 	return NewBlueprint;
 }
 
-void FKismetEditorUtilities::AddComponentsToBlueprint(UBlueprint* Blueprint, const TArray<UActorComponent*>& Components, bool bHarvesting, USCS_Node* OptionalNewRootNode, bool bKeepMobility /*= false*/)
+void FKismetEditorUtilities::AddComponentsToBlueprint(UBlueprint* Blueprint, const TArray<UActorComponent*>& Components, FKismetEditorUtilities::EAddComponentToBPHarvestMode HarvestMode, USCS_Node* OptionalNewRootNode, const bool bKeepMobility /*= false*/)
 {
 	USimpleConstructionScript* SCS = Blueprint->SimpleConstructionScript;
 
@@ -1085,43 +1085,49 @@ void FKismetEditorUtilities::AddComponentsToBlueprint(UBlueprint* Blueprint, con
 		}
 	};
 
-	struct FAddComponentsToBlueprintImpl
+	/** 
+	 * Creates a new USCS_Node in the TargetSCS, duplicating the specified 
+	 * component (leaving the new node unattached). If a copy was already 
+	 * made (found  in NewSceneComponents) then that will be returned instead.
+	 */	
+	auto MakeComponentCopy = [SCS, HarvestMode, bKeepMobility, &InstanceComponentToNodeMap](UActorComponent* ActorComponent)
 	{
-		/** 
-		 * Creates a new USCS_Node in the TargetSCS, duplicating the specified 
-		 * component (leaving the new node unattached). If a copy was already 
-		 * made (found  in NewSceneComponents) then that will be returned instead.
-		 */
-		static USCS_Node* MakeComponentCopy(UActorComponent* ActorComponent, USimpleConstructionScript* TargetSCS, TMap<USceneComponent*, USCS_Node*>& NewSceneComponents, bool bInternalKeepMobility)
+		USceneComponent* AsSceneComponent = Cast<USceneComponent>(ActorComponent);
+		if (AsSceneComponent != nullptr)
 		{
-			USceneComponent* AsSceneComponent = Cast<USceneComponent>(ActorComponent);
-			if (AsSceneComponent != nullptr)
+			USCS_Node** ExistingCopy = InstanceComponentToNodeMap.Find(AsSceneComponent);
+			if (ExistingCopy != nullptr)
 			{
-				USCS_Node** ExistingCopy = NewSceneComponents.Find(AsSceneComponent);
-				if (ExistingCopy != nullptr)
-				{
-					return *ExistingCopy;
-				}
+				return *ExistingCopy;
+			}
+		}
+
+		FName NewComponentName = ActorComponent->GetFName();
+		if (HarvestMode == EAddComponentToBPHarvestMode::Havest_AppendOwnerName)
+		{
+			if (AActor* Owner = ActorComponent->GetOwner())
+			{
+				NewComponentName = *(Owner->GetActorLabel() + TEXT("_") + ActorComponent->GetName());
 			}
+		}
 
-			USCS_Node* NewSCSNode = TargetSCS->CreateNode(ActorComponent->GetClass(), ActorComponent->GetFName());
-			UEditorEngine::FCopyPropertiesForUnrelatedObjectsParams Params;
-			Params.bDoDelta = false; // We need a deep copy of parameters here so the CDO values get copied as well
-			UEditorEngine::CopyPropertiesForUnrelatedObjects(ActorComponent, NewSCSNode->ComponentTemplate, Params);
+		USCS_Node* NewSCSNode = SCS->CreateNode(ActorComponent->GetClass(), NewComponentName);
+		UEditorEngine::FCopyPropertiesForUnrelatedObjectsParams Params;
+		Params.bDoDelta = false; // We need a deep copy of parameters here so the CDO values get copied as well
+		UEditorEngine::CopyPropertiesForUnrelatedObjects(ActorComponent, NewSCSNode->ComponentTemplate, Params);
 
-			// Clear the instance component flag
-			NewSCSNode->ComponentTemplate->CreationMethod = EComponentCreationMethod::Native;
+		// Clear the instance component flag
+		NewSCSNode->ComponentTemplate->CreationMethod = EComponentCreationMethod::Native;
 
-			if (AsSceneComponent != nullptr)
+		if (AsSceneComponent != nullptr)
+		{
+			InstanceComponentToNodeMap.Add(AsSceneComponent, NewSCSNode);
+			if (!bKeepMobility)
 			{
-				NewSceneComponents.Add(AsSceneComponent, NewSCSNode);
-				if (!bInternalKeepMobility)
-				{
-					Cast<USceneComponent>(NewSCSNode->ComponentTemplate)->SetMobility(EComponentMobility::Movable);
-				}
+				Cast<USceneComponent>(NewSCSNode->ComponentTemplate)->SetMobility(EComponentMobility::Movable);
 			}
-			return NewSCSNode;
 		}
+		return NewSCSNode;
 	};
 
 	// Associate a scene node to its first attached parent if there is one.
@@ -1206,22 +1212,18 @@ void FKismetEditorUtilities::AddComponentsToBlueprint(UBlueprint* Blueprint, con
 	}
 
 	// The easy part to add the non-scene components.
-	for (int32 CompIndex = 0; CompIndex < ActorComponents.Num(); ++CompIndex)
+	for (UActorComponent* ActorComponent : ActorComponents)
 	{
-		UActorComponent* ActorComponent = ActorComponents[CompIndex];
-
-		USCS_Node* SCSNode = FAddComponentsToBlueprintImpl::MakeComponentCopy(ActorComponent, SCS, InstanceComponentToNodeMap, bKeepMobility);
-
+		USCS_Node* SCSNode = MakeComponentCopy(ActorComponent);
 		SCS->AddNode(SCSNode);
 	}
 
 	// The loop to add the scene components
-	for (int32 CompIndex = 0; CompIndex < SceneComponentNodes.Num(); ++CompIndex)
+	for (const FSceneComponentAndFirstParent& ComponentNode : SceneComponentNodes)
 	{
-		const FSceneComponentAndFirstParent& ComponentNode = SceneComponentNodes[CompIndex];
 		USceneComponent* SceneComponent = ComponentNode.SceneComponent;
-
-		USCS_Node* SCSNode = FAddComponentsToBlueprintImpl::MakeComponentCopy(SceneComponent, SCS, InstanceComponentToNodeMap, bKeepMobility);
+		
+		USCS_Node* SCSNode = MakeComponentCopy(SceneComponent);
 
 		USceneComponent* FirstAttachParent = ComponentNode.FirstAttachParent;
 		
@@ -1301,7 +1303,7 @@ void FKismetEditorUtilities::AddComponentsToBlueprint(UBlueprint* Blueprint, con
 				ParentSCSNode->AddChildNode(SCSNode);
 			}
 		}
-		else if ((SceneComponent->GetAttachParent()->CreationMethod == EComponentCreationMethod::Native) && !bHarvesting)
+		else if ((SceneComponent->GetAttachParent()->CreationMethod == EComponentCreationMethod::Native) && (HarvestMode == EAddComponentToBPHarvestMode::None))
 		{
 			// If we're attached to a component that will be native in the new blueprint
 			SCS->AddNode(SCSNode);
@@ -1391,14 +1393,14 @@ UBlueprint* FKismetEditorUtilities::CreateBlueprintFromActor(const FName Bluepri
 						{
 							TArray<UActorComponent*> InstanceComponents = Actor->GetInstanceComponents();
 							InstanceComponents.Remove(RootComponent);
-							AddComponentsToBlueprint(NewBlueprint, InstanceComponents, false, NewBlueprint->SimpleConstructionScript->GetDefaultSceneRootNode(), bKeepMobility);
+							AddComponentsToBlueprint(NewBlueprint, InstanceComponents, FKismetEditorUtilities::EAddComponentToBPHarvestMode::None, NewBlueprint->SimpleConstructionScript->GetDefaultSceneRootNode(), bKeepMobility);
 						}
 					}
 				}
 
 				if (!bUsedDefaultSceneRoot)
 				{
-					AddComponentsToBlueprint(NewBlueprint, Actor->GetInstanceComponents(), false, nullptr, bKeepMobility);
+					AddComponentsToBlueprint(NewBlueprint, Actor->GetInstanceComponents(), FKismetEditorUtilities::EAddComponentToBPHarvestMode::None, nullptr, bKeepMobility);
 				}
 			}
 
@@ -1627,7 +1629,7 @@ UBlueprint* FKismetEditorUtilities::CreateBlueprintFromActors(const FName Bluepr
 			RecursiveCreateChildActorTemplates(Actor, nullptr);
 		}
 
-		FKismetEditorUtilities::AddComponentsToBlueprint(AssemblyProps.Blueprint, ChildActorComponents, /*bHarvesting=*/ true, AssemblyProps.RootNodeOverride);
+		FKismetEditorUtilities::AddComponentsToBlueprint(AssemblyProps.Blueprint, ChildActorComponents, FKismetEditorUtilities::EAddComponentToBPHarvestMode::Harvest_UseComponentName, AssemblyProps.RootNodeOverride);
 
 		// Since the names we create are well defined relative to the SCS but created in the transient package, we could end up reusing objects
 		// unless we rename these temporary components out of the way
@@ -1699,7 +1701,8 @@ UBlueprint* FKismetEditorUtilities::HarvestBlueprintFromActors(const FName Bluep
 			}
 		}
 
-		FKismetEditorUtilities::AddComponentsToBlueprint(AssemblyProps.Blueprint, AllSelectedComponents, /*bHarvesting=*/ true, AssemblyProps.RootNodeOverride);
+		const FKismetEditorUtilities::EAddComponentToBPHarvestMode HarvestMode = (AssemblyProps.Actors.Num() > 1 ? FKismetEditorUtilities::EAddComponentToBPHarvestMode::Havest_AppendOwnerName : FKismetEditorUtilities::EAddComponentToBPHarvestMode::Harvest_UseComponentName);
+		FKismetEditorUtilities::AddComponentsToBlueprint(AssemblyProps.Blueprint, AllSelectedComponents, HarvestMode, AssemblyProps.RootNodeOverride);
 
 		// Replace the modified components to their relative transform
 		for (const TPair<USceneComponent*, FTransform >& Pair : SceneComponentOldRelativeTransforms)
diff --git a/Engine/Source/Editor/UnrealEd/Public/Kismet2/KismetEditorUtilities.h b/Engine/Source/Editor/UnrealEd/Public/Kismet2/KismetEditorUtilities.h
--- a/Engine/Source/Editor/UnrealEd/Public/Kismet2/KismetEditorUtilities.h
+++ b/Engine/Source/Editor/UnrealEd/Public/Kismet2/KismetEditorUtilities.h
@@ -146,8 +146,27 @@ public:
 	/** @return true is it's possible to create a blueprint from the specified class */
 	static bool CanCreateBlueprintOfClass(const UClass* Class);
 
+	enum class EAddComponentToBPHarvestMode : uint8
+	{
+		/* Not harvesting */
+		None,
+
+		/* Harvesting, use the component name for the new component directly */
+		Harvest_UseComponentName,
+
+		/* Harvesting, new component name will be OwnerName_ComponentName */
+		Havest_AppendOwnerName
+	};
+
 	/** Take a list of components that belong to a single Actor and add them to a blueprint as SCSNodes */
-	static void AddComponentsToBlueprint(UBlueprint* Blueprint, const TArray<UActorComponent*>& Components, bool bHarvesting = false, class USCS_Node* OptionalNewRootNode = nullptr, bool bKeepMobility = false);
+	static void AddComponentsToBlueprint(UBlueprint* Blueprint, const TArray<UActorComponent*>& Components, EAddComponentToBPHarvestMode HarvestMode = EAddComponentToBPHarvestMode::None, class USCS_Node* OptionalNewRootNode = nullptr, bool bKeepMobility = false);
+
+	UE_DEPRECATED(4.25, "Use version that specifies harvest mode via enum instead of boolean")
+	static void AddComponentsToBlueprint(UBlueprint* Blueprint, const TArray<UActorComponent*>& Components, bool bHarvesting, class USCS_Node* OptionalNewRootNode = nullptr, bool bKeepMobility = false)
+	{
+		AddComponentsToBlueprint(Blueprint, Components, (bHarvesting ? EAddComponentToBPHarvestMode::Harvest_UseComponentName : EAddComponentToBPHarvestMode::None), OptionalNewRootNode, bKeepMobility);
+	}
+
 
 	/** 
 	 * Take an Actor and generate a blueprint based on it. Uses the Actors type as the parent class. 
