HLOD Volumes have trouble overlapping meshes offset from origin

We can't rely on the actor location & EncompassesPoint() as the mesh pivot can be offset, the only way to test properly is to use the components bounding box

diff --git a/Engine/Source/Editor/UnrealEd/Private/HierarchicalLOD.cpp b/Engine/Source/Editor/UnrealEd/Private/HierarchicalLOD.cpp
--- a/Engine/Source/Editor/UnrealEd/Private/HierarchicalLOD.cpp
+++ b/Engine/Source/Editor/UnrealEd/Private/HierarchicalLOD.cpp
@@ -173,12 +173,17 @@ void FHierarchicalLODBuilder::BuildClusters(ULevel* InLevel)
 							if (VolumePtr)
 							{
 								AHierarchicalLODVolume* Volume = *VolumePtr;
+								FBox HLODVolumeBox = Volume->GetComponentsBoundingBox(true);
+
 								for (AActor* Actor : ValidStaticMeshActorsInLevel)
 								{
-									if (ShouldGenerateCluster(Actor, LODId) && Volume->EncompassesPoint(Actor->GetActorLocation(), Volume->bIncludeOverlappingActors ? Actor->GetComponentsBoundingBox().GetSize().Size() : 0.0f, nullptr))
+									if (ShouldGenerateCluster(Actor, LODId))
 									{
-										FLODCluster ActorCluster(Actor);
-										PreviousActorCluster += ActorCluster;
+										FBox ActorBox = Actor->GetComponentsBoundingBox(true);
+										if (HLODVolumeBox.IsInside(ActorBox) || (Volume->bIncludeOverlappingActors && HLODVolumeBox.Intersect(ActorBox)))
+										{
+											PreviousActorCluster += Actor;
+										}
 									}
 								}
 
@@ -187,16 +192,17 @@ void FHierarchicalLODBuilder::BuildClusters(ULevel* InLevel)
 								{
 									for (AActor* Actor : RejectedActorsInLevel)
 									{
-										if (Actor && Volume->EncompassesPoint(Actor->GetActorLocation(), Volume->bIncludeOverlappingActors ? Actor->GetComponentsBoundingBox().GetSize().Size() : 0.0f, nullptr))
+										if (!ShouldGenerateCluster(Actor, LODId - 1) && ShouldGenerateCluster(Actor, LODId))
 										{
-											if (!ShouldGenerateCluster(Actor, LODId - 1) && ShouldGenerateCluster(Actor, LODId))
+											FBox ActorBox = Actor->GetComponentsBoundingBox(true);
+											if (HLODVolumeBox.IsInside(ActorBox) || (Volume->bIncludeOverlappingActors && HLODVolumeBox.Intersect(ActorBox)))
 											{
 												PreviousActorCluster += Actor;
 											}
 										}
 									}
 								}
-							}							
+							}
 						}
 						else
 						{
@@ -455,19 +461,14 @@ void FHierarchicalLODBuilder::InitializeClusters(ULevel* InLevel, const int32 LO
 	// Check whether or not this actor falls within a HierarchicalLODVolume, if so add to the Volume's cluster and exclude from normal process
 	auto ProcessVolumeClusters = [this](AActor* InActor) -> bool
 	{
+		FBox ActorBox = InActor->GetComponentsBoundingBox(true);
 		for (TPair<AHierarchicalLODVolume*, FLODCluster>& Cluster : HLODVolumeClusters)
 		{
-			if (Cluster.Key->EncompassesPoint(InActor->GetActorLocation(), Cluster.Key->bIncludeOverlappingActors ? InActor->GetComponentsBoundingBox().GetSize().Size() : 0.0f, nullptr))
+			FBox HLODVolumeBox = Cluster.Key->GetComponentsBoundingBox(true);
+			if (HLODVolumeBox.IsInside(ActorBox) || (Cluster.Key->bIncludeOverlappingActors && HLODVolumeBox.Intersect(ActorBox)))
 			{
-				FBox BoundingBox = InActor->GetComponentsBoundingBox(true);
-				FBox VolumeBox = Cluster.Key->GetComponentsBoundingBox(true);
-
-				if (VolumeBox.IsInside(BoundingBox) || (Cluster.Key->bIncludeOverlappingActors && VolumeBox.Intersect(BoundingBox)))
-				{
-					FLODCluster ActorCluster(InActor);
-					Cluster.Value += ActorCluster;
-					return true;
-				}
+				Cluster.Value += InActor;
+				return true;
 			}
 		}
 
