Editor: Use the same PIE Instance index when saving window positions as when loading. This fixes an issue where separate servers caused multiple PIE clients to crash on close.

#rnx

diff --git a/Engine/Source/Editor/UnrealEd/Private/PlayLevel.cpp b/Engine/Source/Editor/UnrealEd/Private/PlayLevel.cpp
--- a/Engine/Source/Editor/UnrealEd/Private/PlayLevel.cpp
+++ b/Engine/Source/Editor/UnrealEd/Private/PlayLevel.cpp
@@ -3005,7 +3005,7 @@ TSharedRef<SPIEViewport> UEditorEngine::GeneratePIEViewportWindow(const FRequest
 	{
 		struct FLocal
 		{
-			static void OnPIEWindowClosed(const TSharedRef<SWindow>& WindowBeingClosed, TWeakPtr<SViewport> PIEViewportWidget, TWeakObjectPtr<UEditorEngine> OwningEditorEngine, int32 InInstanceIndex, bool bRestoreRootWindow)
+			static void OnPIEWindowClosed(const TSharedRef<SWindow>& WindowBeingClosed, TWeakPtr<SViewport> PIEViewportWidget, TWeakObjectPtr<UEditorEngine> OwningEditorEngine, int32 ViewportIndex, bool bRestoreRootWindow)
 			{
 				// Save off the window position
 				const FVector2D PIEWindowPos = WindowBeingClosed->GetPositionInScreen();
@@ -3015,7 +3015,7 @@ TSharedRef<SPIEViewport> UEditorEngine::GeneratePIEViewportWindow(const FRequest
 
 				if (OwningEditorEngine.IsValid())
 				{
-					OwningEditorEngine->StoreWindowSizeAndPositionForInstanceIndex(InInstanceIndex, WindowSize, WindowPosition);
+					OwningEditorEngine->StoreWindowSizeAndPositionForInstanceIndex(ViewportIndex, WindowSize, WindowPosition);
 				}
 
 				// Route the callback
@@ -3033,11 +3033,8 @@ TSharedRef<SPIEViewport> UEditorEngine::GeneratePIEViewportWindow(const FRequest
 			}
 		};
 
-		const bool bHasServer = InNetMode == EPlayNetMode::PIE_Client;
-		int32 InstanceIndex = bHasServer ? InWorldContext.PIEInstance - 1 : InWorldContext.PIEInstance;
-
 		PieWindow->SetOnWindowClosed(FOnWindowClosed::CreateStatic(&FLocal::OnPIEWindowClosed, TWeakPtr<SViewport>(PieViewportWidget), TWeakObjectPtr<UEditorEngine>(this),
-			InstanceIndex, bShouldMinimizeRootWindow));
+			InViewportIndex, bShouldMinimizeRootWindow));
 	}
 
 	// Create a new viewport that the viewport widget will use to render the game
