Fix prompt suggesting to install "Visual Studio 2017" rather than 2019.

diff --git a/Engine/Source/Editor/UnrealEd/Private/SourceCodeNavigation.cpp b/Engine/Source/Editor/UnrealEd/Private/SourceCodeNavigation.cpp
--- a/Engine/Source/Editor/UnrealEd/Private/SourceCodeNavigation.cpp
+++ b/Engine/Source/Editor/UnrealEd/Private/SourceCodeNavigation.cpp
@@ -1749,7 +1749,7 @@ FText FSourceCodeNavigation::GetSuggestedSourceCodeIDE(bool bShortIDEName)
 	}
 	else
 	{
-		return LOCTEXT("SuggestedCodeIDE_Windows", "Visual Studio 2017");
+		return LOCTEXT("SuggestedCodeIDE_Windows", "Visual Studio 2019");
 	}
 #elif PLATFORM_MAC
 	return LOCTEXT("SuggestedCodeIDE_Mac", "Xcode");
