Experimental support for basic fix ticking in PIE
-This allows worlds to be ticked in "accumulated fix tick" mode where variable engine frames accumulate time towwards ticking the pie world.
-can be used to simulate 20/30 hz dedicated server.
-Clients can also use this thhough some issues are expected at very high fixed tick frame rates.

diff --git a/Engine/Source/Editor/UnrealEd/Classes/Settings/LevelEditorPlaySettings.h b/Engine/Source/Editor/UnrealEd/Classes/Settings/LevelEditorPlaySettings.h
--- a/Engine/Source/Editor/UnrealEd/Classes/Settings/LevelEditorPlaySettings.h
+++ b/Engine/Source/Editor/UnrealEd/Classes/Settings/LevelEditorPlaySettings.h
@@ -370,6 +370,14 @@ public:
 	UPROPERTY(config, EditAnywhere, Category = "Multiplayer Options|Server")
 	FString AdditionalServerLaunchParameters;
 
+	/** If > 0, Tick dedicated server at a fixed frame rate. Does not impact Listen Server (use ClientFixedFPS setting). This is the target frame rate, e.g, "20" for 20fps, which will result in 1/20 second tick steps. */
+	UPROPERTY(config, EditAnywhere, Category = "Multiplayer Options|Server", meta=(EditCondition = "PlayNetMode == EPlayNetMode::PIE_Client && RunUnderOneProcess"))
+	int32 ServerFixedFPS;
+
+	/** If > 0, Tick clients at a fixed frame rate. Each client instance will map to an element in the list, wrapping around if num clients exceeds size of list. Includes Listen Server. This is the target frame rate, e.g, "20" for 20fps, which will result in 1/20 second tick steps. */
+	UPROPERTY(config, EditAnywhere, Category = "Multiplayer Options|Client", meta=(EditCondition = "PlayNetMode != EPlayNetMode::PIE_Standalone && RunUnderOneProcess"))
+	TArray<int32> ClientFixedFPS;
+
 public:
 
 	/** Should network emulation settings be applied or not */
diff --git a/Engine/Source/Editor/UnrealEd/Private/EditorEngine.cpp b/Engine/Source/Editor/UnrealEd/Private/EditorEngine.cpp
--- a/Engine/Source/Editor/UnrealEd/Private/EditorEngine.cpp
+++ b/Engine/Source/Editor/UnrealEd/Private/EditorEngine.cpp
@@ -1657,58 +1657,75 @@ void UEditorEngine::Tick( float DeltaSeconds, bool bIdleMode )
 			{
 				OldGWorld->TransferBlueprintDebugReferences(PlayWorld);
 			}
-
-			// Tick all travel and Pending NetGames (Seamless, server, client)
-			TickWorldTravel(PieContext, DeltaSeconds);
-
-			// Updates 'connecting' message in PIE network games
-			UpdateTransitionType(PlayWorld);
-
-			// Update streaming for dedicated servers in PIE
-			if (PieContext.RunAsDedicated)
+			
+			float TickDeltaSeconds; // How much time to use per tick
+			
+			if (PieContext.PIEFixedTickSeconds > 0.f)
 			{
-				SCOPE_CYCLE_COUNTER(STAT_UpdateLevelStreaming);
-				PlayWorld->UpdateLevelStreaming();
+				PieContext.PIEAccumulatedTickSeconds += DeltaSeconds;
+				TickDeltaSeconds = PieContext.PIEFixedTickSeconds;
 			}
-
-			// Release mouse if the game is paused. The low level input code might ignore the request when e.g. in fullscreen mode.
-			if ( GameViewport != NULL && GameViewport->Viewport != NULL )
+			else
 			{
-				// Decide whether to drop high detail because of frame rate
-				GameViewport->SetDropDetail(DeltaSeconds);
+				PieContext.PIEAccumulatedTickSeconds = DeltaSeconds;
+				TickDeltaSeconds = DeltaSeconds;
 			}
+			
+			for ( ; PieContext.PIEAccumulatedTickSeconds >= TickDeltaSeconds; PieContext.PIEAccumulatedTickSeconds -= TickDeltaSeconds)
+			{
+				// Tick all travel and Pending NetGames (Seamless, server, client)
+				TickWorldTravel(PieContext, TickDeltaSeconds);
 
-			// Update the level.
-			GameCycles=0;
-			CLOCK_CYCLES(GameCycles);
+				// Updates 'connecting' message in PIE network games
+				UpdateTransitionType(PlayWorld);
 
-			{
-				// So that hierarchical stats work in PIE
-				SCOPE_CYCLE_COUNTER(STAT_FrameTime);
+				// Update streaming for dedicated servers in PIE
+				if (PieContext.RunAsDedicated)
+				{
+					SCOPE_CYCLE_COUNTER(STAT_UpdateLevelStreaming);
+					PlayWorld->UpdateLevelStreaming();
+				}
 
-				FKismetDebugUtilities::NotifyDebuggerOfStartOfGameFrame(PieContext.World());
+				// Release mouse if the game is paused. The low level input code might ignore the request when e.g. in fullscreen mode.
+				if ( GameViewport != NULL && GameViewport->Viewport != NULL )
+				{
+					// Decide whether to drop high detail because of frame rate
+					GameViewport->SetDropDetail(TickDeltaSeconds);
+				}
 
-				// tick the level
-				PieContext.World()->Tick( LEVELTICK_All, DeltaSeconds );
-				bAWorldTicked = true;
-				TickType = LEVELTICK_All;
+				// Update the level.
+				GameCycles=0;
+				CLOCK_CYCLES(GameCycles);
 
-				if (!bFirstTick)
 				{
-					// Update sky light first because sky diffuse will be visible in reflection capture indirect specular
-					USkyLightComponent::UpdateSkyCaptureContents(PlayWorld);
-					UReflectionCaptureComponent::UpdateReflectionCaptureContents(PlayWorld);
+					// So that hierarchical stats work in PIE
+					SCOPE_CYCLE_COUNTER(STAT_FrameTime);
+
+					FKismetDebugUtilities::NotifyDebuggerOfStartOfGameFrame(PieContext.World());
+
+					// tick the level
+					PieContext.World()->Tick( LEVELTICK_All, TickDeltaSeconds );
+					bAWorldTicked = true;
+					TickType = LEVELTICK_All;
+
+					if (!bFirstTick)
+					{
+						// Update sky light first because sky diffuse will be visible in reflection capture indirect specular
+						USkyLightComponent::UpdateSkyCaptureContents(PlayWorld);
+						UReflectionCaptureComponent::UpdateReflectionCaptureContents(PlayWorld);
+					}
+
+					FKismetDebugUtilities::NotifyDebuggerOfEndOfGameFrame(PieContext.World());
 				}
 
-				FKismetDebugUtilities::NotifyDebuggerOfEndOfGameFrame(PieContext.World());
-			}
 
-			UNCLOCK_CYCLES(GameCycles);
+				UNCLOCK_CYCLES(GameCycles);
 
-			// Tick the viewports.
-			if ( GameViewport != NULL )
-			{
-				GameViewport->Tick(DeltaSeconds);
+				// Tick the viewports.
+				if ( GameViewport != NULL )
+				{
+					GameViewport->Tick(TickDeltaSeconds);
+				}
 			}
 
 			// Pop the world
diff --git a/Engine/Source/Editor/UnrealEd/Private/PlayLevel.cpp b/Engine/Source/Editor/UnrealEd/Private/PlayLevel.cpp
--- a/Engine/Source/Editor/UnrealEd/Private/PlayLevel.cpp
+++ b/Engine/Source/Editor/UnrealEd/Private/PlayLevel.cpp
@@ -121,7 +121,7 @@ DEFINE_LOG_CATEGORY(LogPlayLevel);
 const static FName NAME_CategoryPIE("PIE");
 
 // Forward declare local utility functions
-FText GeneratePIEViewportWindowTitle(const EPlayNetMode InNetMode, const ERHIFeatureLevel::Type InFeatureLevel, const FRequestPlaySessionParams& InSessionParams, const int32 ClientIndex = -1);
+FText GeneratePIEViewportWindowTitle(const EPlayNetMode InNetMode, const ERHIFeatureLevel::Type InFeatureLevel, const FRequestPlaySessionParams& InSessionParams, const int32 ClientIndex, const float FixedTick);
 bool PromptMatineeClose();
 
 
@@ -1668,6 +1668,21 @@ void UEditorEngine::CreateNewPlayInEditorInstance(FRequestPlaySessionParams &InR
 		PIELoginInfo.WorldContextHandle = PieWorldContext.ContextHandle;
 		PIELoginInfo.PIEInstanceIndex = PieWorldContext.PIEInstance;
 
+		// Fixed tick setting
+		if (InRequestParams.EditorPlaySettings->ServerFixedFPS > 0 && bInDedicatedInstance)
+		{
+			PieWorldContext.PIEFixedTickSeconds = 1.f / (float)InRequestParams.EditorPlaySettings->ServerFixedFPS;
+		}
+		if (InRequestParams.EditorPlaySettings->ClientFixedFPS.Num() > 0 && !bInDedicatedInstance)
+		{
+			const int32 ClientFixedIdx = PlayInEditorSessionInfo->NumClientInstancesCreated % InRequestParams.EditorPlaySettings->ClientFixedFPS.Num();
+			const int32 DesiredFPS = InRequestParams.EditorPlaySettings->ClientFixedFPS[ClientFixedIdx];
+			if (DesiredFPS > 0)
+			{
+				PieWorldContext.PIEFixedTickSeconds = 1.f / (float)DesiredFPS;
+			}
+		}
+
 		// We increment how many PIE instances we think are going to boot up as it is decremented
 		// in OnLoginPIEComplete_Deferred and used to check if all instances have booted up.
 		PlayInEditorSessionInfo->NumOutstandingPIELogins++;
@@ -2816,7 +2831,7 @@ UGameInstance* UEditorEngine::CreateInnerProcessPIEGameInstance(FRequestPlaySess
 	return GameInstance;
 }
 
-FText GeneratePIEViewportWindowTitle(const EPlayNetMode InNetMode, const ERHIFeatureLevel::Type InFeatureLevel, const FRequestPlaySessionParams& InSessionParams, const int32 ClientIndex)
+FText GeneratePIEViewportWindowTitle(const EPlayNetMode InNetMode, const ERHIFeatureLevel::Type InFeatureLevel, const FRequestPlaySessionParams& InSessionParams, const int32 ClientIndex, const float FixedTick)
 {
 #if PLATFORM_64BITS
 	const FString PlatformBitsString(TEXT("64"));
@@ -2854,7 +2869,17 @@ FText GeneratePIEViewportWindowTitle(const EPlayNetMode InNetMode, const ERHIFea
 		Args.Add(TEXT("XRSystemName"), FText::GetEmpty());
 	}
 
-	return FText::Format(NSLOCTEXT("UnrealEd", "PlayInEditor_WindowTitleFormat", "{GameName} Preview [NetMode: {NetMode}] ({PlatformBits}-bit/{RHIName}) {XRSystemName}"), Args);
+	if (FixedTick > 0.f)
+	{
+		int32 FixedFPS = (int32)(1.f / FixedTick);
+		Args.Add(TEXT("FixedFPS"), FText::FromString(FString::Printf(TEXT("Fixed %dfps"), FixedFPS)));
+	}
+	else
+	{
+		Args.Add(TEXT("FixedFPS"), FText::GetEmpty());
+	}
+
+	return FText::Format(NSLOCTEXT("UnrealEd", "PlayInEditor_WindowTitleFormat", "{GameName} Preview [NetMode: {NetMode}] {FixedFPS} ({PlatformBits}-bit/{RHIName}) {XRSystemName}"), Args);
 }
 
 void UEditorEngine::TransferEditorSelectionToSIEInstances()
@@ -2920,8 +2945,7 @@ TSharedRef<SPIEViewport> UEditorEngine::GeneratePIEViewportWindow(const FRequest
 	// If they haven't provided a Slate Window (common), we will create one.
 	if (!bHasCustomWindow)
 	{
-
-		FText ViewportName = GeneratePIEViewportWindowTitle(InNetMode, PreviewPlatform.GetEffectivePreviewFeatureLevel(), InSessionParams, InViewportIndex);
+		FText ViewportName = GeneratePIEViewportWindowTitle(InNetMode, PreviewPlatform.GetEffectivePreviewFeatureLevel(), InSessionParams,InViewportIndex, InWorldContext.PIEFixedTickSeconds);
 		PieWindow = SNew(SWindow)
 			.Title(ViewportName)
 			.ScreenPosition(FVector2D(WindowPosition.X, WindowPosition.Y))
diff --git a/Engine/Source/Runtime/Engine/Classes/Engine/Engine.h b/Engine/Source/Runtime/Engine/Classes/Engine/Engine.h
--- a/Engine/Source/Runtime/Engine/Classes/Engine/Engine.h
+++ b/Engine/Source/Runtime/Engine/Classes/Engine/Engine.h
@@ -371,6 +371,10 @@ struct FWorldContext
 	/** Custom description to be display in blueprint debugger UI */
 	FString CustomDescription;
 
+	// If > 0, tick this world at a fixed rate in PIE
+	float PIEFixedTickSeconds  = 0.f;
+	float PIEAccumulatedTickSeconds = 0.f;
+
 	/**************************************************************/
 
 	/** Outside pointers to CurrentWorld that should be kept in sync if current world changes  */
