Abort generation of streaming install manifest if max chunk size is too small

diff --git a/Engine/Source/Editor/UnrealEd/Private/Commandlets/AssetRegistryGenerator.cpp b/Engine/Source/Editor/UnrealEd/Private/Commandlets/AssetRegistryGenerator.cpp
--- a/Engine/Source/Editor/UnrealEd/Private/Commandlets/AssetRegistryGenerator.cpp
+++ b/Engine/Source/Editor/UnrealEd/Private/Commandlets/AssetRegistryGenerator.cpp
@@ -369,7 +369,7 @@ bool FAssetRegistryGenerator::GenerateStreamingInstallManifest(int64 InExtraFlav
 
 	// generate per-chunk pak list files
 	bool bSucceeded = true;
-	for (int32 PakchunkIndex = 0; PakchunkIndex < FinalChunkManifests.Num(); ++PakchunkIndex)
+	for (int32 PakchunkIndex = 0; PakchunkIndex < FinalChunkManifests.Num() && bSucceeded; ++PakchunkIndex)
 	{
 		const FChunkPackageSet* Manifest = FinalChunkManifests[PakchunkIndex];
 
@@ -485,8 +485,16 @@ bool FAssetRegistryGenerator::GenerateStreamingInstallManifest(int64 InExtraFlav
 				PakListFile->Serialize(TCHAR_TO_ANSI(*PakListLine), PakListLine.Len());
 			}
 
+			const bool bAddedFilesToPakList = PakListFile->Tell() > 0;
 			PakListFile->Close();
 
+			if (!bAddedFilesToPakList)
+			{
+				UE_LOG(LogAssetRegistryGenerator, Error, TEXT("Failed to add file(s) to paklist '%s', max chunk size '%d' too small"), *PakListFilename, MaxChunkSize);
+				bSucceeded = false;
+				break;
+			}
+
 			// add this pakfilelist to our master list of pakfilelists
 			FString PakChunkListLine = FString::Printf(TEXT("%s%s\r\n"), *PakChunkFilename, *PakChunkOptions);
 			PakChunkListFile->Serialize(TCHAR_TO_ANSI(*PakChunkListLine), PakChunkListLine.Len());
