When creating a blueprint from selected actors with a class that has an existing root component, attach the components to that root
#rnx

diff --git a/Engine/Source/Editor/UnrealEd/Private/Kismet2/Kismet2.cpp b/Engine/Source/Editor/UnrealEd/Private/Kismet2/Kismet2.cpp
--- a/Engine/Source/Editor/UnrealEd/Private/Kismet2/Kismet2.cpp
+++ b/Engine/Source/Editor/UnrealEd/Private/Kismet2/Kismet2.cpp
@@ -1066,6 +1066,12 @@ void FKismetEditorUtilities::AddComponentsToBlueprint(UBlueprint* Blueprint, con
 {
 	USimpleConstructionScript* SCS = Blueprint->SimpleConstructionScript;
 
+	USceneComponent* RootTemplate = nullptr;
+	if (OptionalNewRootNode == nullptr)
+	{
+		RootTemplate = SCS->GetSceneRootComponentTemplate();
+	}
+
 	TArray<UBlueprint*> ParentBPStack;
 	UBlueprint::GetBlueprintHierarchyFromClass(Blueprint->GeneratedClass, ParentBPStack);
 
@@ -1177,7 +1183,7 @@ void FKismetEditorUtilities::AddComponentsToBlueprint(UBlueprint* Blueprint, con
 
 	TFunction<void (USceneComponent*)> ConstructSceneComponentNodes;
 	// Fill the SceneComponentNodes array
-	ConstructSceneComponentNodes = [&SceneComponentsMap, &SceneComponentNodes, &ConstructSceneComponentNodes] (USceneComponent* SceneComponent)
+	ConstructSceneComponentNodes = [&SceneComponentsMap, &SceneComponentNodes, &ConstructSceneComponentNodes, RootTemplate] (USceneComponent* SceneComponent)
 	{
 		// The scene component should always be present in the map
 		if (*SceneComponentsMap.Find(SceneComponent))
@@ -1199,6 +1205,11 @@ void FKismetEditorUtilities::AddComponentsToBlueprint(UBlueprint* Blueprint, con
 				Parent = Parent->GetAttachParent();
 			}
 
+			if (Parent == nullptr && RootTemplate)
+			{
+				Parent = RootTemplate;
+			}
+
 			SceneComponentNodes.Emplace(SceneComponent, Parent);
 			SceneComponentsMap.Add(SceneComponent, true);
 		}
@@ -1261,6 +1272,10 @@ void FKismetEditorUtilities::AddComponentsToBlueprint(UBlueprint* Blueprint, con
 			}
 			else
 			{
+				if (RootTemplate)
+				{
+					SCSNode->SetParent(RootTemplate);
+				}
 				SCS->AddNode(SCSNode);
 			}
 		}
@@ -1508,11 +1523,22 @@ void CreateBlueprintFromActors_Internal(UBlueprint* Blueprint, const TArray<AAct
 		}
 	}
 
+	bool bAdjustTopLevelNodes = false;
+
 	// If there is not one unique root actor then create a new scene component to serve as the shared root node
 	if (RootActors.Num() != 1)
 	{
-		AssemblyProps.RootNodeOverride = SCS->CreateNode(USceneComponent::StaticClass(), TEXT("SharedRoot"));
-		SCS->AddNode(AssemblyProps.RootNodeOverride);
+		// If the new blueprint already has a root node, then we will use it, but that will require adjusting
+		// the positions of the SCS's root nodes as an inherited root is not in this blueprint's SCS
+		if (SCS->GetSceneRootComponentTemplate() == nullptr)
+		{
+			AssemblyProps.RootNodeOverride = SCS->CreateNode(USceneComponent::StaticClass(), TEXT("SharedRoot"));
+			SCS->AddNode(AssemblyProps.RootNodeOverride);
+		}
+		else
+		{
+			bAdjustTopLevelNodes = true;
+		}
 	}
 
 	AssembleBlueprintFunc(AssemblyProps);
@@ -1546,13 +1572,22 @@ void CreateBlueprintFromActors_Internal(UBlueprint* Blueprint, const TArray<AAct
 		{
 			if (USceneComponent* TestRoot = Cast<USceneComponent>(TopLevelNode->ComponentTemplate))
 			{
-				for (USCS_Node* ChildNode : TopLevelNode->GetChildNodes())
+				if (bAdjustTopLevelNodes)
 				{
-					if (USceneComponent* ChildComponent = Cast<USceneComponent>(ChildNode->ComponentTemplate))
+					//The relative transform for those component was converted into the world space
+					const FTransform NewRelativeTransform = TestRoot->GetRelativeTransform().GetRelativeTransform(NewActorTransform);
+					TestRoot->SetRelativeTransform(NewRelativeTransform);
+				}
+				else
+				{
+					for (USCS_Node* ChildNode : TopLevelNode->GetChildNodes())
 					{
-						//The relative transform for those component was converted into the world space
-						const FTransform NewRelativeTransform = ChildComponent->GetRelativeTransform().GetRelativeTransform(NewActorTransform);
-						ChildComponent->SetRelativeTransform(NewRelativeTransform);
+						if (USceneComponent* ChildComponent = Cast<USceneComponent>(ChildNode->ComponentTemplate))
+						{
+							//The relative transform for those component was converted into the world space
+							const FTransform NewRelativeTransform = ChildComponent->GetRelativeTransform().GetRelativeTransform(NewActorTransform);
+							ChildComponent->SetRelativeTransform(NewRelativeTransform);
+						}
 					}
 				}
 			}
