Fix component changing position when converted to a blueprint via child actors if the base class already has a root component.
#rnx

diff --git a/Engine/Source/Editor/UnrealEd/Private/Kismet2/Kismet2.cpp b/Engine/Source/Editor/UnrealEd/Private/Kismet2/Kismet2.cpp
--- a/Engine/Source/Editor/UnrealEd/Private/Kismet2/Kismet2.cpp
+++ b/Engine/Source/Editor/UnrealEd/Private/Kismet2/Kismet2.cpp
@@ -1523,30 +1523,50 @@ void CreateBlueprintFromActors_Internal(UBlueprint* Blueprint, const TArray<AAct
 		}
 	}
 
-	bool bAdjustTopLevelNodes = false;
+	FTransform NewActorTransform = FTransform::Identity;
+	bool bRepositionTopLevelNodes = false;
 
-	// If there is not one unique root actor then create a new scene component to serve as the shared root node
-	if (RootActors.Num() != 1)
+	// If the new blueprint already has a root node, then we will use it, but that will require adjusting
+	// the positions of the SCS's root nodes as an inherited root is not in this blueprint's SCS
+	if (SCS->GetSceneRootComponentTemplate() == nullptr)
 	{
-		// If the new blueprint already has a root node, then we will use it, but that will require adjusting
-		// the positions of the SCS's root nodes as an inherited root is not in this blueprint's SCS
-		if (SCS->GetSceneRootComponentTemplate() == nullptr)
+		// If there is not one unique root actor then create a new scene component to serve as the shared root node
+		if (RootActors.Num() != 1)
 		{
 			AssemblyProps.RootNodeOverride = SCS->CreateNode(USceneComponent::StaticClass(), TEXT("SharedRoot"));
 			SCS->AddNode(AssemblyProps.RootNodeOverride);
 		}
-		else
-		{
-			bAdjustTopLevelNodes = true;
-		}
+	}
+	else
+	{
+		bRepositionTopLevelNodes = true;
 	}
 
 	AssembleBlueprintFunc(AssemblyProps);
 
-	FTransform NewActorTransform = FTransform::Identity;
+	// Reposition nodes to recenter them around the new pivot
+	auto RepositionNodes = [&NewActorTransform](const TArray<USCS_Node*>& Nodes)
+	{
+		for (USCS_Node* Node : Nodes)
+		{
+			if (USceneComponent* NodeTemplate = Cast<USceneComponent>(Node->ComponentTemplate))
+			{
+				//The relative transform for those component was converted into the world space
+				const FTransform NewRelativeTransform = NodeTemplate->GetRelativeTransform().GetRelativeTransform(NewActorTransform);
+				NodeTemplate->SetRelativeLocation_Direct(NewRelativeTransform.GetLocation());
+				NodeTemplate->SetRelativeRotation_Direct(NewRelativeTransform.GetRotation().Rotator());
+				NodeTemplate->SetRelativeScale3D_Direct(NewRelativeTransform.GetScale3D());
+			}
+		}
+	};
+
 	if (RootActors.Num() == 1)
 	{
 		NewActorTransform = RootActors[0]->GetTransform();
+		if (bRepositionTopLevelNodes)
+		{
+			RepositionNodes(SCS->GetRootNodes());
+		}
 	}
 	else if (RootActors.Num() > 1)
 	{
@@ -1567,28 +1587,17 @@ void CreateBlueprintFromActors_Internal(UBlueprint* Blueprint, const TArray<AAct
 			NewActorTransform.SetTranslation(AverageLocation);
 		}
 
-		// Reposition all of the children of the root node to recenter them around the new pivot
-		for (USCS_Node* TopLevelNode : SCS->GetRootNodes())
+		if (bRepositionTopLevelNodes)
 		{
-			if (USceneComponent* TestRoot = Cast<USceneComponent>(TopLevelNode->ComponentTemplate))
+			RepositionNodes(SCS->GetRootNodes());
+		}
+		else
+		{
+			for (USCS_Node* TopLevelNode : SCS->GetRootNodes())
 			{
-				if (bAdjustTopLevelNodes)
+				if (USceneComponent* TestRoot = Cast<USceneComponent>(TopLevelNode->ComponentTemplate))
 				{
-					//The relative transform for those component was converted into the world space
-					const FTransform NewRelativeTransform = TestRoot->GetRelativeTransform().GetRelativeTransform(NewActorTransform);
-					TestRoot->SetRelativeTransform(NewRelativeTransform);
-				}
-				else
-				{
-					for (USCS_Node* ChildNode : TopLevelNode->GetChildNodes())
-					{
-						if (USceneComponent* ChildComponent = Cast<USceneComponent>(ChildNode->ComponentTemplate))
-						{
-							//The relative transform for those component was converted into the world space
-							const FTransform NewRelativeTransform = ChildComponent->GetRelativeTransform().GetRelativeTransform(NewActorTransform);
-							ChildComponent->SetRelativeTransform(NewRelativeTransform);
-						}
-					}
+					RepositionNodes(TopLevelNode->GetChildNodes());
 				}
 			}
 		}
