Fix assert when max chunk size is too small.

diff --git a/Engine/Source/Editor/UnrealEd/Private/Commandlets/AssetRegistryGenerator.cpp b/Engine/Source/Editor/UnrealEd/Private/Commandlets/AssetRegistryGenerator.cpp
--- a/Engine/Source/Editor/UnrealEd/Private/Commandlets/AssetRegistryGenerator.cpp
+++ b/Engine/Source/Editor/UnrealEd/Private/Commandlets/AssetRegistryGenerator.cpp
@@ -488,7 +488,7 @@ bool FAssetRegistryGenerator::GenerateStreamingInstallManifest(int64 InExtraFlav
 			const bool bAddedFilesToPakList = PakListFile->Tell() > 0;
 			PakListFile->Close();
 
-			if (!bAddedFilesToPakList)
+			if (!bFinishedAllFiles && !bAddedFilesToPakList)
 			{
 				UE_LOG(LogAssetRegistryGenerator, Error, TEXT("Failed to add file(s) to paklist '%s', max chunk size '%d' too small"), *PakListFilename, MaxChunkSize);
 				bSucceeded = false;
