No filtering on anim compression toast notification. Reduce confusion about when anim compression is running

diff --git a/Engine/Source/Editor/UnrealEd/Private/Animation/AsyncAnimCompressionUI.cpp b/Engine/Source/Editor/UnrealEd/Private/Animation/AsyncAnimCompressionUI.cpp
--- a/Engine/Source/Editor/UnrealEd/Private/Animation/AsyncAnimCompressionUI.cpp
+++ b/Engine/Source/Editor/UnrealEd/Private/Animation/AsyncAnimCompressionUI.cpp
@@ -20,7 +20,7 @@ FAnimCompressionNotificationImpl GAnimCompressionNotification;
 
 bool FAnimCompressionNotificationImpl::ShouldShowNotification(const bool bIsNotificationAlreadyActive) const
 {
-	const uint32 RemainingJobsThreshold = bIsNotificationAlreadyActive ? 0 : 10;
+	const uint32 RemainingJobsThreshold = 0;
 	const uint32 ActiveJobs = GAsyncCompressedAnimationsTracker ? GAsyncCompressedAnimationsTracker->GetNumRemainingJobs() : 0;
 
 	return ActiveJobs > RemainingJobsThreshold;
