Fix a potential memory stomp when closing the editor.

It was possible that the FLevelEditorViewportClient would be destroyed before the timer callback issued in ::ReceivedFocus has returned. Which would then write over the old memory location of bReceivedFocusRecently which could in theory have been reallocated.

To fix it we now store the FTimerHandle of the event and clear the timer during the destructor of FLevelEditorViewportClient.

#rnx

diff --git a/Engine/Source/Editor/UnrealEd/Private/LevelEditorViewport.cpp b/Engine/Source/Editor/UnrealEd/Private/LevelEditorViewport.cpp
--- a/Engine/Source/Editor/UnrealEd/Private/LevelEditorViewport.cpp
+++ b/Engine/Source/Editor/UnrealEd/Private/LevelEditorViewport.cpp
@@ -1826,10 +1826,14 @@ FLevelEditorViewportClient::~FLevelEditorViewportClient()
 		Layers->RemoveViewFromActorViewVisibility(this);
 
 		GEditor->RemoveLevelViewportClients(this);
+		
+		if (FocusTimerHandle.IsValid())
+		{
+			GEditor->GetTimerManager()->ClearTimer(FocusTimerHandle);
+		}
 
 		GetMutableDefault<ULevelEditorViewportSettings>()->OnSettingChanged().RemoveAll(this);
 
-
 		RemoveReferenceToWorldContext(GEditor->GetEditorWorldContext());
 	}
 
@@ -2117,13 +2121,13 @@ void FLevelEditorViewportClient::ReceivedFocus(FViewport* InViewport)
 		bReceivedFocusRecently = true;
 
 		// A few frames can pass between receiving focus and processing a click, so we use a timer to track whether we have recently received focus.
-		FTimerHandle DummyHandle;
 		FTimerDelegate ResetFocusReceivedTimer;
 		ResetFocusReceivedTimer.BindLambda([&] ()
 		{
 			bReceivedFocusRecently = false;
+			FocusTimerHandle.Invalidate(); // The timer will only execute once, so we can invalidate now.
 		});
-		GEditor->GetTimerManager()->SetTimer(DummyHandle, ResetFocusReceivedTimer, 0.1f, false);
+		GEditor->GetTimerManager()->SetTimer(FocusTimerHandle, ResetFocusReceivedTimer, 0.1f, false);	
 	}
 
 	FEditorViewportClient::ReceivedFocus(InViewport);
diff --git a/Engine/Source/Editor/UnrealEd/Public/LevelEditorViewport.h b/Engine/Source/Editor/UnrealEd/Public/LevelEditorViewport.h
--- a/Engine/Source/Editor/UnrealEd/Public/LevelEditorViewport.h
+++ b/Engine/Source/Editor/UnrealEd/Public/LevelEditorViewport.h
@@ -792,5 +792,8 @@ private:
 	/** Stores the previous frame's value of bEditorCameraCut in order to reset it back to false on the next frame */
 	bool					bWasEditorCameraCut;
 
-	bool bApplyCameraSpeedScaleByDistance;
+	bool					bApplyCameraSpeedScaleByDistance;
+
+	/** Handle to a timer event raised in ::ReceivedFocus*/
+	FTimerHandle			FocusTimerHandle;
 };
