Fix incremental

#rnx

diff --git a/Engine/Source/Editor/UnrealEd/Public/SEditorViewportToolBarButton.h b/Engine/Source/Editor/UnrealEd/Public/SEditorViewportToolBarButton.h
--- a/Engine/Source/Editor/UnrealEd/Public/SEditorViewportToolBarButton.h
+++ b/Engine/Source/Editor/UnrealEd/Public/SEditorViewportToolBarButton.h
@@ -9,6 +9,7 @@
 #include "Widgets/SCompoundWidget.h"
 #include "Framework/SlateDelegates.h"
 #include "Framework/Commands/UICommandInfo.h"
+#include "EditorStyleSet.h"
 
 struct FSlateBrush;
 enum class ECheckBoxState : uint8;
