Fix incremental take 2

#rnx

diff --git a/Engine/Source/Editor/UnrealEd/Public/SEditorViewportToolBarButton.h b/Engine/Source/Editor/UnrealEd/Public/SEditorViewportToolBarButton.h
--- a/Engine/Source/Editor/UnrealEd/Public/SEditorViewportToolBarButton.h
+++ b/Engine/Source/Editor/UnrealEd/Public/SEditorViewportToolBarButton.h
@@ -10,6 +10,7 @@
 #include "Framework/SlateDelegates.h"
 #include "Framework/Commands/UICommandInfo.h"
 #include "EditorStyleSet.h"
+#include "Styling/SlateTypes.h"
 
 struct FSlateBrush;
 enum class ECheckBoxState : uint8;
