Even more incremental fixes

#rnx

diff --git a/Engine/Source/Editor/UnrealEd/Public/SEditorViewportToolBarButton.h b/Engine/Source/Editor/UnrealEd/Public/SEditorViewportToolBarButton.h
--- a/Engine/Source/Editor/UnrealEd/Public/SEditorViewportToolBarButton.h
+++ b/Engine/Source/Editor/UnrealEd/Public/SEditorViewportToolBarButton.h
@@ -11,6 +11,7 @@
 #include "Framework/Commands/UICommandInfo.h"
 #include "EditorStyleSet.h"
 #include "Styling/SlateTypes.h"
+#include "Styling/SlateWidgetStyleAsset.h"
 
 struct FSlateBrush;
 enum class ECheckBoxState : uint8;
