Ok this time incremental is fixed

#rnx

diff --git a/Engine/Source/Editor/UnrealEd/Public/SEditorViewportToolBarMenu.h b/Engine/Source/Editor/UnrealEd/Public/SEditorViewportToolBarMenu.h
--- a/Engine/Source/Editor/UnrealEd/Public/SEditorViewportToolBarMenu.h
+++ b/Engine/Source/Editor/UnrealEd/Public/SEditorViewportToolBarMenu.h
@@ -13,7 +13,7 @@
 #include "Framework/SlateDelegates.h"
 #include "Styling/SlateTypes.h"
 #include "EditorStyleSet.h"
-
+#include "Styling/SlateWidgetStyleAsset.h"
 
 class SMenuAnchor;
 struct FSlateBrush;
