Back out changelist 10951469 for now. This change made it so renaming map files didnt properly save a redirector at the old location in the umap file, it saved it in a new uasset file (the extension was not right)

diff --git a/Engine/Source/Editor/UnrealEd/Private/FileHelpers.cpp b/Engine/Source/Editor/UnrealEd/Private/FileHelpers.cpp
--- a/Engine/Source/Editor/UnrealEd/Private/FileHelpers.cpp
+++ b/Engine/Source/Editor/UnrealEd/Private/FileHelpers.cpp
@@ -206,31 +206,7 @@ namespace FileDialogHelpers
 	}
 }
 
-namespace FileHelperPackageUtil
-{
-	/**
-	 * DoesPackageExist helper that rely on the AssetRegistry to validate if a package exists instead of hitting the FS
-	 * Fallback to the FS if the asset registry initial scan isn't done or we aren't in Editor
-	 */
-	bool DoesPackageExist(UPackage* Package, FString* OutFilename = nullptr)
-	{
-		// Test using asset registry to figure out existence
-		IAssetRegistry& AssetRegistry = FAssetRegistryModule::GetRegistry();
-		if (!AssetRegistry.IsLoadingAssets() || !GIsEditor)
-		{
-			TArray<FAssetData> Data;
-			FAssetRegistryModule::GetRegistry().GetAssetsByPackageName(Package->GetFName(), Data, true);
 
-			if (Data.Num() > 0 && OutFilename)
-			{
-				*OutFilename = FPackageName::LongPackageNameToFilename(Package->GetName(), Package->ContainsMap() ? FPackageName::GetMapPackageExtension() : FPackageName::GetAssetPackageExtension());
-			}
-
-			return Data.Num() > 0;
-		}
-		return FPackageName::DoesPackageExist(Package->GetName(), nullptr, OutFilename);
-	}
-}
 
 /**
  * Queries the user if they want to quit out of interpolation editing before save.
@@ -576,7 +552,7 @@ static bool SaveWorld(UWorld* World,
 	FString	CleanFilename;
 
 	// Does a filename already exist for this package?
-	const bool bPackageExists = FileHelperPackageUtil::DoesPackageExist(Package, &ExistingFilename );
+	const bool bPackageExists = FPackageName::DoesPackageExist( PackageName, NULL, &ExistingFilename );
 
 	if ( ForceFilename )
 	{
@@ -1396,7 +1372,7 @@ bool FEditorFileUtils::AddCheckoutPackageItems(bool bCheckDirty, TArray<UPackage
 				}
 				
 				FString Filename;
-				if (FileHelperPackageUtil::DoesPackageExist(Package, &Filename))
+				if (FPackageName::DoesPackageExist(Package->GetName(), NULL, &Filename))
 				{
 					if (IFileManager::Get().IsReadOnly(*Filename))
 					{
@@ -1450,7 +1426,7 @@ bool FEditorFileUtils::AddCheckoutPackageItems(bool bCheckDirty, TArray<UPackage
 		bool bPkgReadOnly = true;
 		bool bCareAboutReadOnly = SourceControlProvider.UsesLocalReadOnlyState();
 		// Find the filename for this package
-		bool bFoundFile = FileHelperPackageUtil::DoesPackageExist(CurPackage, &Filename);
+		bool bFoundFile = FPackageName::DoesPackageExist(CurPackage->GetName(), NULL, &Filename);
 		if (bFoundFile)
 		{
 			// determine if the package file is read only
@@ -1682,7 +1658,7 @@ bool FEditorFileUtils::PromptToCheckoutPackages(bool bCheckDirty, const TArray<U
 					UPackage* PackageToMakeWritable = *PkgsToMakeWritableIter;
 					FString Filename;
 
-					bool bFoundFile = FileHelperPackageUtil::DoesPackageExist( PackageToMakeWritable, &Filename );
+					bool bFoundFile = FPackageName::DoesPackageExist( PackageToMakeWritable->GetName(), NULL, &Filename );
 					if( bFoundFile )
 					{
 						// If we're ignoring the package due to the user ignoring it for saving, remove it from the ignore list
@@ -2908,7 +2884,7 @@ static InternalSavePackageResult InternalSavePackage(UPackage* PackageToSave, bo
 	if( bIsValidPath )
 	{
 		FString ExistingFilename;
-		const bool bPackageAlreadyExists = FileHelperPackageUtil::DoesPackageExist(PackageToSave, &ExistingFilename);
+		const bool bPackageAlreadyExists = FPackageName::DoesPackageExist(PackageName, NULL, &ExistingFilename);
 		if (!bPackageAlreadyExists)
 		{
 			// Construct a filename from long package name.
@@ -3242,7 +3218,7 @@ static bool InternalSavePackagesFast(const TArray<UPackage*>& PackagesToSave, bo
 
 		// Check if a file exists for this package
 		FString Filename;
-		bool bFoundFile = FileHelperPackageUtil::DoesPackageExist(CurPackage, &Filename);
+		bool bFoundFile = FPackageName::DoesPackageExist(CurPackage->GetName(), NULL, &Filename);
 		if (bFoundFile)
 		{
 			// determine if the package file is read only
