Fix thumbnail related ensure while scrolling in content browser

diff --git a/Engine/Source/Editor/UnrealEd/Private/UnrealEdMisc.cpp b/Engine/Source/Editor/UnrealEd/Private/UnrealEdMisc.cpp
--- a/Engine/Source/Editor/UnrealEd/Private/UnrealEdMisc.cpp
+++ b/Engine/Source/Editor/UnrealEd/Private/UnrealEdMisc.cpp
@@ -94,7 +94,7 @@ DEFINE_LOG_CATEGORY_STATIC(LogUnrealEdMisc, Log, All);
 
 bool FTickableEditorObject::bCollectionIntact = true;
 bool FTickableEditorObject::bIsTickingObjects = false;
-bool FTickableEditorObject::bIsInObjectsTick = false;
+FTickableObjectBase* FTickableEditorObject::ObjectBeingTicked = nullptr;
 
 
 namespace
diff --git a/Engine/Source/Editor/UnrealEd/Public/TickableEditorObject.h b/Engine/Source/Editor/UnrealEd/Public/TickableEditorObject.h
--- a/Engine/Source/Editor/UnrealEd/Public/TickableEditorObject.h
+++ b/Engine/Source/Editor/UnrealEd/Public/TickableEditorObject.h
@@ -40,9 +40,9 @@ public:
 				{
 					if ((TickableEntry.TickType == ETickableTickType::Always) || TickableObject->IsTickable())
 					{
-						bIsInObjectsTick = true;
+						ObjectBeingTicked = TickableObject;
 						TickableObject->Tick(DeltaSeconds);
-						bIsInObjectsTick = false;
+						ObjectBeingTicked = nullptr;
 					}
 
 					// In case it was removed during tick
@@ -78,7 +78,7 @@ public:
 	/** Removes this instance from the static array of tickable objects. */
 	virtual ~FTickableEditorObject()
 	{
-		ensureMsgf(!bIsInObjectsTick, TEXT("Detected possible memory stomp. We are in the Tickable objects Tick function but hit its deconstructor, the 'this' pointer for the Object will now be invalid"));
+		ensureMsgf(ObjectBeingTicked != this, TEXT("Detected possible memory stomp. We are in the Tickable objects Tick function but hit its deconstructor, the 'this' pointer for the Object will now be invalid"));
 
 		ensure(IsInGameThread() || IsInAsyncLoadingThread());
 		if (bCollectionIntact && GetPendingTickableObjects().Remove(this) == 0)
@@ -111,8 +111,9 @@ private:
 	/** True if currently ticking of tickable editor objects. */
 	static bool bIsTickingObjects;
 
-	/** True if we are in the Tick function for an editor tickable object */
-	static bool bIsInObjectsTick;
+	/** Set if we are in the Tick function for an editor tickable object */
+	static FTickableObjectBase* ObjectBeingTicked;
+
 
 	static TArray<FTickableObjectEntry>& GetTickableObjects();
 
