Fix crash when converting an actor attached to another actor to a blueprint for which the class already has a root component

When converting attached actors to a blueprint, the new blueprint instance now correctly ends up parented to the first common ancestor

When converting actors with other actors attached to them, the attached actors now correctly end up parented to the new blueprint instance

diff --git a/Engine/Source/Editor/UnrealEd/Private/Kismet2/Kismet2.cpp b/Engine/Source/Editor/UnrealEd/Private/Kismet2/Kismet2.cpp
--- a/Engine/Source/Editor/UnrealEd/Private/Kismet2/Kismet2.cpp
+++ b/Engine/Source/Editor/UnrealEd/Private/Kismet2/Kismet2.cpp
@@ -1264,7 +1264,7 @@ void FKismetEditorUtilities::AddComponentsToBlueprint(UBlueprint* Blueprint, con
 		}
 
 		AActor* Actor = SceneComponent->GetOwner();
-		if (((Actor == nullptr) || (SceneComponent == Actor->GetRootComponent())) && (SceneComponent->GetAttachParent() == nullptr || FirstAttachParent == nullptr))
+		if (((Actor == nullptr) || (SceneComponent == Actor->GetRootComponent())) && (SceneComponent->GetAttachParent() == nullptr || FirstAttachParent == RootTemplate))
 		{
 			if (OptionalNewRootNode != nullptr)
 			{
@@ -1447,7 +1447,7 @@ UBlueprint* FKismetEditorUtilities::CreateBlueprintFromActor(const FName Bluepri
 				FVector Location = Actor->GetActorLocation();
 				FRotator Rotator = Actor->GetActorRotation();
 
-				AActor* NewActor = CreateBlueprintInstanceFromSelection(NewBlueprint, Actors, Location, Rotator);
+				AActor* NewActor = CreateBlueprintInstanceFromSelection(NewBlueprint, Actors, Location, Rotator, Actor->GetAttachParentActor());
 				if (NewActor)
 				{
 					NewActor->SetActorScale3D(Actor->GetActorScale3D());
@@ -1618,10 +1618,69 @@ void CreateBlueprintFromActors_Internal(UBlueprint* Blueprint, const TArray<AAct
 		FVector Location = NewActorTransform.GetLocation();
 		FRotator Rotator = NewActorTransform.Rotator();
 
-		AActor* NewActor = FKismetEditorUtilities::CreateBlueprintInstanceFromSelection(Blueprint, Actors, Location, Rotator);
+		AActor* NewAttachParent = nullptr;
+		if (RootActors.Num() == 1)
+		{
+			NewAttachParent = RootActors[0]->GetAttachParentActor();
+		}
+		else if (RootActors.Num() > 1)
+		{
+			TArray<AActor*> PossibleCommonAttachParents;
+
+			// Using the first actor, gather the hierarchy of actors it is attached to starting with the direct attachment
+			AActor* AttachParent = RootActors[0]->GetAttachParentActor();
+			while (AttachParent)
+			{
+				PossibleCommonAttachParents.Add(AttachParent);
+				AttachParent = AttachParent->GetAttachParentActor();
+			}
+
+			// For the remaining root actors, evaluate which, if any, of the first actor's parents they are also attached to
+			// Whatever is left in the 0 index of the possible list is the first common ancestor
+			for (int32 RootActorIndex = 1; PossibleCommonAttachParents.Num() > 0 && RootActorIndex < RootActors.Num(); ++RootActorIndex)
+			{
+				AActor* RootActor = RootActors[RootActorIndex];
+				for (int32 PossibleIndex = PossibleCommonAttachParents.Num() - 1; PossibleIndex >= 0; --PossibleIndex)
+				{
+					if (!RootActor->IsAttachedTo(PossibleCommonAttachParents[PossibleIndex]))
+					{
+						// If we're not attached to a given actor, then we can't possibly also be attached to its children, so clear all of them out
+						PossibleCommonAttachParents.RemoveAt(0, PossibleIndex + 1, false);
+						break;
+					}
+				}
+			}
+
+			if (PossibleCommonAttachParents.Num() > 0)
+			{
+				NewAttachParent = PossibleCommonAttachParents[0];
+			}
+		}
+
+		// Get all the actors attached directly attached to actors we're about to destroy and replace and attach them to the new actor
+		TArray<AActor*> AttachedActors;
+		for (AActor* Actor : Actors)
+		{
+			Actor->GetAttachedActors(AttachedActors, false);
+		}
+		for (int32 Index = AttachedActors.Num() - 1; Index >= 0; --Index)
+		{
+			// Remove attached actors that are also in the set of actors being converted to blueprint
+			if (Actors.Contains(AttachedActors[Index]))
+			{
+				AttachedActors.RemoveAtSwap(Index, 1, false);
+			}
+		}
+
+		AActor* NewActor = FKismetEditorUtilities::CreateBlueprintInstanceFromSelection(Blueprint, Actors, Location, Rotator, NewAttachParent);
 		if (NewActor)
 		{
 			NewActor->SetActorScale3D(NewActorTransform.GetScale3D());
+
+			for (AActor* AttachedActor : AttachedActors)
+			{
+				AttachedActor->AttachToActor(NewActor, FAttachmentTransformRules::KeepWorldTransform);
+			}
 		}
 	}
 
@@ -1643,6 +1702,11 @@ UBlueprint* FKismetEditorUtilities::CreateBlueprintFromActors(const FName Bluepr
 			UChildActorComponent* CAC = NewObject<UChildActorComponent>(GetTransientPackage(), ComponentName, RF_ArchetypeObject);
 			CAC->SetChildActorClass(Actor->GetClass(), Actor);
 
+			if (USceneComponent* ActorRootComp = Actor->GetRootComponent())
+			{
+				CAC->SetMobility(ActorRootComp->Mobility);
+			}
+
 			// Clear any properties that can't be on the template
 			if (USceneComponent* RootComponent = CAC->GetChildActorTemplate()->GetRootComponent())
 			{
@@ -1673,7 +1737,7 @@ UBlueprint* FKismetEditorUtilities::CreateBlueprintFromActors(const FName Bluepr
 			RecursiveCreateChildActorTemplates(Actor, nullptr);
 		}
 
-		FKismetEditorUtilities::AddComponentsToBlueprint(AssemblyProps.Blueprint, ChildActorComponents, FKismetEditorUtilities::EAddComponentToBPHarvestMode::Harvest_UseComponentName, AssemblyProps.RootNodeOverride);
+		FKismetEditorUtilities::AddComponentsToBlueprint(AssemblyProps.Blueprint, ChildActorComponents, FKismetEditorUtilities::EAddComponentToBPHarvestMode::Harvest_UseComponentName, AssemblyProps.RootNodeOverride, /*bKeepMobility*/ true);
 
 		// Since the names we create are well defined relative to the SCS but created in the transient package, we could end up reusing objects
 		// unless we rename these temporary components out of the way
@@ -1774,7 +1838,7 @@ UBlueprint* FKismetEditorUtilities::HarvestBlueprintFromActors(const FName Bluep
 	return Blueprint;
 }
 
-AActor* FKismetEditorUtilities::CreateBlueprintInstanceFromSelection(UBlueprint* Blueprint, const TArray<AActor*>& SelectedActors, const FVector& Location, const FRotator& Rotator)
+AActor* FKismetEditorUtilities::CreateBlueprintInstanceFromSelection(UBlueprint* Blueprint, const TArray<AActor*>& SelectedActors, const FVector& Location, const FRotator& Rotator, AActor* AttachParent)
 {
 	check (SelectedActors.Num() > 0 );
 
@@ -1802,6 +1866,12 @@ AActor* FKismetEditorUtilities::CreateBlueprintInstanceFromSelection(UBlueprint*
 	AActor* NewActor = World->SpawnActor(Blueprint->GeneratedClass, &Location, &Rotator);
 	Layers->InitializeNewActorLayers(NewActor);
 	FActorLabelUtilities::SetActorLabelUnique(NewActor, Blueprint->GetName());
+
+	if (AttachParent)
+	{
+		NewActor->AttachToActor(AttachParent, FAttachmentTransformRules::KeepWorldTransform);
+	}
+
 	// Quietly ensure that no components are selected
 	USelection* ComponentSelection = GEditor->GetSelectedComponents();
 	ComponentSelection->BeginBatchSelectOperation();
diff --git a/Engine/Source/Editor/UnrealEd/Public/Kismet2/KismetEditorUtilities.h b/Engine/Source/Editor/UnrealEd/Public/Kismet2/KismetEditorUtilities.h
--- a/Engine/Source/Editor/UnrealEd/Public/Kismet2/KismetEditorUtilities.h
+++ b/Engine/Source/Editor/UnrealEd/Public/Kismet2/KismetEditorUtilities.h
@@ -230,12 +230,13 @@ public:
 
 	/** 
 	 * Creates a new blueprint instance and replaces the provided actor list with the new actor
-	 * @param Blueprint			The blueprint class to create an actor instance from
-	 * @param SelectedActors	The list of currently selected actors in the editor
-	 * @param Location			The location of the newly created actor
-	 * @param Rotator			The rotation of the newly created actor
+	 * @param Blueprint             The blueprint class to create an actor instance from
+	 * @param SelectedActors        The list of currently selected actors in the editor
+	 * @param Location              The location of the newly created actor
+	 * @param Rotator               The rotation of the newly created actor
+	 * @param AttachParent          The actor the newly created instance should be attached to if any
 	 */
-	static AActor* CreateBlueprintInstanceFromSelection(class UBlueprint* Blueprint, const TArray<AActor*>& SelectedActors, const FVector& Location, const FRotator& Rotator);
+	static AActor* CreateBlueprintInstanceFromSelection(class UBlueprint* Blueprint, const TArray<AActor*>& SelectedActors, const FVector& Location, const FRotator& Rotator, AActor* AttachParent = nullptr);
 
 	/** 
 	 * Create a new Blueprint from the supplied base class. Pops up window to let user select location and name.
