Add lower bound protection against invalid parameters in convex hull decomposition tool

#rnx

diff --git a/Engine/Source/Editor/UnrealEd/Private/ConvexDecompTool.cpp b/Engine/Source/Editor/UnrealEd/Private/ConvexDecompTool.cpp
--- a/Engine/Source/Editor/UnrealEd/Private/ConvexDecompTool.cpp
+++ b/Engine/Source/Editor/UnrealEd/Private/ConvexDecompTool.cpp
@@ -132,8 +132,8 @@ static void InitParameters(IVHACD::Parameters &VHACD_Params, uint32 InHullCount,
 	VHACD_Params.m_profiler = &VHACDProfiler;
 #endif
 	VHACD_Params.m_taskRunner = &VHACDTaskRunner; // any async tasks will be run on existing UE thread pools.
-	VHACD_Params.m_resolution = InResolution; // Maximum number of voxels generated during the voxelization stage (default=100,000, range=10,000-16,000,000)
-	VHACD_Params.m_maxNumVerticesPerCH = InMaxHullVerts; // Controls the maximum number of triangles per convex-hull (default=64, range=4-1024)
+	VHACD_Params.m_resolution = FMath::Max(InResolution, (uint32)10000); // Maximum number of voxels generated during the voxelization stage (default=100,000, range=10,000-16,000,000)
+	VHACD_Params.m_maxNumVerticesPerCH = FMath::Max(InMaxHullVerts, (uint32)4); // Controls the maximum number of triangles per convex-hull (default=64, range=4-1024)
 	VHACD_Params.m_concavity = 0;		// Concavity is set to zero so that we consider any concave shape as a potential hull. The number of output hulls is better controlled by recursion depth and the max convex hulls parameter
 	VHACD_Params.m_maxConvexHulls = InHullCount;	// The number of convex hulls requested by the artists/designer
 	VHACD_Params.m_oclAcceleration = false;
