#UE4 Material and MaterialInstance editors now reports the temporary top level preview objects they spawn to the asset toolkit so external code can see them (useful in property editor code, since users are editing the preview objects, not the real ones). Also made the objects in RegenerateArrays children of the object they are spawned for, also for discovery reasons.

diff --git a/Engine/Source/Editor/MaterialEditor/Private/MaterialEditor.cpp b/Engine/Source/Editor/MaterialEditor/Private/MaterialEditor.cpp
--- a/Engine/Source/Editor/MaterialEditor/Private/MaterialEditor.cpp
+++ b/Engine/Source/Editor/MaterialEditor/Private/MaterialEditor.cpp
@@ -535,6 +535,7 @@ void FMaterialEditor::InitMaterialEditor( const EToolkitMode::Type Mode, const T
 	TArray< UObject* > ObjectsToEdit;
 	ObjectsToEdit.Add(ObjectToEdit);
 	ObjectsToEdit.Add(Material);
+	ObjectsToEdit.Add(MaterialEditorInstance);
 	FAssetEditorToolkit::InitAssetEditor( Mode, InitToolkitHost, MaterialEditorAppIdentifier, StandaloneDefaultLayout, bCreateDefaultStandaloneMenu, bCreateDefaultToolbar, ObjectsToEdit, false );
 
 	AddMenuExtender(GetMenuExtensibilityManager()->GetAllExtenders(GetToolkitCommands(), GetEditingObjects()));
diff --git a/Engine/Source/Editor/MaterialEditor/Private/MaterialInstanceEditor.cpp b/Engine/Source/Editor/MaterialEditor/Private/MaterialInstanceEditor.cpp
--- a/Engine/Source/Editor/MaterialEditor/Private/MaterialInstanceEditor.cpp
+++ b/Engine/Source/Editor/MaterialEditor/Private/MaterialInstanceEditor.cpp
@@ -450,7 +450,10 @@ void FMaterialInstanceEditor::InitMaterialInstanceEditor( const EToolkitMode::Ty
 
 	const bool bCreateDefaultStandaloneMenu = true;
 	const bool bCreateDefaultToolbar = true;
-	FAssetEditorToolkit::InitAssetEditor( Mode, InitToolkitHost, MaterialInstanceEditorAppIdentifier, StandaloneDefaultLayout, bCreateDefaultStandaloneMenu, bCreateDefaultToolbar, ObjectToEdit );
+	TArray<UObject*> ObjectsToEdit;
+	ObjectsToEdit.Add(ObjectToEdit);
+	ObjectsToEdit.Add(MaterialEditorInstance);
+	FAssetEditorToolkit::InitAssetEditor( Mode, InitToolkitHost, MaterialInstanceEditorAppIdentifier, StandaloneDefaultLayout, bCreateDefaultStandaloneMenu, bCreateDefaultToolbar, ObjectsToEdit );
 
 	AddMenuExtender(MaterialEditorModule->GetMenuExtensibilityManager()->GetAllExtenders(GetToolkitCommands(), GetEditingObjects()));
 
@@ -1020,6 +1023,22 @@ FText FMaterialInstanceEditor::GetBaseToolkitName() const
 	return LOCTEXT("AppLabel", "Material Instance Editor");
 }
 
+FText FMaterialInstanceEditor::GetToolkitName() const
+{
+	const UObject* EditingObject = GetEditingObjects()[0];
+	check(EditingObject);
+
+	return GetLabelForObject(EditingObject);
+}
+
+FText FMaterialInstanceEditor::GetToolkitToolTipText() const
+{
+	const UObject* EditingObject = GetEditingObjects()[0];
+
+	// Overridden to accommodate editing of multiple objects (original and preview materials)
+	return FAssetEditorToolkit::GetToolTipTextForObject(EditingObject);
+}
+
 FString FMaterialInstanceEditor::GetWorldCentricTabPrefix() const
 {
 	return LOCTEXT("WorldCentricTabPrefix", "Material Instance ").ToString();
diff --git a/Engine/Source/Editor/MaterialEditor/Private/MaterialInstanceEditor.h b/Engine/Source/Editor/MaterialEditor/Private/MaterialInstanceEditor.h
--- a/Engine/Source/Editor/MaterialEditor/Private/MaterialInstanceEditor.h
+++ b/Engine/Source/Editor/MaterialEditor/Private/MaterialInstanceEditor.h
@@ -58,6 +58,8 @@ public:
 	/** IToolkit interface */
 	virtual FName GetToolkitFName() const override;
 	virtual FText GetBaseToolkitName() const override;
+	virtual FText GetToolkitName() const override;
+	virtual FText GetToolkitToolTipText() const override;
 	virtual FString GetWorldCentricTabPrefix() const override;
 	virtual void InitToolMenuContext(struct FToolMenuContext& MenuContext) override;
 
diff --git a/Engine/Source/Editor/UnrealEd/Private/PreviewMaterial.cpp b/Engine/Source/Editor/UnrealEd/Private/PreviewMaterial.cpp
--- a/Engine/Source/Editor/UnrealEd/Private/PreviewMaterial.cpp
+++ b/Engine/Source/Editor/UnrealEd/Private/PreviewMaterial.cpp
@@ -357,7 +357,7 @@ void UMaterialEditorPreviewParameters::RegenerateArrays()
 		// Vector Parameters.
 		for (int32 ParameterIdx = 0; ParameterIdx < ParameterInfo.Num(); ParameterIdx++)
 		{
-			UDEditorVectorParameterValue& ParameterValue = *(NewObject<UDEditorVectorParameterValue>());
+			UDEditorVectorParameterValue& ParameterValue = *(NewObject<UDEditorVectorParameterValue>(this));
 			FName ParameterName = ParameterInfo[ParameterIdx].Name;
 			FLinearColor Value;
 			int32 SortPriority;
@@ -385,7 +385,7 @@ void UMaterialEditorPreviewParameters::RegenerateArrays()
 		ParentMaterial->GetAllScalarParameterInfo(ParameterInfo, Guids);
 		for (int32 ParameterIdx = 0; ParameterIdx < ParameterInfo.Num(); ParameterIdx++)
 		{
-			UDEditorScalarParameterValue& ParameterValue = *(NewObject<UDEditorScalarParameterValue>());
+			UDEditorScalarParameterValue& ParameterValue = *(NewObject<UDEditorScalarParameterValue>(this));
 			FName ParameterName = ParameterInfo[ParameterIdx].Name;
 			float Value;
 			int32 SortPriority;
@@ -415,7 +415,7 @@ void UMaterialEditorPreviewParameters::RegenerateArrays()
 		ParentMaterial->GetAllTextureParameterInfo(ParameterInfo, Guids);
 		for (int32 ParameterIdx = 0; ParameterIdx < ParameterInfo.Num(); ParameterIdx++)
 		{
-			UDEditorTextureParameterValue& ParameterValue = *(NewObject<UDEditorTextureParameterValue>());
+			UDEditorTextureParameterValue& ParameterValue = *(NewObject<UDEditorTextureParameterValue>(this));
 			FName ParameterName = ParameterInfo[ParameterIdx].Name;
 			UTexture* Value;
 			int32 SortPriority;
@@ -444,7 +444,7 @@ void UMaterialEditorPreviewParameters::RegenerateArrays()
 		ParentMaterial->GetAllTextureParameterInfo(ParameterInfo, Guids);
 		for (int32 ParameterIdx = 0; ParameterIdx < ParameterInfo.Num(); ParameterIdx++)
 		{
-			UDEditorRuntimeVirtualTextureParameterValue& ParameterValue = *(NewObject<UDEditorRuntimeVirtualTextureParameterValue>());
+			UDEditorRuntimeVirtualTextureParameterValue& ParameterValue = *(NewObject<UDEditorRuntimeVirtualTextureParameterValue>(this));
 			FName ParameterName = ParameterInfo[ParameterIdx].Name;
 			URuntimeVirtualTexture* Value;
 			int32 SortPriority;
@@ -472,7 +472,7 @@ void UMaterialEditorPreviewParameters::RegenerateArrays()
 		ParentMaterial->GetAllFontParameterInfo(ParameterInfo, Guids);
 		for (int32 ParameterIdx = 0; ParameterIdx < ParameterInfo.Num(); ParameterIdx++)
 		{
-			UDEditorFontParameterValue& ParameterValue = *(NewObject<UDEditorFontParameterValue>());
+			UDEditorFontParameterValue& ParameterValue = *(NewObject<UDEditorFontParameterValue>(this));
 			FName ParameterName = ParameterInfo[ParameterIdx].Name;
 			UFont* FontValue;
 			int32 FontPage;
@@ -573,7 +573,7 @@ void UMaterialEditorPreviewParameters::RegenerateArrays()
 		{
 			int32 SortPriority;
 			FStaticMaterialLayersParameter MaterialLayersParameterValue = FStaticMaterialLayersParameter(SourceStaticParameters.MaterialLayersParameters[ParameterIdx]);
-			UDEditorMaterialLayersParameterValue& ParameterValue = *(NewObject<UDEditorMaterialLayersParameterValue>());
+			UDEditorMaterialLayersParameterValue& ParameterValue = *(NewObject<UDEditorMaterialLayersParameterValue>(this));
 			ParameterValue.ParameterValue = MaterialLayersParameterValue.Value;
 			ParameterValue.bOverride = MaterialLayersParameterValue.bOverride;
 			ParameterValue.ParameterInfo = MaterialLayersParameterValue.ParameterInfo;
@@ -595,7 +595,7 @@ void UMaterialEditorPreviewParameters::RegenerateArrays()
 		{
 			int32 SortPriority;
 			FStaticSwitchParameter StaticSwitchParameterValue = FStaticSwitchParameter(SourceStaticParameters.StaticSwitchParameters[ParameterIdx]);
-			UDEditorStaticSwitchParameterValue& ParameterValue = *(NewObject<UDEditorStaticSwitchParameterValue>());
+			UDEditorStaticSwitchParameterValue& ParameterValue = *(NewObject<UDEditorStaticSwitchParameterValue>(this));
 			ParameterValue.ParameterValue = StaticSwitchParameterValue.Value;
 			ParameterValue.bOverride = StaticSwitchParameterValue.bOverride;
 			ParameterValue.ParameterInfo = StaticSwitchParameterValue.ParameterInfo;
@@ -618,7 +618,7 @@ void UMaterialEditorPreviewParameters::RegenerateArrays()
 		{
 			int32 SortPriority;
 			FStaticComponentMaskParameter StaticComponentMaskParameterValue = FStaticComponentMaskParameter(SourceStaticParameters.StaticComponentMaskParameters[ParameterIdx]);
-			UDEditorStaticComponentMaskParameterValue& ParameterValue = *(NewObject<UDEditorStaticComponentMaskParameterValue>());
+			UDEditorStaticComponentMaskParameterValue& ParameterValue = *(NewObject<UDEditorStaticComponentMaskParameterValue>(this));
 			ParameterValue.ParameterValue.R = StaticComponentMaskParameterValue.R;
 			ParameterValue.ParameterValue.G = StaticComponentMaskParameterValue.G;
 			ParameterValue.ParameterValue.B = StaticComponentMaskParameterValue.B;
@@ -1018,7 +1018,7 @@ void UMaterialEditorInstanceConstant::RegenerateArrays()
 		for (int32 ParameterIdx = 0; ParameterIdx < SourceStaticParameters.MaterialLayersParameters.Num(); ParameterIdx++)
 		{
 			FStaticMaterialLayersParameter MaterialLayersParameterParameterValue = FStaticMaterialLayersParameter(SourceStaticParameters.MaterialLayersParameters[ParameterIdx]);
-			UDEditorMaterialLayersParameterValue& ParameterValue = *(NewObject<UDEditorMaterialLayersParameterValue>());
+			UDEditorMaterialLayersParameterValue& ParameterValue = *(NewObject<UDEditorMaterialLayersParameterValue>(this));
 
 			ParameterValue.ParameterValue = MaterialLayersParameterParameterValue.Value;
 			ParameterValue.bOverride = MaterialLayersParameterParameterValue.bOverride;
@@ -1034,7 +1034,7 @@ void UMaterialEditorInstanceConstant::RegenerateArrays()
 		SourceInstance->GetAllScalarParameterInfo(OutParameterInfo, Guids);
 		for (int32 ParameterIdx=0; ParameterIdx<OutParameterInfo.Num(); ParameterIdx++)
 		{			
-			UDEditorScalarParameterValue& ParameterValue = *(NewObject<UDEditorScalarParameterValue>());
+			UDEditorScalarParameterValue& ParameterValue = *(NewObject<UDEditorScalarParameterValue>(this));
 			const FMaterialParameterInfo& ParameterInfo = OutParameterInfo[ParameterIdx];
 
 			ParameterValue.bOverride = false;
@@ -1069,7 +1069,7 @@ void UMaterialEditorInstanceConstant::RegenerateArrays()
 		SourceInstance->GetAllVectorParameterInfo(OutParameterInfo, Guids);
 		for(int32 ParameterIdx=0; ParameterIdx<OutParameterInfo.Num(); ParameterIdx++)
 		{
-			UDEditorVectorParameterValue& ParameterValue = *(NewObject<UDEditorVectorParameterValue>());
+			UDEditorVectorParameterValue& ParameterValue = *(NewObject<UDEditorVectorParameterValue>(this));
 			const FMaterialParameterInfo& ParameterInfo = OutParameterInfo[ParameterIdx];
 			
 			ParameterValue.bOverride = false;
@@ -1102,7 +1102,7 @@ void UMaterialEditorInstanceConstant::RegenerateArrays()
 		SourceInstance->GetAllTextureParameterInfo(OutParameterInfo, Guids);
 		for(int32 ParameterIdx=0; ParameterIdx<OutParameterInfo.Num(); ParameterIdx++)
 		{			
-			UDEditorTextureParameterValue& ParameterValue = *(NewObject<UDEditorTextureParameterValue>());
+			UDEditorTextureParameterValue& ParameterValue = *(NewObject<UDEditorTextureParameterValue>(this));
 			const FMaterialParameterInfo& ParameterInfo = OutParameterInfo[ParameterIdx];
 
 			ParameterValue.bOverride = false;
@@ -1140,7 +1140,7 @@ void UMaterialEditorInstanceConstant::RegenerateArrays()
 		SourceInstance->GetAllRuntimeVirtualTextureParameterInfo(OutParameterInfo, Guids);
 		for (int32 ParameterIdx = 0; ParameterIdx < OutParameterInfo.Num(); ParameterIdx++)
 		{
-			UDEditorRuntimeVirtualTextureParameterValue& ParameterValue = *(NewObject<UDEditorRuntimeVirtualTextureParameterValue>());
+			UDEditorRuntimeVirtualTextureParameterValue& ParameterValue = *(NewObject<UDEditorRuntimeVirtualTextureParameterValue>(this));
 			const FMaterialParameterInfo& ParameterInfo = OutParameterInfo[ParameterIdx];
 
 			ParameterValue.bOverride = false;
@@ -1177,7 +1177,7 @@ void UMaterialEditorInstanceConstant::RegenerateArrays()
 		SourceInstance->GetAllFontParameterInfo(OutParameterInfo, Guids);
 		for(int32 ParameterIdx=0; ParameterIdx<OutParameterInfo.Num(); ParameterIdx++)
 		{
-			UDEditorFontParameterValue& ParameterValue = *(NewObject<UDEditorFontParameterValue>());
+			UDEditorFontParameterValue& ParameterValue = *(NewObject<UDEditorFontParameterValue>(this));
 			const FMaterialParameterInfo& ParameterInfo = OutParameterInfo[ParameterIdx];
 
 			ParameterValue.bOverride = false;
@@ -1213,7 +1213,7 @@ void UMaterialEditorInstanceConstant::RegenerateArrays()
 		for(int32 ParameterIdx=0; ParameterIdx<SourceStaticParameters.StaticSwitchParameters.Num(); ParameterIdx++)
 		{			
 			FStaticSwitchParameter StaticSwitchParameterValue = FStaticSwitchParameter(SourceStaticParameters.StaticSwitchParameters[ParameterIdx]);
-			UDEditorStaticSwitchParameterValue& ParameterValue = *(NewObject<UDEditorStaticSwitchParameterValue>());
+			UDEditorStaticSwitchParameterValue& ParameterValue = *(NewObject<UDEditorStaticSwitchParameterValue>(this));
 
 			ParameterValue.ParameterValue = StaticSwitchParameterValue.Value;
 			ParameterValue.bOverride = StaticSwitchParameterValue.bOverride;
@@ -1231,7 +1231,7 @@ void UMaterialEditorInstanceConstant::RegenerateArrays()
 		for(int32 ParameterIdx=0; ParameterIdx<SourceStaticParameters.StaticComponentMaskParameters.Num(); ParameterIdx++)
 		{
 			FStaticComponentMaskParameter StaticComponentMaskParameterValue = FStaticComponentMaskParameter(SourceStaticParameters.StaticComponentMaskParameters[ParameterIdx]);
-			UDEditorStaticComponentMaskParameterValue& ParameterValue = *(NewObject<UDEditorStaticComponentMaskParameterValue>());
+			UDEditorStaticComponentMaskParameterValue& ParameterValue = *(NewObject<UDEditorStaticComponentMaskParameterValue>(this));
 
 			ParameterValue.ParameterValue.R = StaticComponentMaskParameterValue.R;
 			ParameterValue.ParameterValue.G = StaticComponentMaskParameterValue.G;
