Hide the option to toggle the S3DDC if the current project isn't using it.

diff --git a/Engine/Source/Editor/UnrealEd/Private/Settings/EditorSettings.cpp b/Engine/Source/Editor/UnrealEd/Private/Settings/EditorSettings.cpp
--- a/Engine/Source/Editor/UnrealEd/Private/Settings/EditorSettings.cpp
+++ b/Engine/Source/Editor/UnrealEd/Private/Settings/EditorSettings.cpp
@@ -15,6 +15,16 @@ UEditorSettings::UEditorSettings(const FObjectInitializer& ObjectInitializer)
 	bCopyStarterContentPreference = false;
 	bEditorAnalyticsEnabled_DEPRECATED = true;
 	AutoScalabilityWorkScaleAmount = 1;
+
+	FProperty* S3Property = StaticClass()->FindPropertyByName("bEnableS3DDC");
+	if (ensure(S3Property != nullptr))
+	{
+		bool bValue = false;
+		if (!GConfig->GetBool(TEXT("EditorSettings"), TEXT("bShowEnableS3DDC"), bValue, GEditorIni) || !bValue)
+		{
+			S3Property->ClearPropertyFlags(CPF_Edit);
+		}
+	}
 }
 
 void UEditorSettings::PostEditChangeProperty( struct FPropertyChangedEvent& PropertyChangedEvent)
diff --git a/Engine/Source/Editor/UnrealEd/Public/Settings/EditorSettings.h b/Engine/Source/Editor/UnrealEd/Public/Settings/EditorSettings.h
--- a/Engine/Source/Editor/UnrealEd/Public/Settings/EditorSettings.h
+++ b/Engine/Source/Editor/UnrealEd/Public/Settings/EditorSettings.h
@@ -31,12 +31,8 @@ class UNREALED_API UEditorSettings : public UObject
 	UPROPERTY(EditAnywhere, config, Category=DerivedData, meta = (ConfigRestartRequired = true))
 	FDirectoryPath SharedDerivedDataCache;
 
-	/** Whether to show the S3DDC settings. Needs to be enabled on a project-by-project basis. */
-	UPROPERTY(config)
-	bool bShowEnableS3DDC;
-
 	/** Whether to enable the S3 derived data cache backend */
-	UPROPERTY(EditAnywhere, config, Category=DerivedData, meta = (DisplayName = "Enable AWS S3 Cache", Tooltip = "Enables the experimental S3 derived data cache.", ConfigRestartRequired = true, EditCondition="bShowEnableS3DDC"))
+	UPROPERTY(EditAnywhere, config, Category=DerivedData, meta = (DisplayName = "Enable AWS S3 Cache", Tooltip = "Enables the experimental S3 derived data cache.", ConfigRestartRequired = true))
 	bool bEnableS3DDC = true;
 
 	// =====================================================================
