Adding missing shutdown/tool ending code to mesh paint tools and mesh paint mode.

diff --git a/Engine/Plugins/MeshPainting/Source/MeshPaintingToolset/Private/MeshTexturePaintingTool.cpp b/Engine/Plugins/MeshPainting/Source/MeshPaintingToolset/Private/MeshTexturePaintingTool.cpp
--- a/Engine/Plugins/MeshPainting/Source/MeshPaintingToolset/Private/MeshTexturePaintingTool.cpp
+++ b/Engine/Plugins/MeshPainting/Source/MeshPaintingToolset/Private/MeshTexturePaintingTool.cpp
@@ -124,6 +124,26 @@ void UMeshTexturePaintingTool::Setup()
 void UMeshTexturePaintingTool::Shutdown(EToolShutdownType ShutdownType)
 {
 	FinishPainting();
+	// If the user has pending changes and the editor is not exiting, we want to do the commit for all the modified textures.
+	if ((GetNumberOfPendingPaintChanges() > 0) && !IsEngineExitRequested())
+	{
+		CommitAllPaintedTextures();
+	}
+	else
+	{
+		ClearAllTextureOverrides();
+	}
+
+	PaintTargetData.Empty();
+
+	// Remove any existing texture targets
+	TexturePaintTargetList.Empty();
+	UMeshToolManager* MeshToolManager = Cast<UMeshToolManager>(GetToolManager());
+	if (MeshToolManager)
+	{
+		MeshToolManager->Refresh();
+	}
+
 	BrushProperties->SaveProperties(this);
 	Super::Shutdown(ShutdownType);
 }
diff --git a/Engine/Plugins/MeshPainting/Source/MeshPaintingToolset/Private/MeshVertexPaintingTool.cpp b/Engine/Plugins/MeshPainting/Source/MeshPaintingToolset/Private/MeshVertexPaintingTool.cpp
--- a/Engine/Plugins/MeshPainting/Source/MeshPaintingToolset/Private/MeshVertexPaintingTool.cpp
+++ b/Engine/Plugins/MeshPainting/Source/MeshPaintingToolset/Private/MeshVertexPaintingTool.cpp
@@ -100,6 +100,11 @@ void UMeshVertexPaintingTool::Shutdown(EToolShutdownType ShutdownType)
 {
 	FinishPainting();
 	BrushProperties->SaveProperties(this);
+	UMeshToolManager* MeshToolManager = Cast<UMeshToolManager>(GetToolManager());
+	if (MeshToolManager)
+	{
+		MeshToolManager->Refresh();
+	}
 	Super::Shutdown(ShutdownType);
 }
 
@@ -553,6 +558,14 @@ void UMeshColorPaintingTool::Setup()
 		EToolMessageLevel::UserNotification);
 }
 
+void UMeshColorPaintingTool::Shutdown(EToolShutdownType ShutdownType)
+{
+	//If we're painting vertex colors then propagate the painting done on LOD0 to all lower LODs. 
+	//Then stop forcing the LOD level of the mesh to LOD0.
+	ApplyForcedLODIndex(-1);
+	Super::Shutdown(ShutdownType);
+}
+
 void UMeshColorPaintingTool::CacheSelectionData()
 {
 	Super::CacheSelectionData();
diff --git a/Engine/Plugins/MeshPainting/Source/MeshPaintingToolset/Public/MeshVertexPaintingTool.h b/Engine/Plugins/MeshPainting/Source/MeshPaintingToolset/Public/MeshVertexPaintingTool.h
--- a/Engine/Plugins/MeshPainting/Source/MeshPaintingToolset/Public/MeshVertexPaintingTool.h
+++ b/Engine/Plugins/MeshPainting/Source/MeshPaintingToolset/Public/MeshVertexPaintingTool.h
@@ -247,6 +247,7 @@ class MESHPAINTINGTOOLSET_API UMeshColorPaintingTool : public UMeshVertexPaintin
 public:
 	UMeshColorPaintingTool();
 	virtual void Setup() override;
+	virtual void Shutdown(EToolShutdownType ShutdownType) override;
 	void LODPaintStateChanged(const bool bLODPaintingEnabled);
 	int32 GetMaxLODIndexToPaint() const;
 	void PaintLODChanged();
diff --git a/Engine/Source/Editor/UnrealEd/Private/Tools/UEdMode.cpp b/Engine/Source/Editor/UnrealEd/Private/Tools/UEdMode.cpp
--- a/Engine/Source/Editor/UnrealEd/Private/Tools/UEdMode.cpp
+++ b/Engine/Source/Editor/UnrealEd/Private/Tools/UEdMode.cpp
@@ -413,7 +413,7 @@ void UEdMode::TerminateActiveToolsOnPIEStart()
 {
 	DeactivateAllActiveTools();
 }
-void UEdMode::TerminateActiveToolsOnSaveWorld()
+void UEdMode::TerminateActiveToolsOnOnMapChanged(uint32 MapChangeFlags)
 {
 	DeactivateAllActiveTools();
 }
@@ -788,9 +788,9 @@ void UEdMode::Enter()
 	{
 		TerminateActiveToolsOnPIEStart();
 	});
-	PreSaveWorldDelegateHandle = FEditorDelegates::PreSaveWorld.AddLambda([this](uint32 SaveFlags, UWorld* World)
+	MapChangeDelegateHandle = FEditorDelegates::MapChange.AddLambda([this](uint32 ChangeFlags)
 	{
-		TerminateActiveToolsOnSaveWorld();
+		TerminateActiveToolsOnOnMapChanged(ChangeFlags);
 	});
 	bInvalidationPending = false;
 
@@ -811,7 +811,7 @@ void UEdMode::RegisterTool(TSharedPtr<FUICommandInfo> UICommand, FString ToolIde
 void UEdMode::Exit()
 {
 	FEditorDelegates::BeginPIE.Remove(BeginPIEDelegateHandle);
-	FEditorDelegates::PreSaveWorld.Remove(PreSaveWorldDelegateHandle);
+	FEditorDelegates::MapChange.Remove(MapChangeDelegateHandle);
 	GetToolManager()->OnToolStarted.RemoveAll(this);
 	GetToolManager()->OnToolEnded.RemoveAll(this);
 	// auto-accept any in-progress tools
diff --git a/Engine/Source/Editor/UnrealEd/Public/Tools/UEdMode.h b/Engine/Source/Editor/UnrealEd/Public/Tools/UEdMode.h
--- a/Engine/Source/Editor/UnrealEd/Public/Tools/UEdMode.h
+++ b/Engine/Source/Editor/UnrealEd/Public/Tools/UEdMode.h
@@ -260,7 +260,8 @@ private:
 	virtual void TerminateActiveToolsOnPIEStart();
 
 	// default behavior is to accept active tool
-	virtual void TerminateActiveToolsOnSaveWorld();
+	virtual void TerminateActiveToolsOnOnMapChanged(uint32 MapChangeFlags);
+
 
 	FRay GetRayFromMousePos(FEditorViewportClient* ViewportClient, FViewport* Viewport, int MouseX, int MouseY);
 	
@@ -334,8 +335,8 @@ protected:
 
 	// called when PIE is about to start, shuts down active tools
 	FDelegateHandle BeginPIEDelegateHandle;
-	// called before a Save starts. This currently shuts down active tools.
-	FDelegateHandle PreSaveWorldDelegateHandle;
+	// called before a map changes. This currently shuts down active tools.
+	FDelegateHandle MapChangeDelegateHandle;
 
 	UPROPERTY()
 	TSoftClassPtr<UObject> SettingsClass;
