Fix needing to manually enable building when launching projects from installed engine builds. Rather than relying solely on the engine promoted flag to determine whether a target should be built, we now check for the presence of a matching target receipt with the correct version information. If it's present, we don't trigger the build.

diff --git a/Engine/Source/Developer/DesktopPlatform/Private/TargetReceipt.cpp b/Engine/Source/Developer/DesktopPlatform/Private/TargetReceipt.cpp
--- a/Engine/Source/Developer/DesktopPlatform/Private/TargetReceipt.cpp
+++ b/Engine/Source/Developer/DesktopPlatform/Private/TargetReceipt.cpp
@@ -6,6 +6,22 @@
 #include "Dom/JsonObject.h"
 #include "Serialization/JsonSerializer.h"
 
+bool TryParseBuildVersion(const FJsonObject& Object, FBuildVersion& OutVersion)
+{
+	if (Object.TryGetNumberField(TEXT("MajorVersion"), OutVersion.MajorVersion) && Object.TryGetNumberField(TEXT("MinorVersion"), OutVersion.MinorVersion) && Object.TryGetNumberField(TEXT("PatchVersion"), OutVersion.PatchVersion))
+	{
+		Object.TryGetNumberField(TEXT("Changelist"), OutVersion.Changelist);
+		Object.TryGetNumberField(TEXT("CompatibleChangelist"), OutVersion.CompatibleChangelist);
+		Object.TryGetNumberField(TEXT("IsLicenseeVersion"), OutVersion.IsLicenseeVersion);
+		Object.TryGetNumberField(TEXT("IsPromotedBuild"), OutVersion.IsPromotedBuild);
+		Object.TryGetStringField(TEXT("BranchName"), OutVersion.BranchName);
+		Object.TryGetStringField(TEXT("BuildId"), OutVersion.BuildId);
+		Object.TryGetStringField(TEXT("BuildVersion"), OutVersion.BuildVersion);
+		return true;
+	}
+	return false;
+}
+
 bool FTargetReceipt::Read(const FString& FileName)
 {
 	// Read the file from disk
@@ -59,6 +75,13 @@ bool FTargetReceipt::Read(const FString& FileName)
 		return false;
 	}
 
+	// Read the version information
+	const TSharedPtr<FJsonObject>* VersionObject;
+	if (!Object->TryGetObjectField(TEXT("Version"), VersionObject) || !VersionObject->IsValid() || !TryParseBuildVersion(*VersionObject->Get(), Version))
+	{
+		return false;
+	}
+
 	// Get the launch path
 	if(!Object->TryGetStringField(TEXT("Launch"), Launch))
 	{
diff --git a/Engine/Source/Developer/DesktopPlatform/Public/TargetReceipt.h b/Engine/Source/Developer/DesktopPlatform/Public/TargetReceipt.h
--- a/Engine/Source/Developer/DesktopPlatform/Public/TargetReceipt.h
+++ b/Engine/Source/Developer/DesktopPlatform/Public/TargetReceipt.h
@@ -67,6 +67,9 @@ struct DESKTOPPLATFORM_API FTargetReceipt
 	/** The type of the target */
 	EBuildTargetType TargetType;
 
+	/** The version information for this target */
+	FBuildVersion Version;
+
 	/** The exectuable to launch for this target */
 	FString Launch;
 
diff --git a/Engine/Source/Editor/UnrealEd/Private/PlayLevelViaLauncher.cpp b/Engine/Source/Editor/UnrealEd/Private/PlayLevelViaLauncher.cpp
--- a/Engine/Source/Editor/UnrealEd/Private/PlayLevelViaLauncher.cpp
+++ b/Engine/Source/Editor/UnrealEd/Private/PlayLevelViaLauncher.cpp
@@ -17,6 +17,10 @@
 #include "PlayLevel.h"
 #include "Async/Async.h"
 #include "Logging/MessageLog.h"
+#include "TargetReceipt.h"
+#include "DesktopPlatformModule.h"
+#include "TargetPlatform.h"
+#include "PlatformInfo.h"
 
 #define LOCTEXT_NAMESPACE "PlayLevel"
 
@@ -49,6 +53,57 @@ static void HandleOutputReceived(const FString& InMessage)
 	}
 }
 
+bool HasPromotedTarget(const TCHAR* BaseDir, const TCHAR* TargetName, const TCHAR* Platform, EBuildConfiguration Configuration, const TCHAR* Architecture)
+{
+	// Get the path to the receipt, and check it exists
+	FString ReceiptPath = FTargetReceipt::GetDefaultPath(BaseDir, TargetName, Platform, Configuration, Architecture);
+	if (!FPaths::FileExists(*ReceiptPath))
+	{
+		UE_LOG(LogPlayLevel, Log, TEXT("Unable to use promoted target - %s does not exist."), *ReceiptPath);
+		return false;
+	}
+
+	// Read the receipt for this target
+	FTargetReceipt Receipt;
+	if (!Receipt.Read(ReceiptPath))
+	{
+		UE_LOG(LogPlayLevel, Log, TEXT("Unable to use promoted target - cannot read %s"), *ReceiptPath);
+		return false;
+	}
+
+	// Check the receipt is for a promoted build
+	if (!Receipt.Version.IsPromotedBuild)
+	{
+		UE_LOG(LogPlayLevel, Log, TEXT("Unable to use promoted target - receipt %s is not for a promoted target"), *ReceiptPath);
+		return false;
+	}
+
+	// Make sure it matches the current build info
+	FEngineVersion ReceiptVersion = Receipt.Version.GetEngineVersion();
+	FEngineVersion CurrentVersion = FEngineVersion::Current();
+	if (!ReceiptVersion.ExactMatch(CurrentVersion))
+	{
+		UE_LOG(LogPlayLevel, Log, TEXT("Unable to use promoted target - receipt version (%s) is not exact match with current engine version (%s)"), *ReceiptVersion.ToString(), *CurrentVersion.ToString());
+		return false;
+	}
+
+	return true;
+}
+
+bool TryGetDefaultTargetName(const FString& ProjectFile, EBuildTargetType TargetType, FString& OutTargetName)
+{
+	const TArray<FTargetInfo>& Targets = FDesktopPlatformModule::Get()->GetTargetsForProject(ProjectFile);
+	for (const FTargetInfo& Target : Targets)
+	{
+		if (Target.Type == TargetType)
+		{
+			OutTargetName = Target.Name;
+			return true;
+		}
+	}
+	return false;
+}
+
 void UEditorEngine::StartPlayUsingLauncherSession(FRequestPlaySessionParams& InRequestParams)
 {
 	check(InRequestParams.SessionDestination == EPlaySessionDestinationType::Launcher);
@@ -164,16 +219,45 @@ void UEditorEngine::StartPlayUsingLauncherSession(FRequestPlaySessionParams& InR
 	}
 	else if (EditorPlaySettings->BuildGameBeforeLaunch == EPlayOnBuildMode::PlayOnBuild_Default)
 	{
-		LauncherSessionInfo->bPlayUsingLauncherBuild = LauncherSessionInfo->bPlayUsingLauncherHasCode || !FApp::GetEngineIsPromotedBuild();
-		if (!LauncherSessionInfo->bPlayUsingLauncherBuild)
+		// If this isn't a promoted build, we always have to build
+		bool bBuildTarget = true;
+		if (FApp::GetEngineIsPromotedBuild())
 		{
-			FText Reason;
-			if (LaunchPlatform->RequiresTempTarget(LauncherSessionInfo->bPlayUsingLauncherHasCode, BuildConfiguration, false, Reason))
+			// Figure out which target we're building
+			FString ReceiptDir;
+			FString TargetName;
+			if (TryGetDefaultTargetName(FPaths::GetProjectFilePath(), EBuildTargetType::Game, TargetName))
+			{
+				ReceiptDir = FPaths::GetPath(FPaths::GetProjectFilePath());
+			}
+			else
+			{
+				TargetName = TEXT("UE4Game");
+
+				FText Reason;
+				if (LaunchPlatform->RequiresTempTarget(LauncherSessionInfo->bPlayUsingLauncherHasCode, BuildConfiguration, false, Reason))
+				{
+					UE_LOG(LogPlayLevel, Log, TEXT("Project requires temp target (%s)"), *Reason.ToString());
+					ReceiptDir = FPaths::GetPath(FPaths::GetProjectFilePath());
+				}
+				else
+				{
+					UE_LOG(LogPlayLevel, Log, TEXT("Project does not require temp target"));
+					ReceiptDir = FPaths::EngineDir();
+				}
+			}
+
+			// Get the target platform
+			const ITargetPlatform* TargetPlatform = GetTargetPlatformManager()->FindTargetPlatform(LaunchPlatformName);
+			FString PlatformName = TargetPlatform->GetPlatformInfo().UBTTargetId.ToString();
+
+			// Check if the existing target is valid
+			if (HasPromotedTarget(*ReceiptDir, *TargetName, *PlatformName, BuildConfiguration, nullptr))
 			{
-				UE_LOG(LogPlayLevel, Log, TEXT("Project requires temp target (%s)"), *Reason.ToString());
-				LauncherSessionInfo->bPlayUsingLauncherBuild = true;
+				bBuildTarget = false;
 			}
 		}
+		LauncherSessionInfo->bPlayUsingLauncherBuild = bBuildTarget;
 	}
 	else if (EditorPlaySettings->BuildGameBeforeLaunch == EPlayOnBuildMode::PlayOnBuild_IfEditorBuiltLocally)
 	{
diff --git a/Engine/Source/Runtime/Core/Private/Misc/EngineVersion.cpp b/Engine/Source/Runtime/Core/Private/Misc/EngineVersion.cpp
--- a/Engine/Source/Runtime/Core/Private/Misc/EngineVersion.cpp
+++ b/Engine/Source/Runtime/Core/Private/Misc/EngineVersion.cpp
@@ -100,6 +100,11 @@ void FEngineVersion::Empty()
 	Set(0, 0, 0, 0, FString());
 }
 
+bool FEngineVersion::ExactMatch(const FEngineVersion& Other) const
+{
+	return Major == Other.Major && Minor == Other.Minor && Patch == Other.Patch && Changelist == Other.Changelist && Branch == Other.Branch;
+}
+
 bool FEngineVersion::IsCompatibleWith(const FEngineVersionBase &Other) const
 {
 	// If this or the other is not a promoted build, always assume compatibility. 
diff --git a/Engine/Source/Runtime/Core/Private/Modules/BuildVersion.cpp b/Engine/Source/Runtime/Core/Private/Modules/BuildVersion.cpp
--- a/Engine/Source/Runtime/Core/Private/Modules/BuildVersion.cpp
+++ b/Engine/Source/Runtime/Core/Private/Modules/BuildVersion.cpp
@@ -18,6 +18,31 @@ FBuildVersion::FBuildVersion()
 {
 }
 
+int FBuildVersion::GetEffectiveCompatibleChangelist() const
+{
+	return (Changelist != 0 && CompatibleChangelist != 0) ? CompatibleChangelist : Changelist;
+}
+
+FEngineVersion FBuildVersion::GetEngineVersion() const
+{
+	int EncodedChangelist = Changelist;
+	if (IsLicenseeVersion)
+	{
+		EncodedChangelist = FEngineVersion::EncodeLicenseeChangelist(EncodedChangelist);
+	}
+	return FEngineVersion((uint16)MajorVersion, (uint16)MinorVersion, (uint16)PatchVersion, EncodedChangelist, BranchName);
+}
+
+FEngineVersion FBuildVersion::GetCompatibleEngineVersion() const
+{
+	int EncodedChangelist = GetEffectiveCompatibleChangelist();
+	if (IsLicenseeVersion)
+	{
+		EncodedChangelist = FEngineVersion::EncodeLicenseeChangelist(EncodedChangelist);
+	}
+	return FEngineVersion((uint16)MajorVersion, (uint16)MinorVersion, (uint16)PatchVersion, EncodedChangelist, BranchName);
+}
+
 FString FBuildVersion::GetDefaultFileName()
 {
 	return FPaths::EngineDir() / TEXT("Build/Build.version");
diff --git a/Engine/Source/Runtime/Core/Public/Misc/EngineVersion.h b/Engine/Source/Runtime/Core/Public/Misc/EngineVersion.h
--- a/Engine/Source/Runtime/Core/Public/Misc/EngineVersion.h
+++ b/Engine/Source/Runtime/Core/Public/Misc/EngineVersion.h
@@ -24,6 +24,9 @@ public:
 	/** Clears the object. */
 	void Empty();
 
+	/** Checks whether this engine version is an exact match for another engine version */
+	bool ExactMatch(const FEngineVersion& Other) const;
+
 	/** Checks compatibility with another version object. */
 	bool IsCompatibleWith(const FEngineVersionBase &Other) const;
 
diff --git a/Engine/Source/Runtime/Core/Public/Modules/BuildVersion.h b/Engine/Source/Runtime/Core/Public/Modules/BuildVersion.h
--- a/Engine/Source/Runtime/Core/Public/Modules/BuildVersion.h
+++ b/Engine/Source/Runtime/Core/Public/Modules/BuildVersion.h
@@ -4,6 +4,7 @@
 
 #include "CoreTypes.h"
 #include "Containers/UnrealString.h"
+#include "Misc/EngineVersion.h"
 
 class FJsonObject;
 
@@ -68,6 +69,24 @@ public:
 	 */
 	FBuildVersion();
 
+	/// <summary>
+	/// Gets the compatible changelist if set, otherwise the default compatible changelist
+	/// </summary>
+	/// <returns>The compatible changelist</returns>
+	int GetEffectiveCompatibleChangelist() const;
+
+	/// <summary>
+	/// Get an engine version object for this build version
+	/// </summary>
+	/// <returns>New engine version object</returns>
+	FEngineVersion GetEngineVersion() const;
+
+	/// <summary>
+	/// Get a compatible engine version object for this build version
+	/// </summary>
+	/// <returns>New engine version object</returns>
+	FEngineVersion GetCompatibleEngineVersion() const;
+
 	/// <summary>
 	/// Get the default path to the build.version file on disk
 	/// </summary>
