Remove legacy header include to fix CIS.

#rnx

diff --git a/Engine/Source/Editor/UnrealEd/Private/PlayLevelViaLauncher.cpp b/Engine/Source/Editor/UnrealEd/Private/PlayLevelViaLauncher.cpp
--- a/Engine/Source/Editor/UnrealEd/Private/PlayLevelViaLauncher.cpp
+++ b/Engine/Source/Editor/UnrealEd/Private/PlayLevelViaLauncher.cpp
@@ -19,7 +19,6 @@
 #include "Logging/MessageLog.h"
 #include "TargetReceipt.h"
 #include "DesktopPlatformModule.h"
-#include "TargetPlatform.h"
 #include "PlatformInfo.h"
 
 #define LOCTEXT_NAMESPACE "PlayLevel"
