Additional check for which root template to use when determining instance changes in blueprints

#rnx
#okfrogithub public

diff --git a/Engine/Source/Editor/UnrealEd/Private/Kismet2/Kismet2.cpp b/Engine/Source/Editor/UnrealEd/Private/Kismet2/Kismet2.cpp
--- a/Engine/Source/Editor/UnrealEd/Private/Kismet2/Kismet2.cpp
+++ b/Engine/Source/Editor/UnrealEd/Private/Kismet2/Kismet2.cpp
@@ -1290,7 +1290,15 @@ void FKismetEditorUtilities::AddComponentsToBlueprint(UBlueprint* Blueprint, con
 			USCS_Node* ParentSCSNode = nullptr;
 			if (FirstAttachParent)
 			{
-				ParentSCSNode = InstanceComponentToNodeMap.FindChecked(FirstAttachParent);
+				// If we are using the root template then it will not be in the component node map, the scene component will be
+				if(SceneComponent && (SceneComponent->GetAttachParent() == nullptr || FirstAttachParent == RootTemplate))
+				{
+					ParentSCSNode = InstanceComponentToNodeMap.FindChecked(SceneComponent);
+				}
+				else
+				{
+					ParentSCSNode = InstanceComponentToNodeMap.FindChecked(FirstAttachParent);
+				}
 			}
 			else
 			{
