Renamed "Play Offline" to "Play Standalone" and updated the comment to better reflect what it is doing under the hood. This brings it in line with the other existing net mode terminology and reduces confusion about what networking support Offline/Standalone has.

diff --git a/Engine/Source/Editor/UnrealEd/Classes/Settings/LevelEditorPlaySettings.h b/Engine/Source/Editor/UnrealEd/Classes/Settings/LevelEditorPlaySettings.h
--- a/Engine/Source/Editor/UnrealEd/Classes/Settings/LevelEditorPlaySettings.h
+++ b/Engine/Source/Editor/UnrealEd/Classes/Settings/LevelEditorPlaySettings.h
@@ -85,8 +85,8 @@ enum EPlayModeType
 UENUM()
 enum EPlayNetMode
 {
-	/** A non-networked game will be started. This can be used in combination with bLaunchSeparateServer to test menu -> server flow in your game. */
-	PIE_Standalone UMETA(DisplayName="Play Offline"),
+	/** A standalone game will be started. This will not create a dedicated server, nor automatically connect to one. A server can be launched by enabling bLaunchSeparateServer if you need to test offline -> server connection flow for your game. */
+	PIE_Standalone UMETA(DisplayName="Play Standalone"),
 	/** The editor will act as both a Server and a Client. Additional instances may be opened beyond that depending on the number of clients. */
 	PIE_ListenServer UMETA(DisplayName="Play As Listen Server"),
 	/** The editor will act as a Client. A server will be started for you behind the scenes to connect to. */
