Skip VerifyEDLCookInfo when cooking DLC

Reported by [REDACTED_PERSON] at Gearbox Software.

[RN] fix Core/Cooker

diff --git a/Engine/Source/Editor/UnrealEd/Private/Commandlets/CookCommandlet.cpp b/Engine/Source/Editor/UnrealEd/Private/Commandlets/CookCommandlet.cpp
--- a/Engine/Source/Editor/UnrealEd/Private/Commandlets/CookCommandlet.cpp
+++ b/Engine/Source/Editor/UnrealEd/Private/Commandlets/CookCommandlet.cpp
@@ -1073,7 +1073,7 @@ bool UCookCommandlet::CookByTheBook( const TArray<ITargetPlatform*>& Platforms,
 		}
 	} while (bTestCook);
 
-	if (!bIterativeCooking)
+	if (!bIterativeCooking && StartupOptions.DLCName.IsEmpty())
 	{
 		VerifyEDLCookInfo();
 	}
