Edigrate: Disable depth writes (MP_PixelDepthOffset) in Lightmass materials

diff --git a/Engine/Source/Editor/UnrealEd/Private/Lightmass/LightmassRender.cpp b/Engine/Source/Editor/UnrealEd/Private/Lightmass/LightmassRender.cpp
--- a/Engine/Source/Editor/UnrealEd/Private/Lightmass/LightmassRender.cpp
+++ b/Engine/Source/Editor/UnrealEd/Private/Lightmass/LightmassRender.cpp
@@ -448,7 +448,7 @@ public:
 	
 			return Compiler->Constant(0.0f);
 		}
-		else if( Property == MP_WorldPositionOffset)
+		else if( Property == MP_WorldPositionOffset || Property == MP_PixelDepthOffset )
 		{
 			//This property MUST return 0 as a default or during the process of rendering textures out for lightmass to use, pixels will be off by 1.
 			return Compiler->Constant(0.0f);
