Fix CIS issue with dereferencing null

#rnx

diff --git a/Engine/Source/Editor/UnrealEd/Private/Kismet2/Kismet2.cpp b/Engine/Source/Editor/UnrealEd/Private/Kismet2/Kismet2.cpp
--- a/Engine/Source/Editor/UnrealEd/Private/Kismet2/Kismet2.cpp
+++ b/Engine/Source/Editor/UnrealEd/Private/Kismet2/Kismet2.cpp
@@ -1291,7 +1291,7 @@ void FKismetEditorUtilities::AddComponentsToBlueprint(UBlueprint* Blueprint, con
 			if (FirstAttachParent)
 			{
 				// If we are using the root template then it will not be in the component node map, the scene component will be
-				if(SceneComponent && (SceneComponent->GetAttachParent() == nullptr || FirstAttachParent == RootTemplate))
+				if(SceneComponent->GetAttachParent() == nullptr || FirstAttachParent == RootTemplate)
 				{
 					ParentSCSNode = InstanceComponentToNodeMap.FindChecked(SceneComponent);
 				}
