- Add list of particle systems with warm up enabled

#rnx

diff --git a/Engine/Source/Editor/UnrealEd/Classes/Commandlets/ParticleSystemAuditCommandlet.h b/Engine/Source/Editor/UnrealEd/Classes/Commandlets/ParticleSystemAuditCommandlet.h
--- a/Engine/Source/Editor/UnrealEd/Classes/Commandlets/ParticleSystemAuditCommandlet.h
+++ b/Engine/Source/Editor/UnrealEd/Classes/Commandlets/ParticleSystemAuditCommandlet.h
@@ -32,6 +32,8 @@ class UParticleSystemAuditCommandlet : public UCommandlet
 	TSet<FString> ParticleSystemsWithFarLODDistance;
 	/** All particle systems w/ bone location sources that do not match between LODs */
 	TSet<FString> ParticleSystemsWithBoneLocationMismatches;
+	/** All particle systems with warmup time on them */
+	TSet<FString> ParticleSystemsWithWarmupTime;
 
 	/** If a particle system has a spawn rate or burst count greater than this value, it will be reported */
 	UPROPERTY(config)
diff --git a/Engine/Source/Editor/UnrealEd/Private/Commandlets/ParticleSystemAuditCommandlet.cpp b/Engine/Source/Editor/UnrealEd/Private/Commandlets/ParticleSystemAuditCommandlet.cpp
--- a/Engine/Source/Editor/UnrealEd/Private/Commandlets/ParticleSystemAuditCommandlet.cpp
+++ b/Engine/Source/Editor/UnrealEd/Private/Commandlets/ParticleSystemAuditCommandlet.cpp
@@ -332,6 +332,11 @@ bool UParticleSystemAuditCommandlet::ProcessParticleSystems()
 				ParticleSystemsWithBoneLocationMismatches.Add(PSys->GetPathName());
 			}
 
+			if (PSys->WarmupTime > 0.0f)
+			{
+				ParticleSystemsWithWarmupTime.Add(PSys->GetPathName());
+			}
+
 			if ((PSys->LODMethod == PARTICLESYSTEMLODMETHOD_Automatic) &&
 				(bInvalidLOD == false) && (bSingleLOD == false) &&
 				(PSys->LODDistanceCheckTime == 0.0f))
@@ -389,6 +394,7 @@ void UParticleSystemAuditCommandlet::DumpResults()
 	DumpSimplePSysSet(ParticleSystemsWithHighSpawnRateOrBurst, TEXT("PSysHighSpawnRateOrBurst"));
 	DumpSimplePSysSet(ParticleSystemsWithFarLODDistance, TEXT("PSysFarLODDistance"));
 	DumpSimplePSysSet(ParticleSystemsWithBoneLocationMismatches, TEXT("PSysBoneLocationLODMismatches"));
+	DumpSimplePSysSet(ParticleSystemsWithWarmupTime, TEXT("PSysWithWarmupTime"));
 }
 
 /**
