Fix playing Steam-enabled games in PIE when "Run Under One Process" is unchecked by restoring missing command line options, including -MultiprocessOSS.

diff --git a/Engine/Source/Editor/UnrealEd/Private/PlayLevelNewProcess.cpp b/Engine/Source/Editor/UnrealEd/Private/PlayLevelNewProcess.cpp
--- a/Engine/Source/Editor/UnrealEd/Private/PlayLevelNewProcess.cpp
+++ b/Engine/Source/Editor/UnrealEd/Private/PlayLevelNewProcess.cpp
@@ -91,17 +91,14 @@ void UEditorEngine::LaunchNewProcess(const FRequestPlaySessionParams& InParams,
 		}
 	}
 
-	if(!bIsDedicatedServer)
-	{
-		// Allow loading specific GameUserSettings from the ini which differ per-process.
-		FString GameUserSettingsOverride = GGameUserSettingsIni.Replace(TEXT("GameUserSettings"), *FString::Printf(TEXT("PIEGameUserSettings%d"), InInstanceNum));
-
-		// Construct parms:
-		//	-Override GameUserSettings.ini
-		//	-Force no steam
-		//	-Allow saving of config files (since we are giving them an override INI)
-		CommandLine += FString::Printf(TEXT("GameUserSettingsINI=\"%s\" -MultiprocessSaveConfig -MultiprocessOSS "), *GameUserSettingsOverride);
-	}
+	// Allow loading specific GameUserSettings from the ini which differ per-process.
+	FString GameUserSettingsOverride = GGameUserSettingsIni.Replace(TEXT("GameUserSettings"), *FString::Printf(TEXT("PIEGameUserSettings%d"), InInstanceNum));
+
+	// Construct parms:
+	//	-Override GameUserSettings.ini
+	//	-Force no steam
+	//	-Allow saving of config files (since we are giving them an override INI)
+	CommandLine += FString::Printf(TEXT("GameUserSettingsINI=\"%s\" -MultiprocessSaveConfig -MultiprocessOSS "), *GameUserSettingsOverride);
 
 	if (bIsDedicatedServer)
 	{
