TBA: Stood up a simple schema export pass from the text asset commandlet. Currently hardcoded and incomplete, but enough to test the framework

diff --git a/Engine/Source/Editor/UnrealEd/Classes/Commandlets/TextAssetCommandlet.h b/Engine/Source/Editor/UnrealEd/Classes/Commandlets/TextAssetCommandlet.h
--- a/Engine/Source/Editor/UnrealEd/Classes/Commandlets/TextAssetCommandlet.h
+++ b/Engine/Source/Editor/UnrealEd/Classes/Commandlets/TextAssetCommandlet.h
@@ -18,7 +18,8 @@ enum class ETextAssetCommandletMode
 	RoundTrip,
 	LoadBinary,
 	LoadText,
-	FindMismatchedSerializers
+	FindMismatchedSerializers,
+	GenerateSchema
 };
 
 UCLASS(config=Editor)
diff --git a/Engine/Source/Editor/UnrealEd/Private/Commandlets/TextAssetCommandlet.cpp b/Engine/Source/Editor/UnrealEd/Private/Commandlets/TextAssetCommandlet.cpp
--- a/Engine/Source/Editor/UnrealEd/Private/Commandlets/TextAssetCommandlet.cpp
+++ b/Engine/Source/Editor/UnrealEd/Private/Commandlets/TextAssetCommandlet.cpp
@@ -88,6 +88,172 @@ static void RepairDamagedFiles()
 	IFileManager::Get().IterateDirectoryRecursively(*FPaths::EngineContentDir(), RepairVisitor);
 }
 
+typedef TFunction<void(FStructuredArchiveRecord, TArray<FString>&)> FSimpleSchemaFieldPropertyGenerator;
+
+namespace StringConstants
+{
+	static FString Object(TEXT("object"));
+	static FString String(TEXT("string"));
+	static FString Number(TEXT("number"));
+	static FString Array(TEXT("array"));
+	static FString Boolean(TEXT("boolean"));
+	static FString Properties(TEXT("properties"));
+	static FString Type(TEXT("type"));
+	static FString Nil(TEXT("nil"));
+}
+
+inline void WriteSimpleSchemaField(FStructuredArchiveRecord Record, const TCHAR* FieldName, FString& Type, FSimpleSchemaFieldPropertyGenerator PropertiesCallback = FSimpleSchemaFieldPropertyGenerator())
+{
+	FStructuredArchiveRecord FieldRecord = Record.EnterField(SA_FIELD_NAME(FieldName)).EnterRecord();
+	FieldRecord << SA_VALUE(TEXT("type"), Type);
+
+	if (PropertiesCallback)
+	{
+		FStructuredArchiveRecord PropertiesRecord = FieldRecord.EnterField(SA_FIELD_NAME(TEXT("properties"))).EnterRecord();
+		TArray<FString> Required;
+		PropertiesCallback(PropertiesRecord, Required);
+		if (Required.Num() > 0)
+		{
+			int32 NumRequired = Required.Num();
+			FStructuredArchiveArray RequiredArray = FieldRecord.EnterField(SA_FIELD_NAME(TEXT("required"))).EnterArray(NumRequired);
+			for (FString& RequiredProperty : Required)
+			{
+				RequiredArray.EnterElement() << RequiredProperty;
+			}
+		}
+	}
+}
+
+TSet<FName> GMissingThings;
+
+void GeneratePropertySchema(FProperty* Property, FStructuredArchiveRecord Record, TArray<FString>& Required)
+{
+	static const FName NAME_ClassProperty(TEXT("ClassProperty"));
+	static const FName NAME_WeakObjectProperty(TEXT("WeakObjectProperty"));
+
+	const FFieldClass* PropertyClass = Property->GetClass();
+	const FName PropertyClassName = PropertyClass->GetFName();
+
+	if (PropertyClassName == NAME_ArrayProperty)
+	{
+		WriteSimpleSchemaField(Record, TEXT("__Type"), StringConstants::String);
+		WriteSimpleSchemaField(Record, TEXT("__InnerType"), StringConstants::String);
+		WriteSimpleSchemaField(Record, TEXT("__InnerStructName"), StringConstants::String);
+		WriteSimpleSchemaField(Record, TEXT("__Value"), StringConstants::Array);
+
+		FStructuredArchiveRecord ItemsRecord = Record.EnterRecord(SA_FIELD_NAME(TEXT("items")));
+	}
+	else
+	{
+		WriteSimpleSchemaField(Record, TEXT("__Type"), StringConstants::String);
+		
+		// We need to describe the data that this property writes out, which only the derived property class knows. We'll have to add something to the FProperty API to do that
+		// but for now I'm just going to hardcode things here
+		if (PropertyClassName == NAME_StrProperty
+ 		||  PropertyClassName == NAME_ObjectProperty
+		||  PropertyClassName == NAME_NameProperty)
+		{
+			WriteSimpleSchemaField(Record, TEXT("__Value"), StringConstants::String);
+		}
+		else if (PropertyClassName == NAME_EnumProperty || (PropertyClassName == NAME_ByteProperty && ((const FByteProperty*)(Property))->Enum != nullptr))
+		{
+			WriteSimpleSchemaField(Record, TEXT("__EnumName"), StringConstants::String);
+			WriteSimpleSchemaField(Record, TEXT("__Value"), StringConstants::String);
+		}
+		else if (PropertyClass->IsChildOf(FNumericProperty::StaticClass()))
+		{
+			WriteSimpleSchemaField(Record, TEXT("__Value"), StringConstants::Number);
+		}
+		else if (PropertyClassName == NAME_StructProperty || PropertyClassName == NAME_ClassProperty)
+		{
+			WriteSimpleSchemaField(Record, TEXT("__StructName"), StringConstants::String);
+			WriteSimpleSchemaField(Record, TEXT("__Value"), StringConstants::Object);
+		}
+		else if (PropertyClassName == NAME_TextProperty)
+		{
+			WriteSimpleSchemaField(Record, TEXT("__Value"), StringConstants::Object);
+		}
+		else if (PropertyClassName == NAME_BoolProperty)
+		{
+			WriteSimpleSchemaField(Record, TEXT("__Value"), StringConstants::Boolean);
+		}
+		else if (PropertyClassName == NAME_SetProperty)
+		{
+			WriteSimpleSchemaField(Record, TEXT("__InnerType"), StringConstants::String);
+			WriteSimpleSchemaField(Record, TEXT("__Value"), StringConstants::Object);
+		}
+		else if (PropertyClassName == NAME_InterfaceProperty)
+		{
+			WriteSimpleSchemaField(Record, TEXT("__Value"), StringConstants::Object);
+		}
+		else if (PropertyClassName == NAME_WeakObjectProperty)
+		{
+			WriteSimpleSchemaField(Record, TEXT("__Value"), StringConstants::Object);
+		}
+		else if (PropertyClassName == NAME_MapProperty)
+		{
+			WriteSimpleSchemaField(Record, TEXT("__InnerType"), StringConstants::String);
+			WriteSimpleSchemaField(Record, TEXT("__ValueType"), StringConstants::String);
+			WriteSimpleSchemaField(Record, TEXT("__Value"), StringConstants::Object);
+		}
+		else
+		{
+			if (!GMissingThings.Contains(PropertyClassName))
+			{
+				UE_LOG(LogTextAsset, Warning, TEXT("Unhandled property type: %s"), *PropertyClassName.ToString());
+				GMissingThings.Add(PropertyClassName);
+			}
+		}
+	}
+}
+
+void GenerateClassSchema(UClass* Class, FStructuredArchiveRecord Record, TArray<FString>& Required)
+{
+	FProperty* CurProperty = Class->PropertyLink;
+	while (CurProperty != nullptr)
+	{
+		WriteSimpleSchemaField(Record, *CurProperty->GetName(), StringConstants::Object, [CurProperty](FStructuredArchiveRecord InRecord, TArray<FString>& InRequired) { GeneratePropertySchema(CurProperty, InRecord, InRequired); });
+		CurProperty = CurProperty->PropertyLinkNext;
+	}
+}
+
+void GenerateSchema()
+{
+	//static const FName NAME_SpecificClass(TEXT("TextAssetTestObject"));
+	static const FName NAME_SpecificClass(NAME_None);
+
+	FString OutputFilename;
+	if (!FParse::Value(FCommandLine::Get(), TEXT("-schemaoutput="), OutputFilename))
+	{
+		OutputFilename = FPaths::ProjectConfigDir() / TEXT("Schemas/TextAssetExports.json");
+	}
+
+	TUniquePtr<FArchive> OutputAr(IFileManager::Get().CreateFileWriter(*OutputFilename));
+	FJsonArchiveOutputFormatter JsonFormatter(*OutputAr.Get());
+	FStructuredArchive StructuredArchive(JsonFormatter);
+	FStructuredArchiveRecord RootRecord = StructuredArchive.Open().EnterRecord();
+
+	for (FObjectIterator It(UClass::StaticClass()); It; ++It)
+	{
+		UClass* Class = Cast<UClass>(*It);
+		if (Class)
+		{
+			if (NAME_SpecificClass == NAME_None || Class->GetFName() == NAME_SpecificClass)
+			{
+				FStructuredArchiveRecord ClassRecord = RootRecord.EnterRecord(SA_FIELD_NAME(*Class->GetFullName()));
+				ClassRecord << SA_VALUE(*StringConstants::Type, StringConstants::Object);
+
+				WriteSimpleSchemaField(ClassRecord, TEXT("__Class"), StringConstants::Object);
+				WriteSimpleSchemaField(ClassRecord, TEXT("__Outer"), StringConstants::String);
+				WriteSimpleSchemaField(ClassRecord, TEXT("__bNotAlwaysLoadedForEditorGame"), StringConstants::Boolean);
+				WriteSimpleSchemaField(ClassRecord, TEXT("__Value"), StringConstants::Object, [Class](FStructuredArchiveRecord Record, TArray<FString>& Required) { GenerateClassSchema(Class, Record, Required); });
+			}
+		}
+	}
+
+	StructuredArchive.Close();
+}
+
 bool UTextAssetCommandlet::DoTextAssetProcessing(const FString& InCommandLine)
 {
 	FProcessingArgs Args;
@@ -142,10 +308,18 @@ bool UTextAssetCommandlet::DoTextAssetProcessing(const FProcessingArgs& InArgs)
 
 	TArray<FString> Blacklist;
 
-	if (InArgs.ProcessingMode == ETextAssetCommandletMode::FindMismatchedSerializers)
+	switch (InArgs.ProcessingMode)
 	{
+	case ETextAssetCommandletMode::FindMismatchedSerializers:
 		FindMismatchedSerializers();
 		return true;
+
+	case ETextAssetCommandletMode::GenerateSchema:
+		GenerateSchema();
+		break;
+
+	default:
+		break;
 	}
 
 	TArray<UObject*> Objects;		
