TBA: Removed an unused constant 'Nil' that seems to be defined as a macro on mac

diff --git a/Engine/Source/Editor/UnrealEd/Private/Commandlets/TextAssetCommandlet.cpp b/Engine/Source/Editor/UnrealEd/Private/Commandlets/TextAssetCommandlet.cpp
--- a/Engine/Source/Editor/UnrealEd/Private/Commandlets/TextAssetCommandlet.cpp
+++ b/Engine/Source/Editor/UnrealEd/Private/Commandlets/TextAssetCommandlet.cpp
@@ -99,7 +99,6 @@ namespace StringConstants
 	static FString Boolean(TEXT("boolean"));
 	static FString Properties(TEXT("properties"));
 	static FString Type(TEXT("type"));
-	static FString Nil(TEXT("nil"));
 }
 
 inline void WriteSimpleSchemaField(FStructuredArchiveRecord Record, const TCHAR* FieldName, FString& Type, FSimpleSchemaFieldPropertyGenerator PropertiesCallback = FSimpleSchemaFieldPropertyGenerator())
