Missing include

diff --git a/Engine/Source/Editor/UnrealEd/Private/Fbx/FbxMainExport.cpp b/Engine/Source/Editor/UnrealEd/Private/Fbx/FbxMainExport.cpp
--- a/Engine/Source/Editor/UnrealEd/Private/Fbx/FbxMainExport.cpp
+++ b/Engine/Source/Editor/UnrealEd/Private/Fbx/FbxMainExport.cpp
@@ -105,7 +105,11 @@
 #include "Framework/Application/SlateApplication.h"
 #include "Interfaces/IMainFrameModule.h"
 #include "UObject/MetaData.h"
+
+#if WITH_CHAOS
 #include "ChaosCheck.h"
+#include "Chaos/Convex.h"
+#endif
 
 namespace UnFbx
 {
