Expanded the Disk space check to multiple directories with larger space requirements.

diff --git a/Engine/Source/Editor/UnrealEd/Classes/Editor/UnrealEdEngine.h b/Engine/Source/Editor/UnrealEd/Classes/Editor/UnrealEdEngine.h
--- a/Engine/Source/Editor/UnrealEd/Classes/Editor/UnrealEdEngine.h
+++ b/Engine/Source/Editor/UnrealEd/Classes/Editor/UnrealEdEngine.h
@@ -240,6 +240,7 @@ public:
 
 	//~ Begin UEngine Interface.
 	virtual void Init(IEngineLoop* InEngineLoop) override;
+
 	virtual void PreExit() override;
 	virtual void Tick(float DeltaSeconds, bool bIdleMode) override;
 	//~ End UEngine Interface.
@@ -861,4 +862,10 @@ private:
 	* @return	Returns the number of dirty packages that require checkout. If bCheckIfAny is true, returns 1 if any packages will require checkout.
 	*/
 	int32 InternalGetNumDirtyPackagesThatNeedCheckout(bool bCheckIfAny) const;
+
+	/**
+	* Internal function to validate free space on drives used by Unreal Engine
+	* The intent is to notify the user of situations where stability maybe impacted.
+	*/
+	void ValidateFreeDiskSpace() const;
 };
diff --git a/Engine/Source/Editor/UnrealEd/Private/UnrealEdEngine.cpp b/Engine/Source/Editor/UnrealEd/Private/UnrealEdEngine.cpp
--- a/Engine/Source/Editor/UnrealEd/Private/UnrealEdEngine.cpp
+++ b/Engine/Source/Editor/UnrealEd/Private/UnrealEdEngine.cpp
@@ -74,21 +74,8 @@ void UUnrealEdEngine::Init(IEngineLoop* InEngineLoop)
 {
 	Super::Init(InEngineLoop);
 
-	uint64 TotalDiskSpace = 0;
-	uint64 FreeDiskSpace = 0;
-	const TCHAR* BaseDir = FPlatformProcess::BaseDir();
-	if (FPlatformMisc::GetDiskTotalAndFreeSpace(BaseDir, TotalDiskSpace, FreeDiskSpace))
-	{
-		const uint64 HardDriveFreeMB = FreeDiskSpace / (1024 * 1024);
-		if (HardDriveFreeMB < 512)
-		{
-			FEngineAnalytics::LowDriveSpaceDetected();
-			
-			const FText Title = NSLOCTEXT("DriveSpaceDialog", "LowHardDriveSpaceMsgTitle", "Low drive space warning");
-			const FText Message = NSLOCTEXT("DriveSpaceDialog", "LowHardDriveSpaceMsg", "The available space on the drive containing this project is low.\nIt is recommended that you free some space to avoid issues.");
-			FMessageDialog::Open(EAppMsgType::Ok, Message, &Title);
-		}
-	}
+	// Display warnings to the user about disk space issues
+	ValidateFreeDiskSpace();
 
 	// Build databases used by source code navigation
 	FSourceCodeNavigation::Initialize();
@@ -1487,3 +1474,42 @@ void UUnrealEdEngine::UpdateEdModeOnMatineeClose(const FEditorModeID& EditorMode
 		GLevelEditorModeTools().OnEditorModeIDChanged().Remove(UpdateEdModeOnMatineeCloseDelegateHandle);
 	}
 }
+
+bool IsBelowFreeDiskSpaceLimit(const TCHAR* TestDir, FText& OutAppendMessage, const FText& LocationDescriptor, const uint64 MinMB = 5120)
+{
+	uint64 TotalDiskSpace = 0;
+	uint64 FreeDiskSpace = 0;
+
+	if (FPlatformMisc::GetDiskTotalAndFreeSpace(TestDir, TotalDiskSpace, FreeDiskSpace))
+	{
+		const uint64 HardDriveFreeMB = FreeDiskSpace / (1024 * 1024);
+		if (HardDriveFreeMB < MinMB)
+		{
+			static const FText AppendWarning = NSLOCTEXT("DriveSpaceDialog", "LowHardDriveSpaceFormatMsg", "{0}\n  {1} MB Free \t\t {2}\n \t\t\t\t {3} \n");
+
+			OutAppendMessage = FText::Format(AppendWarning, OutAppendMessage, HardDriveFreeMB, FText::FromString(FPaths::ConvertRelativePathToFull(TestDir)), LocationDescriptor);
+
+			return true;
+		}
+	}
+	return false;
+}
+
+void UUnrealEdEngine::ValidateFreeDiskSpace() const
+{
+	FText Message = NSLOCTEXT("DriveSpaceDialog", "LowHardDriveSpaceMsgHeader", "The following drive locations have limited free space.\nIt is recommended that you free some space to avoid issues such as crashed and data loss due to files not being able to be written and saved.\n");
+	
+	bool bShowWarning = false;
+	bShowWarning |= IsBelowFreeDiskSpaceLimit(FPlatformProcess::BaseDir(), Message, NSLOCTEXT("DriveSpaceDialog", "BaseDirDescriptor", "The base engine directory."));
+	bShowWarning |= IsBelowFreeDiskSpaceLimit(FPlatformMisc::ProjectDir(), Message, NSLOCTEXT("DriveSpaceDialog", "ProjectDirDescriptor", "The project directory."));
+	bShowWarning |= IsBelowFreeDiskSpaceLimit(FPlatformProcess::UserDir(), Message, NSLOCTEXT("DriveSpaceDialog", "UserDirDescriptor", "User directory where user specific settings are stored."), 1024);
+
+	if (bShowWarning)
+	{
+		FEngineAnalytics::LowDriveSpaceDetected();
+
+		const FText Title = NSLOCTEXT("DriveSpaceDialog", "LowHardDriveSpaceMsgTitle", "Low drive space warning");
+
+		FMessageDialog::Open(EAppMsgType::Ok, Message, &Title);
+	}
+}
