Cooker garbage collection - change behavior when both MinFreeMemory and MemoryAllowance are set - collect only when all are triggered rather than any.
#rnx

diff --git a/Engine/Source/Editor/UnrealEd/Classes/CookOnTheSide/CookOnTheFlyServer.h b/Engine/Source/Editor/UnrealEd/Classes/CookOnTheSide/CookOnTheFlyServer.h
--- a/Engine/Source/Editor/UnrealEd/Classes/CookOnTheSide/CookOnTheFlyServer.h
+++ b/Engine/Source/Editor/UnrealEd/Classes/CookOnTheSide/CookOnTheFlyServer.h
@@ -600,6 +600,9 @@ private:
 	*/
 	void CookByTheBookFinished();
 
+	/** Print some stats when finishing or cancelling CookByTheBook */
+	void PrintFinishStats();
+
 	/**
 	* Get all the packages which are listed in asset registry passed in.  
 	*
diff --git a/Engine/Source/Editor/UnrealEd/Private/CookOnTheFlyServer.cpp b/Engine/Source/Editor/UnrealEd/Private/CookOnTheFlyServer.cpp
--- a/Engine/Source/Editor/UnrealEd/Private/CookOnTheFlyServer.cpp
+++ b/Engine/Source/Editor/UnrealEd/Private/CookOnTheFlyServer.cpp
@@ -3746,16 +3746,6 @@ bool UCookOnTheFlyServer::HasExceededMaxMemory() const
 		return false;
 	}
 
-	const FPlatformMemoryStats MemStats = FPlatformMemory::GetStats();
-
-	//  if we have less emmory free then we should have then gc some stuff
-	if ((MemStats.AvailablePhysical < MinFreeMemory) && 
-		(MinFreeMemory != 0) )
-	{
-		UE_LOG(LogCook, Display, TEXT("Available physical memory low %d kb, exceeded max memory"), MemStats.AvailablePhysical / 1024);
-		return true;
-	}
-
 #if UE_GC_TRACK_OBJ_AVAILABLE
 	if (GUObjectArray.GetObjectArrayEstimatedAvailable() < MinFreeUObjectIndicesBeforeGC)
 	{
@@ -3764,22 +3754,58 @@ bool UCookOnTheFlyServer::HasExceededMaxMemory() const
 	}
 #endif // UE_GC_TRACK_OBJ_AVAILABLE
 
+
+	// Only report exceeded memory if all the active memory usage triggers have fired
+	int ActiveTriggers = 0;
+	int FiredTriggers = 0;
+
+	TStringBuilder<256> TriggerMessages;
+	const FPlatformMemoryStats MemStats = FPlatformMemory::GetStats();
+	if (MinFreeMemory != 0)
+	{
+		++ActiveTriggers;
+		if (MemStats.AvailablePhysical < MinFreeMemory)
+		{
+			++FiredTriggers;
+			TriggerMessages.Appendf(TEXT("\n  CookSettings.MinFreeMemory: Available physical memory %dMiB is less than %dMiB."),
+				static_cast<uint32>(MemStats.AvailablePhysical / 1024 / 1024), static_cast<uint32>(MinFreeMemory / 1024 / 1024));
+		}
+	}
+
 	// don't gc if we haven't reached our min gc level yet
-	if (MemStats.UsedVirtual < MinMemoryBeforeGC)
+	if (MinMemoryBeforeGC > 0)
 	{
-		return false;
+		++ActiveTriggers;
+		if (MemStats.UsedVirtual >= MinMemoryBeforeGC)
+		{
+			TriggerMessages.Appendf(TEXT("\n  CookSettings.MinMemoryBeforeGC: Used virtual memory %dMiB is greater than %dMiB."),
+				static_cast<uint32>(MemStats.UsedVirtual / 1024 / 1024), static_cast<uint32>(MinMemoryBeforeGC / 1024 / 1024));
+			++FiredTriggers;
+		}
 	}
 
-	//uint64 UsedMemory = MemStats.UsedVirtual; 
-	uint64 UsedMemory = MemStats.UsedPhysical; //should this be used virtual?
-	if ((UsedMemory >= MaxMemoryAllowance) &&
-		(MaxMemoryAllowance > 0u))
+	if (MaxMemoryAllowance > 0u)
 	{
-		UE_LOG(LogCook, Display, TEXT("Used memory high %d kb, exceeded max memory"), MemStats.UsedPhysical / 1024);
-		return true;
+		++ActiveTriggers;
+		//uint64 UsedMemory = MemStats.UsedVirtual; 
+		uint64 UsedMemory = MemStats.UsedPhysical; //should this be used virtual?
+		if (UsedMemory >= MaxMemoryAllowance)
+		{
+			TriggerMessages.Appendf(TEXT("\n  CookSettings.MaxMemoryAllowance: Used physical memory %dMiB is greater than %dMiB."),
+				static_cast<uint32>(MemStats.UsedPhysical / 1024 / 1024), static_cast<uint32>(MaxMemoryAllowance / 1024 / 1024));
+			++FiredTriggers;
+		}
 	}
 
-	return false;
+	if (ActiveTriggers > 0 && FiredTriggers == ActiveTriggers)
+	{
+		UE_LOG(LogCook, Display, TEXT("Exceeded max memory on all configured triggers:%s"), TriggerMessages.ToString());
+		return true;
+	}
+	else
+	{
+		return false;
+	}
 }
 
 TArray<UPackage*> UCookOnTheFlyServer::GetUnsolicitedPackages(const TArray<const ITargetPlatform*>& TargetPlatforms) const
@@ -7014,19 +7040,13 @@ void UCookOnTheFlyServer::CookByTheBookFinished()
 		}
 	}
 
-	const float TotalCookTime = (float)(FPlatformTime::Seconds() - CookByTheBookOptions->CookStartTime);
-	UE_LOG(LogCook, Display, TEXT("Cook by the book total time in tick %fs total time %f"), CookByTheBookOptions->CookTime, TotalCookTime);
-
 	CookByTheBookOptions->BasedOnReleaseCookedPackages.Empty();
-
 	CookByTheBookOptions->bRunning = false;
 	CookByTheBookOptions->bFullLoadAndSave = false;
 
-	const FPlatformMemoryStats MemStats = FPlatformMemory::GetStats();
-
 	PlatformManager->ClearSessionPlatforms();
 
-	UE_LOG(LogCook, Display, TEXT("Peak Used virtual %uMB Peak Used physical %uMB"), MemStats.PeakUsedVirtual / 1024 / 1024, MemStats.PeakUsedPhysical / 1024 / 1024 );
+	PrintFinishStats();
 
 	OutputHierarchyTimers();
 	ClearHierarchyTimers();
@@ -7034,6 +7054,15 @@ void UCookOnTheFlyServer::CookByTheBookFinished()
 	UE_LOG(LogCook, Display, TEXT("Done!"));
 }
 
+void UCookOnTheFlyServer::PrintFinishStats()
+{
+	const float TotalCookTime = (float)(FPlatformTime::Seconds() - CookByTheBookOptions->CookStartTime);
+	UE_LOG(LogCook, Display, TEXT("Cook by the book total time in tick %fs total time %f"), CookByTheBookOptions->CookTime, TotalCookTime);
+
+	const FPlatformMemoryStats MemStats = FPlatformMemory::GetStats();
+	UE_LOG(LogCook, Display, TEXT("Peak Used virtual %u MiB Peak Used physical %u MiB"), MemStats.PeakUsedVirtual / 1024 / 1024, MemStats.PeakUsedPhysical / 1024 / 1024);
+}
+
 void UCookOnTheFlyServer::BuildMapDependencyGraph(const ITargetPlatform* TargetPlatform)
 {
 	auto& MapDependencyGraph = CookByTheBookOptions->MapDependencyGraphs.FindChecked(TargetPlatform);
@@ -7117,6 +7146,8 @@ void UCookOnTheFlyServer::CancelCookByTheBook()
 		CookByTheBookOptions->bRunning = false;
 
 		SandboxFile = nullptr;
+
+		PrintFinishStats();
 	}	
 }
 
