#UE4 CIS fix

diff --git a/Engine/Source/Editor/UnrealEd/Private/PlayLevel.cpp b/Engine/Source/Editor/UnrealEd/Private/PlayLevel.cpp
--- a/Engine/Source/Editor/UnrealEd/Private/PlayLevel.cpp
+++ b/Engine/Source/Editor/UnrealEd/Private/PlayLevel.cpp
@@ -2697,6 +2697,7 @@ UGameInstance* UEditorEngine::CreateInnerProcessPIEGameInstance(FRequestPlaySess
 		ViewportClient = NewObject<UGameViewportClient>(this, GameViewportClientClass);
 		ViewportClient->Init(*PieWorldContext, GameInstance, bCreateNewAudioDevice);
 
+		ULevelEditorPlaySettings* PlayInSettings = GetMutableDefault<ULevelEditorPlaySettings>();
 		ViewportClient->EngineShowFlags.SetServerDrawDebug(PlayInSettings->ShowServerDebugDrawingByDefault());
 
 		if (!InParams.EditorPlaySettings->EnableGameSound)
