Fix crash opening recent level that changed to a redirector

diff --git a/Engine/Source/Editor/LevelEditor/Private/LevelEditorActions.cpp b/Engine/Source/Editor/LevelEditor/Private/LevelEditorActions.cpp
--- a/Engine/Source/Editor/LevelEditor/Private/LevelEditorActions.cpp
+++ b/Engine/Source/Editor/LevelEditor/Private/LevelEditorActions.cpp
@@ -269,10 +269,8 @@ void FLevelEditorActionCallbacks::OpenRecentFile( int32 RecentFileIndex )
 	IMainFrameModule& MainFrameModule = FModuleManager::LoadModuleChecked<IMainFrameModule>( "MainFrame" );
 	FMainMRUFavoritesList* RecentsAndFavorites = MainFrameModule.GetMRUFavoritesList();
 
-	// Save the name of the file we are attempting to load as VerifyFile/AskSaveChanges might rearrange the MRU list on us
-	const FString NewPackageName = RecentsAndFavorites->GetMRUItem( RecentFileIndex );
-	
-	if( RecentsAndFavorites->VerifyMRUFile( RecentFileIndex ) )
+	FString NewPackageName;
+	if( RecentsAndFavorites->VerifyMRUFile( RecentFileIndex, NewPackageName ) )
 	{
 		// Prompt the user to save any outstanding changes.
 		if( FEditorFileUtils::SaveDirtyPackages(true, true, false) )
diff --git a/Engine/Source/Editor/UnrealEd/Private/MRUList.cpp b/Engine/Source/Editor/UnrealEd/Private/MRUList.cpp
--- a/Engine/Source/Editor/UnrealEd/Private/MRUList.cpp
+++ b/Engine/Source/Editor/UnrealEd/Private/MRUList.cpp
@@ -5,6 +5,7 @@
 #include "Misc/PackageName.h"
 #include "Misc/ConfigCacheIni.h"
 #include "Logging/MessageLog.h"
+#include "AssetRegistryModule.h"
 
 FMRUList::FMRUList(const FString& InINISection, const int32 InitMaxItems)
 	:	MaxItems( InitMaxItems ),
@@ -165,10 +166,26 @@ void FMRUList::InternalWriteINI( const TArray<FString>& InItems, const FString&
 }
 
 
-bool FMRUList::VerifyMRUFile(int32 InItem)
+bool FMRUList::VerifyMRUFile(int32 InItem, FString& OutPackageName)
 {
 	check( InItem > -1 && InItem < GetMaxItems() );
-	const FString PackageName = Items[InItem];
+
+	// Handle redirector
+	const FString OriginalPackageName = Items[InItem];
+	const FName OriginalObjectPath = FName(*(OriginalPackageName + TEXT('.') + FPackageName::GetShortName(OriginalPackageName)));
+	FAssetRegistryModule& AssetRegistryModule = FModuleManager::LoadModuleChecked<FAssetRegistryModule>("AssetRegistry");
+	const FName RedirectedObjectPath = AssetRegistryModule.Get().GetRedirectedObjectPath(OriginalObjectPath);
+
+	FString PackageName;
+	FString RedirectedPackageName;
+	if (RedirectedObjectPath != OriginalObjectPath && FPackageName::TryConvertFilenameToLongPackageName(RedirectedObjectPath.ToString(), RedirectedPackageName))
+	{
+		PackageName = RedirectedPackageName;
+	}
+	else
+	{
+		PackageName = OriginalPackageName;
+	}
 
 	FString Filename;
 	bool bSuccess = FPackageName::TryConvertLongPackageNameToFilename(PackageName, Filename, FPackageName::GetMapPackageExtension());
@@ -186,10 +203,21 @@ bool FMRUList::VerifyMRUFile(int32 InItem)
 
 		return false;
 	}
+	else if (!RedirectedPackageName.IsEmpty())
+	{
+		// Remove old path, add new path
+		RemoveMRUItem( InItem );
+		// Note: WriteToINI is called by AddMRUItem
+		AddMRUItem(PackageName);
+	}
+	else
+	{
+		// Otherwise, move the file to the top of the list
+		MoveToTop( InItem );
+		WriteToINI();
+	}
 
-	// Otherwise, move the file to the top of the list
-	MoveToTop( InItem );
-	WriteToINI();
+	OutPackageName = PackageName;
 
 	return true;
 }
diff --git a/Engine/Source/Editor/UnrealEd/Public/MRUList.h b/Engine/Source/Editor/UnrealEd/Public/MRUList.h
--- a/Engine/Source/Editor/UnrealEd/Public/MRUList.h
+++ b/Engine/Source/Editor/UnrealEd/Public/MRUList.h
@@ -66,9 +66,10 @@ public:
 	 * from the list and return false.
 	 *
 	 * @param	ItemIndex		Index of the item to query
+	 * @param	OutPackageName	Long package name after following redirector
 	 * @return					true if the item exists, false if it doesn't
 	 */
-	bool VerifyMRUFile(int32 ItemIndex);
+	bool VerifyMRUFile(int32 ItemIndex, FString& OutPackageName);
 
 	/**
 	 * Accessor.
