USD fixes for Houdini scene
- Fixed UDIMs not finding additional tiles
- Fixed point instancers not finding their meshes because they were part of a collapsed prim.

diff --git a/Engine/Plugins/Importers/USDImporter/Source/USDSchemas/Private/USDGeomPointInstancerTranslator.cpp b/Engine/Plugins/Importers/USDImporter/Source/USDSchemas/Private/USDGeomPointInstancerTranslator.cpp
--- a/Engine/Plugins/Importers/USDImporter/Source/USDSchemas/Private/USDGeomPointInstancerTranslator.cpp
+++ b/Engine/Plugins/Importers/USDImporter/Source/USDSchemas/Private/USDGeomPointInstancerTranslator.cpp
@@ -3,12 +3,14 @@
 #include "USDGeomPointInstancerTranslator.h"
 
 #include "USDConversionUtils.h"
+#include "USDSchemasModule.h"
 #include "USDTypesConversion.h"
 
 #if USE_USD_SDK
 
 #include "Components/HierarchicalInstancedStaticMeshComponent.h"
 #include "Engine/StaticMesh.h"
+#include "Modules/ModuleManager.h"
 
 #include "USDIncludesStart.h"
 	#include "pxr/usd/usd/prim.h"
@@ -83,6 +85,31 @@ namespace UsdGeomPointInstancerTranslatorImpl
 
 		return ( StaticMesh != nullptr );
 	}
+
+	TUsdStore< pxr::SdfPath > UnwindToNonCollapsedPrim( FUsdSchemaTranslator::ECollapsingType CollapsingType, TUsdStore< pxr::UsdPrim > UsdPrim, const TSharedRef< FUsdSchemaTranslationContext >& TranslationContext )
+	{
+		IUsdSchemasModule& UsdSchemasModule = FModuleManager::Get().LoadModuleChecked< IUsdSchemasModule >( TEXT("USDSchemas") );
+
+		TUsdStore< pxr::SdfPath > UsdPrimPath = UsdPrim.Get().GetPrimPath();
+
+		if ( TSharedPtr< FUsdSchemaTranslator > SchemaTranslator = UsdSchemasModule.GetTranslatorRegistry().CreateTranslatorForSchema( TranslationContext, pxr::UsdTyped( UsdPrim.Get() ) ) )
+		{
+			while ( SchemaTranslator->IsCollapsed( CollapsingType ) )
+			{
+				UsdPrimPath = UsdPrimPath.Get().GetParentPath();
+				UsdPrim = UsdPrim.Get().GetStage()->GetPrimAtPath( UsdPrimPath.Get() );
+
+				SchemaTranslator = UsdSchemasModule.GetTranslatorRegistry().CreateTranslatorForSchema( TranslationContext, pxr::UsdTyped( UsdPrim.Get() ) );
+
+				if ( !SchemaTranslator.IsValid() )
+				{
+					break;
+				}
+			}
+		}
+
+		return UsdPrimPath;
+	};
 }
 
 void FUsdGeomPointInstancerTranslator::UpdateComponents( USceneComponent* PointInstancerRootComponent )
@@ -107,6 +134,8 @@ void FUsdGeomPointInstancerTranslator::UpdateComponents( USceneComponent* PointI
 	const pxr::UsdRelationship& Prototypes = PointInstancer.GetPrototypesRel();
 	pxr::SdfPathVector PrototypesPaths;
 
+	TSet< FString > ProcessedPrims;
+
 	if ( Prototypes.GetTargets( &PrototypesPaths ) )
 	{
 		TGuardValue< USceneComponent* > ParentComponentGuard( Context->ParentComponent, PointInstancerRootComponent );
@@ -126,22 +155,32 @@ void FUsdGeomPointInstancerTranslator::UpdateComponents( USceneComponent* PointI
 
 			TGuardValue< USceneComponent* > ParentComponentGuard2( Context->ParentComponent, PrototypeXformComponent );
 
-			// Find first UsdGeomMesh in prototype childs
+			// Find UsdGeomMeshes in prototype childs
 			TArray< TUsdStore< pxr::UsdPrim > > ChildGeomMeshPrims = UsdUtils::GetAllPrimsOfType( PrototypePrim, pxr::TfType::Find< pxr::UsdGeomMesh >() );
 
-			if ( ChildGeomMeshPrims.Num() > 0 )
+			for ( const TUsdStore< pxr::UsdPrim >& PrototypeGeomMeshPrim : ChildGeomMeshPrims )
 			{
-				pxr::UsdPrim PrototypeGeomMeshPrim = ChildGeomMeshPrims[ 0 ].Get();
-				pxr::UsdGeomMesh PrototypeGeomMesh( PrototypeGeomMeshPrim );
-
-				FUsdGeomXformableTranslator PrototypeGeomMeshTranslator( UHierarchicalInstancedStaticMeshComponent::StaticClass(), Context, PrototypeGeomMesh );
+				TUsdStore< pxr::SdfPath > PrototypeTargetPrimPath = UsdGeomPointInstancerTranslatorImpl::UnwindToNonCollapsedPrim( FUsdSchemaTranslator::ECollapsingType::Assets, PrototypeGeomMeshPrim, Context );
 
-				USceneComponent* UsdGeomPrimComponent = PrototypeGeomMeshTranslator.CreateComponentsEx( {}, bNeedsActor );
+				const FString UEPrototypeTargetPrimPath = UsdToUnreal::ConvertPath( PrototypeTargetPrimPath.Get() );
 
-				if ( UHierarchicalInstancedStaticMeshComponent* HismComponent = Cast< UHierarchicalInstancedStaticMeshComponent >( UsdGeomPrimComponent ) )
+				if ( !ProcessedPrims.Contains( UEPrototypeTargetPrimPath ) )
 				{
-					UsdGeomPointInstancerTranslatorImpl::SetStaticMesh( PrototypeGeomMesh, *HismComponent, Context->PrimPathsToAssets );
-					UsdGeomPointInstancerTranslatorImpl::ConvertGeomPointInstancer( Prim.GetStage(), PointInstancer, PrototypeIndex, *HismComponent, pxr::UsdTimeCode( Context->Time ) );
+					ProcessedPrims.Add( UEPrototypeTargetPrimPath );
+
+					if ( pxr::UsdPrim PrototypeTargetPrim = Prim.GetStage()->GetPrimAtPath( PrototypeTargetPrimPath.Get() ) )
+					{
+						pxr::UsdGeomMesh PrototypeGeomMesh( PrototypeTargetPrim );
+						FUsdGeomXformableTranslator PrototypeGeomXformableTranslator( UHierarchicalInstancedStaticMeshComponent::StaticClass(), Context, PrototypeGeomMesh );
+
+						USceneComponent* UsdGeomPrimComponent = PrototypeGeomXformableTranslator.CreateComponentsEx( {}, bNeedsActor );
+
+						if ( UHierarchicalInstancedStaticMeshComponent* HismComponent = Cast< UHierarchicalInstancedStaticMeshComponent >( UsdGeomPrimComponent ) )
+						{
+							UsdGeomPointInstancerTranslatorImpl::SetStaticMesh( PrototypeGeomMesh, *HismComponent, Context->PrimPathsToAssets );
+							UsdGeomPointInstancerTranslatorImpl::ConvertGeomPointInstancer( Prim.GetStage(), PointInstancer, PrototypeIndex, *HismComponent, pxr::UsdTimeCode( Context->Time ) );
+						}
+					}
 				}
 			}
 		}
@@ -149,3 +188,4 @@ void FUsdGeomPointInstancerTranslator::UpdateComponents( USceneComponent* PointI
 }
 
 #endif // #if USE_USD_SDK
+
diff --git a/Engine/Source/Editor/UnrealEd/Private/Factories/EditorFactories.cpp b/Engine/Source/Editor/UnrealEd/Private/Factories/EditorFactories.cpp
--- a/Engine/Source/Editor/UnrealEd/Private/Factories/EditorFactories.cpp
+++ b/Engine/Source/Editor/UnrealEd/Private/Factories/EditorFactories.cpp
@@ -4126,7 +4126,7 @@ bool UTextureFactory::DoesSupportClass(UClass* Class)
 	return Class == UTexture2D::StaticClass() || Class == UTextureCube::StaticClass();
 }
 
-static int32 ParseUDIMName(const FString& Name, const FString& UdimRegexPattern, FString& OutRootName)
+static int32 ParseUDIMName(const FString& Name, const FString& UdimRegexPattern, FString& OutPrefixName, FString& OutPostfixName)
 {
 	FRegexPattern RegexPattern( UdimRegexPattern );
 	FRegexMatcher RegexMatcher( RegexPattern, Name );
@@ -4147,11 +4147,11 @@ static int32 ParseUDIMName(const FString& Name, const FString& UdimRegexPattern,
 		{
 			LexFromString( UdimValue, *Name.Mid( StartOfCaptureGroup2, EndOfCaptureGroup2 - StartOfCaptureGroup2 ) );
 
-			OutRootName = Name.Mid( StartOfCaptureGroup1, EndOfCaptureGroup1 - StartOfCaptureGroup1 );
+			OutPrefixName = Name.Mid( StartOfCaptureGroup1, EndOfCaptureGroup1 - StartOfCaptureGroup1 );
 
 			if ( StartOfCaptureGroup3 != INDEX_NONE && EndOfCaptureGroup3 != INDEX_NONE )
 			{
-				OutRootName.Append( Name.Mid( StartOfCaptureGroup3, EndOfCaptureGroup3 - StartOfCaptureGroup3 ) );
+				OutPostfixName = Name.Mid( StartOfCaptureGroup3, EndOfCaptureGroup3 - StartOfCaptureGroup3 );
 			}
 		}
 	}
@@ -4181,15 +4181,19 @@ UObject* UTextureFactory::FactoryCreateBinary
 	TMap<int32, FString> UDIMIndexToFile;
 	{
 		const FString FilenameNoExtension = FPaths::GetBaseFilename(CurrentFilename);
-		FString BaseUDIMName;
-		const int32 BaseUDIMIndex = ParseUDIMName(FilenameNoExtension, UdimRegexPattern, BaseUDIMName);
+
+		FString PreUDIMName;
+		FString PostUDIMName;
+		const int32 BaseUDIMIndex = ParseUDIMName(FilenameNoExtension, UdimRegexPattern, PreUDIMName, PostUDIMName);
+
+		const FString BaseUDIMName = PreUDIMName + PostUDIMName;
 		if (BaseUDIMIndex != INDEX_NONE)
 		{
 			UDIMIndexToFile.Add(BaseUDIMIndex, CurrentFilename);
 
 			// Filter for other potential UDIM pages, with the same base name and file extension
 			const FString Path = FPaths::GetPath(CurrentFilename);
-			const FString UDIMFilter = (Path / BaseUDIMName) + TEXT("*") + FPaths::GetExtension(CurrentFilename, true);
+			const FString UDIMFilter = (Path / PreUDIMName) + TEXT("*") + PostUDIMName + FPaths::GetExtension(CurrentFilename, true);
 
 			TArray<FString> UDIMFiles;
 			IFileManager::Get().FindFiles(UDIMFiles, *UDIMFilter, true, false);
@@ -4198,8 +4202,8 @@ UObject* UTextureFactory::FactoryCreateBinary
 			{
 				if (!CurrentFilename.EndsWith(UDIMFile) && FactoryCanImport(UDIMFile))
 				{
-					FString UDIMName;
-					const int32 UDIMIndex = ParseUDIMName(FPaths::GetBaseFilename(UDIMFile), UdimRegexPattern, UDIMName);
+					const int32 UDIMIndex = ParseUDIMName(FPaths::GetBaseFilename(UDIMFile), UdimRegexPattern, PreUDIMName, PostUDIMName);
+					const FString UDIMName = PreUDIMName + PostUDIMName;
 					if (!UDIMIndexToFile.Contains(UDIMIndex) && UDIMName == BaseUDIMName)
 					{
 						UDIMIndexToFile.Add(UDIMIndex, Path / UDIMFile);
@@ -4212,36 +4216,41 @@ UObject* UTextureFactory::FactoryCreateBinary
 				// Exclude UDIM number from the name of the UE4 texture asset we create
 				TextureName = *BaseUDIMName;
 
-				// Need to rename the package to match the new texture name, since package was already created
-				// Package name will be the same as the object name, except will contain additional path information,
-				// so we take the existing package name, then extract the UDIM index in order to preserve the path
-				FString PackageName;
-				InParent->GetName(PackageName);
-
-				FString PackageUDIMName;
-				const int32 PackageUDIMIndex = ParseUDIMName(PackageName, UdimRegexPattern, PackageUDIMName);
-				if (PackageUDIMIndex == -1)
+				// Don't try to rename the package if its the transient package
+				if ( InParent != GetTransientPackage() )
 				{
-					// If we're re-importing UDIM texture, the package will already be correctly named after the UDIM base name
-					// In this case we'll fail to parse the UDIM name, but the package should already have the proper name
-					check(PackageName.EndsWith(BaseUDIMName, ESearchCase::CaseSensitive));
-				}
-				else
-				{
-					check(PackageUDIMIndex == BaseUDIMIndex);
-					check(PackageUDIMName.EndsWith(BaseUDIMName, ESearchCase::CaseSensitive));
-
-					// In normal case, higher level code would have already checked for duplicate package name
-					// But since we're changing package name here, check to see if package with the new name already exists...
-					// If it does, code later in this method will prompt user to overwrite the existing asset
-					UPackage* ExistingPackage = FindPackage(InParent->GetOuter(), *PackageUDIMName);
-					if (ExistingPackage)
+					// Need to rename the package to match the new texture name, since package was already created
+					// Package name will be the same as the object name, except will contain additional path information,
+					// so we take the existing package name, then extract the UDIM index in order to preserve the path
+					FString PackageName;
+					InParent->GetName(PackageName);
+
+					const int32 PackageUDIMIndex = ParseUDIMName(PackageName, UdimRegexPattern, PreUDIMName, PostUDIMName);
+					const FString PackageUDIMName = PreUDIMName + PostUDIMName;
+
+					if (PackageUDIMIndex == -1)
 					{
-						InParent = ExistingPackage;
+						// If we're re-importing UDIM texture, the package will already be correctly named after the UDIM base name
+						// In this case we'll fail to parse the UDIM name, but the package should already have the proper name
+						check(PackageName.EndsWith(BaseUDIMName, ESearchCase::CaseSensitive));
 					}
 					else
 					{
-						verify(InParent->Rename(*PackageUDIMName, nullptr, REN_DontCreateRedirectors));
+						check(PackageUDIMIndex == BaseUDIMIndex);
+						check(PackageUDIMName.EndsWith(BaseUDIMName, ESearchCase::CaseSensitive));
+
+						// In normal case, higher level code would have already checked for duplicate package name
+						// But since we're changing package name here, check to see if package with the new name already exists...
+						// If it does, code later in this method will prompt user to overwrite the existing asset
+						UPackage* ExistingPackage = FindPackage(InParent->GetOuter(), *PackageUDIMName);
+						if (ExistingPackage)
+						{
+							InParent = ExistingPackage;
+						}
+						else
+						{
+							verify(InParent->Rename(*PackageUDIMName, nullptr, REN_DontCreateRedirectors));
+						}
 					}
 				}
 			}
