Only inject encryption key asset metadata for platforms that actually require it

diff --git a/Engine/Source/Editor/UnrealEd/Private/Commandlets/AssetRegistryGenerator.cpp b/Engine/Source/Editor/UnrealEd/Private/Commandlets/AssetRegistryGenerator.cpp
--- a/Engine/Source/Editor/UnrealEd/Private/Commandlets/AssetRegistryGenerator.cpp
+++ b/Engine/Source/Editor/UnrealEd/Private/Commandlets/AssetRegistryGenerator.cpp
@@ -1237,7 +1237,10 @@ bool FAssetRegistryGenerator::SaveAssetRegistry(const FString& SandboxPath, bool
 				FAssetRegistryState NewState;
 				NewState.InitializeFromExistingAndPrune(State, CookedPackages, TSet<FName>(), ChunkBucketElement.Value, SaveOptions);
 
-				InjectEncryptionData(NewState);
+				if (!TargetPlatform->HasSecurePackageFormat())
+				{
+					InjectEncryptionData(NewState);
+				}
 
 				// Create runtime registry data
 				FArrayWriter SerializedAssetRegistry;
