Refactored stage time only encryption config to be "explicit" groups that can take normal assets and also external non-uasset based files for encryption

diff --git a/Engine/Source/Editor/UnrealEd/Private/Commandlets/AssetRegistryGenerator.cpp b/Engine/Source/Editor/UnrealEd/Private/Commandlets/AssetRegistryGenerator.cpp
--- a/Engine/Source/Editor/UnrealEd/Private/Commandlets/AssetRegistryGenerator.cpp
+++ b/Engine/Source/Editor/UnrealEd/Private/Commandlets/AssetRegistryGenerator.cpp
@@ -313,29 +313,22 @@ bool FAssetRegistryGenerator::GenerateStreamingInstallManifest(int64 InExtraFlav
 		}
 	}
 
-	// Add manifests for any staging-time only groups
+	// Update manifests for any encryption groups that contain non-asset files files
 	if (bUseAssetManager && !TargetPlatform->HasSecurePackageFormat())
 	{
-		UE_LOG(LogAssetRegistryGenerator, Log, TEXT("Updating stage-only encryption manifests for platform %s"), *TargetPlatform->IniPlatformName());
 		FContentEncryptionConfig ContentEncryptionConfig;
 		UAssetManager::Get().GetContentEncryptionConfig(ContentEncryptionConfig);
-		const FContentEncryptionConfig::TGroupMap& EncryptedNonUFSFileGroups = ContentEncryptionConfig.GetPackageGroupMap();
+		const FContentEncryptionConfig::TGroupMap& EncryptionGroups = ContentEncryptionConfig.GetPackageGroupMap();
 		
-		for (const FContentEncryptionConfig::TGroupMap::ElementType& Element : EncryptedNonUFSFileGroups)
+		for (const FContentEncryptionConfig::TGroupMap::ElementType& GroupElement : EncryptionGroups)
 		{
-			if (Element.Value.bStageTimeOnly)
-			{
-				UE_LOG(LogAssetRegistryGenerator, Log, TEXT("Adding stage-time only manifest for group '%s'"), *Element.Key.ToString());
+			const FName GroupName = GroupElement.Key;
+			const FContentEncryptionConfig::FGroup& EncryptionGroup = GroupElement.Value;
 
-				FName GroupName = Element.Key;
-				const TSet<FName>& PackageNames = Element.Value.PackageNames;
+			if (EncryptionGroup.NonAssetFiles.Num() > 0)
+			{
+				UE_LOG(LogAssetRegistryGenerator, Log, TEXT("Updating non-asset files in manifest for group '%s'"), *GroupName.ToString());
 				
-				FChunkPackageSet* NewManifest = new FChunkPackageSet();
-				for (FName PackageName : PackageNames)
-				{
-					NewManifest->Add(PackageName, FPackageName::LongPackageNameToFilename(PackageName.ToString()));
-				}
-
 				int32 ChunkID = UAssetManager::Get().GetContentEncryptionGroupChunkID(GroupName);
 				int32 PakchunkIndex = GetPakchunkIndex(ChunkID);
 				if (PakchunkIndex >= FinalChunkManifests.Num())
@@ -346,14 +339,18 @@ bool FAssetRegistryGenerator::GenerateStreamingInstallManifest(int64 InExtraFlav
 					FinalChunkManifests.AddZeroed(PakchunkIndex - FinalChunkManifests.Num() + 1);
 				}
 				checkf(PakchunkIndex < FinalChunkManifests.Num(), TEXT("Chunk %i out of range. %i manifests available"), PakchunkIndex, FinalChunkManifests.Num() - 1);
-				checkf(FinalChunkManifests[PakchunkIndex] == nullptr || FinalChunkManifests[PakchunkIndex]->Num() == 0, TEXT("Manifest already exists for chunk %i"), PakchunkIndex);
 
-				if (FinalChunkManifests[PakchunkIndex])
+				FChunkPackageSet* Manifest = FinalChunkManifests[PakchunkIndex];
+				if (Manifest == nullptr)
 				{
-					delete FinalChunkManifests[PakchunkIndex];
+					Manifest = new FChunkPackageSet();
+					FinalChunkManifests[PakchunkIndex] = Manifest;
 				}
 
-				FinalChunkManifests[PakchunkIndex] = NewManifest;
+				for (const FString& NonAssetFile : EncryptionGroup.NonAssetFiles)
+				{
+					Manifest->Add(*NonAssetFile, FPaths::RootDir() / NonAssetFile); // Paths added as relative to the root. The staging code will need to map this onto the target path of all staged assets
+				}
 			}
 		}
 	}
diff --git a/Engine/Source/Runtime/Engine/Classes/Engine/ContentEncryptionConfig.h b/Engine/Source/Runtime/Engine/Classes/Engine/ContentEncryptionConfig.h
--- a/Engine/Source/Runtime/Engine/Classes/Engine/ContentEncryptionConfig.h
+++ b/Engine/Source/Runtime/Engine/Classes/Engine/ContentEncryptionConfig.h
@@ -13,10 +13,17 @@ public:
 		All,
 	};
 
+	enum class EGroupType
+	{
+		Root,
+		Explicit
+	};
+
 	struct FGroup
 	{
 		TSet<FName> PackageNames;
-		bool bStageTimeOnly = false;
+		TSet<FString> NonAssetFiles;
+		EGroupType GroupType = EGroupType::Root;
 		EAllowedReferences AllowedReferences = EAllowedReferences::None;
 	};
 
@@ -27,9 +34,14 @@ public:
 		PackageGroups.FindOrAdd(InGroupName).PackageNames.Add(InPackageName);
 	}
 
-	void SetGroupAsStageTimeOnly(FName InGroupName, bool bInStageTimeOnly)
+	void AddNonAssetFile(FName InGroupName, const FString& InFilename)
+	{
+		PackageGroups.FindOrAdd(InGroupName).NonAssetFiles.Add(InFilename);
+	}
+
+	void SetGroupType(FName InGroupName, EGroupType InGroupType)
 	{
-		PackageGroups.FindOrAdd(InGroupName).bStageTimeOnly = bInStageTimeOnly;
+		PackageGroups.FindOrAdd(InGroupName).GroupType = InGroupType;
 	}
 
 	void SetAllowedReferences(FName InGroupName, EAllowedReferences InAllowedReferences)
