Remove BMP export or URuntimeVirtualTexure
Replace it with BMP export of UVirtualTextureBuilder

diff --git a/Engine/Source/Editor/UnrealEd/Classes/Exporters/TextureExporterBMP.h b/Engine/Source/Editor/UnrealEd/Classes/Exporters/TextureExporterBMP.h
--- a/Engine/Source/Editor/UnrealEd/Classes/Exporters/TextureExporterBMP.h
+++ b/Engine/Source/Editor/UnrealEd/Classes/Exporters/TextureExporterBMP.h
@@ -26,7 +26,7 @@ protected:
 };
 
 UCLASS()
-class URuntimeVirtualTextureExporterBMP : public UTextureExporterBMP
+class UVirtualTextureBuilderExporterBMP : public UTextureExporterBMP
 {
 	GENERATED_UCLASS_BODY()
 };
diff --git a/Engine/Source/Editor/UnrealEd/Private/Factories/EditorFactories.cpp b/Engine/Source/Editor/UnrealEd/Private/Factories/EditorFactories.cpp
--- a/Engine/Source/Editor/UnrealEd/Private/Factories/EditorFactories.cpp
+++ b/Engine/Source/Editor/UnrealEd/Private/Factories/EditorFactories.cpp
@@ -217,8 +217,8 @@
 
 #include "Animation/DebugSkelMeshComponent.h"
 
-#include "VT/RuntimeVirtualTexture.h"
-#include "VT/RuntimeVirtualTextureStreamingProxy.h"
+#include "VT/VirtualTexture.h"
+#include "VT/VirtualTextureBuilder.h"
 
 #if PLATFORM_WINDOWS
 	// Needed for DDS support.
@@ -4897,10 +4897,10 @@ UTextureExporterBMP::UTextureExporterBMP(const FObjectInitializer& ObjectInitial
 
 UTexture2D* UTextureExporterBMP::GetExportTexture(UObject* Object) const
 {
-	URuntimeVirtualTexture* RuntimeVirtualTexture = Cast<URuntimeVirtualTexture>(Object);
-	if (RuntimeVirtualTexture != nullptr && RuntimeVirtualTexture->GetStreamingTexture() != nullptr)
+	UVirtualTextureBuilder* VirtualTextureBuilder = Cast<UVirtualTextureBuilder>(Object);
+	if (VirtualTextureBuilder != nullptr)
 	{
-		return RuntimeVirtualTexture->GetStreamingTexture();
+		return VirtualTextureBuilder->Texture;
 	}
 	return Cast<UTexture2D>(Object);
 }
@@ -4957,7 +4957,7 @@ bool UTextureExporterBMP::ExportBinary( UObject* Object, const TCHAR* Type, FArc
 	int32 SizeX = Texture->Source.GetSizeX();
 	int32 SizeY = Texture->Source.GetSizeY();
 	TArray64<uint8> RawData;
-	Texture->Source.GetMipData(RawData, LayerIndex, 0);
+	Texture->Source.GetMipData(RawData, 0, LayerIndex, 0);
 
 	FBitmapFileHeader bmf;
 	FBitmapInfoHeader bmhdr;
@@ -5019,14 +5019,13 @@ bool UTextureExporterBMP::ExportBinary( UObject* Object, const TCHAR* Type, FArc
 	return true;
 }
 
-URuntimeVirtualTextureExporterBMP::URuntimeVirtualTextureExporterBMP(const FObjectInitializer& ObjectInitializer)
+UVirtualTextureBuilderExporterBMP::UVirtualTextureBuilderExporterBMP(const FObjectInitializer& ObjectInitializer)
 	: Super(ObjectInitializer)
 {
-	SupportedClass = URuntimeVirtualTexture::StaticClass();
+	SupportedClass = UVirtualTextureBuilder::StaticClass();
 	PreferredFormatIndex = 0;
 	FormatExtension.Add(TEXT("BMP"));
 	FormatDescription.Add(TEXT("Windows Bitmap"));
-
 }
 
 /*------------------------------------------------------------------------------
diff --git a/Engine/Source/Editor/VirtualTexturingEditor/Private/RuntimeVirtualTextureAssetTypeActions.cpp b/Engine/Source/Editor/VirtualTexturingEditor/Private/RuntimeVirtualTextureAssetTypeActions.cpp
--- a/Engine/Source/Editor/VirtualTexturingEditor/Private/RuntimeVirtualTextureAssetTypeActions.cpp
+++ b/Engine/Source/Editor/VirtualTexturingEditor/Private/RuntimeVirtualTextureAssetTypeActions.cpp
@@ -227,7 +227,7 @@ void FAssetTypeActions_RuntimeVirtualTexture::GetActions(TArray<UObject*> const&
 
 		MenuBuilder.AddMenuEntry(
 			LOCTEXT("RuntimeVirtualTexture_FindMaterials", "Find Materials Using This"),
-			LOCTEXT("RuntimeVirtualTexture_FindMaterialsTooltip", "Finds all materials that use this material in the content browser."),
+			LOCTEXT("RuntimeVirtualTexture_FindMaterialsTooltip", "Finds all materials that use this texture in the content browser."),
 			FSlateIcon(),
 			FUIAction(
 				FExecuteAction::CreateSP(this, &FAssetTypeActions_RuntimeVirtualTexture::ExecuteFindMaterials, RuntimeVirtualTextures[0]),
