Fix for MFA not being handled in PIE

#tests PIE on prod account with MFA enabled

diff --git a/Engine/Source/Editor/UnrealEd/Classes/Editor/EditorEngine.h b/Engine/Source/Editor/UnrealEd/Classes/Editor/EditorEngine.h
--- a/Engine/Source/Editor/UnrealEd/Classes/Editor/EditorEngine.h
+++ b/Engine/Source/Editor/UnrealEd/Classes/Editor/EditorEngine.h
@@ -3030,6 +3030,9 @@ protected:
 	/** Above function but called a frame later, to stop PIE login from happening from a network callback */
 	virtual void OnLoginPIEComplete_Deferred(int32 LocalUserNum, bool bWasSuccessful, FString ErrorString, FPieLoginStruct DataStruct);
 
+	/** allow for game specific override to determine if login should be treated as successful for pass-through handling instead */
+	virtual bool IsLoginPIESuccessful(int32 LocalUserNum, bool bWasSuccessful, const FString& ErrorString, const FPieLoginStruct& DataStruct) { return bWasSuccessful; }
+
 	/** Called when all PIE instances have been successfully logged in */
 	virtual void OnAllPIEInstancesStarted();
 	
diff --git a/Engine/Source/Editor/UnrealEd/Private/PlayLevel.cpp b/Engine/Source/Editor/UnrealEd/Private/PlayLevel.cpp
--- a/Engine/Source/Editor/UnrealEd/Private/PlayLevel.cpp
+++ b/Engine/Source/Editor/UnrealEd/Private/PlayLevel.cpp
@@ -1467,16 +1467,17 @@ void UEditorEngine::OnLoginPIEComplete_Deferred(int32 LocalUserNum, bool bWasSuc
 		return;
 	}
 
-	// Create a new Game Instance for this.
-	UGameInstance* GameInstance = CreateInnerProcessPIEGameInstance(PlayInEditorSessionInfo->OriginalRequestParams, DataStruct.GameInstancePIEParameters, DataStruct.PIEInstanceIndex);
+	// Detect success based on override
+	bool bPIELoginSucceeded = IsLoginPIESuccessful(LocalUserNum, bWasSuccessful, ErrorString, DataStruct);
 
-	bool bAnyPIELoginFailed = false;
+	// Create a new Game Instance for this.
+	UGameInstance* GameInstance = CreateInnerProcessPIEGameInstance(PlayInEditorSessionInfo->OriginalRequestParams, DataStruct.GameInstancePIEParameters, DataStruct.PIEInstanceIndex);	
 	if (GameInstance)
 	{
 		GameInstance->GetWorldContext()->bWaitingOnOnlineSubsystem = false;
 
 		// Logging after the create so a new MessageLog Page is created
-		if (bWasSuccessful)
+		if (bPIELoginSucceeded)
 		{
 			if (DataStruct.GameInstancePIEParameters.NetMode != EPlayNetMode::PIE_Client )
 			{
@@ -1489,7 +1490,6 @@ void UEditorEngine::OnLoginPIEComplete_Deferred(int32 LocalUserNum, bool bWasSuc
 		}
 		else
 		{
-			bAnyPIELoginFailed = true;
 			if (DataStruct.GameInstancePIEParameters.NetMode != EPlayNetMode::PIE_Client)
 			{
 				FMessageLog(NAME_CategoryPIE).Error(FText::Format(LOCTEXT("LoggedInServerFailure", "Server failed to login. {0}"), FText::FromString(ErrorString)));
@@ -1513,7 +1513,7 @@ void UEditorEngine::OnLoginPIEComplete_Deferred(int32 LocalUserNum, bool bWasSuc
 		// If there are no more instances waiting to log-in then we can do post-launch notifications.
 		if (PlayInEditorSessionInfo->NumOutstandingPIELogins == 0)
 		{
-			if (bAnyPIELoginFailed)
+			if (!bPIELoginSucceeded)
 			{
 				UE_LOG(LogPlayLevel, Warning, TEXT("At least one of the requested PIE instances failed to start, ending PIE session."));
 				EndPlayMap();
