Fix formatting on sanbox cleanup timing message

#rnx

diff --git a/Engine/Source/Editor/UnrealEd/Private/CookOnTheFlyServer.cpp b/Engine/Source/Editor/UnrealEd/Private/CookOnTheFlyServer.cpp
--- a/Engine/Source/Editor/UnrealEd/Private/CookOnTheFlyServer.cpp
+++ b/Engine/Source/Editor/UnrealEd/Private/CookOnTheFlyServer.cpp
@@ -7352,6 +7352,7 @@ void UCookOnTheFlyServer::InitializeSandbox(const TArrayView<const ITargetPlatfo
 	{
 		PlatformNames += Target->PlatformName() + TEXT(" ");
 	}
+	PlatformNames.TrimEndInline();
 	UE_LOG(LogCook, Display, TEXT("Sandbox cleanup took %5.3f seconds for platforms %s"), CleanSandboxTime, *PlatformNames);
 #endif
 }
