Unregister StaticMeshComponent before making changes to it's MeshDescription, otherwise Rendering code may attempt to operate on this component while it is in an invalid state (somewhat-frequently occurs if the build takes long enough for a progress bar to appear, which causes a hitproxy render call that ends up trying to rebuild renderproxy for the currently-being-rebuilt component)
#rnx

diff --git a/Engine/Source/Editor/UnrealEd/Private/Tools/EditorComponentSourceFactory.cpp b/Engine/Source/Editor/UnrealEd/Private/Tools/EditorComponentSourceFactory.cpp
--- a/Engine/Source/Editor/UnrealEd/Private/Tools/EditorComponentSourceFactory.cpp
+++ b/Engine/Source/Editor/UnrealEd/Private/Tools/EditorComponentSourceFactory.cpp
@@ -6,6 +6,7 @@
 #include "Components/PrimitiveComponent.h"
 #include "Components/StaticMeshComponent.h"
 #include "StaticMeshAttributes.h"
+#include "ComponentReregisterContext.h"
 #include "PhysicsEngine/BodySetup.h"
 
 
@@ -47,6 +48,12 @@ void FStaticMeshComponentTarget::CommitMaterialSetUpdate(const FComponentMateria
 	
 	UStaticMesh* StaticMesh = Cast<UStaticMeshComponent>(Component)->GetStaticMesh();
 
+	// flush any pending rendering commands, which might touch this component while we are rebuilding it's mesh
+	FlushRenderingCommands();
+
+	// unregister the component while we update it's static mesh
+	TUniquePtr<FComponentReregisterContext> ComponentReregisterContext = MakeUnique<FComponentReregisterContext>(Component);
+
 	// make sure transactional flag is on
 	StaticMesh->SetFlags(RF_Transactional);
 
@@ -93,7 +100,13 @@ void FStaticMeshComponentTarget::CommitMesh( const FCommitter& Committer )
 	//check(bSaved);
 	UStaticMesh* StaticMesh = Cast<UStaticMeshComponent>(Component)->GetStaticMesh();
 
-	// make sure transactional flag is on
+	// flush any pending rendering commands, which might touch this component while we are rebuilding it's mesh
+	FlushRenderingCommands();
+
+	// unregister the component while we update it's static mesh
+	TUniquePtr<FComponentReregisterContext> ComponentReregisterContext = MakeUnique<FComponentReregisterContext>(Component);
+
+	// make sure transactional flag is on for this asset
 	StaticMesh->SetFlags(RF_Transactional);
 
 	bool bSavedToTransactionBuffer = StaticMesh->Modify();
