Fix crashes on Launch On after editing a material (UE-91496)

- Fixed by [REDACTED_PERSON].

diff --git a/Engine/Source/Editor/UnrealEd/Private/CookOnTheFlyServer.cpp b/Engine/Source/Editor/UnrealEd/Private/CookOnTheFlyServer.cpp
--- a/Engine/Source/Editor/UnrealEd/Private/CookOnTheFlyServer.cpp
+++ b/Engine/Source/Editor/UnrealEd/Private/CookOnTheFlyServer.cpp
@@ -6655,10 +6655,6 @@ void UCookOnTheFlyServer::InitShaderCodeLibrary(void)
 
             if (ShaderFormats.Num() > 0)
 			{
-				if (!IsCookFlagSet(ECookInitializationFlags::Iterative))
-				{
-					FShaderCodeLibrary::CleanDirectories(ShaderFormats);
-				}
 				FShaderCodeLibrary::CookShaderFormats(ShaderFormatsWithStableKeys);
 			}
         }
@@ -6839,8 +6835,10 @@ void UCookOnTheFlyServer::CleanShaderCodeLibraries()
 	const UProjectPackagingSettings* const PackagingSettings = GetDefault<UProjectPackagingSettings>();
     bool const bCacheShaderLibraries = IsUsingShaderCodeLibrary();
 	ITargetPlatformManagerModule& TPM = GetTargetPlatformManagerRef();
+	bool bIterativeCook = IsCookFlagSet(ECookInitializationFlags::Iterative) ||	PackageTracker->CookedPackages.Num() != 0;
+
 	// If not iterative then clean up our temporary files
-	if (bCacheShaderLibraries && PackagingSettings->bShareMaterialShaderCode && !IsCookFlagSet(ECookInitializationFlags::Iterative))
+	if (bCacheShaderLibraries && PackagingSettings->bShareMaterialShaderCode && !bIterativeCook)
 	{
 		for (const ITargetPlatform* TargetPlatform : PlatformManager->GetSessionPlatforms())
 		{
@@ -7280,21 +7278,17 @@ void UCookOnTheFlyServer::InitializeSandbox(const TArrayView<const ITargetPlatfo
 			FPlatformData* PlatformData = PlatformManager->GetPlatformData(Target);
 			const bool bIsIniSettingsOutOfDate = IniSettingsOutOfDate(Target); // needs to be executed for side effects even if non-iterative
 
-			bool bShouldClearCookedContent = false;
-			if (PlatformData->bIsSandboxInitialized)
-			{
-				// We have constructed the sandbox in an earlier cook in this process (e.g. in the editor) and should not clear it again
-				bShouldClearCookedContent = false;
-			}
-			else if (bIsDiffOnly)
+			bool bShouldClearCookedContent = true;
+			if (bIsDiffOnly)
 			{
 				// When looking for deterministic cooking differences in cooked packages, don't delete the packages on disk
 				bShouldClearCookedContent = false;
 			}
-			else if (bIsIterativeCook)
+			else if (bIsIterativeCook || PlatformData->bIsSandboxInitialized)
 			{
 				if (!bIsIniSettingsOutOfDate)
 				{
+					// We have constructed the sandbox in an earlier cook in this process (e.g. in the editor) and should not clear it again
 					bShouldClearCookedContent = false;
 				}
 				else
