Blueprints: Consolidated UBlueprintProjectSettings into UBlueprintEditorProjectSettings

diff --git a/Engine/Source/Editor/BlueprintGraph/Public/BlueprintEditorSettings.h b/Engine/Source/Editor/BlueprintGraph/Public/BlueprintEditorSettings.h
--- a/Engine/Source/Editor/BlueprintGraph/Public/BlueprintEditorSettings.h
+++ b/Engine/Source/Editor/BlueprintGraph/Public/BlueprintEditorSettings.h
@@ -175,14 +175,3 @@ protected:
 	virtual void PostEditChangeProperty(struct FPropertyChangedEvent& PropertyChangedEvent) override;
 	//~ End UObject Interface
 };
-
-UCLASS(config=Editor, defaultconfig)
-class BLUEPRINTGRAPH_API UBlueprintProjectSettings : public UDeveloperSettings
-{
-public:
-	GENERATED_BODY()
-
-	// The list of namespaces to always expose in any Blueprint (for all users of the game/project)
-	UPROPERTY(EditAnywhere, config, Category=Experimental)
-	TArray<FString> NamespacesToAlwaysInclude;
-};
diff --git a/Engine/Source/Editor/Kismet/Private/BlueprintEditor.cpp b/Engine/Source/Editor/Kismet/Private/BlueprintEditor.cpp
--- a/Engine/Source/Editor/Kismet/Private/BlueprintEditor.cpp
+++ b/Engine/Source/Editor/Kismet/Private/BlueprintEditor.cpp
@@ -139,6 +139,7 @@
 #include "AnimationTransitionGraph.h"
 #include "BlueprintEditorModes.h"
 #include "BlueprintEditorSettings.h"
+#include "Settings/EditorProjectSettings.h"
 #include "K2Node_SwitchString.h"
 #include "AnimGraphNode_StateMachineBase.h"
 #include "AnimationStateMachineGraph.h"
@@ -1737,7 +1738,7 @@ struct FBlueprintNamespaceHelper
 	FBlueprintNamespaceHelper()
 	{
 		AddNamespaces(GetDefault<UBlueprintEditorSettings>()->NamespacesToAlwaysInclude);
-		AddNamespaces(GetDefault<UBlueprintProjectSettings>()->NamespacesToAlwaysInclude);
+		AddNamespaces(GetDefault<UBlueprintEditorProjectSettings>()->NamespacesToAlwaysInclude);
 	}
 
 	
diff --git a/Engine/Source/Editor/UnrealEd/Public/Settings/EditorProjectSettings.h b/Engine/Source/Editor/UnrealEd/Public/Settings/EditorProjectSettings.h
--- a/Engine/Source/Editor/UnrealEd/Public/Settings/EditorProjectSettings.h
+++ b/Engine/Source/Editor/UnrealEd/Public/Settings/EditorProjectSettings.h
@@ -158,7 +158,7 @@ public:
 };
 
 
-UCLASS(config=Editor, meta=(DisplayName="Blueprints"), defaultconfig)
+UCLASS(config=Editor, meta=(DisplayName="Blueprint Project Settings"), defaultconfig)
 class UNREALED_API UBlueprintEditorProjectSettings : public UDeveloperSettings
 {
 	GENERATED_UCLASS_BODY()
@@ -191,5 +191,9 @@ public:
 	 */
 	UPROPERTY(EditAnywhere, config, Category= Blueprints, DisplayName = "Compiler Messages Disabled Entirely")
 	TArray<FName> DisabledCompilerMessages;
+
+	// The list of namespaces to always expose in any Blueprint (for all users of the game/project)
+	UPROPERTY(EditAnywhere, config, Category=Experimental)
+	TArray<FString> NamespacesToAlwaysInclude;
 };
 
