Few changes to make ray tracing debug visualization modes consistent with GBuffer visualizations

- Only allow tone mapping on radiance, indirect irradiance and primary rays debug ray tracing view modes
- Apply pow 2.2 on roughness (same as done by BufferVisualization/Roughness engine material graph)

diff --git a/Engine/Shaders/Private/RayTracing/RayTracingDebug.usf b/Engine/Shaders/Private/RayTracing/RayTracingDebug.usf
--- a/Engine/Shaders/Private/RayTracing/RayTracingDebug.usf
+++ b/Engine/Shaders/Private/RayTracing/RayTracingDebug.usf
@@ -68,7 +68,7 @@ RAY_TRACING_ENTRY_RAYGEN(RayTracingDebugMainRGS)
 			Result = float4(Payload.Specular, Payload.Specular, Payload.Specular, 1.0f);
 			break;
 		case RAY_TRACING_DEBUG_VIZ_ROUGHNESS:
-			Result = float4(Payload.Roughness, Payload.Roughness, Payload.Roughness, 1.0f);
+			Result = float4(pow(Payload.Roughness, 2.2f).xxx, 1.0f); // NOTE: pow 2.2 is used for consistency with BufferVisualization/Roughness shader
 			break;
 		case RAY_TRACING_DEBUG_VIZ_IOR:
 			Result = float4(Payload.Ior, Payload.Ior, Payload.Ior, 1.0f);
diff --git a/Engine/Source/Editor/UnrealEd/Private/RayTracingDebugVisualizationMenuCommands.cpp b/Engine/Source/Editor/UnrealEd/Private/RayTracingDebugVisualizationMenuCommands.cpp
--- a/Engine/Source/Editor/UnrealEd/Private/RayTracingDebugVisualizationMenuCommands.cpp
+++ b/Engine/Source/Editor/UnrealEd/Private/RayTracingDebugVisualizationMenuCommands.cpp
@@ -75,13 +75,9 @@ bool FRayTracingDebugVisualizationMenuCommands::DebugModeShouldBeTonemapped(cons
 	static TArray<FName> TonemappedRayTracingDebugVisualizationModes;
 	if (TonemappedRayTracingDebugVisualizationModes.Num() == 0)
 	{
-		// World normal is tonemapped to match raster buffer visualization mode
-		TonemappedRayTracingDebugVisualizationModes.Add(*LOCTEXT("World Normal", "World Normal").ToString());
-		TonemappedRayTracingDebugVisualizationModes.Add(*LOCTEXT("BaseColor", "BaseColor").ToString());
-		TonemappedRayTracingDebugVisualizationModes.Add(*LOCTEXT("DiffuseColor", "DiffuseColor").ToString());
-		TonemappedRayTracingDebugVisualizationModes.Add(*LOCTEXT("SpecularColor", "SpecularColor").ToString());
-		TonemappedRayTracingDebugVisualizationModes.Add(*LOCTEXT("ShadingModelID", "ShadingModelID").ToString());
 		TonemappedRayTracingDebugVisualizationModes.Add(*LOCTEXT("PrimaryRays", "PrimaryRays").ToString());
+		TonemappedRayTracingDebugVisualizationModes.Add(*LOCTEXT("Radiance", "Radiance").ToString());
+		TonemappedRayTracingDebugVisualizationModes.Add(*LOCTEXT("IndirectIrradiance", "IndirectIrradiance").ToString());
 	}
 
 	return TonemappedRayTracingDebugVisualizationModes.Contains(RayTracingDebugModeName);
