Avoid crash when re importing a skeletal mesh and we got a fail reimport.

diff --git a/Engine/Source/Editor/UnrealEd/Private/Fbx/FbxStaticMeshImport.cpp b/Engine/Source/Editor/UnrealEd/Private/Fbx/FbxStaticMeshImport.cpp
--- a/Engine/Source/Editor/UnrealEd/Private/Fbx/FbxStaticMeshImport.cpp
+++ b/Engine/Source/Editor/UnrealEd/Private/Fbx/FbxStaticMeshImport.cpp
@@ -9,6 +9,7 @@
 #include "Misc/Guid.h"
 #include "UObject/Object.h"
 #include "UObject/GarbageCollection.h"
+#include "UObject/MetaData.h"
 #include "UObject/Package.h"
 #include "Misc/PackageName.h"
 #include "Materials/MaterialInterface.h"
@@ -48,6 +49,93 @@ extern void RestoreExistingMeshSettings(struct ExistingStaticMeshData* ExistingM
 extern void RestoreExistingMeshData(struct ExistingStaticMeshData* ExistingMeshDataPtr, UStaticMesh* NewMesh, int32 LodLevel, bool bCanShowDialog, bool bForceConflictingMaterialReset);
 extern void UpdateSomeLodsImportMeshData(UStaticMesh* NewMesh, TArray<int32> *ReimportLodList);
 
+struct FRestoreReimportData
+{
+	UObject* EditorObject = nullptr;
+	UObject* DupObject = nullptr;
+	FString OriginalName;
+	FString OriginalPackageName;
+	TMap<FName, FString> ObjectMetaData;
+
+	FRestoreReimportData()
+	{
+		//Unsupported default constructor
+		check(false);
+	}
+
+	FRestoreReimportData(UStaticMesh* Mesh)
+	{
+		if (!ensure(Mesh))
+		{
+			return;
+		}
+		EditorObject = Mesh;
+		OriginalName = Mesh->GetName();
+		TMap<FName, FString>* ObjectMetaDataPtr = Mesh->GetOutermost()->GetMetaData()->GetMapForObject(Mesh);
+		if (ObjectMetaDataPtr && ObjectMetaDataPtr->Num() > 0)
+		{
+			ObjectMetaData = *ObjectMetaDataPtr;
+		}
+		OriginalPackageName = Mesh->GetOutermost()->GetName();
+		DupObject = StaticDuplicateObject(Mesh, GetTransientPackage());
+		DupObject->AddToRoot();
+	}
+
+	void RestoreMesh(UnFbx::FFbxImporter* FbxImporter)
+	{
+		if (!ensure(EditorObject && DupObject))
+		{
+			return;
+		}
+
+		if(ensure(FbxImporter))
+		{
+			FbxImporter->AddTokenizedErrorMessage(FTokenizedMessage::Create(EMessageSeverity::Error, FText::Format(LOCTEXT("ImportStaticMeshAsSingle", "Fail to reimport mesh {0}. Restoring the original mesh"), FText::FromString(OriginalPackageName))), FFbxErrors::Generic_ImportingNewObjectFailed);
+		}
+
+		UPackage* Package = nullptr;
+		//We have to close any staticmesh editor using this asset
+		int32 EditorCloseCount = GEditor->GetEditorSubsystem<UAssetEditorSubsystem>()->CloseAllEditorsForAsset(EditorObject);
+		//Restore the duplicate asset
+		OriginalPackageName = UPackageTools::SanitizePackageName(OriginalPackageName);
+		Package = CreatePackage(NULL, *OriginalPackageName);
+		UObject* ExistingObject = StaticFindObject(UStaticMesh::StaticClass(), Package, *OriginalName, true);
+		if (ExistingObject)
+		{
+			//Some reimport path did not trash the package in case there is a fail (i.e. scene reimport)
+			//Rename the original mesh and trash it
+			ExistingObject->Rename(nullptr, GetTransientPackage(), REN_DoNotDirty | REN_DontCreateRedirectors | REN_NonTransactional);
+			ExistingObject->MarkPendingKill();
+		}
+		
+		//Rename the dup object
+		DupObject->Rename(*OriginalName, Package, REN_DoNotDirty | REN_DontCreateRedirectors | REN_NonTransactional);
+		if (ObjectMetaData.Num() > 0)
+		{
+			//Package was recreate so restore the metadata
+			UMetaData* PackageMetaData = DupObject->GetOutermost()->GetMetaData();
+			checkSlow(PackageMetaData);
+			PackageMetaData->SetObjectValues(DupObject, ObjectMetaData);
+		}
+		//Since all loaded package are add to root, we have to remove the dup from the root
+		DupObject->RemoveFromRoot();
+		if (EditorCloseCount > 0)
+		{
+			GEditor->GetEditorSubsystem<UAssetEditorSubsystem>()->OpenEditorForAsset(DupObject);
+		}
+	}
+
+	void CleanupDuplicateMesh()
+	{
+		if(DupObject)
+		{
+			DupObject->RemoveFromRoot();
+			DupObject->MarkPendingKill();
+			DupObject = nullptr;
+		}
+	}
+};
+
 static FbxString GetNodeNameWithoutNamespace( FbxNode* Node )
 {
 	FbxString NodeName = Node->GetName();
@@ -956,6 +1044,8 @@ UStaticMesh* UnFbx::FFbxImporter::ReimportSceneStaticMesh(uint64 FbxNodeUniqueId
 {
 	TRACE_CPUPROFILER_EVENT_SCOPE(FFbxImporter::ReimportSceneStaticMesh);
 
+	FRestoreReimportData RestoreData(Mesh);
+
 	TArray<FbxNode*> FbxMeshArray;
 	UStaticMesh* FirstBaseMesh = NULL;
 	FbxNode* Node = NULL;
@@ -1072,8 +1162,18 @@ UStaticMesh* UnFbx::FFbxImporter::ReimportSceneStaticMesh(uint64 FbxNodeUniqueId
 			AddTokenizedErrorMessage(FTokenizedMessage::Create(EMessageSeverity::Error, FText::Format(LOCTEXT("Error_NoFBXMeshFound", "No FBX mesh found when reimport Unreal mesh '{0}'. The FBX file is crashed."), FText::FromString(Mesh->GetName()))), FFbxErrors::Generic_Mesh_MeshNotFound);
 		}
 	}
-	//Don't restore materials when reimporting scene
-	RestoreExistingMeshData(ExistMeshDataPtr, FirstBaseMesh, INDEX_NONE, false, ImportOptions->bResetToFbxOnMaterialConflict);
+
+	if(FirstBaseMesh)
+	{
+		//Don't restore materials when reimporting scene
+		RestoreExistingMeshData(ExistMeshDataPtr, FirstBaseMesh, INDEX_NONE, false, ImportOptions->bResetToFbxOnMaterialConflict);
+		RestoreData.CleanupDuplicateMesh();
+	}
+	else
+	{
+		RestoreData.RestoreMesh(this);
+	}
+
 	return FirstBaseMesh;
 }
 
@@ -1108,6 +1208,8 @@ UStaticMesh* UnFbx::FFbxImporter::ReimportStaticMesh(UStaticMesh* Mesh, UFbxStat
 {
 	TRACE_CPUPROFILER_EVENT_SCOPE(FFbxImporter::ReimportStaticMesh);
 
+	FRestoreReimportData RestoreData(Mesh);
+
 	TArray<FbxNode*> FbxMeshArray;
 	FbxNode* Node = NULL;
 	UStaticMesh* NewMesh = NULL;
@@ -1341,6 +1443,11 @@ UStaticMesh* UnFbx::FFbxImporter::ReimportStaticMesh(UStaticMesh* Mesh, UFbxStat
 	{
 		UpdateSomeLodsImportMeshData(NewMesh, &ReimportLodList);
 		RestoreExistingMeshData(ExistMeshDataPtr, NewMesh, INDEX_NONE, ImportOptions->bCanShowDialog, ImportOptions->bResetToFbxOnMaterialConflict);
+		RestoreData.CleanupDuplicateMesh();
+	}
+	else
+	{
+		RestoreData.RestoreMesh(this);
 	}
 	return NewMesh;
 }
diff --git a/Engine/Source/Editor/UnrealEd/Private/StaticMeshEdit.cpp b/Engine/Source/Editor/UnrealEd/Private/StaticMeshEdit.cpp
--- a/Engine/Source/Editor/UnrealEd/Private/StaticMeshEdit.cpp
+++ b/Engine/Source/Editor/UnrealEd/Private/StaticMeshEdit.cpp
@@ -1017,7 +1017,7 @@ struct ExistingStaticMeshData
 	UStaticMesh::FOnMeshChanged	ExistingOnMeshChanged;
 	UStaticMesh* ExistingComplexCollisionMesh = nullptr;
 
-	TMap<FName, FString>* ExistingUMetaDataTagValues;
+	TMap<FName, FString> ExistingUMetaDataTagValues;
 };
 
 bool IsUsingMaterialSlotNameWorkflow(UAssetImportData* AssetImportData)
@@ -1049,7 +1049,11 @@ ExistingStaticMeshData* SaveExistingStaticMeshData(UStaticMesh* ExistingMesh, Un
 		ExistingMeshDataPtr = new ExistingStaticMeshData();
 
 		//Save the package UMetaData
-		ExistingMeshDataPtr->ExistingUMetaDataTagValues = UMetaData::GetMapForObject(ExistingMesh);
+		TMap<FName, FString>* ExistingUMetaDataTagValues = UMetaData::GetMapForObject(ExistingMesh);
+		if(ExistingUMetaDataTagValues && ExistingUMetaDataTagValues->Num() > 0)
+		{
+			ExistingMeshDataPtr->ExistingUMetaDataTagValues = *ExistingUMetaDataTagValues;
+		}
 
 		ExistingMeshDataPtr->ImportVersion = ExistingMesh->ImportVersion;
 		ExistingMeshDataPtr->UseMaterialNameSlotWorkflow = IsUsingMaterialSlotNameWorkflow(ExistingMesh->AssetImportData);
@@ -1422,11 +1426,11 @@ void RestoreExistingMeshData(ExistingStaticMeshData* ExistingMeshDataPtr, UStati
 	}
 
 	//Restore the package metadata
-	if (ExistingMeshDataPtr->ExistingUMetaDataTagValues)
+	if (ExistingMeshDataPtr->ExistingUMetaDataTagValues.Num() > 0)
 	{
 		UMetaData* PackageMetaData = NewMesh->GetOutermost()->GetMetaData();
 		checkSlow(PackageMetaData);
-		PackageMetaData->SetObjectValues(NewMesh, *ExistingMeshDataPtr->ExistingUMetaDataTagValues);
+		PackageMetaData->SetObjectValues(NewMesh, ExistingMeshDataPtr->ExistingUMetaDataTagValues);
 	}
 
 	//Create a remap material Index use to find the matching section later
