Fix CIS issues.

#rnx

diff --git a/Engine/Plugins/MovieScene/ActorSequence/Source/ActorSequenceEditor/Private/ActorSequenceEditorTabSummoner.cpp b/Engine/Plugins/MovieScene/ActorSequence/Source/ActorSequenceEditor/Private/ActorSequenceEditorTabSummoner.cpp
--- a/Engine/Plugins/MovieScene/ActorSequence/Source/ActorSequenceEditor/Private/ActorSequenceEditorTabSummoner.cpp
+++ b/Engine/Plugins/MovieScene/ActorSequence/Source/ActorSequenceEditor/Private/ActorSequenceEditorTabSummoner.cpp
@@ -103,7 +103,7 @@ private:
 			ComponentIcon = FSlateIconFinder::FindIconBrushForClass( InNodePtr->GetComponentTemplate()->GetClass(), TEXT("SCS.Component") );
 		}
 
-		FText Label = InNodePtr->IsInherited() && !bIsInEditMode
+		FText Label = InNodePtr->IsInheritedComponent() && !bIsInEditMode
 			? FText::Format(LOCTEXT("NativeComponentFormatString","{0} (Inherited)"), FText::FromString(InNodePtr->GetDisplayString()))
 			: FText::FromString(InNodePtr->GetDisplayString());
 
diff --git a/Engine/Source/Editor/UnrealEd/Private/Settings/EditorProjectSettings.cpp b/Engine/Source/Editor/UnrealEd/Private/Settings/EditorProjectSettings.cpp
--- a/Engine/Source/Editor/UnrealEd/Private/Settings/EditorProjectSettings.cpp
+++ b/Engine/Source/Editor/UnrealEd/Private/Settings/EditorProjectSettings.cpp
@@ -156,25 +156,32 @@ void UBlueprintEditorProjectSettings::PostEditChangeProperty(struct FPropertyCha
 	const FName Name = PropertyChangedEvent.Property ? PropertyChangedEvent.Property->GetFName() : NAME_None;
 	if (Name == GET_MEMBER_NAME_CHECKED(UBlueprintEditorProjectSettings, bEnableChildActorExpansionInTreeView))
 	{
+		if (!GEditor)
+		{
+			return;
+		}
+
 		// Find open blueprint editors and refresh them
-		UAssetEditorSubsystem* AssetEditorSubsystem = GEditor->GetEditorSubsystem<UAssetEditorSubsystem>();
-		TArray<UObject*> EditedAssets = AssetEditorSubsystem->GetAllEditedAssets();
-		for (UObject* Asset : EditedAssets)
+		if (UAssetEditorSubsystem* AssetEditorSubsystem = GEditor->GetEditorSubsystem<UAssetEditorSubsystem>())
 		{
-			if (Asset && Asset->IsA<UBlueprint>())
+			TArray<UObject*> EditedAssets = AssetEditorSubsystem->GetAllEditedAssets();
+			for (UObject* Asset : EditedAssets)
 			{
-				TSharedPtr<IToolkit> AssetEditorPtr = FToolkitManager::Get().FindEditorForAsset(Asset);
-				if (AssetEditorPtr.IsValid() && AssetEditorPtr->IsBlueprintEditor())
+				if (Asset && Asset->IsA<UBlueprint>())
 				{
-					TSharedPtr<IBlueprintEditor> BlueprintEditorPtr = StaticCastSharedPtr<IBlueprintEditor>(AssetEditorPtr);
-					BlueprintEditorPtr->RefreshEditors();
+					TSharedPtr<IToolkit> AssetEditorPtr = FToolkitManager::Get().FindEditorForAsset(Asset);
+					if (AssetEditorPtr.IsValid() && AssetEditorPtr->IsBlueprintEditor())
+					{
+						TSharedPtr<IBlueprintEditor> BlueprintEditorPtr = StaticCastSharedPtr<IBlueprintEditor>(AssetEditorPtr);
+						BlueprintEditorPtr->RefreshEditors();
+					}
 				}
 			}
 		}
 
 		// Deselect actors so we are forced to clear the current tree view
 		// @todo - Figure out how to update the tree view directly instead?
-		if (GEditor && GEditor->GetSelectedActorCount() > 0)
+		if (GEditor->GetSelectedActorCount() > 0)
 		{
 			const bool bNoteSelectionChange = true;
 			const bool bDeselectBSPSurfaces = true;
