Fully compile the blueprint when harvesting or creating via child actor components
#rnx

diff --git a/Engine/Source/Editor/UnrealEd/Private/Kismet2/Kismet2.cpp b/Engine/Source/Editor/UnrealEd/Private/Kismet2/Kismet2.cpp
--- a/Engine/Source/Editor/UnrealEd/Private/Kismet2/Kismet2.cpp
+++ b/Engine/Source/Editor/UnrealEd/Private/Kismet2/Kismet2.cpp
@@ -1611,8 +1611,7 @@ void CreateBlueprintFromActors_Internal(UBlueprint* Blueprint, const TArray<AAct
 		}
 	}
 
-	// Regenerate skeleton class as components have been added since initial generation
-	FKismetEditorUtilities::GenerateBlueprintSkeleton(Blueprint, /*bForceRegeneration=*/ true);
+	FKismetEditorUtilities::CompileBlueprint(Blueprint);
 
 	// Notify the asset registry
 	FAssetRegistryModule::AssetCreated(Blueprint);
