Fix missing include
#rnx

diff --git a/Engine/Source/Editor/UnrealEd/Private/Fbx/FbxStaticMeshImport.cpp b/Engine/Source/Editor/UnrealEd/Private/Fbx/FbxStaticMeshImport.cpp
--- a/Engine/Source/Editor/UnrealEd/Private/Fbx/FbxStaticMeshImport.cpp
+++ b/Engine/Source/Editor/UnrealEd/Private/Fbx/FbxStaticMeshImport.cpp
@@ -22,6 +22,7 @@
 #include "Engine/StaticMeshSocket.h"
 #include "Editor.h"
 #include "Modules/ModuleManager.h"
+#include "Subsystems/AssetEditorSubsystem.h"
 
 #include "StaticMeshResources.h"
 #include "ObjectTools.h"
