Workaround for performance issues when using fly mode camera in editor (RMB down + WASD keys).

Only perform TraceWorldForPositionWithDefault (which currently flushes the render thread) on the initial mouse down event, not on all key input events.

diff --git a/Engine/Source/Editor/UnrealEd/Private/LevelEditorViewport.cpp b/Engine/Source/Editor/UnrealEd/Private/LevelEditorViewport.cpp
--- a/Engine/Source/Editor/UnrealEd/Private/LevelEditorViewport.cpp
+++ b/Engine/Source/Editor/UnrealEd/Private/LevelEditorViewport.cpp
@@ -2770,7 +2770,7 @@ bool FLevelEditorViewportClient::InputKey(FViewport* InViewport, int32 Controlle
 	FSceneView* View = CalcSceneView( &ViewFamily );
 
 	// Compute the click location.
-	if ( InputState.IsAnyMouseButtonDown() )
+	if ( InputState.IsMouseButtonEvent() && InputState.IsAnyMouseButtonDown() )
 	{
 		const FViewportCursorLocation Cursor(View, this, HitX, HitY);
 		const FActorPositionTraceResult TraceResult = FActorPositioning::TraceWorldForPositionWithDefault(Cursor, *View);
