When harvesting components in to a class with an existing root component, only root components should get attached to the inherited root.
Fixes relative position of a previously instanced component being incorrect after applying changes to a blueprint.
#rnx

diff --git a/Engine/Source/Editor/UnrealEd/Private/Kismet2/Kismet2.cpp b/Engine/Source/Editor/UnrealEd/Private/Kismet2/Kismet2.cpp
--- a/Engine/Source/Editor/UnrealEd/Private/Kismet2/Kismet2.cpp
+++ b/Engine/Source/Editor/UnrealEd/Private/Kismet2/Kismet2.cpp
@@ -1207,7 +1207,12 @@ void FKismetEditorUtilities::AddComponentsToBlueprint(UBlueprint* Blueprint, con
 
 			if (Parent == nullptr && RootTemplate)
 			{
-				Parent = RootTemplate;
+				AActor* Owner = SceneComponent->GetOwner();
+				const bool bIsRootComponent = (Owner == nullptr || Owner->GetRootComponent() == SceneComponent);
+				if (bIsRootComponent)
+				{
+					Parent = RootTemplate;
+				}
 			}
 
 			SceneComponentNodes.Emplace(SceneComponent, Parent);
