Feature and content Paks can be imported again in the editor.

# This was a side effect of removal of hash storage in index entries from CL 11442058

#rnx

diff --git a/Engine/Source/Developer/PakFileUtilities/Private/PakFileUtilities.cpp b/Engine/Source/Developer/PakFileUtilities/Private/PakFileUtilities.cpp
--- a/Engine/Source/Developer/PakFileUtilities/Private/PakFileUtilities.cpp
+++ b/Engine/Source/Developer/PakFileUtilities/Private/PakFileUtilities.cpp
@@ -3722,7 +3722,7 @@ bool DiffFilesInPaks(const FString& InPakFilename1, const FString& InPakFilename
 			FPakEntry EntryInfo2;
 			EntryInfo2.Serialize(PakReader2, PakFile2.GetInfo().Version);
 
-			if (EntryInfo2 == Entry2)
+			if (EntryInfo2.IndexDataEquals(Entry2))
 			{
 				const FString& PAK2FileName = *It.TryGetFilename();
 				FPakEntry Entry1;
diff --git a/Engine/Source/Editor/UnrealEd/Private/Factories/PackFactory.cpp b/Engine/Source/Editor/UnrealEd/Private/Factories/PackFactory.cpp
--- a/Engine/Source/Editor/UnrealEd/Private/Factories/PackFactory.cpp
+++ b/Engine/Source/Editor/UnrealEd/Private/Factories/PackFactory.cpp
@@ -369,7 +369,7 @@ UObject* UPackFactory::FactoryCreateBinary
 				FPakEntry EntryInfo;
 				EntryInfo.Serialize(PakReader, PakFile.GetInfo().Version);
 
-				if (EntryInfo == Entry)
+				if (EntryInfo.IndexDataEquals(Entry))
 				{
 					FString ConfigString;
 					PackFactoryHelper::ExtractFileToString(Entry, PakReader, CopyBuffer, PersistentCompressionBuffer, ConfigString, PakFile);
@@ -377,7 +377,7 @@ UObject* UPackFactory::FactoryCreateBinary
 				}
 				else
 				{
-					UE_LOG(LogPackFactory, Error, TEXT("Serialized hash mismatch for \"%s\"."), **EntryFilename);
+					UE_LOG(LogPackFactory, Error, TEXT("Index data mismatch for entry: \"%s\"."), **EntryFilename);
 					ErrorCount++;
 				}
 			}
@@ -480,7 +480,7 @@ UObject* UPackFactory::FactoryCreateBinary
 			FPakEntry EntryInfo;
 			EntryInfo.Serialize(PakReader, PakFile.GetInfo().Version);
 
-			if (EntryInfo == Entry)
+			if (EntryInfo.IndexDataEquals(Entry))
 			{
 				if (EntryFilename->StartsWith(TEXT("Source/")) || EntryFilename->Contains(TEXT("/Source/")))
 				{
@@ -558,7 +558,7 @@ UObject* UPackFactory::FactoryCreateBinary
 			}
 			else
 			{
-				UE_LOG(LogPackFactory, Error, TEXT("Serialized hash mismatch for \"%s\"."), **EntryFilename);
+				UE_LOG(LogPackFactory, Error, TEXT("Index data mismatch for entry: \"%s\"."), **EntryFilename);
 				ErrorCount++;
 			}
 		}
