- Modified particle commandlet to output information about lights unsed in systems

#rnx

diff --git a/Engine/Source/Editor/UnrealEd/Classes/Commandlets/ParticleSystemAuditCommandlet.h b/Engine/Source/Editor/UnrealEd/Classes/Commandlets/ParticleSystemAuditCommandlet.h
--- a/Engine/Source/Editor/UnrealEd/Classes/Commandlets/ParticleSystemAuditCommandlet.h
+++ b/Engine/Source/Editor/UnrealEd/Classes/Commandlets/ParticleSystemAuditCommandlet.h
@@ -34,6 +34,8 @@ class UParticleSystemAuditCommandlet : public UCommandlet
 	TSet<FString> ParticleSystemsWithBoneLocationMismatches;
 	/** All particle systems with warmup time on them */
 	TSet<FString> ParticleSystemsWithWarmupTime;
+	/** All particle systems with lights and the details */
+	TSet<FString> ParticleSystemsWithLights;
 
 	/** If a particle system has a spawn rate or burst count greater than this value, it will be reported */
 	UPROPERTY(config)
@@ -49,6 +51,9 @@ class UParticleSystemAuditCommandlet : public UCommandlet
 	/** Only assets in this collection will be considered. If this is left blank, no assets will be filtered by collection */
 	FString FilterCollection;
 
+	/** Package paths to include */
+	TArray<FName> PackagePaths;
+
 	/** Entry point */
 	int32 Main(const FString& Params) override;
 
@@ -66,10 +71,10 @@ class UParticleSystemAuditCommandlet : public UCommandlet
 	 *
 	 *	@return	bool			true if successful, false if not
 	 */
-	bool DumpSimplePSysSet(TSet<FString>& InPSysSet, const TCHAR* InShortFilename);
+	bool DumpSimplePSysSet(TSet<FString>& InPSysSet, const TCHAR* InShortFilename, const TCHAR* OptionalHeader = nullptr);
 
 	/** Generic function to handle dumping values to a CSV file */
-	bool DumpSimpleSet(TSet<FString>& InSet, const TCHAR* InShortFilename, const TCHAR* InObjectClassName);
+	bool DumpSimpleSet(TSet<FString>& InSet, const TCHAR* InShortFilename, const TCHAR* InObjectClassName, const TCHAR* OptionalHeader = nullptr);
 
 	/** Gets an archive to write to an output file */
 	FArchive* GetOutputFile(const TCHAR* InShortFilename);
diff --git a/Engine/Source/Editor/UnrealEd/Private/Commandlets/ParticleSystemAuditCommandlet.cpp b/Engine/Source/Editor/UnrealEd/Private/Commandlets/ParticleSystemAuditCommandlet.cpp
--- a/Engine/Source/Editor/UnrealEd/Private/Commandlets/ParticleSystemAuditCommandlet.cpp
+++ b/Engine/Source/Editor/UnrealEd/Private/Commandlets/ParticleSystemAuditCommandlet.cpp
@@ -9,6 +9,7 @@
 #include "Distributions/DistributionFloatConstant.h"
 #include "AssetData.h"
 #include "Particles/ParticleModuleRequired.h"
+#include "Particles/Light/ParticleModuleLight.h"
 #include "Particles/Spawn/ParticleModuleSpawn.h"
 #include "Particles/Spawn/ParticleModuleSpawnPerUnit.h"
 #include "Particles/TypeData/ParticleModuleTypeDataRibbon.h"
@@ -44,6 +45,22 @@ int32 UParticleSystemAuditCommandlet::Main(const FString& Params)
 
 	FParse::Value(*Params, TEXT("FilterCollection="), FilterCollection);
 
+	FString PackagePathsString;
+	if (FParse::Value(*Params, TEXT("PackagePaths="), PackagePathsString))
+	{
+		TArray<FString> PackagePathsStrings;
+		PackagePathsString.ParseIntoArray(PackagePathsStrings, TEXT(","));
+		for (const FString& v : PackagePathsStrings)
+		{
+			PackagePaths.Add(FName(v));
+		}
+	}
+
+	if (PackagePaths.Num() == 0)
+	{
+		PackagePaths.Add(FName(TEXT("/Game")));
+	}
+
 	ProcessParticleSystems();
 	DumpResults();
 
@@ -57,7 +74,7 @@ bool UParticleSystemAuditCommandlet::ProcessParticleSystems()
 	AssetRegistry.SearchAllAssets(true);
 
 	FARFilter Filter;
-	Filter.PackagePaths.Add(TEXT("/Game"));
+	Filter.PackagePaths = PackagePaths;
 	Filter.bRecursivePaths = true;
 
 	Filter.ClassNames.Add(UParticleSystem::StaticClass()->GetFName());
@@ -123,6 +140,9 @@ bool UParticleSystemAuditCommandlet::ProcessParticleSystems()
 			bool bHasOnlyBeamsOrHasNoEmitters = true;
 			bool bHasSpawnPerUnit = false;
 			bool bMismatchedLODBoneModules = false;
+			bool bHasLights = false;
+			bool bHasShadowCastingLights = false;
+			bool bHasHiQualityLights = false;
 			for (int32 EmitterIdx = 0; EmitterIdx < PSys->Emitters.Num(); EmitterIdx++)
 			{
 				UParticleEmitter* Emitter = PSys->Emitters[EmitterIdx];
@@ -192,6 +212,12 @@ bool UParticleSystemAuditCommandlet::ProcessParticleSystems()
 								{
 									bHasSpawnPerUnit = true;
 								}
+								else if (UParticleModuleLight* LightModule = Cast<UParticleModuleLight>(Module))
+								{
+									bHasLights = true;
+									bHasShadowCastingLights |= LightModule->bShadowCastingLights;
+									bHasHiQualityLights |= LightModule->bHighQualityLights;
+								}
 							}
 						}
 					}
@@ -337,6 +363,17 @@ bool UParticleSystemAuditCommandlet::ProcessParticleSystems()
 				ParticleSystemsWithWarmupTime.Add(PSys->GetPathName());
 			}
 
+			if (bHasLights)
+			{
+				FString SystemString = FString::Printf(
+					TEXT("%s,%s,%s"),
+					*PSys->GetPathName(),
+					bHasHiQualityLights ? TEXT("true") : TEXT("false"),
+					bHasShadowCastingLights ? TEXT("true") : TEXT("false")
+				);
+				ParticleSystemsWithLights.Add(SystemString);
+			}
+
 			if ((PSys->LODMethod == PARTICLESYSTEMLODMETHOD_Automatic) &&
 				(bInvalidLOD == false) && (bSingleLOD == false) &&
 				(PSys->LODDistanceCheckTime == 0.0f))
@@ -395,6 +432,7 @@ void UParticleSystemAuditCommandlet::DumpResults()
 	DumpSimplePSysSet(ParticleSystemsWithFarLODDistance, TEXT("PSysFarLODDistance"));
 	DumpSimplePSysSet(ParticleSystemsWithBoneLocationMismatches, TEXT("PSysBoneLocationLODMismatches"));
 	DumpSimplePSysSet(ParticleSystemsWithWarmupTime, TEXT("PSysWithWarmupTime"));
+	DumpSimplePSysSet(ParticleSystemsWithLights, TEXT("PSysWithLights"), TEXT("Name,HiQuality,ShadowCasting"));
 }
 
 /**
@@ -405,12 +443,12 @@ void UParticleSystemAuditCommandlet::DumpResults()
  *
  *	@return	bool			true if successful, false if not
  */
-bool UParticleSystemAuditCommandlet::DumpSimplePSysSet(TSet<FString>& InPSysSet, const TCHAR* InShortFilename)
+bool UParticleSystemAuditCommandlet::DumpSimplePSysSet(TSet<FString>& InPSysSet, const TCHAR* InShortFilename, const TCHAR* OptionalHeader)
 {
-	return DumpSimpleSet(InPSysSet, InShortFilename, TEXT("ParticleSystem"));
+	return DumpSimpleSet(InPSysSet, InShortFilename, TEXT("ParticleSystem"), OptionalHeader);
 }
 
-bool UParticleSystemAuditCommandlet::DumpSimpleSet(TSet<FString>& InSet, const TCHAR* InShortFilename, const TCHAR* InObjectClassName)
+bool UParticleSystemAuditCommandlet::DumpSimpleSet(TSet<FString>& InSet, const TCHAR* InShortFilename, const TCHAR* InObjectClassName, const TCHAR* OptionalHeader)
 {
 	if (InSet.Num() > 0)
 	{
@@ -421,7 +459,14 @@ bool UParticleSystemAuditCommandlet::DumpSimpleSet(TSet<FString>& InSet, const T
 		if (OutputStream != NULL)
 		{
 			UE_LOG(LogParticleSystemAuditCommandlet, Log, TEXT("Dumping '%s' results..."), InShortFilename);
-			OutputStream->Logf(TEXT("%s,..."), InObjectClassName);
+			if (OptionalHeader != nullptr)
+			{
+				OutputStream->Logf(TEXT("%s"), OptionalHeader);
+			}
+			else
+			{
+				OutputStream->Logf(TEXT("%s,..."), InObjectClassName);
+			}
 			for (TSet<FString>::TIterator DumpIt(InSet); DumpIt; ++DumpIt)
 			{
 				FString ObjName = *DumpIt;
