Fix failure to compile a previously codeless project after adding a featurepack

diff --git a/Engine/Source/Editor/UnrealEd/Private/Factories/PackFactory.cpp b/Engine/Source/Editor/UnrealEd/Private/Factories/PackFactory.cpp
--- a/Engine/Source/Editor/UnrealEd/Private/Factories/PackFactory.cpp
+++ b/Engine/Source/Editor/UnrealEd/Private/Factories/PackFactory.cpp
@@ -701,6 +701,10 @@ UObject* UPackFactory::FactoryCreateBinary
 					}
 					else
 					{
+						// We didn't previously have source, so the UBT target name will be UE4Editor, and attempts to recompile will end up building the wrong target. Now that we have source,
+						// we need to change the UBT target to be the newly created editor module
+						FPlatformMisc::SetUBTTargetName(*(FString(FApp::GetProjectName()) + TEXT("Editor")));
+
 						if (!HotReloadSupport.RecompileModule(FApp::GetProjectName(), *GWarn, ERecompileModuleFlags::ReloadAfterRecompile | ERecompileModuleFlags::ForceCodeProject))
 						{
 							FMessageDialog::Open(EAppMsgType::Ok, NSLOCTEXT("PackFactory", "FailedToCompileNewGameModule", "Failed to compile newly created game module."));
