FBX import material should scan local path only instead of all the assets

diff --git a/Engine/Source/Editor/UnrealEd/Private/Factories/MaterialImportHelpers.cpp b/Engine/Source/Editor/UnrealEd/Private/Factories/MaterialImportHelpers.cpp
--- a/Engine/Source/Editor/UnrealEd/Private/Factories/MaterialImportHelpers.cpp
+++ b/Engine/Source/Editor/UnrealEd/Private/Factories/MaterialImportHelpers.cpp
@@ -72,7 +72,9 @@ UMaterialInterface* UMaterialImportHelpers::FindExistingMaterial(const FString&
 	TArray<FAssetData> AssetData;
 	FARFilter Filter;
 
-	AssetRegistry.SearchAllAssets(true);
+	TArray<FString> ScanPaths;
+	ScanPaths.Add(BasePath);
+	AssetRegistry.ScanPathsSynchronous(ScanPaths, true);
 
 	Filter.bRecursiveClasses = true;
 	Filter.bRecursivePaths = bRecursivePaths;
