In fbx importer, avoid forcing a rescan when scaning material path.
#rnx

diff --git a/Engine/Source/Editor/UnrealEd/Private/Factories/MaterialImportHelpers.cpp b/Engine/Source/Editor/UnrealEd/Private/Factories/MaterialImportHelpers.cpp
--- a/Engine/Source/Editor/UnrealEd/Private/Factories/MaterialImportHelpers.cpp
+++ b/Engine/Source/Editor/UnrealEd/Private/Factories/MaterialImportHelpers.cpp
@@ -74,7 +74,8 @@ UMaterialInterface* UMaterialImportHelpers::FindExistingMaterial(const FString&
 
 	TArray<FString> ScanPaths;
 	ScanPaths.Add(BasePath);
-	AssetRegistry.ScanPathsSynchronous(ScanPaths, true);
+	const bool bForceRescan = false;
+	AssetRegistry.ScanPathsSynchronous(ScanPaths, bForceRescan);
 
 	Filter.bRecursiveClasses = true;
 	Filter.bRecursivePaths = bRecursivePaths;
