Optimised GetAllPackageFilenamesFromAssetRegistry by generating the data in parallel and populating the cache fully.

diff --git a/Engine/Source/Editor/UnrealEd/Private/CookOnTheFlyServer.cpp b/Engine/Source/Editor/UnrealEd/Private/CookOnTheFlyServer.cpp
--- a/Engine/Source/Editor/UnrealEd/Private/CookOnTheFlyServer.cpp
+++ b/Engine/Source/Editor/UnrealEd/Private/CookOnTheFlyServer.cpp
@@ -8237,31 +8237,63 @@ void UCookOnTheFlyServer::HandleNetworkFileServerRecompileShaders(const FShaderR
 
 bool UCookOnTheFlyServer::GetAllPackageFilenamesFromAssetRegistry( const FString& AssetRegistryPath, TArray<FName>& OutPackageFilenames ) const
 {
+	SCOPE_TIMER(GetAllPackageFilenamesFromAssetRegistry);
 	FArrayReader SerializedAssetData;
 	if (FFileHelper::LoadFileToArray(SerializedAssetData, *AssetRegistryPath))
 	{
 		FAssetRegistryState TempState;
-		FAssetRegistrySerializationOptions LoadOptions;
-		LoadOptions.bSerializeDependencies = false;
-		LoadOptions.bSerializePackageData = false;
-
-		TempState.Serialize(SerializedAssetData, LoadOptions);
+		TempState.Serialize(SerializedAssetData, FAssetRegistrySerializationOptions());
 
 		const TMap<FName, const FAssetData*>& RegistryDataMap = TempState.GetObjectPathToAssetDataMap();
 
+		check(OutPackageFilenames.Num() == 0);
+		OutPackageFilenames.SetNum(RegistryDataMap.Num());
+
+		TArray<FName> PackageNames;
+		PackageNames.Reserve(RegistryDataMap.Num());
+
 		for (const TPair<FName, const FAssetData*>& RegistryData : RegistryDataMap)
 		{
-			const FAssetData* NewAssetData = RegistryData.Value;
-			FName CachedPackageFileFName = PackageNameCache->GetCachedStandardPackageFileFName(NewAssetData->ObjectPath);
-			if (CachedPackageFileFName != NAME_None)
+			int32 AddedIndex = PackageNames.Add(RegistryData.Value->ObjectPath);
+			if (PackageNameCache->Contains(PackageNames.Last()))
 			{
-				OutPackageFilenames.Add(CachedPackageFileFName);
+				OutPackageFilenames[AddedIndex] = PackageNameCache->GetCachedStandardPackageFileFName(PackageNames.Last());
 			}
-			else
+		}
+
+		TArray<TTuple<FName, FString, FString>> PackageToStandardFileNames;
+		PackageToStandardFileNames.SetNum(RegistryDataMap.Num());
+
+		ParallelFor(PackageNames.Num(), [&AssetRegistryPath, &OutPackageFilenames, &PackageToStandardFileNames, &PackageNames, this](int32 AssetIndex)
+			{
+				if (!OutPackageFilenames[AssetIndex].IsNone())
+				{
+					return;
+				}
+
+				const FName PackageName = PackageNames[AssetIndex];
+
+				FString PackageFilename;
+				FString StandardFilename;
+				FName StandardFileFName;
+				if (!PackageNameCache->CalculateCacheData(PackageName, PackageFilename, StandardFilename, StandardFileFName))
+				{
+					UE_LOG(LogCook, Warning, TEXT("Could not resolve package %s from %s"), *PackageName.ToString(), *AssetRegistryPath);
+				}
+
+				OutPackageFilenames[AssetIndex] = StandardFileFName;
+				PackageToStandardFileNames[AssetIndex] = { PackageName, MoveTemp(PackageFilename), MoveTemp(StandardFilename) };
+			});
+
+		for (int32 Idx = OutPackageFilenames.Num() - 1; Idx >= 0; --Idx)
+		{
+			if (OutPackageFilenames[Idx] == NAME_None)
 			{
-				UE_LOG(LogCook, Warning, TEXT("Could not resolve package %s from %s"), *NewAssetData->ObjectPath.ToString(), *AssetRegistryPath);
+				OutPackageFilenames.RemoveAtSwap(Idx);
 			}
 		}
+
+		PackageNameCache->AppendCacheResults(MoveTemp(PackageToStandardFileNames));
 		return true;
 	}
 
diff --git a/Engine/Source/Editor/UnrealEd/Private/Cooker/PackageNameCache.h b/Engine/Source/Editor/UnrealEd/Private/Cooker/PackageNameCache.h
--- a/Engine/Source/Editor/UnrealEd/Private/Cooker/PackageNameCache.h
+++ b/Engine/Source/Editor/UnrealEd/Private/Cooker/PackageNameCache.h
@@ -46,6 +46,12 @@ struct FPackageNameCache
 	void			ClearPackageFilenameCache(IAssetRegistry* InAssetRegistry) const;
 	bool			ClearPackageFilenameCacheForPackage(const UPackage* Package) const;
 
+	void			AppendCacheResults(TArray<TTuple<FName, FString, FString>>&& PackageToStandardFileNames);
+
+	bool			CalculateCacheData(FName PackageName, FString& OutPackageFilename, FString& OutStandardFilename, FName& OutStandardFileFName) const;
+
+	bool			Contains(FName PackageName) const;
+	
 private:
 	bool DoesPackageExist(const FName& PackageName, FString* OutFilename) const;
 	const FCachedPackageFilename& Cache(const FName& PackageName) const;
@@ -89,6 +95,21 @@ bool FPackageNameCache::ClearPackageFilenameCacheForPackage(const UPackage* Pack
 	return PackageFilenameCache.Remove(Package->GetFName()) >= 1;
 }
 
+void FPackageNameCache::AppendCacheResults(TArray<TTuple<FName, FString, FString>>&& PackageToStandardFileNames)
+{
+	check(IsInGameThread());
+	for (auto& Entry : PackageToStandardFileNames)
+	{
+		FName PackageName = Entry.Get<0>();
+		FString& PackageFilename = Entry.Get<1>();
+		FString& StandardFilename = Entry.Get<2>();
+
+		FName StandardFileFName(*StandardFilename);
+		PackageFilenameToPackageFNameCache.Add(StandardFileFName, PackageName);
+		PackageFilenameCache.Emplace(PackageName, FCachedPackageFilename(MoveTemp(PackageFilename), MoveTemp(StandardFilename), StandardFileFName));
+	}
+}
+
 bool FPackageNameCache::DoesPackageExist(const FName& PackageName, FString* OutFilename) const
 {
 	if (!AssetRegistry)
@@ -115,6 +136,24 @@ bool FPackageNameCache::DoesPackageExist(const FName& PackageName, FString* OutF
 	return true;
 }
 
+bool FPackageNameCache::CalculateCacheData(FName PackageName, FString& OutPackageFilename, FString& OutStandardFilename, FName& OutStandardFileFName) const
+{
+	if (DoesPackageExist(PackageName, &OutPackageFilename))
+	{
+		OutStandardFilename = OutPackageFilename = FPaths::ConvertRelativePathToFull(OutPackageFilename);
+
+		FPaths::MakeStandardFilename(OutStandardFilename);
+		OutStandardFileFName = FName(*OutStandardFilename);
+		return true;
+	}
+	return false;
+}
+
+bool FPackageNameCache::Contains(FName PackageName) const
+{
+	return PackageFilenameCache.Contains(PackageName);
+}
+
 const FCachedPackageFilename& FPackageNameCache::Cache(const FName& PackageName) const
 {
 	check(IsInGameThread());
@@ -128,18 +167,11 @@ const FCachedPackageFilename& FPackageNameCache::Cache(const FName& PackageName)
 
 	// cache all the things, like it's your birthday!
 
-	FString Filename;
 	FString PackageFilename;
 	FString StandardFilename;
 	FName StandardFileFName = NAME_None;
 
-	if (DoesPackageExist(PackageName, &Filename))
-	{
-		StandardFilename = PackageFilename = FPaths::ConvertRelativePathToFull(Filename);
-
-		FPaths::MakeStandardFilename(StandardFilename);
-		StandardFileFName = FName(*StandardFilename);
-	}
+	CalculateCacheData(PackageName, PackageFilename, StandardFilename, StandardFileFName);
 
 	PackageFilenameToPackageFNameCache.Add(StandardFileFName, PackageName);
 
