Fix for Mac Editor compile

diff --git a/Engine/Source/Editor/UnrealEd/Private/CookOnTheFlyServer.cpp b/Engine/Source/Editor/UnrealEd/Private/CookOnTheFlyServer.cpp
--- a/Engine/Source/Editor/UnrealEd/Private/CookOnTheFlyServer.cpp
+++ b/Engine/Source/Editor/UnrealEd/Private/CookOnTheFlyServer.cpp
@@ -8282,7 +8282,7 @@ bool UCookOnTheFlyServer::GetAllPackageFilenamesFromAssetRegistry( const FString
 				}
 
 				OutPackageFilenames[AssetIndex] = StandardFileFName;
-				PackageToStandardFileNames[AssetIndex] = { PackageName, MoveTemp(PackageFilename), MoveTemp(StandardFilename) };
+				PackageToStandardFileNames[AssetIndex] = TTuple<FName, FString, FString>(PackageName, MoveTemp(PackageFilename), MoveTemp(StandardFilename));
 			});
 
 		for (int32 Idx = OutPackageFilenames.Num() - 1; Idx >= 0; --Idx)
