Clean up FWidget code a bit

#rnx

diff --git a/Engine/Source/Editor/UnrealEd/Private/UnrealWidget.cpp b/Engine/Source/Editor/UnrealEd/Private/UnrealWidget.cpp
--- a/Engine/Source/Editor/UnrealEd/Private/UnrealWidget.cpp
+++ b/Engine/Source/Editor/UnrealEd/Private/UnrealWidget.cpp
@@ -1,188 +1,62 @@
 // Copyright Epic Games, Inc. All Rights Reserved.
 
 #include "UnrealWidget.h"
-#include "Materials/Material.h"
+
 #include "CanvasItem.h"
-#include "Settings/LevelEditorViewportSettings.h"
-#include "Materials/MaterialInstanceDynamic.h"
-#include "EditorViewportClient.h"
+#include "CanvasTypes.h"
+#include "DynamicMeshBuilder.h"
 #include "EdMode.h"
 #include "EditorModeManager.h"
+#include "EditorViewportClient.h"
+#include "Materials/Material.h"
+#include "Materials/MaterialInstanceDynamic.h"
+#include "Settings/LevelEditorViewportSettings.h"
 #include "SnappingUtils.h"
-#include "DynamicMeshBuilder.h"
-#include "CanvasTypes.h"
 
-IMPLEMENT_HIT_PROXY(HWidgetAxis,HHitProxy);
+IMPLEMENT_HIT_PROXY(HWidgetAxis, HHitProxy);
 
-static const float AXIS_LENGTH = 35.0f;
-static const float TRANSLATE_ROTATE_AXIS_CIRCLE_RADIUS = 20.0f;
-static const float TWOD_AXIS_CIRCLE_RADIUS = 10.0f;
 static const float INNER_AXIS_CIRCLE_RADIUS = 48.0f;
-static const float OUTER_AXIS_CIRCLE_RADIUS = 56.0f;
-static const float ROTATION_TEXT_RADIUS = 75.0f;
-static const int32 AXIS_CIRCLE_SIDES = 24;
-static const float ARCALL_RELATIVE_INNER_SIZE = 0.75f;
-static const float AXIS_LENGTH_SCALE_OFFSET = 5.0f;
-
-/*
- *  Simple struct used to create and group data related to the current window's / viewport's space,
- *  orientation, and scale.
- */
-struct FSpaceDescriptor
-{
-    /*
-     *  Creates a new FSpaceDescriptor.
-     *
-     *  @param View         The virtual view for the space.
-     *  @param Viewport     The real viewport for the space.
-     *  @param InLocation   The location of the camera in the virtual space.
-     */
-    FSpaceDescriptor(const FSceneView* View, const FEditorViewportClient* Viewport, const FVector& InLocation) :
-        bIsPerspective(View->ViewMatrices.GetProjectionMatrix().M[3][3] < 1.0f),
-        bIsLocalSpace(Viewport->GetWidgetCoordSystemSpace() == COORD_Local),
-        bIsOrthoXY(!bIsPerspective && FMath::Abs(View->ViewMatrices.GetViewMatrix().M[2][2]) > 0.0f),
-        bIsOrthoXZ(!bIsPerspective && FMath::Abs(View->ViewMatrices.GetViewMatrix().M[1][2]) > 0.0f),
-        bIsOrthoYZ(!bIsPerspective && FMath::Abs(View->ViewMatrices.GetViewMatrix().M[0][2]) > 0.0f),
-        UniformScale(View->WorldToScreen(InLocation).W * (4.0f / View->UnscaledViewRect.Width() / View->ViewMatrices.GetProjectionMatrix().M[0][0])),
-        Scale(CreateScale())
-    {
-    }
-
-    // Wether or not the view is perspective.
-    const bool bIsPerspective;
-
-    // Whether or not the view is in local space.
-    const bool bIsLocalSpace;
-
-    // Whether or not the view is orthogonal to the XY plane.
-    const bool bIsOrthoXY;
-
-    // Whether or not the view is orthogonal to the XZ plane.
-    const bool bIsOrthoXZ;
-
-    // Whether or not the view is orthogonal to the YZ plane.
-    const bool bIsOrthoYZ;
-
-    // The uniform scale for the space.
-    const float UniformScale;
-
-    // The scale vector for the space based on orientation.
-    const FVector Scale;
-
-    /*
-     *  Used to determine whether or not the X axis should be drawn.
-     *
-     *  @param AxisToDraw   The desired axis to draw.
-     *
-     *  @return True if the axis should be drawn. False otherwise.
-     */
-    bool ShouldDrawAxisX(const EAxisList::Type AxisToDraw)
-    {
-        return ShouldDrawAxis(EAxisList::X, AxisToDraw, bIsOrthoYZ);
-    }
-
-    /*
-     *  Used to determine whether or not the Y axis should be drawn.
-     *
-     *  @param AxisToDraw   The desired axis to draw.
-     *
-     *  @return True if the axis should be drawn. False otherwise.
-     */
-    bool ShouldDrawAxisY(const EAxisList::Type AxisToDraw)
-    {
-        return ShouldDrawAxis(EAxisList::Y, AxisToDraw, bIsOrthoXZ);
-    }
-
-    /*
-     *  Used to determine whether or not the Z axis should be drawn.
-     *
-     *  @param AxisToDraw   The desired axis to draw.
-     *
-     *  @return True if the axis should be drawn. False otherwise.
-     */
-    bool ShouldDrawAxisZ(const EAxisList::Type AxisToDraw)
-    {
-        return ShouldDrawAxis(EAxisList::Z, AxisToDraw, bIsOrthoXY);
-    }
-        
-private:
-
-    /*
-     *  Creates a space scale vector from the determined orientation and uniform scale.
-     *
-     *  @return Space scale vector.
-     */
-    FVector CreateScale()
-    {
-        if (bIsOrthoXY)
-        {
-            return FVector(UniformScale, UniformScale, 1.0f);
-        }
-        else if (bIsOrthoXZ)
-        {
-            return FVector(UniformScale, 1.0f, UniformScale);
-        }
-        else if (bIsOrthoYZ)
-        {
-            return FVector(1.0f, UniformScale, UniformScale);
-        }
-        else
-        {
-            return FVector(UniformScale, UniformScale, UniformScale);
-        }
-    }
-    
-    /*
-     *  Used to determine whether or not a specific axis should be drawn.
-     *
-     *  @param AxisToCheck  The axis to check.
-     *  @param AxisToDraw   The desired axis to draw.
-     *  @param bIsOrtho     Whether or not the axis to check is orthogonal to the viewing orientation.
-     *
-     *  @return True if the axis should be drawn. False otherwise.
-     */
-    bool ShouldDrawAxis(const EAxisList::Type AxisToCheck, const EAxisList::Type AxisToDraw, const bool bIsOrtho)
-    {
-        return (AxisToCheck & AxisToDraw) && (bIsPerspective || bIsLocalSpace || !bIsOrtho);
-    }
-};
 
 FWidget::FWidget()
 {
-	EditorModeTools = NULL;
-	TotalDeltaRotation = 0;
+	EditorModeTools      = NULL;
+	TotalDeltaRotation   = 0;
 	CurrentDeltaRotation = 0;
 
-	AxisColorX = FLinearColor(0.594f,0.0197f,0.0f);
-	AxisColorY = FLinearColor(0.1349f,0.3959f,0.0f);
-	AxisColorZ = FLinearColor(0.0251f,0.207f,0.85f);
-	ScreenAxisColor = FLinearColor(0.76, 0.72, 0.14f);
-	PlaneColorXY = FColor::Yellow;
-	ArcBallColor = FColor(128, 128, 128, 6);
-
-	ScreenSpaceColor  = FColor(196, 196, 196);
-	CurrentColor = FColor::Yellow;
+	AxisColorX       = FLinearColor(0.594f, 0.0197f, 0.0f);
+	AxisColorY       = FLinearColor(0.1349f, 0.3959f, 0.0f);
+	AxisColorZ       = FLinearColor(0.0251f, 0.207f, 0.85f);
+	ScreenAxisColor  = FLinearColor(0.76, 0.72, 0.14f);
+	PlaneColorXY     = FColor::Yellow;
+	ArcBallColor     = FColor(128, 128, 128, 6);
+	ScreenSpaceColor = FColor(196, 196, 196);
+	CurrentColor     = FColor::Yellow;
 
 	UMaterial* AxisMaterialBase = GEngine->ArrowMaterial;
 
-	AxisMaterialX = UMaterialInstanceDynamic::Create( AxisMaterialBase, NULL );
-	AxisMaterialX->SetVectorParameterValue( "GizmoColor", AxisColorX );
+	AxisMaterialX = UMaterialInstanceDynamic::Create(AxisMaterialBase, NULL);
+	AxisMaterialX->SetVectorParameterValue("GizmoColor", AxisColorX);
 
-	AxisMaterialY = UMaterialInstanceDynamic::Create( AxisMaterialBase, NULL );
-	AxisMaterialY->SetVectorParameterValue( "GizmoColor", AxisColorY );
+	AxisMaterialY = UMaterialInstanceDynamic::Create(AxisMaterialBase, NULL);
+	AxisMaterialY->SetVectorParameterValue("GizmoColor", AxisColorY);
 
-	AxisMaterialZ = UMaterialInstanceDynamic::Create( AxisMaterialBase, NULL );
-	AxisMaterialZ->SetVectorParameterValue( "GizmoColor", AxisColorZ );
+	AxisMaterialZ = UMaterialInstanceDynamic::Create(AxisMaterialBase, NULL);
+	AxisMaterialZ->SetVectorParameterValue("GizmoColor", AxisColorZ);
 
-	CurrentAxisMaterial = UMaterialInstanceDynamic::Create( AxisMaterialBase, NULL );
-	CurrentAxisMaterial->SetVectorParameterValue("GizmoColor", CurrentColor );
+	CurrentAxisMaterial = UMaterialInstanceDynamic::Create(AxisMaterialBase, NULL);
+	CurrentAxisMaterial->SetVectorParameterValue("GizmoColor", CurrentColor);
 
-	OpaquePlaneMaterialXY = UMaterialInstanceDynamic::Create( AxisMaterialBase, NULL );
-	OpaquePlaneMaterialXY->SetVectorParameterValue( "GizmoColor", FLinearColor::White );
+	OpaquePlaneMaterialXY = UMaterialInstanceDynamic::Create(AxisMaterialBase, NULL);
+	OpaquePlaneMaterialXY->SetVectorParameterValue("GizmoColor", FLinearColor::White);
 
-	TransparentPlaneMaterialXY = (UMaterial*)StaticLoadObject( UMaterial::StaticClass(),NULL,TEXT("/Engine/EditorMaterials/WidgetVertexColorMaterial.WidgetVertexColorMaterial"),NULL,LOAD_None,NULL );
-	
-	GridMaterial = (UMaterial*)StaticLoadObject( UMaterial::StaticClass(),NULL,TEXT("/Engine/EditorMaterials/WidgetGridVertexColorMaterial_Ma.WidgetGridVertexColorMaterial_Ma"),NULL,LOAD_None,NULL );
+	TransparentPlaneMaterialXY = (UMaterial*)StaticLoadObject(
+	    UMaterial::StaticClass(), NULL,
+	    TEXT("/Engine/EditorMaterials/WidgetVertexColorMaterial.WidgetVertexColorMaterial"), NULL, LOAD_None, NULL);
+
+	GridMaterial = (UMaterial*)StaticLoadObject(
+	    UMaterial::StaticClass(), NULL,
+	    TEXT("/Engine/EditorMaterials/WidgetGridVertexColorMaterial_Ma.WidgetGridVertexColorMaterial_Ma"), NULL,
+	    LOAD_None, NULL);
 	if (!GridMaterial)
 	{
 		GridMaterial = TransparentPlaneMaterialXY;
@@ -190,1276 +64,385 @@ FWidget::FWidget()
 
 	CurrentAxis = EAxisList::None;
 
-	CustomCoordSystem = FMatrix::Identity;
+	CustomCoordSystem      = FMatrix::Identity;
 	CustomCoordSystemSpace = COORD_World;
 
 	bAbsoluteTranslationInitialOffsetCached = false;
-	InitialTranslationOffset = FVector::ZeroVector;
-	InitialTranslationPosition = FVector(0, 0, 0);
+	InitialTranslationOffset                = FVector::ZeroVector;
+	InitialTranslationPosition              = FVector(0, 0, 0);
 
-	bDragging = false;
-	bSnapEnabled = false;
-	bDefaultVisibility = true;
+	bDragging               = false;
+	bSnapEnabled            = false;
+	bDefaultVisibility      = true;
 	bIsOrthoDrawingFullRing = false;
 
-	Origin = FVector2D::ZeroVector;
-	XAxisDir = FVector2D::ZeroVector;
-	YAxisDir = FVector2D::ZeroVector;
-	ZAxisDir = FVector2D::ZeroVector;
+	Origin       = FVector2D::ZeroVector;
+	XAxisDir     = FVector2D::ZeroVector;
+	YAxisDir     = FVector2D::ZeroVector;
+	ZAxisDir     = FVector2D::ZeroVector;
 	DragStartPos = FVector2D::ZeroVector;
-	LastDragPos = FVector2D::ZeroVector;
-}
-
-extern ENGINE_API void StringSize(UFont* Font,int32& XL,int32& YL,const TCHAR* Text, FCanvas* Canvas);
-
-void FWidget::SetUsesEditorModeTools( FEditorModeTools* InEditorModeTools )
-{	
-	EditorModeTools = InEditorModeTools;
-}
-
-/**
- * Renders any widget specific HUD text
- * @param Canvas - Canvas to use for 2d rendering
- */
-void FWidget::DrawHUD (FCanvas* Canvas)
-{
-	if (HUDString.Len())
-	{
-		const float DPIScale = Canvas->GetDPIScale();
-		int32 StringPosX = FMath::FloorToInt(HUDInfoPos.X/DPIScale);
-		int32 StringPosY = FMath::FloorToInt(HUDInfoPos.Y/DPIScale);
-
-		//measure string size
-		int32 StringSizeX, StringSizeY;
-		StringSize(GEngine->GetSmallFont(), StringSizeX, StringSizeY, *HUDString);
-		
-		//add some padding to the outside
-		const int32 Border = 5;
-		int32 FillMinX = StringPosX - Border - (StringSizeX>>1);
-		int32 FillMinY = StringPosY - Border;// - (StringSizeY>>1);
-		StringSizeX += 2*Border;
-		StringSizeY += 2*Border;
-
-		//mostly alpha'ed black
-		FCanvasTileItem TileItem( FVector2D( FillMinX, FillMinY), GWhiteTexture, FVector2D( StringSizeX, StringSizeY ), FLinearColor( 0.0f, 0.0f, 0.0f, .25f ) );
-		TileItem.BlendMode = SE_BLEND_Translucent;
-		Canvas->DrawItem( TileItem );
-		FCanvasTextItem TextItem( FVector2D( StringPosX, StringPosY), FText::FromString( HUDString ), GEngine->GetSmallFont(), FLinearColor::White );
-		TextItem.bCentreX = true;
-		Canvas->DrawItem( TextItem );	
-	}
-}
-
-void FWidget::Render( const FSceneView* View,FPrimitiveDrawInterface* PDI, FEditorViewportClient* ViewportClient )
-{
-	check(ViewportClient);
-
-
-	//reset HUD text
-	HUDString.Empty();
-
-	const bool bShowFlagsSupportsWidgetDrawing = View->Family->EngineShowFlags.ModeWidgets;
-	const bool bEditorModeToolsSupportsWidgetDrawing = EditorModeTools ? EditorModeTools->GetShowWidget() : true;
-	bool bDrawWidget;
-
-	// Because the movement routines use the widget axis to determine how to transform mouse movement into
-	// editor object movement, we need to still run through the Render routine even though widget drawing may be
-	// disabled.  So we keep a flag that is used to determine whether or not to actually render anything.  This way
-	// we can still update the widget axis' based on the Context's transform matrices, even though drawing is disabled.
-	if(bDefaultVisibility && bShowFlagsSupportsWidgetDrawing && bEditorModeToolsSupportsWidgetDrawing)
-	{
-		bDrawWidget = true;
-
-		// See if there is a custom coordinate system we should be using, only change it if we are drawing widgets.
-		CustomCoordSystem = ViewportClient->GetWidgetCoordSystem();
-	}
-	else
-	{
-		bDrawWidget = false;
-	}
-
-	CustomCoordSystemSpace = ViewportClient->GetWidgetCoordSystemSpace();
-
-	// If the current modes don't want to use the widget, don't draw it.
-	if( EditorModeTools && !EditorModeTools->UsesTransformWidget() )
-	{
-		CurrentAxis = EAxisList::None;
-		return;
-	}
-
-	FVector WidgetLocation = ViewportClient->GetWidgetLocation();
-	FVector2D NewOrigin;
-	if (View->ScreenToPixel(View->WorldToScreen(WidgetLocation), NewOrigin))
-	{
-		// Only update the viewport-space origin if the position was in front of the camera
-		Origin = NewOrigin;
-	}
-
-	switch( ViewportClient->GetWidgetMode() )
-	{
-		case WM_Translate:
-			Render_Translate(View, PDI, ViewportClient, WidgetLocation, bDrawWidget);
-			break;
-
-		case WM_Rotate:
-			Render_Rotate(View, PDI, ViewportClient, WidgetLocation, bDrawWidget);
-			break;
-
-		case WM_Scale:
-			Render_Scale(View, PDI, ViewportClient, WidgetLocation, bDrawWidget);
-			break;
-
-		case WM_TranslateRotateZ:
-			Render_TranslateRotateZ(View, PDI, ViewportClient, WidgetLocation, bDrawWidget);
-			break;
-
-		case WM_2D:
-			Render_2D(View, PDI, ViewportClient, WidgetLocation, bDrawWidget);
-			break;
-
-		default:
-			break;
-	}
-};
-
-/**
- * Draws an arrow head line for a specific axis.
- */
-void FWidget::Render_Axis(const FSceneView* View, FPrimitiveDrawInterface* PDI, EAxisList::Type InAxis, FMatrix& InMatrix, UMaterialInterface* InMaterial, const FLinearColor& InColor, FVector2D& OutAxisDir, const FVector& InScale, bool bDrawWidget, bool bCubeHead, float AxisLengthOffset)
-{
-	FMatrix AxisRotation = FMatrix::Identity;
-	if( InAxis == EAxisList::Y )
-	{
-		AxisRotation = FRotationMatrix::MakeFromXZ(FVector(0, 1, 0), FVector(0, 0, 1));
-	}
-	else if( InAxis == EAxisList::Z )
-	{
-		AxisRotation = FRotationMatrix::MakeFromXY(FVector(0, 0, 1), FVector(0, 1, 0));
-	}
-
-	FMatrix ArrowToWorld = AxisRotation * InMatrix;
-
-	// The scale that is passed in potentially leaves one component with a scale of 1, if that happens
-	// we need to extract the inform scale and use it to construct the scale that transforms the primitives
-	float UniformScale = InScale.GetMax() > 1.0f ? InScale.GetMax() : InScale.GetMin() < 1.0f ? InScale.GetMin() : 1.0f;
-	// After the primitives have been scaled and transformed, we apply this inverse scale that flattens the dimension
-	// that was scaled up to prevent it from intersecting with the near plane.  In perspective this won't have any effect,
-	// but in the ortho viewports it will prevent scaling in the direction of the camera and thus intersecting the near plane.
-	FVector FlattenScale = FVector(InScale.Component(0) == 1.0f ? 1.0f / UniformScale : 1.0f, InScale.Component(1) == 1.0f ? 1.0f / UniformScale : 1.0f, InScale.Component(2) == 1.0f ? 1.0f / UniformScale : 1.0f);
-
-	FScaleMatrix Scale(UniformScale);
-	ArrowToWorld = Scale * ArrowToWorld;
-
-	if( bDrawWidget )
-	{
-		const bool bDisabled = EditorModeTools ? (EditorModeTools->IsDefaultModeActive() && GEditor->HasLockedActors() ) : false;
-		PDI->SetHitProxy( new HWidgetAxis( InAxis, bDisabled) );
-
-		const float AxisLength = AXIS_LENGTH + GetDefault<ULevelEditorViewportSettings>()->TransformWidgetSizeAdjustment - (AxisLengthOffset * 2);
-		const float HalfHeight = AxisLength/2.0f;
-		const float CylinderRadius = 1.2f;
-		const FVector Offset(0, 0, HalfHeight + AxisLengthOffset);
-
-		switch( InAxis )
-		{
-			case EAxisList::X:
-			{
-				DrawCylinder(PDI, ( Scale * FRotationMatrix(FRotator(-90, 0.f, 0)) * InMatrix ) * FScaleMatrix(FlattenScale), Offset, FVector(1, 0, 0), FVector(0, 1, 0), FVector(0, 0, 1), CylinderRadius, HalfHeight, 16, InMaterial->GetRenderProxy(), SDPG_Foreground);
-				break;
-			}
-			case EAxisList::Y:
-			{
-				DrawCylinder(PDI, (Scale * FRotationMatrix(FRotator(0, 0, 90)) * InMatrix)* FScaleMatrix(FlattenScale), Offset, FVector(1, 0, 0), FVector(0, 1, 0), FVector(0, 0, 1), CylinderRadius, HalfHeight, 16, InMaterial->GetRenderProxy(), SDPG_Foreground );
-				break;
-			}
-			case EAxisList::Z:
-			{
-				DrawCylinder(PDI, ( Scale * InMatrix ) * FScaleMatrix(FlattenScale), Offset, FVector(1, 0, 0), FVector(0, 1, 0), FVector(0, 0, 1), CylinderRadius, HalfHeight, 16, InMaterial->GetRenderProxy(), SDPG_Foreground);
-				break;
-			}
-		}
-
-		if ( bCubeHead )
-		{
-			const float CubeHeadOffset = 3.0f;
-			FVector RootPos(AxisLength + CubeHeadOffset + AxisLengthOffset, 0, 0);
-
-			Render_Cube(PDI, (FTranslationMatrix(RootPos) * ArrowToWorld) * FScaleMatrix(FlattenScale), InMaterial, FVector(4.0f));
-		}
-		else
-		{
-			const float ConeHeadOffset = 12.0f;
-			FVector RootPos(AxisLength + ConeHeadOffset + AxisLengthOffset, 0, 0);
-
-			float Angle = FMath::DegreesToRadians( PI * 5 );
-			DrawCone(PDI, ( FScaleMatrix(-13) * FTranslationMatrix(RootPos) * ArrowToWorld ) * FScaleMatrix(FlattenScale), Angle, Angle, 32, false, FColor::White, InMaterial->GetRenderProxy(), SDPG_Foreground);
-		}
-	
-		PDI->SetHitProxy( NULL );
-	}
-
-	FVector2D NewOrigin;
-	FVector2D AxisEnd;
-	const FVector AxisEndWorld = ArrowToWorld.TransformPosition(FVector(64, 0, 0));
-	const FVector WidgetOrigin = InMatrix.GetOrigin();
-
-	if (View->ScreenToPixel(View->WorldToScreen(WidgetOrigin), NewOrigin) &&
-		View->ScreenToPixel(View->WorldToScreen(AxisEndWorld), AxisEnd))
-	{
-		// If both the origin and the axis endpoint are in front of the camera, trivially calculate the viewport space axis direction
-		OutAxisDir = (AxisEnd - NewOrigin).GetSafeNormal();
-	}
-	else
-	{
-		// If either the origin or axis endpoint are behind the camera, translate the entire widget in front of the camera in the view direction before performing the
-		// viewport space calculation
-		const FMatrix InvViewMatrix = View->ViewMatrices.GetInvViewMatrix();
-		const FVector ViewLocation = InvViewMatrix.GetOrigin();
-		const FVector ViewDirection = InvViewMatrix.GetUnitAxis(EAxis::Z);
-		const FVector Offset = ViewDirection * (FVector::DotProduct(ViewLocation - WidgetOrigin, ViewDirection) + 100.0f);
-		const FVector AdjustedWidgetOrigin = WidgetOrigin + Offset;
-		const FVector AdjustedWidgetAxisEnd = AxisEndWorld + Offset;
-
-		if (View->ScreenToPixel(View->WorldToScreen(AdjustedWidgetOrigin), NewOrigin) &&
-			View->ScreenToPixel(View->WorldToScreen(AdjustedWidgetAxisEnd), AxisEnd))
-		{
-			OutAxisDir = -(AxisEnd - NewOrigin).GetSafeNormal();
-		}
-	}
+	LastDragPos  = FVector2D::ZeroVector;
 }
 
-void FWidget::Render_Cube( FPrimitiveDrawInterface* PDI, const FMatrix& InMatrix, const UMaterialInterface* InMaterial, const FVector& InScale )
-{
-	const FMatrix CubeToWorld = FScaleMatrix(InScale) * InMatrix;
-	DrawBox( PDI, CubeToWorld, FVector(1,1,1), InMaterial->GetRenderProxy(), SDPG_Foreground );
-}
+extern ENGINE_API void StringSize(UFont* Font, int32& XL, int32& YL, const TCHAR* Text, FCanvas* Canvas);
 
-void DrawCornerHelper( FPrimitiveDrawInterface* PDI, const FMatrix& LocalToWorld, const FVector& Length, float Thickness, const FMaterialRenderProxy* MaterialRenderProxy,uint8 DepthPriorityGroup  )
+void FWidget::SetUsesEditorModeTools(FEditorModeTools* InEditorModeTools)
 {
-	const float TH = Thickness;
-
-	float TX = Length.X/2;
-	float TY = Length.Y/2;
-	float TZ = Length.Z/2;
-
-	FDynamicMeshBuilder MeshBuilder(PDI->View->GetFeatureLevel());
-
-	// Top
-	{
-		int32 VertexIndices[4];
-		VertexIndices[0] = MeshBuilder.AddVertex( FVector(-TX, -TY, +TZ), FVector2D::ZeroVector, FVector(1,0,0), FVector(0,1,0), FVector(0,0,1), FColor::White );
-		VertexIndices[1] = MeshBuilder.AddVertex( FVector(-TX, +TY, +TZ), FVector2D::ZeroVector, FVector(1,0,0), FVector(0,1,0), FVector(0,0,1), FColor::White );
-		VertexIndices[2] = MeshBuilder.AddVertex( FVector(+TX, +TY, +TZ), FVector2D::ZeroVector, FVector(1,0,0), FVector(0,1,0), FVector(0,0,1), FColor::White );
-		VertexIndices[3] = MeshBuilder.AddVertex( FVector(+TX, -TY, +TZ), FVector2D::ZeroVector, FVector(1,0,0), FVector(0,1,0), FVector(0,0,1), FColor::White );
-
-		MeshBuilder.AddTriangle(VertexIndices[0],VertexIndices[1],VertexIndices[2]);
-		MeshBuilder.AddTriangle(VertexIndices[0],VertexIndices[2],VertexIndices[3]);
-	}
-
-	//Left
-	{
-		int32 VertexIndices[4];
-		VertexIndices[0] = MeshBuilder.AddVertex( FVector(-TX,  -TY, TZ-TH),	FVector2D::ZeroVector, FVector(0,0,1), FVector(0,1,0), FVector(-1,0,0), FColor::White );
-		VertexIndices[1] = MeshBuilder.AddVertex( FVector(-TX, -TY, TZ),		FVector2D::ZeroVector, FVector(0,0,1), FVector(0,1,0), FVector(-1,0,0), FColor::White );
-		VertexIndices[2] = MeshBuilder.AddVertex( FVector(-TX, +TY, TZ),		FVector2D::ZeroVector, FVector(0,0,1), FVector(0,1,0), FVector(-1,0,0), FColor::White );
-		VertexIndices[3] = MeshBuilder.AddVertex( FVector(-TX, +TY, TZ-TH),		FVector2D::ZeroVector, FVector(0,0,1), FVector(0,1,0), FVector(-1,0,0), FColor::White );
-
-
-		MeshBuilder.AddTriangle(VertexIndices[0],VertexIndices[1],VertexIndices[2]);
-		MeshBuilder.AddTriangle(VertexIndices[0],VertexIndices[2],VertexIndices[3]);
-	}
-
-	// Front
-	{
-		int32 VertexIndices[5];
-		VertexIndices[0] = MeshBuilder.AddVertex( FVector(-TX,	+TY, TZ-TH),	FVector2D::ZeroVector, FVector(1,0,0), FVector(0,0,-1), FVector(0,1,0), FColor::White );
-		VertexIndices[1] = MeshBuilder.AddVertex( FVector(-TX,	+TY, +TZ  ),	FVector2D::ZeroVector, FVector(1,0,0), FVector(0,0,-1), FVector(0,1,0), FColor::White );
-		VertexIndices[2] = MeshBuilder.AddVertex( FVector(+TX-TH, +TY, +TX  ),	FVector2D::ZeroVector, FVector(1,0,0), FVector(0,0,-1), FVector(0,1,0), FColor::White );
-		VertexIndices[3] = MeshBuilder.AddVertex( FVector(+TX,	+TY, +TZ  ),	FVector2D::ZeroVector, FVector(1,0,0), FVector(0,0,-1), FVector(0,1,0), FColor::White );
-		VertexIndices[4] = MeshBuilder.AddVertex( FVector(+TX-TH, +TY, TZ-TH),	FVector2D::ZeroVector, FVector(1,0,0), FVector(0,0,-1), FVector(0,1,0), FColor::White );
-
-		MeshBuilder.AddTriangle(VertexIndices[0],VertexIndices[1],VertexIndices[2]);
-		MeshBuilder.AddTriangle(VertexIndices[0],VertexIndices[2],VertexIndices[4]);
-		MeshBuilder.AddTriangle(VertexIndices[4],VertexIndices[2],VertexIndices[3]);
-	}
-
-	// Back
-	{
-		int32 VertexIndices[5];
-		VertexIndices[0] = MeshBuilder.AddVertex( FVector(-TX,	-TY, TZ-TH),	FVector2D::ZeroVector, FVector(1,0,0), FVector(0,0,1), FVector(0,-1,0), FColor::White );
-		VertexIndices[1] = MeshBuilder.AddVertex( FVector(-TX,	-TY, +TZ),		FVector2D::ZeroVector, FVector(1,0,0), FVector(0,0,1), FVector(0,-1,0), FColor::White );
-		VertexIndices[2] = MeshBuilder.AddVertex( FVector(+TX-TH, -TY, +TX),	FVector2D::ZeroVector, FVector(1,0,0), FVector(0,0,1), FVector(0,-1,0), FColor::White );
-		VertexIndices[3] = MeshBuilder.AddVertex( FVector(+TX,	-TY, +TZ),		FVector2D::ZeroVector, FVector(1,0,0), FVector(0,0,1), FVector(0,-1,0), FColor::White );
-		VertexIndices[4] = MeshBuilder.AddVertex( FVector(+TX-TH, -TY, TZ-TH),	FVector2D::ZeroVector, FVector(1,0,0), FVector(0,0,1), FVector(0,-1,0), FColor::White );
-
-		MeshBuilder.AddTriangle(VertexIndices[0],VertexIndices[1],VertexIndices[2]);
-		MeshBuilder.AddTriangle(VertexIndices[0],VertexIndices[2],VertexIndices[4]);
-		MeshBuilder.AddTriangle(VertexIndices[4],VertexIndices[2],VertexIndices[3]);
-	}
-	// Bottom
-	{
-		int32 VertexIndices[4];
-		VertexIndices[0] = MeshBuilder.AddVertex( FVector(-TX, -TY, TZ-TH),		FVector2D::ZeroVector, FVector(1,0,0), FVector(0,0,-1), FVector(0,0,1), FColor::White );
-		VertexIndices[1] = MeshBuilder.AddVertex( FVector(-TX, +TY, TZ-TH),		FVector2D::ZeroVector, FVector(1,0,0), FVector(0,0,-1), FVector(0,0,1), FColor::White );
-		VertexIndices[2] = MeshBuilder.AddVertex( FVector(+TX-TH, +TY, TZ-TH),	FVector2D::ZeroVector, FVector(1,0,0), FVector(0,0,-1), FVector(0,0,1), FColor::White );
-		VertexIndices[3] = MeshBuilder.AddVertex( FVector(+TX-TH, -TY, TZ-TH),	FVector2D::ZeroVector, FVector(1,0,0), FVector(0,0,-1), FVector(0,0,1), FColor::White );
-
-		MeshBuilder.AddTriangle(VertexIndices[0],VertexIndices[1],VertexIndices[2]);
-		MeshBuilder.AddTriangle(VertexIndices[0],VertexIndices[2],VertexIndices[3]);
-	}
-	MeshBuilder.Draw(PDI,LocalToWorld,MaterialRenderProxy,DepthPriorityGroup,0.f);
-}
-
-void DrawDualAxis( FPrimitiveDrawInterface* PDI, const FMatrix& BoxToWorld,const FVector& Length, float Thickness, const FMaterialRenderProxy* AxisMat,const FMaterialRenderProxy* Axis2Mat )
-{
-	DrawCornerHelper(PDI, BoxToWorld, Length, Thickness, Axis2Mat, SDPG_Foreground);
-	DrawCornerHelper(PDI, FScaleMatrix(FVector(-1, 1, 1)) * FRotationMatrix(FRotator(-90, 0, 0)) * BoxToWorld, Length, Thickness, AxisMat, SDPG_Foreground);
+	EditorModeTools = InEditorModeTools;
 }
-/**
- * Draws the translation widget.
- */
-void FWidget::Render_Translate( const FSceneView* View, FPrimitiveDrawInterface* PDI, FEditorViewportClient* ViewportClient, const FVector& InLocation, bool bDrawWidget )
-{
-	// Figure out axis colors
-	const FLinearColor& XColor = ( CurrentAxis&EAxisList::X ? (FLinearColor)CurrentColor : AxisColorX );
-	const FLinearColor& YColor = ( CurrentAxis&EAxisList::Y ? (FLinearColor)CurrentColor : AxisColorY );
-	const FLinearColor& ZColor = ( CurrentAxis&EAxisList::Z ? (FLinearColor)CurrentColor : AxisColorZ );
-	FColor CurrentScreenColor = ( CurrentAxis & EAxisList::Screen ? CurrentColor : ScreenSpaceColor );
-
-	// Figure out axis matrices
-	FMatrix WidgetMatrix = CustomCoordSystem * FTranslationMatrix( InLocation );
-	const EAxisList::Type DrawAxis = GetAxisToDraw( ViewportClient->GetWidgetMode() );
-	const bool bDisabled = IsWidgetDisabled();
-
-    FSpaceDescriptor Space(View, ViewportClient, InLocation);
-
-	// Draw the axis lines with arrow heads
-	if( Space.ShouldDrawAxisX(DrawAxis) )
-	{
-		UMaterialInstanceDynamic* XMaterial = ( CurrentAxis&EAxisList::X ? CurrentAxisMaterial : AxisMaterialX );
-		Render_Axis( View, PDI, EAxisList::X, WidgetMatrix, XMaterial, XColor, XAxisDir, Space.Scale, bDrawWidget );
-	}
-
-	if( Space.ShouldDrawAxisY(DrawAxis) )
-	{
-		UMaterialInstanceDynamic* YMaterial = ( CurrentAxis&EAxisList::Y ? CurrentAxisMaterial : AxisMaterialY );
-		Render_Axis( View, PDI, EAxisList::Y, WidgetMatrix, YMaterial, YColor, YAxisDir, Space.Scale, bDrawWidget );
-	}
-
-	if( Space.ShouldDrawAxisZ(DrawAxis) )
-	{
-		UMaterialInstanceDynamic* ZMaterial = ( CurrentAxis&EAxisList::Z ? CurrentAxisMaterial : AxisMaterialZ );
-		Render_Axis( View, PDI, EAxisList::Z, WidgetMatrix, ZMaterial, ZColor, ZAxisDir, Space.Scale, bDrawWidget );
-	}
-
-	// Draw the grabbers
-	if( bDrawWidget )
-	{
-		FVector CornerPos = FVector(7, 0, 7) * Space.UniformScale;
-		FVector AxisSize = FVector(12, 1.2, 12) * Space.UniformScale;
-		float CornerLength = 1.2f * Space.UniformScale;
-
-		// After the primitives have been scaled and transformed, we apply this inverse scale that flattens the dimension
-		// that was scaled up to prevent it from intersecting with the near plane.  In perspective this won't have any effect,
-		// but in the ortho viewports it will prevent scaling in the direction of the camera and thus intersecting the near plane.
-		FVector FlattenScale = FVector(Space.Scale.Component(0) == 1.0f ? 1.0f / Space.UniformScale : 1.0f, Space.Scale.Component(1) == 1.0f ? 1.0f / Space.UniformScale : 1.0f, Space.Scale.Component(2) == 1.0f ? 1.0f / Space.UniformScale : 1.0f);
 
-		if (Space.bIsPerspective || Space.bIsLocalSpace || Space.bIsOrthoXY)
-		{
-			if( (DrawAxis&EAxisList::XY) == EAxisList::XY )							// Top
-			{
-				UMaterialInstanceDynamic* XMaterial = ( (CurrentAxis&EAxisList::XY) == EAxisList::XY ? CurrentAxisMaterial : AxisMaterialX );
-				UMaterialInstanceDynamic* YMaterial = ( (CurrentAxis&EAxisList::XY) == EAxisList::XY? CurrentAxisMaterial : AxisMaterialY );
-
-				PDI->SetHitProxy( new HWidgetAxis(EAxisList::XY, bDisabled) );
+void FWidget::ConvertMouseToAxis_Translate(FVector2D DragDir, FVector& InOutDelta, FVector& OutDrag) const
+{
+	// Get drag delta in widget axis space
+	OutDrag = FVector((CurrentAxis & EAxisList::X) ? FVector2D::DotProduct(XAxisDir, DragDir) : 0.0f,
+	                  (CurrentAxis & EAxisList::Y) ? FVector2D::DotProduct(YAxisDir, DragDir) : 0.0f,
+	                  (CurrentAxis & EAxisList::Z) ? FVector2D::DotProduct(ZAxisDir, DragDir) : 0.0f);
+
+	// Snap to grid in widget axis space
+	const FVector GridSize = FVector(GEditor->GetGridSize());
+	FSnappingUtils::SnapPointToGrid(OutDrag, GridSize);
+
+	// Convert to effective screen space delta, and replace input delta, adjusted for inverted screen space Y axis
+	const FVector2D EffectiveDelta = OutDrag.X * XAxisDir + OutDrag.Y * YAxisDir + OutDrag.Z * ZAxisDir;
+	InOutDelta                     = FVector(EffectiveDelta.X, -EffectiveDelta.Y, 0.0f);
+
+	// Transform drag delta into world space
+	OutDrag = CustomCoordSystem.TransformPosition(OutDrag);
+}
+
+void FWidget::ConvertMouseToAxis_Rotate(FVector2D TangentDir, FVector2D DragDir, FSceneView* InView,
+                                        FEditorViewportClient* InViewportClient, FVector& InOutDelta,
+                                        FRotator& OutRotation)
+{
+	if (CurrentAxis == EAxisList::X)
+	{
+		FRotator Rotation;
+		FVector2D EffectiveDelta;
+		// Get screen direction representing positive rotation
+		const FVector2D AxisDir = bIsOrthoDrawingFullRing ? TangentDir : XAxisDir;
+
+		// Get rotation in widget local space
+		Rotation = FRotator(0, 0, FVector2D::DotProduct(AxisDir, DragDir));
+		FSnappingUtils::SnapRotatorToGrid(Rotation);
+
+		// Record delta rotation (used by the widget to render the accumulated delta)
+		CurrentDeltaRotation = Rotation.Roll;
+
+		// Use to calculate the new input delta
+		EffectiveDelta = AxisDir * Rotation.Roll;
+		// Adjust the input delta according to how much rotation was actually applied
+		InOutDelta = FVector(EffectiveDelta.X, -EffectiveDelta.Y, 0.0f);
+		// Need to get the delta rotation in the current coordinate space of the widget
+		OutRotation = (CustomCoordSystem.Inverse() * FRotationMatrix(Rotation) * CustomCoordSystem).Rotator();
+	}
+	else if (CurrentAxis == EAxisList::Y)
+	{
+		FRotator Rotation;
+		FVector2D EffectiveDelta;
+		// TODO: Determine why -TangentDir is necessary here, and fix whatever is causing it
+		const FVector2D AxisDir = bIsOrthoDrawingFullRing ? -TangentDir : YAxisDir;
+
+		Rotation = FRotator(FVector2D::DotProduct(AxisDir, DragDir), 0, 0);
+		FSnappingUtils::SnapRotatorToGrid(Rotation);
+
+		CurrentDeltaRotation = Rotation.Pitch;
+		EffectiveDelta       = AxisDir * Rotation.Pitch;
+		// Adjust the input delta according to how much rotation was actually applied
+		InOutDelta = FVector(EffectiveDelta.X, -EffectiveDelta.Y, 0.0f);
+		// Need to get the delta rotation in the current coordinate space of the widget
+		OutRotation = (CustomCoordSystem.Inverse() * FRotationMatrix(Rotation) * CustomCoordSystem).Rotator();
+	}
+	else if (CurrentAxis == EAxisList::Z)
+	{
+		FRotator Rotation;
+		FVector2D EffectiveDelta;
+		const FVector2D AxisDir = bIsOrthoDrawingFullRing ? TangentDir : ZAxisDir;
+
+		Rotation = FRotator(0, FVector2D::DotProduct(AxisDir, DragDir), 0);
+		FSnappingUtils::SnapRotatorToGrid(Rotation);
+
+		CurrentDeltaRotation = Rotation.Yaw;
+		EffectiveDelta       = AxisDir * Rotation.Yaw;
+		// Adjust the input delta according to how much rotation was actually applied
+		InOutDelta = FVector(EffectiveDelta.X, -EffectiveDelta.Y, 0.0f);
+		// Need to get the delta rotation in the current coordinate space of the widget
+		OutRotation = (CustomCoordSystem.Inverse() * FRotationMatrix(Rotation) * CustomCoordSystem).Rotator();
+	}
+	else if (CurrentAxis == EAxisList::XYZ) //arcball rotate
+	{
+		//To do this we need to calculate the rotation axis and rotation angle.
+		//The Axis is the cross product of the current ray from eye to pixel in world space with the previous ray.
+		//The Angle is angle amount we rotate from the object's location to the imaginary sphere that matches up with the difference between the current and previous ray
+		//From the Camera. Those rays form a triangle, which we can bisect with a common side.
+		FVector2D MousePosition(InViewportClient->Viewport->GetMouseX(), InViewportClient->Viewport->GetMouseY());
+		FViewportCursorLocation OldMouseViewportRay(InView, InViewportClient, LastDragPos.X, LastDragPos.Y);
+		FViewportCursorLocation MouseViewportRay(InView, InViewportClient, MousePosition.X, MousePosition.Y);
+
+		LastDragPos               = MousePosition;
+		FVector DirectionToWidget = InViewportClient->GetWidgetLocation() - MouseViewportRay.GetOrigin();
+		float Length              = DirectionToWidget.Size();
+		if (!FMath::IsNearlyZero(Length))
+		{
+			//Calc Axis
+			DirectionToWidget /= Length;
+			const FVector CameraToPixelDir    = MouseViewportRay.GetDirection();
+			const FVector OldCameraToPixelDir = OldMouseViewportRay.GetDirection();
+			FVector RotationAxis              = FVector::CrossProduct(OldCameraToPixelDir, CameraToPixelDir);
+			RotationAxis.Normalize();
+			float RotationAngle     = 0.0f;
+			FVector4 ScreenLocation = InView->WorldToScreen(InViewportClient->GetWidgetLocation());
+			FVector2D PixelLocation;
+			InView->ScreenToPixel(ScreenLocation, PixelLocation);
+			float Distance      = FVector2D::Distance(PixelLocation, MousePosition);
+			const float MaxDiff = 2.0f * (INNER_AXIS_CIRCLE_RADIUS);
+			//If outside radius, do screen rotate instead, like other DCC's
+
+			if (Distance > MaxDiff)
+			{
+				FPlane Plane(InViewportClient->GetWidgetLocation(), DirectionToWidget);
+				FVector StartOnPlane =
+				    FMath::RayPlaneIntersection(MouseViewportRay.GetOrigin(), CameraToPixelDir, Plane);
+				FVector OldOnPlane =
+				    FMath::RayPlaneIntersection(MouseViewportRay.GetOrigin(), OldCameraToPixelDir, Plane);
+				StartOnPlane -= InViewportClient->GetWidgetLocation();
+				OldOnPlane -= InViewportClient->GetWidgetLocation();
+				StartOnPlane.Normalize();
+				OldOnPlane.Normalize();
+				RotationAngle = FMath::Acos(FVector::DotProduct(StartOnPlane, OldOnPlane));
+
+				FVector Cross = FVector::CrossProduct(OldCameraToPixelDir, CameraToPixelDir);
+				if (FVector::DotProduct(DirectionToWidget, Cross) < 0.0f)
 				{
-					DrawDualAxis(PDI, ( FTranslationMatrix(CornerPos) * FRotationMatrix(FRotator(0, 0, 90)) * WidgetMatrix ) * FScaleMatrix(FlattenScale), AxisSize, CornerLength, XMaterial->GetRenderProxy(), YMaterial->GetRenderProxy());
+					RotationAngle *= -1.0f;
 				}
-				PDI->SetHitProxy( NULL );
+				RotationAxis = DirectionToWidget;
 			}
-		}
-
-		if (Space.bIsPerspective || Space.bIsLocalSpace || Space.bIsOrthoXZ)		// Front
-		{
-			if( (DrawAxis&EAxisList::XZ) == EAxisList::XZ ) 
+			else
 			{
-				UMaterialInstanceDynamic* XMaterial = ( (CurrentAxis&EAxisList::XZ) == EAxisList::XZ ? CurrentAxisMaterial : AxisMaterialX );
-				UMaterialInstanceDynamic* ZMaterial = ( (CurrentAxis&EAxisList::XZ) == EAxisList::XZ ? CurrentAxisMaterial : AxisMaterialZ );
-
-				PDI->SetHitProxy( new HWidgetAxis(EAxisList::XZ, bDisabled) );
-				{
-					DrawDualAxis(PDI, (FTranslationMatrix(CornerPos) * WidgetMatrix) * FScaleMatrix(FlattenScale), AxisSize, CornerLength, XMaterial->GetRenderProxy(), ZMaterial->GetRenderProxy() );
-				}
-				PDI->SetHitProxy( NULL );
+				const float Scale = ScreenLocation.W *
+				    (4.0f / InView->UnscaledViewRect.Width() / InView->ViewMatrices.GetProjectionMatrix().M[0][0]);
+				const float InnerRadius = (INNER_AXIS_CIRCLE_RADIUS * Scale) +
+				    GetDefault<ULevelEditorViewportSettings>()->TransformWidgetSizeAdjustment;
+				const float LengthOfAdjacent = Length - InnerRadius;
+				RotationAngle                = FMath::Acos(FVector::DotProduct(OldCameraToPixelDir, CameraToPixelDir));
+				const float OppositeSize     = FMath::Tan(RotationAngle) * LengthOfAdjacent;
+				RotationAngle                = FMath::Atan(OppositeSize / InnerRadius);
+				RotationAngle                = RotationAngle < 0.0f ? RotationAngle : -RotationAngle;
 			}
-		}
-
-		if( Space.bIsPerspective || Space.bIsLocalSpace || Space.bIsOrthoYZ )		// Side
-		{
-			if( (DrawAxis&EAxisList::YZ) == EAxisList::YZ ) 
-			{
-				UMaterialInstanceDynamic* YMaterial = ( (CurrentAxis&EAxisList::YZ) == EAxisList::YZ ? CurrentAxisMaterial : AxisMaterialY );
-				UMaterialInstanceDynamic* ZMaterial = ( (CurrentAxis&EAxisList::YZ) == EAxisList::YZ ? CurrentAxisMaterial : AxisMaterialZ );
 
-				PDI->SetHitProxy( new HWidgetAxis(EAxisList::YZ, bDisabled) );
-				{
-					DrawDualAxis(PDI, (FTranslationMatrix(CornerPos) * FRotationMatrix(FRotator(0, 90, 0)) * WidgetMatrix) * FScaleMatrix(FlattenScale), AxisSize, CornerLength, YMaterial->GetRenderProxy(), ZMaterial->GetRenderProxy() );
-				}
-				PDI->SetHitProxy( NULL );
-			}
+			const FQuat QuatRotation(RotationAxis, RotationAngle);
+			OutRotation = FRotator(QuatRotation);
 		}
+		return;
 	}
-
-	// Draw screen-space movement handle (circle)
-	if( bDrawWidget && ( DrawAxis & EAxisList::Screen ) && Space.bIsPerspective )
-	{
-		PDI->SetHitProxy( new HWidgetAxis(EAxisList::Screen, bDisabled) );
-		const FVector CameraXAxis = View->ViewMatrices.GetViewMatrix().GetColumn(0);
-		const FVector CameraYAxis = View->ViewMatrices.GetViewMatrix().GetColumn(1);
-		const FVector CameraZAxis = View->ViewMatrices.GetViewMatrix().GetColumn(2);
-
-		UMaterialInstanceDynamic* XYZMaterial = ( CurrentAxis&EAxisList::Screen) ? CurrentAxisMaterial : OpaquePlaneMaterialXY;
-		DrawSphere( PDI, InLocation, FRotator::ZeroRotator, 4.0f * Space.Scale, 10, 5, XYZMaterial->GetRenderProxy(), SDPG_Foreground );
-
-		PDI->SetHitProxy( NULL );
-	}
-}
-
-uint8 LargeInnerAlpha = 0x3f;
-uint8 SmallInnerAlpha = 0x0f;
-uint8 LargeOuterAlpha = 0x7f;
-uint8 SmallOuterAlpha = 0x0f;
-
-void FWidget::DrawColoredSphere(FPrimitiveDrawInterface* PDI, const FVector& Center, const FRotator& Orientation, FColor Color, const FVector& Radii, int32 NumSides, int32 NumRings, const FMaterialRenderProxy* MaterialRenderProxy, uint8 DepthPriority, bool bDisableBackfaceCulling)
-{
-	// Use a mesh builder to draw the sphere.
-	FDynamicMeshBuilder MeshBuilder(PDI->View->GetFeatureLevel());
+	else if (CurrentAxis == EAxisList::Screen)
 	{
-		// The first/last arc are on top of each other.
-		int32 NumVerts = (NumSides + 1) * (NumRings + 1);
-		FDynamicMeshVertex* Verts = (FDynamicMeshVertex*)FMemory::Malloc(NumVerts * sizeof(FDynamicMeshVertex));
-
-		// Calculate verts for one arc
-		FDynamicMeshVertex* ArcVerts = (FDynamicMeshVertex*)FMemory::Malloc((NumRings + 1) * sizeof(FDynamicMeshVertex));
-
-		for (int32 i = 0; i < NumRings + 1; i++)
-		{
-			FDynamicMeshVertex* ArcVert = &ArcVerts[i];
-
-			float angle = ((float)i / NumRings) * PI;
-
-			ArcVert->Color = Color;
-			// Note- unit sphere, so position always has mag of one. We can just use it for normal!			
-			ArcVert->Position.X = 0.0f;
-			ArcVert->Position.Y = FMath::Sin(angle);
-			ArcVert->Position.Z = FMath::Cos(angle);
-
-			ArcVert->SetTangents(
-				FVector(1, 0, 0),
-				FVector(0.0f, -ArcVert->Position.Z, ArcVert->Position.Y),
-				ArcVert->Position
-			);
-
-			ArcVert->TextureCoordinate[0].X = 0.0f;
-			ArcVert->TextureCoordinate[0].Y = ((float)i / NumRings);
-		}
-
-		// Then rotate this arc NumSides+1 times.
-		for (int32 s = 0; s < NumSides + 1; s++)
-		{
-			FRotator ArcRotator(0, 360.f * (float)s / NumSides, 0);
-			FRotationMatrix ArcRot(ArcRotator);
-			float XTexCoord = ((float)s / NumSides);
+		FVector2D MousePosition(InViewportClient->Viewport->GetMouseX(), InViewportClient->Viewport->GetMouseY());
+		FViewportCursorLocation OldMouseViewportRay(InView, InViewportClient, LastDragPos.X, LastDragPos.Y);
+		FViewportCursorLocation MouseViewportRay(InView, InViewportClient, MousePosition.X, MousePosition.Y);
 
-			for (int32 v = 0; v < NumRings + 1; v++)
-			{
-				int32 VIx = (NumRings + 1)*s + v;
-				Verts[VIx].Color = Color;
-				Verts[VIx].Position = ArcRot.TransformPosition(ArcVerts[v].Position);
-
-				Verts[VIx].SetTangents(
-					ArcRot.TransformVector(ArcVerts[v].TangentX.ToFVector()),
-					ArcRot.TransformVector(ArcVerts[v].GetTangentY()),
-					ArcRot.TransformVector(ArcVerts[v].TangentZ.ToFVector())
-				);
-
-				Verts[VIx].TextureCoordinate[0].X = XTexCoord;
-				Verts[VIx].TextureCoordinate[0].Y = ArcVerts[v].TextureCoordinate[0].Y;
-			}
-		}
-
-		// Add all of the vertices we generated to the mesh builder.
-		for (int32 VertIdx = 0; VertIdx < NumVerts; VertIdx++)
-		{
-			MeshBuilder.AddVertex(Verts[VertIdx]);
-		}
+		LastDragPos               = MousePosition;
+		FVector DirectionToWidget = InViewportClient->GetWidgetLocation() - MouseViewportRay.GetOrigin();
+		float Length              = DirectionToWidget.Size();
 
-		// Add all of the triangles we generated to the mesh builder.
-		for (int32 s = 0; s < NumSides; s++)
+		if (!FMath::IsNearlyZero(Length))
 		{
-			int32 a0start = (s + 0) * (NumRings + 1);
-			int32 a1start = (s + 1) * (NumRings + 1);
+			DirectionToWidget /= Length;
 
-			for (int32 r = 0; r < NumRings; r++)
+			const FVector CameraToPixelDir    = MouseViewportRay.GetDirection();
+			const FVector OldCameraToPixelDir = OldMouseViewportRay.GetDirection();
+			FPlane Plane(InViewportClient->GetWidgetLocation(), DirectionToWidget);
+			FVector StartOnPlane = FMath::RayPlaneIntersection(MouseViewportRay.GetOrigin(), CameraToPixelDir, Plane);
+			FVector OldOnPlane = FMath::RayPlaneIntersection(MouseViewportRay.GetOrigin(), OldCameraToPixelDir, Plane);
+			StartOnPlane -= InViewportClient->GetWidgetLocation();
+			OldOnPlane -= InViewportClient->GetWidgetLocation();
+			StartOnPlane.Normalize();
+			OldOnPlane.Normalize();
+			float RotationAngle = FMath::Acos(FVector::DotProduct(StartOnPlane, OldOnPlane));
+			FVector Cross       = FVector::CrossProduct(OldCameraToPixelDir, CameraToPixelDir);
+			if (FVector::DotProduct(DirectionToWidget, Cross) < 0.0f)
 			{
-				MeshBuilder.AddTriangle(a0start + r + 0, a1start + r + 0, a0start + r + 1);
-				MeshBuilder.AddTriangle(a1start + r + 0, a1start + r + 1, a0start + r + 1);
+				RotationAngle *= -1.0f;
 			}
+			const FQuat QuatRotation(DirectionToWidget, RotationAngle);
+			OutRotation = FRotator(QuatRotation);
 		}
-
-		// Free our local copy of verts and arc verts
-		FMemory::Free(Verts);
-		FMemory::Free(ArcVerts);
+		return;
 	}
-	MeshBuilder.Draw(PDI, FScaleMatrix(Radii) * FRotationMatrix(Orientation) * FTranslationMatrix(Center), MaterialRenderProxy, DepthPriority, bDisableBackfaceCulling);
 }
-/**
- * Draws the rotation widget.
- */
-void FWidget::Render_Rotate( const FSceneView* View,FPrimitiveDrawInterface* PDI, FEditorViewportClient* ViewportClient, const FVector& InLocation, bool bDrawWidget )
-{
-	float Scale = View->WorldToScreen(InLocation).W * (4.0f / View->UnscaledViewRect.Width() / View->ViewMatrices.GetProjectionMatrix().M[0][0]);
-
-	//get the axes 
-	FVector XAxis = CustomCoordSystem.TransformVector(FVector(1, 0, 0));
-	FVector YAxis = CustomCoordSystem.TransformVector(FVector(0, 1, 0));
-	FVector ZAxis = CustomCoordSystem.TransformVector(FVector(0, 0, 1));
 
-	EAxisList::Type DrawAxis = GetAxisToDraw(ViewportClient->GetWidgetMode());
+void FWidget::ConvertMouseToAxis_Scale(FVector2D DragDir, FVector& InOutDelta, FVector& OutScale)
+{
+	FVector2D AxisDir = FVector2D::ZeroVector;
 
-	//if drawing arcball, and dragging make sure axis is showing local
-	if (DrawAxis&EAxisList::XYZ  && CurrentAxis == EAxisList::XYZ && bDragging)
+	if (CurrentAxis & EAxisList::X)
 	{
-		FMatrix LocalCoordinate = ViewportClient->GetLocalCoordinateSystem();
-		XAxis = LocalCoordinate.TransformVector(FVector(1, 0, 0));
-		YAxis = LocalCoordinate.TransformVector(FVector(0, 1, 0));
-		ZAxis = LocalCoordinate.TransformVector(FVector(0, 0, 1));
+		AxisDir += XAxisDir;
 	}
 
-	FVector DirectionToWidget = View->IsPerspectiveProjection() ? (InLocation - View->ViewMatrices.GetViewOrigin()) : -View->GetViewDirection();
-	DirectionToWidget.Normalize();
-
-	// Draw a circle for each axis
-	if (bDrawWidget || bDragging)
+	if (CurrentAxis & EAxisList::Y)
 	{
-		bIsOrthoDrawingFullRing = false;
-		//now draw the arc segments
-		if (DrawAxis&EAxisList::X)
-		{
-			DrawRotationArc(View, PDI, EAxisList::X, InLocation, ZAxis, YAxis, DirectionToWidget, AxisColorX.ToFColor(true), Scale, XAxisDir);
-		}
-
-		if (DrawAxis&EAxisList::Y)
-		{
-			DrawRotationArc(View, PDI, EAxisList::Y, InLocation, XAxis, ZAxis, DirectionToWidget, AxisColorY.ToFColor(true), Scale, YAxisDir);
-		}
-
-		if (DrawAxis&EAxisList::Z)
-		{
-			DrawRotationArc(View, PDI, EAxisList::Z, InLocation, XAxis, YAxis, DirectionToWidget, AxisColorZ.ToFColor(true), Scale, ZAxisDir);
-		}
-
-		if (DrawAxis&EAxisList::XYZ  && (!bDragging || CurrentAxis == EAxisList::XYZ) && GetDefault<ULevelEditorViewportSettings>()->bAllowArcballRotate)
-		{
-			FVector Center = InLocation;
-			FRotator Orientation = FRotator::ZeroRotator;
-			const float InnerDistance = (INNER_AXIS_CIRCLE_RADIUS * ARCALL_RELATIVE_INNER_SIZE * Scale) + GetDefault<ULevelEditorViewportSettings>()->TransformWidgetSizeAdjustment;
-			FVector Radii(InnerDistance, InnerDistance, InnerDistance);
-			const bool bDisabled = IsWidgetDisabled();
-			FColor ArcColor = ArcBallColor;
-			ArcColor.A = (CurrentAxis==EAxisList::XYZ) ? ArcBallColor.A + 5 : ArcBallColor.A; //less transparent if selected
-
-			//set priority to foreground so axis get hit first
-			PDI->SetHitProxy(new HWidgetAxis(EAxisList::XYZ, bDisabled, HPP_Foreground));
-			{
-				DrawColoredSphere(PDI, InLocation, FRotator::ZeroRotator, ArcColor, Radii, 32, 24, TransparentPlaneMaterialXY->GetRenderProxy(), SDPG_Foreground, true);
-			}
-			PDI->SetHitProxy(NULL);
-		}
-		
-		if(DrawAxis&EAxisList::Screen && (!bDragging || CurrentAxis == EAxisList::Screen) && GetDefault<ULevelEditorViewportSettings>()->bAllowScreenRotate)
-		{
-			const bool bIsPerspective = (View->ViewMatrices.GetProjectionMatrix().M[3][3] < 1.0f);
-			const bool bIsOrtho = !bIsPerspective;
-
-			const FVector Axis0 = View->GetViewUp();
-			const FVector Axis1 = View->GetViewRight();
-			
-			const float OuterRadius = (OUTER_AXIS_CIRCLE_RADIUS * 1.25f * Scale) + GetDefault<ULevelEditorViewportSettings>()->TransformWidgetSizeAdjustment;
-			const float InnerRadius = ((OUTER_AXIS_CIRCLE_RADIUS -1) * 1.25f * Scale) + GetDefault<ULevelEditorViewportSettings>()->TransformWidgetSizeAdjustment;
-
-			FColor ArcColor = (CurrentAxis==EAxisList::Screen) ? CurrentColor : ScreenAxisColor.ToFColor(true);
-
-			const bool bDisabled = IsWidgetDisabled();
-
-			PDI->SetHitProxy(new HWidgetAxis(EAxisList::Screen, bDisabled));
-			{
-				FThickArcParams OuterArcParams(PDI, InLocation, TransparentPlaneMaterialXY, InnerRadius, OuterRadius);
-				//Pass through alpha
-				DrawThickArc(OuterArcParams, Axis0, Axis1, 0.0f, 2.0f * PI, ArcColor, DirectionToWidget, !bIsPerspective);
-			}
-			PDI->SetHitProxy(NULL);
-		}
+		AxisDir += YAxisDir;
 	}
-}
 
-/**
- * Draws the scaling widget.
- */
-void FWidget::Render_Scale( const FSceneView* View,FPrimitiveDrawInterface* PDI, FEditorViewportClient* ViewportClient, const FVector& InLocation, bool bDrawWidget)
-{
-	// Figure out axis colors
-	const FLinearColor& XColor = ( CurrentAxis&EAxisList::X ? (FLinearColor)CurrentColor : AxisColorX );
-	const FLinearColor& YColor = ( CurrentAxis&EAxisList::Y ? (FLinearColor)CurrentColor : AxisColorY );
-	const FLinearColor& ZColor = ( CurrentAxis&EAxisList::Z ? (FLinearColor)CurrentColor : AxisColorZ );
-	FColor CurrentScreenColor = ( CurrentAxis & EAxisList::Screen ? CurrentColor : ScreenSpaceColor );
-
-	// Figure out axis materials
-
-	UMaterialInstanceDynamic* XMaterial = ( CurrentAxis&EAxisList::X ? CurrentAxisMaterial : AxisMaterialX );
-	UMaterialInstanceDynamic* YMaterial = ( CurrentAxis&EAxisList::Y ? CurrentAxisMaterial : AxisMaterialY );
-	UMaterialInstanceDynamic* ZMaterial = ( CurrentAxis&EAxisList::Z ? CurrentAxisMaterial : AxisMaterialZ );
-	UMaterialInstanceDynamic* XYZMaterial = ( CurrentAxis&EAxisList::XYZ ? CurrentAxisMaterial : OpaquePlaneMaterialXY );
-
-	FMatrix WidgetMatrix = CustomCoordSystem * FTranslationMatrix( InLocation );
-	const EAxisList::Type DrawAxis = GetAxisToDraw( ViewportClient->GetWidgetMode() );
-    FSpaceDescriptor Space(View, ViewportClient, InLocation);
-
-	// Use a constant uniform scale for this widget since orthographic view for it is not supported.
-	const FVector UniformScale(Space.UniformScale);
-
-	// Draw the axis lines with cube heads	
-    if (Space.ShouldDrawAxisX(DrawAxis))
-    {
-        Render_Axis(View, PDI, EAxisList::X, WidgetMatrix, XMaterial, XColor, XAxisDir, UniformScale, bDrawWidget, true, AXIS_LENGTH_SCALE_OFFSET);
-    }
-
-    if (Space.ShouldDrawAxisY(DrawAxis))
-    {
-        Render_Axis(View, PDI, EAxisList::Y, WidgetMatrix, YMaterial, YColor, YAxisDir, UniformScale, bDrawWidget, true, AXIS_LENGTH_SCALE_OFFSET);
-    }
-
-    if (Space.ShouldDrawAxisZ(DrawAxis))
-    {
-        Render_Axis(View, PDI, EAxisList::Z, WidgetMatrix, ZMaterial, ZColor, ZAxisDir, UniformScale, bDrawWidget, true, AXIS_LENGTH_SCALE_OFFSET);
-    }
-
-	// Draw grabber handles and center cube
-	if ( bDrawWidget )
+	if (CurrentAxis & EAxisList::Z)
 	{
-		const bool bDisabled = IsWidgetDisabled();
-
-		// Grabber handles - since orthographic scale widgets are not supported, we should always draw grabber handles if we're drawing
-		// the corresponding axes.
-		if( (DrawAxis&(EAxisList::X|EAxisList::Y)) == (EAxisList::X|EAxisList::Y) )
-		{
-			PDI->SetHitProxy( new HWidgetAxis(EAxisList::XY, bDisabled) );
-			{
-				PDI->DrawLine( WidgetMatrix.TransformPosition(FVector(24,0,0) * UniformScale), WidgetMatrix.TransformPosition(FVector(12,12,0) * UniformScale), XColor, SDPG_Foreground );
-				PDI->DrawLine( WidgetMatrix.TransformPosition(FVector(12,12,0) * UniformScale), WidgetMatrix.TransformPosition(FVector(0,24,0) * UniformScale), YColor, SDPG_Foreground );
-			}
-			PDI->SetHitProxy( NULL );
-		}
-
-		if( (DrawAxis&(EAxisList::X|EAxisList::Z)) == (EAxisList::X|EAxisList::Z) )
-		{
-			PDI->SetHitProxy( new HWidgetAxis(EAxisList::XZ, bDisabled) );
-			{
-				PDI->DrawLine( WidgetMatrix.TransformPosition(FVector(24,0,0) * UniformScale), WidgetMatrix.TransformPosition(FVector(12,0,12) * UniformScale), XColor, SDPG_Foreground );
-				PDI->DrawLine( WidgetMatrix.TransformPosition(FVector(12,0,12) * UniformScale), WidgetMatrix.TransformPosition(FVector(0,0,24) * UniformScale), ZColor, SDPG_Foreground );
-			}
-			PDI->SetHitProxy( NULL );
-		}
+		AxisDir += ZAxisDir;
+	}
 
-		if( (DrawAxis&(EAxisList::Y|EAxisList::Z)) == (EAxisList::Y|EAxisList::Z) )
-		{
-			PDI->SetHitProxy( new HWidgetAxis(EAxisList::YZ, bDisabled) );
-			{
-				PDI->DrawLine( WidgetMatrix.TransformPosition(FVector(0,24,0) * UniformScale), WidgetMatrix.TransformPosition(FVector(0,12,12) * UniformScale), YColor, SDPG_Foreground );
-				PDI->DrawLine( WidgetMatrix.TransformPosition(FVector(0,12,12) * UniformScale), WidgetMatrix.TransformPosition(FVector(0,0,24) * UniformScale), ZColor, SDPG_Foreground );
-			}
-			PDI->SetHitProxy( NULL );
-		}
+	AxisDir.Normalize();
+	const float ScaleDelta = FVector2D::DotProduct(AxisDir, DragDir);
 
-		// Center cube
-		if( (DrawAxis&(EAxisList::XYZ)) == EAxisList::XYZ )
-		{
-			PDI->SetHitProxy( new HWidgetAxis(EAxisList::XYZ, bDisabled) );
+	OutScale =
+	    FVector((CurrentAxis & EAxisList::X) ? ScaleDelta : 0.0f, (CurrentAxis & EAxisList::Y) ? ScaleDelta : 0.0f,
+	            (CurrentAxis & EAxisList::Z) ? ScaleDelta : 0.0f);
 
-			Render_Cube(PDI, WidgetMatrix, XYZMaterial, UniformScale * 4 );
+	// Snap to grid in widget axis space
+	const FVector GridSize = FVector(GEditor->GetGridSize());
+	FSnappingUtils::SnapScale(OutScale, GridSize);
 
-			PDI->SetHitProxy( NULL );
-		}
-	}
+	// Convert to effective screen space delta, and replace input delta, adjusted for inverted screen space Y axis
+	const float ScaleMax           = OutScale.GetMax();
+	const float ScaleMin           = OutScale.GetMin();
+	const float ScaleApplied       = (ScaleMax > -ScaleMin) ? ScaleMax : ScaleMin;
+	const FVector2D EffectiveDelta = AxisDir * ScaleApplied;
+	InOutDelta                     = FVector(EffectiveDelta.X, -EffectiveDelta.Y, 0.0f);
 }
 
-/**
-* Draws the Translate & Rotate Z widget.
-*/
-
-void FWidget::Render_TranslateRotateZ( const FSceneView* View, FPrimitiveDrawInterface* PDI, FEditorViewportClient* ViewportClient, const FVector& InLocation, bool bDrawWidget )
+void FWidget::ConvertMouseToAxis_TranslateRotateZ(FVector2D TangentDir, FVector2D DragDir, FVector& InOutDelta,
+                                                  FVector& OutDrag, FRotator& OutRotation)
 {
-	// Figure out axis colors
-
-	FColor XYPlaneColor  = ( (CurrentAxis&EAxisList::XY) ==  EAxisList::XY) ? CurrentColor : PlaneColorXY;
-	FColor ZRotateColor  = ( ( CurrentAxis&EAxisList::ZRotation ) == EAxisList::ZRotation ) ? CurrentColor : AxisColorZ.ToFColor(true);
-	FColor XColor        = ( ( CurrentAxis&EAxisList::X ) == EAxisList::X ) ? CurrentColor : AxisColorX.ToFColor(true);
-	FColor YColor        = ( ( CurrentAxis&EAxisList::Y ) == EAxisList::Y && CurrentAxis != EAxisList::ZRotation ) ? CurrentColor : AxisColorY.ToFColor(true);
-	FColor ZColor        = ( ( CurrentAxis&EAxisList::Z ) == EAxisList::Z ) ? CurrentColor : AxisColorZ.ToFColor(true);
-
-	// Figure out axis materials
-	UMaterialInterface* ZRotateMaterial = (CurrentAxis&EAxisList::ZRotation) == EAxisList::ZRotation? CurrentAxisMaterial : AxisMaterialZ;
-	UMaterialInterface* XMaterial = CurrentAxis&EAxisList::X? CurrentAxisMaterial : AxisMaterialX;
-	UMaterialInterface* YMaterial = ( CurrentAxis&EAxisList::Y && CurrentAxis != EAxisList::ZRotation ) ? CurrentAxisMaterial : AxisMaterialY;
-	UMaterialInterface* ZMaterial = CurrentAxis&EAxisList::Z ? CurrentAxisMaterial : AxisMaterialZ;
-
-	// Figure out axis matrices
-	FMatrix AxisMatrix = CustomCoordSystem * FTranslationMatrix( InLocation );
-	EAxisList::Type DrawAxis = GetAxisToDraw( ViewportClient->GetWidgetMode() );
-
-	FSpaceDescriptor Space(View, ViewportClient, InLocation);
-
-	// Draw the grabbers
-	if( bDrawWidget )
+	if (CurrentAxis == EAxisList::ZRotation)
 	{
-		// Draw the axis lines with arrow heads
-		if( DrawAxis&EAxisList::X && (Space.bIsPerspective || Space.bIsLocalSpace || View->ViewMatrices.GetViewMatrix().M[0][2] != -1.f) )
-		{
-			Render_Axis( View, PDI, EAxisList::X, AxisMatrix, XMaterial, XColor, XAxisDir, Space.Scale, bDrawWidget );
-		}
-
-		if( DrawAxis&EAxisList::Y && (Space.bIsPerspective || Space.bIsLocalSpace || View->ViewMatrices.GetViewMatrix().M[1][2] != -1.f) )
-		{
-			Render_Axis( View, PDI, EAxisList::Y, AxisMatrix, YMaterial, YColor, YAxisDir, Space.Scale, bDrawWidget );
-		}
-
-		if( DrawAxis&EAxisList::Z && (Space.bIsPerspective || Space.bIsLocalSpace || View->ViewMatrices.GetViewMatrix().M[0][1] != 1.f) )
-		{
-			Render_Axis( View, PDI, EAxisList::Z, AxisMatrix, ZMaterial, ZColor, ZAxisDir, Space.Scale, bDrawWidget );
-		}
+		const FVector2D AxisDir = bIsOrthoDrawingFullRing ? TangentDir : ZAxisDir;
+		FRotator Rotation       = FRotator(0, FVector2D::DotProduct(AxisDir, DragDir), 0);
+		FSnappingUtils::SnapRotatorToGrid(Rotation);
+		CurrentDeltaRotation = Rotation.Yaw;
 
-		const bool bDisabled = IsWidgetDisabled();
+		const FVector2D EffectiveDelta = AxisDir * Rotation.Yaw;
+		InOutDelta                     = FVector(EffectiveDelta.X, -EffectiveDelta.Y, 0.0f);
 
-		const float ScaledRadius = (TRANSLATE_ROTATE_AXIS_CIRCLE_RADIUS * Space.UniformScale) + GetDefault<ULevelEditorViewportSettings>()->TransformWidgetSizeAdjustment;
+		OutRotation = (CustomCoordSystem.Inverse() * FRotationMatrix(Rotation) * CustomCoordSystem).Rotator();
+	}
+	else
+	{
+		// Get drag delta in widget axis space
+		OutDrag = FVector((CurrentAxis & EAxisList::X) ? FVector2D::DotProduct(XAxisDir, DragDir) : 0.0f,
+		                  (CurrentAxis & EAxisList::Y) ? FVector2D::DotProduct(YAxisDir, DragDir) : 0.0f,
+		                  (CurrentAxis & EAxisList::Z) ? FVector2D::DotProduct(ZAxisDir, DragDir) : 0.0f);
 
-		//ZRotation
-		if( DrawAxis&EAxisList::ZRotation && (Space.bIsPerspective || Space.bIsLocalSpace || View->ViewMatrices.GetViewMatrix().M[0][2] != -1.f) )
-		{
-			PDI->SetHitProxy( new HWidgetAxis(EAxisList::ZRotation, bDisabled) );
-			{
-				FVector XAxis = CustomCoordSystem.TransformPosition( FVector(1,0,0).RotateAngleAxis( (EditorModeTools ? EditorModeTools->TranslateRotateXAxisAngle : 0 ), FVector(0,0,1)) );
-				FVector YAxis = CustomCoordSystem.TransformPosition( FVector(0,1,0).RotateAngleAxis( (EditorModeTools ? EditorModeTools->TranslateRotateXAxisAngle : 0 ), FVector(0,0,1)) );
-				FVector BaseArrowPoint = InLocation + XAxis * ScaledRadius;
-				DrawFlatArrow(PDI, BaseArrowPoint, XAxis, YAxis, ZRotateColor, ScaledRadius, ScaledRadius*.5f, ZRotateMaterial->GetRenderProxy(), SDPG_Foreground);
-			}
-			PDI->SetHitProxy( NULL );
-		}
+		// Snap to grid in widget axis space
+		const FVector GridSize = FVector(GEditor->GetGridSize());
+		FSnappingUtils::SnapPointToGrid(OutDrag, GridSize);
 
-		//XY Plane
-		if( Space.bIsPerspective || Space.bIsLocalSpace || View->ViewMatrices.GetViewMatrix().M[0][1] != 1.f )
-		{
-			if( (DrawAxis & EAxisList::XY) == EAxisList::XY ) 
-			{
-				// Add more sides to the circle if we've been scaled up to keep the circle looking circular
-				// An extra side for every 5 extra unreal units seems to produce a nice result
-				const int32 CircleSides = (GetDefault<ULevelEditorViewportSettings>()->TransformWidgetSizeAdjustment > 0) 
-					? AXIS_CIRCLE_SIDES + (GetDefault<ULevelEditorViewportSettings>()->TransformWidgetSizeAdjustment / 5)
-					: AXIS_CIRCLE_SIDES;
+		// Convert to effective screen space delta, and replace input delta, adjusted for inverted screen space Y axis
+		const FVector2D EffectiveDelta = OutDrag.X * XAxisDir + OutDrag.Y * YAxisDir + OutDrag.Z * ZAxisDir;
+		InOutDelta                     = FVector(EffectiveDelta.X, -EffectiveDelta.Y, 0.0f);
 
-				PDI->SetHitProxy( new HWidgetAxis(EAxisList::XY, bDisabled) );
-				{
-					DrawCircle( PDI, InLocation, CustomCoordSystem.TransformPosition( FVector(1,0,0) ), CustomCoordSystem.TransformPosition( FVector(0,1,0) ), XYPlaneColor, ScaledRadius, CircleSides, SDPG_Foreground );
-					XYPlaneColor.A = ((CurrentAxis&EAxisList::XY) == EAxisList::XY) ? 0x3f : 0x0f;	//make the disc transparent
-					DrawDisc  ( PDI, InLocation, CustomCoordSystem.TransformPosition( FVector(1,0,0) ), CustomCoordSystem.TransformPosition( FVector(0,1,0) ), XYPlaneColor, ScaledRadius, CircleSides, TransparentPlaneMaterialXY->GetRenderProxy(), SDPG_Foreground );
-				}
-				PDI->SetHitProxy( NULL );
-			}
-		}
+		// Transform drag delta into world space
+		OutDrag = CustomCoordSystem.TransformPosition(OutDrag);
 	}
 }
 
-/**
-* Draws the 2D widget.
-*/
-
-void FWidget::Render_2D(const FSceneView* View, FPrimitiveDrawInterface* PDI, FEditorViewportClient* ViewportClient, const FVector& InLocation, bool bDrawWidget)
+void FWidget::ConvertMouseToAxis_WM_2D(FVector2D TangentDir, FVector2D DragDir, FVector& InOutDelta, FVector& OutDrag,
+                                       FRotator& OutRotation)
 {
-	//////////////////////////////////////////////////////////////////////////
-	// Translation subwidget
-	//////////////////////////////////////////////////////////////////////////
-
-	// Figure out axis colors
-	const FLinearColor& XColor = (CurrentAxis&EAxisList::X ? (FLinearColor)CurrentColor : AxisColorX);
-	const FLinearColor& YColor = (CurrentAxis&EAxisList::Y ? (FLinearColor)CurrentColor : AxisColorY);
-	const FLinearColor& ZColor = (CurrentAxis&EAxisList::Z ? (FLinearColor)CurrentColor : AxisColorZ);
-	FColor CurrentScreenColor = (CurrentAxis & EAxisList::Screen ? CurrentColor : ScreenSpaceColor);
-
-	// Figure out axis matrices
-	FMatrix WidgetMatrix = CustomCoordSystem * FTranslationMatrix(InLocation);
-
-    FSpaceDescriptor Space(View, ViewportClient, InLocation);
-
-	EAxisList::Type DrawAxis = EAxisList::None;
-	if (Space.bIsOrthoXY)
-	{
-		DrawAxis = EAxisList::X;
-	}
-	else if (Space.bIsOrthoXZ)
-	{
-		DrawAxis = EAxisList::XZ;
-	}
-	else if (Space.bIsOrthoYZ)
-	{
-		DrawAxis = EAxisList::Z;
-	}
-	else if (Space.bIsPerspective)
+	if (CurrentAxis == EAxisList::Rotate2D)
 	{
-		// Find the best plane to move on
-		const FVector CameraZAxis = View->ViewMatrices.GetViewMatrix().GetColumn(2);
-		const FVector LargestAxis = CameraZAxis.GetAbs();
-		if (LargestAxis.X > LargestAxis.Y)
-		{
-			if (LargestAxis.Z > LargestAxis.X)
-			{
-				DrawAxis = EAxisList::X;
-			}
-			else
-			{
-				DrawAxis = EAxisList::Z;
-			}
-		}
-		else
-		{
-			if (LargestAxis.Y > LargestAxis.Z)
-			{
-				DrawAxis = EAxisList::XZ;
-			}
-			else
-			{
-				DrawAxis = EAxisList::X;
-			}
-		}
-	}
-
-	const bool bDisabled = IsWidgetDisabled();
-
-	// Radius
-	const float ScaledRadius = (TWOD_AXIS_CIRCLE_RADIUS * Space.UniformScale) + GetDefault<ULevelEditorViewportSettings>()->TransformWidgetSizeAdjustment;
-	const int32 CircleSides = (GetDefault<ULevelEditorViewportSettings>()->TransformWidgetSizeAdjustment > 0)
-		? AXIS_CIRCLE_SIDES + (GetDefault<ULevelEditorViewportSettings>()->TransformWidgetSizeAdjustment / 5)
-		: AXIS_CIRCLE_SIDES;
-
-	// Draw the grabbers
-	if (bDrawWidget)
-	{
-		FVector CornerPos = FVector(7, 0, 7) * Space.UniformScale;
-		FVector AxisSize = FVector(12, 1.2, 12) * Space.UniformScale;
-		float CornerLength = 1.2f * Space.UniformScale;
+		// TODO: Determine why -TangentDir is necessary here, and fix whatever is causing it
+		const FVector2D AxisDir = bIsOrthoDrawingFullRing ? -TangentDir : YAxisDir;
 
-		// Draw the axis lines with arrow heads
-		if (Space.ShouldDrawAxisX(DrawAxis))
-		{
-			UMaterialInstanceDynamic* XMaterial = (CurrentAxis&EAxisList::X ? CurrentAxisMaterial : AxisMaterialX);
-			Render_Axis(View, PDI, EAxisList::X, WidgetMatrix, XMaterial, XColor, XAxisDir, Space.Scale, bDrawWidget);
-		}
+		FRotator Rotation = FRotator(FVector2D::DotProduct(AxisDir, DragDir), 0, 0);
+		FSnappingUtils::SnapRotatorToGrid(Rotation);
 
-		if (Space.ShouldDrawAxisY(DrawAxis))
-		{
-			UMaterialInstanceDynamic* YMaterial = (CurrentAxis&EAxisList::Y ? CurrentAxisMaterial : AxisMaterialY);
-			Render_Axis(View, PDI, EAxisList::Y, WidgetMatrix, YMaterial, YColor, YAxisDir, Space.Scale, bDrawWidget);
-		}
+		CurrentDeltaRotation     = Rotation.Pitch;
+		FVector2D EffectiveDelta = AxisDir * Rotation.Pitch;
 
-		if (Space.ShouldDrawAxisZ(DrawAxis))
-		{
-			UMaterialInstanceDynamic* ZMaterial = (CurrentAxis&EAxisList::Z ? CurrentAxisMaterial : AxisMaterialZ);
-			Render_Axis(View, PDI, EAxisList::Z, WidgetMatrix, ZMaterial, ZColor, ZAxisDir, Space.Scale, bDrawWidget);
-		}
 
-		// After the primitives have been scaled and transformed, we apply this inverse scale that flattens the dimension
-		// that was scaled up to prevent it from intersecting with the near plane.  In perspective this won't have any effect,
-		// but in the ortho viewports it will prevent scaling in the direction of the camera and thus intersecting the near plane.
-		FVector FlattenScale = FVector(Space.Scale.Component(0) == 1.0f ? 1.0f / Space.UniformScale : 1.0f, Space.Scale.Component(1) == 1.0f ? 1.0f / Space.UniformScale : 1.0f, Space.Scale.Component(2) == 1.0f ? 1.0f / Space.UniformScale : 1.0f);
-		float ArrowRadius = ScaledRadius * 2.0f;
-		float ArrowStartRadius = ScaledRadius * 1.3f;
+		// Adjust the input delta according to how much rotation was actually applied
+		InOutDelta = FVector(EffectiveDelta.X, -EffectiveDelta.Y, 0.0f);
 
-		uint8 HoverAlpha = 0xff;
-		uint8 NormalAlpha = 0x2f;
+		// Need to get the delta rotation in the current coordinate space of the widget
+		OutRotation = (CustomCoordSystem.Inverse() * FRotationMatrix(Rotation) * CustomCoordSystem).Rotator();
+	}
+	else
+	{
+		// Get drag delta in widget axis space
+		OutDrag = FVector((CurrentAxis & EAxisList::X) ? FVector2D::DotProduct(XAxisDir, DragDir) : 0.0f,
+		                  (CurrentAxis & EAxisList::Y) ? FVector2D::DotProduct(YAxisDir, DragDir) : 0.0f,
+		                  (CurrentAxis & EAxisList::Z) ? FVector2D::DotProduct(ZAxisDir, DragDir) : 0.0f);
 
-		if (((DrawAxis&EAxisList::XZ) == EAxisList::XZ) && (Space.bIsPerspective || Space.bIsLocalSpace || Space.bIsOrthoXZ)) // Front
-		{
-			uint8 Alpha = ((CurrentAxis&EAxisList::XZ) == EAxisList::XZ ? HoverAlpha : NormalAlpha);
-			PDI->SetHitProxy(new HWidgetAxis(EAxisList::XZ, bDisabled));
-			{
-				FColor Color = YColor.ToFColor(true);
-				DrawCircle(PDI, InLocation, CustomCoordSystem.TransformPosition(FVector(1, 0, 0)), CustomCoordSystem.TransformPosition(FVector(0, 0, 1)), Color, ScaledRadius, CircleSides, SDPG_Foreground);
-				Color.A = Alpha;
-				DrawDisc(PDI, InLocation, CustomCoordSystem.TransformPosition(FVector(1, 0, 0)), CustomCoordSystem.TransformPosition(FVector(0, 0, 1)), Color, ScaledRadius, CircleSides, TransparentPlaneMaterialXY->GetRenderProxy(), SDPG_Foreground);
-			}
-			PDI->SetHitProxy(NULL);
+		// Snap to grid in widget axis space
+		const FVector GridSize = FVector(GEditor->GetGridSize());
+		FSnappingUtils::SnapPointToGrid(OutDrag, GridSize);
 
-			PDI->SetHitProxy(new HWidgetAxis(EAxisList::Rotate2D, bDisabled));
-			{
-				FColor Color = YColor.ToFColor(true);
-				Color.A = ((CurrentAxis&EAxisList::Rotate2D) == EAxisList::Rotate2D ? HoverAlpha : NormalAlpha);
+		// Convert to effective screen space delta, and replace input delta, adjusted for inverted screen space Y axis
+		const FVector2D EffectiveDelta = OutDrag.X * XAxisDir + OutDrag.Y * YAxisDir + OutDrag.Z * ZAxisDir;
+		InOutDelta                     = FVector(EffectiveDelta.X, -EffectiveDelta.Y, 0.0f);
 
-				FVector XAxis = CustomCoordSystem.TransformPosition(FVector(1, 0, 0).RotateAngleAxis((EditorModeTools ? EditorModeTools->TranslateRotate2DAngle : 0), FVector(0, -1, 0)));
-				FVector YAxis = CustomCoordSystem.TransformPosition(FVector(0, 0, 1).RotateAngleAxis((EditorModeTools ? EditorModeTools->TranslateRotate2DAngle : 0), FVector(0, -1, 0)));
-				FVector BaseArrowPoint = InLocation + XAxis * ArrowStartRadius;
-				DrawFlatArrow(PDI, BaseArrowPoint, XAxis, YAxis, Color, ArrowRadius, ArrowRadius*.5f, TransparentPlaneMaterialXY->GetRenderProxy(), SDPG_Foreground);
-			}
-			PDI->SetHitProxy(NULL);
-		}
+		// Transform drag delta into world space
+		OutDrag = CustomCoordSystem.TransformPosition(OutDrag);
 	}
 }
 
 /**
  * Converts mouse movement on the screen to widget axis movement/rotation.
  */
-void FWidget::ConvertMouseMovementToAxisMovement(FSceneView* InView, FEditorViewportClient* InViewportClient, bool bInUsedDragModifier, FVector& InOutDelta, FVector& OutDrag, FRotator& OutRotation, FVector& OutScale )
+void FWidget::ConvertMouseMovementToAxisMovement(FSceneView* InView, FEditorViewportClient* InViewportClient,
+                                                 bool bInUsedDragModifier, FVector& InOutDelta, FVector& OutDrag,
+                                                 FRotator& OutRotation, FVector& OutScale)
 {
-	OutDrag = FVector::ZeroVector;
+	OutDrag     = FVector::ZeroVector;
 	OutRotation = FRotator::ZeroRotator;
-	OutScale = FVector::ZeroVector; 
-
-	const int32 WidgetMode = InViewportClient->GetWidgetMode();
-
-	// Get input delta as 2D vector, adjusted for inverted screen space Y axis
-	const FVector2D DragDir = FVector2D(InOutDelta.X, -InOutDelta.Y);
-
-	// Get offset of the drag start position from the widget origin
-	const FVector2D DirectionToMousePos = FVector2D(DragStartPos - Origin).GetSafeNormal();
-
-	// For rotations which display as a full ring, calculate the tangent direction representing a clockwise movement
-	FVector2D TangentDir;
-	if (bInUsedDragModifier)
-	{
-		// If a drag modifier has been used, this implies we are not actually touching the widget, so don't attempt to
-		// calculate the tangent dir based on the relative offset of the cursor from the widget location.
-		TangentDir = FVector2D(1, 1).GetSafeNormal();
-	}
-	else
-	{
-		// Treat the tangent dir as perpendicular to the relative offset of the cursor from the widget location.
-		TangentDir = FVector2D(-DirectionToMousePos.Y, DirectionToMousePos.X);
-	}
-
-	switch (WidgetMode)
-	{
-		case WM_Translate:
-		{
-			// Get drag delta in widget axis space
-			OutDrag = FVector(
-				(CurrentAxis & EAxisList::X) ? FVector2D::DotProduct(XAxisDir, DragDir) : 0.0f,
-				(CurrentAxis & EAxisList::Y) ? FVector2D::DotProduct(YAxisDir, DragDir) : 0.0f,
-				(CurrentAxis & EAxisList::Z) ? FVector2D::DotProduct(ZAxisDir, DragDir) : 0.0f
-				);
-
-			// Snap to grid in widget axis space
-			const FVector GridSize = FVector(GEditor->GetGridSize());
-			FSnappingUtils::SnapPointToGrid(OutDrag, GridSize);
-
-			// Convert to effective screen space delta, and replace input delta, adjusted for inverted screen space Y axis
-			const FVector2D EffectiveDelta = OutDrag.X * XAxisDir + OutDrag.Y * YAxisDir + OutDrag.Z * ZAxisDir;
-			InOutDelta = FVector(EffectiveDelta.X, -EffectiveDelta.Y, 0.0f);
-
-			// Transform drag delta into world space
-			OutDrag = CustomCoordSystem.TransformPosition(OutDrag);
-		}
-		break;
-
-		case WM_Rotate:
-			{
-
-				if (CurrentAxis == EAxisList::X)
-				{
-					FRotator Rotation;
-					FVector2D EffectiveDelta;
-					// Get screen direction representing positive rotation
-					const FVector2D AxisDir = bIsOrthoDrawingFullRing ? TangentDir : XAxisDir;
-
-					// Get rotation in widget local space
-					Rotation = FRotator(0, 0, FVector2D::DotProduct(AxisDir, DragDir));
-					FSnappingUtils::SnapRotatorToGrid(Rotation);
-
-					// Record delta rotation (used by the widget to render the accumulated delta)
-					CurrentDeltaRotation = Rotation.Roll;
-
-					// Use to calculate the new input delta
-					EffectiveDelta = AxisDir * Rotation.Roll;
-					// Adjust the input delta according to how much rotation was actually applied
-					InOutDelta = FVector(EffectiveDelta.X, -EffectiveDelta.Y, 0.0f);
-					// Need to get the delta rotation in the current coordinate space of the widget
-					OutRotation = (CustomCoordSystem.Inverse() * FRotationMatrix(Rotation) * CustomCoordSystem).Rotator();
-				}
-				else if (CurrentAxis == EAxisList::Y)
-				{
-					FRotator Rotation;
-					FVector2D EffectiveDelta;
-					// TODO: Determine why -TangentDir is necessary here, and fix whatever is causing it
-					const FVector2D AxisDir = bIsOrthoDrawingFullRing ? -TangentDir : YAxisDir;
-
-					Rotation = FRotator(FVector2D::DotProduct(AxisDir, DragDir), 0, 0);
-					FSnappingUtils::SnapRotatorToGrid(Rotation);
-
-					CurrentDeltaRotation = Rotation.Pitch;
-					EffectiveDelta = AxisDir * Rotation.Pitch;
-					// Adjust the input delta according to how much rotation was actually applied
-					InOutDelta = FVector(EffectiveDelta.X, -EffectiveDelta.Y, 0.0f);
-					// Need to get the delta rotation in the current coordinate space of the widget
-					OutRotation = (CustomCoordSystem.Inverse() * FRotationMatrix(Rotation) * CustomCoordSystem).Rotator();
-				}
-				else if (CurrentAxis == EAxisList::Z)
-				{
-					FRotator Rotation;
-					FVector2D EffectiveDelta;
-					const FVector2D AxisDir = bIsOrthoDrawingFullRing ? TangentDir : ZAxisDir;
-
-					Rotation = FRotator(0, FVector2D::DotProduct(AxisDir, DragDir), 0);
-					FSnappingUtils::SnapRotatorToGrid(Rotation);
-
-					CurrentDeltaRotation = Rotation.Yaw;
-					EffectiveDelta = AxisDir * Rotation.Yaw;
-					// Adjust the input delta according to how much rotation was actually applied
-					InOutDelta = FVector(EffectiveDelta.X, -EffectiveDelta.Y, 0.0f);
-					// Need to get the delta rotation in the current coordinate space of the widget
-					OutRotation = (CustomCoordSystem.Inverse() * FRotationMatrix(Rotation) * CustomCoordSystem).Rotator();
-				}
-				else if (CurrentAxis == EAxisList::XYZ) //arcball rotate
-				{
-					//To do this we need to calculate the rotation axis and rotation angle.
-					//The Axis is the cross product of the current ray from eye to pixel in world space with the previous ray.
-					//The Angle is angle amount we rotate from the object's location to the imaginary sphere that matches up with the difference between the current and previous ray
-					//From the Camera. Those rays form a triangle, which we can bisect with a common side.
-					FVector2D MousePosition(InViewportClient->Viewport->GetMouseX(), InViewportClient->Viewport->GetMouseY());
-					FViewportCursorLocation OldMouseViewportRay(InView, InViewportClient, LastDragPos.X, LastDragPos.Y);
-					FViewportCursorLocation MouseViewportRay(InView, InViewportClient, MousePosition.X, MousePosition.Y);
-					
-					LastDragPos = MousePosition;
-					FVector DirectionToWidget = InViewportClient->GetWidgetLocation() - MouseViewportRay.GetOrigin();
-					float Length = DirectionToWidget.Size();
-					if (!FMath::IsNearlyZero(Length))
-					{
-						//Calc Axis
-						DirectionToWidget /= Length;
-						const FVector CameraToPixelDir = MouseViewportRay.GetDirection();						
-						const FVector OldCameraToPixelDir = OldMouseViewportRay.GetDirection();
-						FVector RotationAxis = FVector::CrossProduct(OldCameraToPixelDir, CameraToPixelDir);
-						RotationAxis.Normalize();
-						float RotationAngle = 0.0f;
-						FVector4 ScreenLocation = InView->WorldToScreen(InViewportClient->GetWidgetLocation());
-						FVector2D PixelLocation;
-						InView->ScreenToPixel(ScreenLocation, PixelLocation);
-						float Distance = FVector2D::Distance(PixelLocation, MousePosition);
-						const float MaxDiff = 2.0f * (INNER_AXIS_CIRCLE_RADIUS);
-						//If outside radius, do screen rotate instead, like other DCC's
-
-						if (Distance > MaxDiff)
-						{
-							FPlane Plane(InViewportClient->GetWidgetLocation(), DirectionToWidget);
-							FVector StartOnPlane = FMath::RayPlaneIntersection(MouseViewportRay.GetOrigin(), CameraToPixelDir, Plane);
-							FVector OldOnPlane = FMath::RayPlaneIntersection(MouseViewportRay.GetOrigin(), OldCameraToPixelDir, Plane);
-							StartOnPlane -= InViewportClient->GetWidgetLocation();
-							OldOnPlane -= InViewportClient->GetWidgetLocation();
-							StartOnPlane.Normalize();
-							OldOnPlane.Normalize();
-							RotationAngle = FMath::Acos(FVector::DotProduct(StartOnPlane, OldOnPlane));
-							
-							FVector Cross = FVector::CrossProduct(OldCameraToPixelDir, CameraToPixelDir);
-							if (FVector::DotProduct(DirectionToWidget, Cross) < 0.0f)
-							{
-								RotationAngle *= -1.0f;
-							}
-							RotationAxis = DirectionToWidget;
-						}
-						else
-						{
-							const float Scale = ScreenLocation.W * (4.0f / InView->UnscaledViewRect.Width() / InView->ViewMatrices.GetProjectionMatrix().M[0][0]);
-							const float InnerRadius = (INNER_AXIS_CIRCLE_RADIUS * Scale) + GetDefault<ULevelEditorViewportSettings>()->TransformWidgetSizeAdjustment;
-							const float LengthOfAdjacent = Length - InnerRadius;
-							RotationAngle = FMath::Acos(FVector::DotProduct(OldCameraToPixelDir, CameraToPixelDir));
-							const float OppositeSize = FMath::Tan(RotationAngle) * LengthOfAdjacent;
-							RotationAngle = FMath::Atan(OppositeSize / InnerRadius);
-							RotationAngle = RotationAngle < 0.0f ? RotationAngle : -RotationAngle;
-						}
-				
-						const FQuat QuatRotation(RotationAxis, RotationAngle);
-						OutRotation = FRotator(QuatRotation);
-					}
-					break;
-				}
-				else if(CurrentAxis == EAxisList::Screen)
-				{
-					FVector2D MousePosition(InViewportClient->Viewport->GetMouseX(), InViewportClient->Viewport->GetMouseY());
-					FViewportCursorLocation OldMouseViewportRay(InView, InViewportClient, LastDragPos.X, LastDragPos.Y);
-					FViewportCursorLocation MouseViewportRay(InView, InViewportClient, MousePosition.X, MousePosition.Y);
-
-					LastDragPos = MousePosition;
-					FVector DirectionToWidget = InViewportClient->GetWidgetLocation() - MouseViewportRay.GetOrigin();
-					float Length = DirectionToWidget.Size();
-
-					if (!FMath::IsNearlyZero(Length))
-					{
-						DirectionToWidget /= Length;
-
-						const FVector CameraToPixelDir = MouseViewportRay.GetDirection();
-						const FVector OldCameraToPixelDir = OldMouseViewportRay.GetDirection();
-						FPlane Plane(InViewportClient->GetWidgetLocation(), DirectionToWidget);
-						FVector StartOnPlane = FMath::RayPlaneIntersection(MouseViewportRay.GetOrigin(), CameraToPixelDir, Plane);
-						FVector OldOnPlane = FMath::RayPlaneIntersection(MouseViewportRay.GetOrigin(), OldCameraToPixelDir, Plane);
-						StartOnPlane -= InViewportClient->GetWidgetLocation();
-						OldOnPlane -= InViewportClient->GetWidgetLocation();
-						StartOnPlane.Normalize();
-						OldOnPlane.Normalize();
-						float RotationAngle = FMath::Acos(FVector::DotProduct(StartOnPlane, OldOnPlane));
-						FVector Cross =  FVector::CrossProduct(OldCameraToPixelDir, CameraToPixelDir);
-						if (FVector::DotProduct(DirectionToWidget, Cross) < 0.0f)
-						{
-							RotationAngle *= -1.0f;
-						}	
-						const FQuat QuatRotation(DirectionToWidget, RotationAngle);
-						OutRotation = FRotator(QuatRotation);
-					}
-					break;
-				}
-			}
-			break;
-
-		case WM_Scale:
-			{
-				FVector2D AxisDir = FVector2D::ZeroVector;
-
-				if (CurrentAxis & EAxisList::X)
-				{
-					AxisDir += XAxisDir;
-				}
-
-				if (CurrentAxis & EAxisList::Y)
-				{
-					AxisDir += YAxisDir;
-				}
-
-				if (CurrentAxis & EAxisList::Z)
-				{
-					AxisDir += ZAxisDir;
-				}
-
-				AxisDir.Normalize();
-				const float ScaleDelta = FVector2D::DotProduct(AxisDir, DragDir);
-
-				OutScale = FVector(
-					(CurrentAxis & EAxisList::X) ? ScaleDelta : 0.0f,
-					(CurrentAxis & EAxisList::Y) ? ScaleDelta : 0.0f,
-					(CurrentAxis & EAxisList::Z) ? ScaleDelta : 0.0f
-					);
-
-				// Snap to grid in widget axis space
-				const FVector GridSize = FVector(GEditor->GetGridSize());
-				FSnappingUtils::SnapScale(OutScale, GridSize);
-
-				// Convert to effective screen space delta, and replace input delta, adjusted for inverted screen space Y axis
-				const float ScaleMax = OutScale.GetMax();
-				const float ScaleMin = OutScale.GetMin();
-				const float ScaleApplied = (ScaleMax > -ScaleMin) ? ScaleMax : ScaleMin;
-				const FVector2D EffectiveDelta = AxisDir * ScaleApplied;
-				InOutDelta = FVector(EffectiveDelta.X, -EffectiveDelta.Y, 0.0f);
-			}
-			break;
+	OutScale    = FVector::ZeroVector;
 
-		case WM_TranslateRotateZ:
-			{
-				if( CurrentAxis == EAxisList::ZRotation )
-				{
-
-					const FVector2D AxisDir = bIsOrthoDrawingFullRing ? TangentDir : ZAxisDir;
-					FRotator Rotation = FRotator(0, FVector2D::DotProduct(AxisDir, DragDir), 0);
-					FSnappingUtils::SnapRotatorToGrid(Rotation);
-					CurrentDeltaRotation = Rotation.Yaw;
-
-					const FVector2D EffectiveDelta = AxisDir * Rotation.Yaw;
-					InOutDelta = FVector(EffectiveDelta.X, -EffectiveDelta.Y, 0.0f);
-
-					OutRotation = (CustomCoordSystem.Inverse() * FRotationMatrix(Rotation) * CustomCoordSystem).Rotator();
-				}
-				else
-				{
-					// Get drag delta in widget axis space
-					OutDrag = FVector(
-						(CurrentAxis & EAxisList::X) ? FVector2D::DotProduct(XAxisDir, DragDir) : 0.0f,
-						(CurrentAxis & EAxisList::Y) ? FVector2D::DotProduct(YAxisDir, DragDir) : 0.0f,
-						(CurrentAxis & EAxisList::Z) ? FVector2D::DotProduct(ZAxisDir, DragDir) : 0.0f
-						);
-
-					// Snap to grid in widget axis space
-					const FVector GridSize = FVector(GEditor->GetGridSize());
-					FSnappingUtils::SnapPointToGrid(OutDrag, GridSize);
-
-					// Convert to effective screen space delta, and replace input delta, adjusted for inverted screen space Y axis
-					const FVector2D EffectiveDelta = OutDrag.X * XAxisDir + OutDrag.Y * YAxisDir + OutDrag.Z * ZAxisDir;
-					InOutDelta = FVector(EffectiveDelta.X, -EffectiveDelta.Y, 0.0f);
-
-					// Transform drag delta into world space
-					OutDrag = CustomCoordSystem.TransformPosition(OutDrag);
-				}
-			}
-			break;
-
-		case WM_2D:
-		{
-			if (CurrentAxis == EAxisList::Rotate2D)
-			{
-				// TODO: Determine why -TangentDir is necessary here, and fix whatever is causing it
-				const FVector2D AxisDir = bIsOrthoDrawingFullRing ? -TangentDir : YAxisDir;
-
-				FRotator Rotation = FRotator(FVector2D::DotProduct(AxisDir, DragDir), 0, 0);
-				FSnappingUtils::SnapRotatorToGrid(Rotation);
-
-				CurrentDeltaRotation = Rotation.Pitch;
-				FVector2D EffectiveDelta = AxisDir * Rotation.Pitch;
-
-
-				// Adjust the input delta according to how much rotation was actually applied
-				InOutDelta = FVector(EffectiveDelta.X, -EffectiveDelta.Y, 0.0f);
-
-				// Need to get the delta rotation in the current coordinate space of the widget
-				OutRotation = (CustomCoordSystem.Inverse() * FRotationMatrix(Rotation) * CustomCoordSystem).Rotator();
-			}
-			else
-			{
-				// Get drag delta in widget axis space
-				OutDrag = FVector(
-					(CurrentAxis & EAxisList::X) ? FVector2D::DotProduct(XAxisDir, DragDir) : 0.0f,
-					(CurrentAxis & EAxisList::Y) ? FVector2D::DotProduct(YAxisDir, DragDir) : 0.0f,
-					(CurrentAxis & EAxisList::Z) ? FVector2D::DotProduct(ZAxisDir, DragDir) : 0.0f
-					);
-
-				// Snap to grid in widget axis space
-				const FVector GridSize = FVector(GEditor->GetGridSize());
-				FSnappingUtils::SnapPointToGrid(OutDrag, GridSize);
-
-				// Convert to effective screen space delta, and replace input delta, adjusted for inverted screen space Y axis
-				const FVector2D EffectiveDelta = OutDrag.X * XAxisDir + OutDrag.Y * YAxisDir + OutDrag.Z * ZAxisDir;
-				InOutDelta = FVector(EffectiveDelta.X, -EffectiveDelta.Y, 0.0f);
-
-				// Transform drag delta into world space
-				OutDrag = CustomCoordSystem.TransformPosition(OutDrag);
-			}
-		}
-			break;
+	const int32 WidgetMode = InViewportClient->GetWidgetMode();
 
+	// Get input delta as 2D vector, adjusted for inverted screen space Y axis
+	const FVector2D DragDir = FVector2D(InOutDelta.X, -InOutDelta.Y);
 
-		default:
-			break;
+	// Get offset of the drag start position from the widget origin
+	const FVector2D DirectionToMousePos = FVector2D(DragStartPos - Origin).GetSafeNormal();
+
+	// For rotations which display as a full ring, calculate the tangent direction representing a clockwise movement
+	FVector2D TangentDir = bInUsedDragModifier ?
+	    // If a drag modifier has been used, this implies we are not actually touching the widget, so don't attempt to
+	    // calculate the tangent dir based on the relative offset of the cursor from the widget location.
+	    FVector2D(1, 1).GetSafeNormal() :
+	    // Treat the tangent dir as perpendicular to the relative offset of the cursor from the widget location.
+	    FVector2D(-DirectionToMousePos.Y, DirectionToMousePos.X);
+
+	switch (WidgetMode)
+	{
+	case WM_Translate:
+		ConvertMouseToAxis_Translate(DragDir, InOutDelta, OutDrag);
+		break;
+	case WM_Rotate:
+		ConvertMouseToAxis_Rotate(TangentDir, DragDir, InView, InViewportClient, InOutDelta, OutRotation);
+		break;
+	case WM_Scale:
+		ConvertMouseToAxis_Scale(DragDir, InOutDelta, OutScale);
+		break;
+	case WM_TranslateRotateZ:
+		ConvertMouseToAxis_TranslateRotateZ(TangentDir, DragDir, InOutDelta, OutDrag, OutRotation);
+		break;
+	case WM_2D:
+		ConvertMouseToAxis_WM_2D(TangentDir, DragDir, InOutDelta, OutDrag, OutRotation);
+		break;
+	default:
+		break;
 	}
 }
 
 /**
  * For axis movement, get the "best" planar normal and axis mask
  * @param InAxis - Axis of movement
- * @param InDirToPixel - 
+ * @param InDirToPixel -
  * @param OutPlaneNormal - Normal of the plane to project the mouse onto
  * @param OutMask - Used to mask out the component of the planar movement we want
  */
-void GetAxisPlaneNormalAndMask(const FMatrix& InCoordSystem, const FVector& InAxis, const FVector& InDirToPixel, FVector& OutPlaneNormal, FVector& NormalToRemove)
+void GetAxisPlaneNormalAndMask(const FMatrix& InCoordSystem, const FVector& InAxis, const FVector& InDirToPixel,
+                               FVector& OutPlaneNormal, FVector& NormalToRemove)
 {
 	FVector XAxis = InCoordSystem.TransformVector(FVector(1, 0, 0));
 	FVector YAxis = InCoordSystem.TransformVector(FVector(0, 1, 0));
@@ -1469,12 +452,12 @@ void GetAxisPlaneNormalAndMask(const FMatrix& InCoordSystem, const FVector& InAx
 	float YDot = FMath::Abs(InDirToPixel | YAxis);
 	float ZDot = FMath::Abs(InDirToPixel | ZAxis);
 
-	if ((InAxis|XAxis) > .1f)
+	if ((InAxis | XAxis) > .1f)
 	{
 		OutPlaneNormal = (YDot > ZDot) ? YAxis : ZAxis;
 		NormalToRemove = (YDot > ZDot) ? ZAxis : YAxis;
 	}
-	else if ((InAxis|YAxis) > .1f)
+	else if ((InAxis | YAxis) > .1f)
 	{
 		OutPlaneNormal = (XDot > ZDot) ? XAxis : ZAxis;
 		NormalToRemove = (XDot > ZDot) ? ZAxis : XAxis;
@@ -1498,191 +481,218 @@ void GetPlaneNormalAndMask(const FVector& InAxis, FVector& OutPlaneNormal, FVect
 	NormalToRemove = InAxis;
 }
 
-/**
- * Absolute Translation conversion from mouse movement on the screen to widget axis movement/rotation.
- */
-void FWidget::AbsoluteTranslationConvertMouseMovementToAxisMovement(FSceneView* InView, FEditorViewportClient* InViewportClient, const FVector& InLocation, const FVector2D& InMousePosition, FVector& OutDrag, FRotator& OutRotation, FVector& OutScale )
+void FWidget::AbsoluteConvertMouseToAxis_Translate(FSceneView* InView, const FMatrix& InputCoordSystem,
+                                                   FAbsoluteMovementParams& InOutParams, FVector& OutDrag)
 {
-	// Compute a world space ray from the screen space mouse coordinates
-	FViewportCursorLocation MouseViewportRay( InView, InViewportClient, InMousePosition.X, InMousePosition.Y );
+	switch (CurrentAxis)
+	{
+	case EAxisList::X:
+		GetAxisPlaneNormalAndMask(InputCoordSystem, InOutParams.XAxis, InOutParams.CameraDir, InOutParams.PlaneNormal,
+		                          InOutParams.NormalToRemove);
+		break;
+	case EAxisList::Y:
+		GetAxisPlaneNormalAndMask(InputCoordSystem, InOutParams.YAxis, InOutParams.CameraDir, InOutParams.PlaneNormal,
+		                          InOutParams.NormalToRemove);
+		break;
+	case EAxisList::Z:
+		GetAxisPlaneNormalAndMask(InputCoordSystem, InOutParams.ZAxis, InOutParams.CameraDir, InOutParams.PlaneNormal,
+		                          InOutParams.NormalToRemove);
+		break;
+	case EAxisList::XY:
+		GetPlaneNormalAndMask(InOutParams.ZAxis, InOutParams.PlaneNormal, InOutParams.NormalToRemove);
+		break;
+	case EAxisList::XZ:
+		GetPlaneNormalAndMask(InOutParams.YAxis, InOutParams.PlaneNormal, InOutParams.NormalToRemove);
+		break;
+	case EAxisList::YZ:
+		GetPlaneNormalAndMask(InOutParams.XAxis, InOutParams.PlaneNormal, InOutParams.NormalToRemove);
+		break;
+	case EAxisList::Screen:
+		InOutParams.XAxis = InView->ViewMatrices.GetViewMatrix().GetColumn(0);
+		InOutParams.YAxis = InView->ViewMatrices.GetViewMatrix().GetColumn(1);
+		InOutParams.ZAxis = InView->ViewMatrices.GetViewMatrix().GetColumn(2);
+		GetPlaneNormalAndMask(InOutParams.ZAxis, InOutParams.PlaneNormal, InOutParams.NormalToRemove);
+		//do not damp the movement in this case, we also want to snap
+		InOutParams.bMovementLockedToCamera = false;
+		break;
+	}
 
-	FAbsoluteMovementParams Params;
-	Params.EyePos = MouseViewportRay.GetOrigin();
-	Params.PixelDir = MouseViewportRay.GetDirection();
-	Params.CameraDir = InView->GetViewDirection();
-	Params.Position = InLocation;
-	//dampen by 
-	Params.bMovementLockedToCamera = InViewportClient->IsShiftPressed();
-	Params.bPositionSnapping = true;
+	OutDrag = GetAbsoluteTranslationDelta(InOutParams);
+}
 
-	FMatrix InputCoordSystem = InViewportClient->GetWidgetCoordSystem();
+void FWidget::AbsoluteConvertMouseToAxis_WM_2D(const FMatrix& InputCoordSystem, FAbsoluteMovementParams& InOutParams,
+                                               FVector& OutDrag, FRotator& OutRotation)
+{
+	switch (CurrentAxis)
+	{
+	case EAxisList::X: {
+		GetAxisPlaneNormalAndMask(InputCoordSystem, InOutParams.XAxis, InOutParams.CameraDir, InOutParams.PlaneNormal,
+		                          InOutParams.NormalToRemove);
+		OutDrag = GetAbsoluteTranslationDelta(InOutParams);
+		break;
+	}
+	case EAxisList::Z: {
+		GetAxisPlaneNormalAndMask(InputCoordSystem, InOutParams.ZAxis, InOutParams.CameraDir, InOutParams.PlaneNormal,
+		                          InOutParams.NormalToRemove);
+		OutDrag = GetAbsoluteTranslationDelta(InOutParams);
+		break;
+	}
+	case EAxisList::XZ: {
+		GetPlaneNormalAndMask(InOutParams.YAxis, InOutParams.PlaneNormal, InOutParams.NormalToRemove);
+		OutDrag = GetAbsoluteTranslationDelta(InOutParams);
+		break;
+	}
 
-	Params.XAxis = InputCoordSystem.TransformVector(FVector(1, 0, 0));
-	Params.YAxis = InputCoordSystem.TransformVector(FVector(0, 1, 0));
-	Params.ZAxis = InputCoordSystem.TransformVector(FVector(0, 0, 1));
+	//Rotate about the y-axis
+	case EAxisList::Rotate2D: {
+		//no position snapping, we'll handle the rotation snapping elsewhere
+		InOutParams.bPositionSnapping = false;
 
-	switch( InViewportClient->GetWidgetMode() )
-	{
-		case WM_Translate:
-		{
-			switch( CurrentAxis )
-			{
-			case EAxisList::X:  GetAxisPlaneNormalAndMask(InputCoordSystem, Params.XAxis, Params.CameraDir, Params.PlaneNormal, Params.NormalToRemove); break;
-			case EAxisList::Y:  GetAxisPlaneNormalAndMask(InputCoordSystem, Params.YAxis, Params.CameraDir, Params.PlaneNormal, Params.NormalToRemove); break;
-			case EAxisList::Z:  GetAxisPlaneNormalAndMask(InputCoordSystem, Params.ZAxis, Params.CameraDir, Params.PlaneNormal, Params.NormalToRemove); break;
-			case EAxisList::XY: GetPlaneNormalAndMask(Params.ZAxis, Params.PlaneNormal, Params.NormalToRemove); break;
-			case EAxisList::XZ: GetPlaneNormalAndMask(Params.YAxis, Params.PlaneNormal, Params.NormalToRemove); break;
-			case EAxisList::YZ: GetPlaneNormalAndMask(Params.XAxis, Params.PlaneNormal, Params.NormalToRemove); break;
-			case EAxisList::Screen:
-				Params.XAxis = InView->ViewMatrices.GetViewMatrix().GetColumn(0);
-				Params.YAxis = InView->ViewMatrices.GetViewMatrix().GetColumn(1);
-				Params.ZAxis = InView->ViewMatrices.GetViewMatrix().GetColumn(2);
-				GetPlaneNormalAndMask(Params.ZAxis, Params.PlaneNormal, Params.NormalToRemove);
-				//do not damp the movement in this case, we also want to snap
-				Params.bMovementLockedToCamera = false;
-				break;
-			}
+		GetPlaneNormalAndMask(InOutParams.YAxis, InOutParams.PlaneNormal, InOutParams.NormalToRemove);
+		//No DAMPING
+		InOutParams.bMovementLockedToCamera = false;
+		//this is the one movement type where we want to always use the widget origin and
+		//NOT the "first click" origin
+		FVector XZPlaneProjectedPosition = GetAbsoluteTranslationDelta(InOutParams) + InitialTranslationOffset;
 
-			OutDrag = GetAbsoluteTranslationDelta(Params);
+		//remove the component along the normal we want to mute
+		float MovementAlongMutedAxis = XZPlaneProjectedPosition | InOutParams.NormalToRemove;
+		XZPlaneProjectedPosition     = XZPlaneProjectedPosition - (InOutParams.NormalToRemove * MovementAlongMutedAxis);
 
-			break;
+		if (!XZPlaneProjectedPosition.Normalize())
+		{
+			XZPlaneProjectedPosition = InOutParams.YAxis;
 		}
 
-		case WM_2D:
-		{
-			switch (CurrentAxis)
-			{
-				case EAxisList::X:
-				{
-					GetAxisPlaneNormalAndMask(InputCoordSystem, Params.XAxis, Params.CameraDir, Params.PlaneNormal, Params.NormalToRemove);
-					OutDrag = GetAbsoluteTranslationDelta(Params);
-					break;
-				}
-				case EAxisList::Z:
-				{
-					GetAxisPlaneNormalAndMask(InputCoordSystem, Params.ZAxis, Params.CameraDir, Params.PlaneNormal, Params.NormalToRemove);
-					OutDrag = GetAbsoluteTranslationDelta(Params);
-					break;
-				}
-				case EAxisList::XZ:
-				{
-					GetPlaneNormalAndMask(Params.YAxis, Params.PlaneNormal, Params.NormalToRemove);
-					OutDrag = GetAbsoluteTranslationDelta(Params);
-					break;
-				}
-			
-				//Rotate about the y-axis
-				case EAxisList::Rotate2D:
-				{
-					//no position snapping, we'll handle the rotation snapping elsewhere
-					Params.bPositionSnapping = false;
+		//NOW, find the rotation around the PlaneNormal to make the xaxis point at InDrag
+		OutRotation = FRotator::ZeroRotator;
 
-					GetPlaneNormalAndMask(Params.YAxis, Params.PlaneNormal, Params.NormalToRemove);
-					//No DAMPING
-					Params.bMovementLockedToCamera = false;
-					//this is the one movement type where we want to always use the widget origin and 
-					//NOT the "first click" origin
-					FVector XZPlaneProjectedPosition = GetAbsoluteTranslationDelta(Params) + InitialTranslationOffset;
+		float PitchDegrees = -FMath::Atan2(-XZPlaneProjectedPosition.Z, XZPlaneProjectedPosition.X) * 180.f / PI;
+		OutRotation.Pitch  = PitchDegrees - (EditorModeTools ? EditorModeTools->TranslateRotate2DAngle : 0);
 
-					//remove the component along the normal we want to mute
-					float MovementAlongMutedAxis = XZPlaneProjectedPosition | Params.NormalToRemove;
-					XZPlaneProjectedPosition = XZPlaneProjectedPosition - (Params.NormalToRemove*MovementAlongMutedAxis);
+		if (bSnapEnabled)
+		{
+			FSnappingUtils::SnapRotatorToGrid(OutRotation);
+		}
 
-					if (!XZPlaneProjectedPosition.Normalize())
-					{
-						XZPlaneProjectedPosition = Params.YAxis;
-					}
+		break;
+	}
+	}
+}
+void FWidget::AbsoluteConvertMouseToAxis_TranslateRotateZ(const FMatrix& InputCoordSystem,
+                                                          FAbsoluteMovementParams& InOutParams, FVector& OutDrag,
+                                                          FRotator& OutRotation)
+{
+	FVector LineToUse;
+	switch (CurrentAxis)
+	{
+	case EAxisList::X: {
+		GetAxisPlaneNormalAndMask(InputCoordSystem, InOutParams.XAxis, InOutParams.CameraDir, InOutParams.PlaneNormal,
+		                          InOutParams.NormalToRemove);
+		OutDrag = GetAbsoluteTranslationDelta(InOutParams);
+		break;
+	}
+	case EAxisList::Y: {
+		GetAxisPlaneNormalAndMask(InputCoordSystem, InOutParams.YAxis, InOutParams.CameraDir, InOutParams.PlaneNormal,
+		                          InOutParams.NormalToRemove);
+		OutDrag = GetAbsoluteTranslationDelta(InOutParams);
+		break;
+	}
+	case EAxisList::Z: {
+		GetAxisPlaneNormalAndMask(InputCoordSystem, InOutParams.ZAxis, InOutParams.CameraDir, InOutParams.PlaneNormal,
+		                          InOutParams.NormalToRemove);
+		OutDrag = GetAbsoluteTranslationDelta(InOutParams);
+		break;
+	}
+	case EAxisList::XY: {
+		GetPlaneNormalAndMask(InOutParams.ZAxis, InOutParams.PlaneNormal, InOutParams.NormalToRemove);
+		OutDrag = GetAbsoluteTranslationDelta(InOutParams);
+		break;
+	}
+	//Rotate about the z-axis
+	case EAxisList::ZRotation: {
+		//no position snapping, we'll handle the rotation snapping elsewhere
+		InOutParams.bPositionSnapping = false;
 
-					//NOW, find the rotation around the PlaneNormal to make the xaxis point at InDrag
-					OutRotation = FRotator::ZeroRotator;
+		//find new point on the
+		GetPlaneNormalAndMask(InOutParams.ZAxis, InOutParams.PlaneNormal, InOutParams.NormalToRemove);
+		//No DAMPING
+		InOutParams.bMovementLockedToCamera = false;
+		//this is the one movement type where we want to always use the widget origin and
+		//NOT the "first click" origin
+		FVector XYPlaneProjectedPosition = GetAbsoluteTranslationDelta(InOutParams) + InitialTranslationOffset;
 
-					float PitchDegrees = -FMath::Atan2(-XZPlaneProjectedPosition.Z, XZPlaneProjectedPosition.X) * 180.f / PI;
-					OutRotation.Pitch = PitchDegrees - (EditorModeTools ? EditorModeTools->TranslateRotate2DAngle : 0);
+		//remove the component along the normal we want to mute
+		float MovementAlongMutedAxis = XYPlaneProjectedPosition | InOutParams.NormalToRemove;
+		XYPlaneProjectedPosition     = XYPlaneProjectedPosition - (InOutParams.NormalToRemove * MovementAlongMutedAxis);
 
-					if (bSnapEnabled)
-					{
-						FSnappingUtils::SnapRotatorToGrid(OutRotation);
-					}
+		if (!XYPlaneProjectedPosition.Normalize())
+		{
+			XYPlaneProjectedPosition = InOutParams.XAxis;
+		}
 
-					break;
-				}
-			}
+		//NOW, find the rotation around the PlaneNormal to make the xaxis point at InDrag
+		OutRotation = FRotator::ZeroRotator;
 
-			break;
-		}
+		OutRotation.Yaw = XYPlaneProjectedPosition.Rotation().Yaw -
+		    (EditorModeTools ? EditorModeTools->TranslateRotateXAxisAngle : 0);
 
-		case WM_TranslateRotateZ:
+		if (bSnapEnabled)
 		{
-			FVector LineToUse;
-			switch( CurrentAxis )
-			{
-				case EAxisList::X:	
-					{
-						GetAxisPlaneNormalAndMask(InputCoordSystem, Params.XAxis, Params.CameraDir, Params.PlaneNormal, Params.NormalToRemove);
-						OutDrag = GetAbsoluteTranslationDelta (Params);
-						break;
-					}
-				case EAxisList::Y:	
-					{
-						GetAxisPlaneNormalAndMask(InputCoordSystem, Params.YAxis, Params.CameraDir, Params.PlaneNormal, Params.NormalToRemove);
-						OutDrag = GetAbsoluteTranslationDelta (Params);
-						break;
-					}
-				case EAxisList::Z:	
-				{
-					GetAxisPlaneNormalAndMask(InputCoordSystem, Params.ZAxis, Params.CameraDir, Params.PlaneNormal, Params.NormalToRemove);
-					OutDrag = GetAbsoluteTranslationDelta (Params);
-					break;
-				}
-				case EAxisList::XY:
-				{
-					GetPlaneNormalAndMask(Params.ZAxis, Params.PlaneNormal, Params.NormalToRemove);
-					OutDrag = GetAbsoluteTranslationDelta (Params);
-					break;
-				}
-				//Rotate about the z-axis
-				case EAxisList::ZRotation:
-				{
-					//no position snapping, we'll handle the rotation snapping elsewhere
-					Params.bPositionSnapping = false;
-
-					//find new point on the 
-					GetPlaneNormalAndMask(Params.ZAxis, Params.PlaneNormal, Params.NormalToRemove);
-					//No DAMPING
-					Params.bMovementLockedToCamera = false;
-					//this is the one movement type where we want to always use the widget origin and 
-					//NOT the "first click" origin
-					FVector XYPlaneProjectedPosition = GetAbsoluteTranslationDelta (Params) + InitialTranslationOffset;
+			FSnappingUtils::SnapRotatorToGrid(OutRotation);
+		}
 
-					//remove the component along the normal we want to mute
-					float MovementAlongMutedAxis = XYPlaneProjectedPosition|Params.NormalToRemove;
-					XYPlaneProjectedPosition = XYPlaneProjectedPosition - (Params.NormalToRemove*MovementAlongMutedAxis);
+		break;
+	}
+	default:
+		break;
+	}
+}
 
-					if (!XYPlaneProjectedPosition.Normalize())
-					{
-						XYPlaneProjectedPosition = Params.XAxis;
-					}
+/**
+ * Absolute Translation conversion from mouse movement on the screen to widget axis movement/rotation.
+ */
+void FWidget::AbsoluteTranslationConvertMouseMovementToAxisMovement(FSceneView* InView,
+                                                                    FEditorViewportClient* InViewportClient,
+                                                                    const FVector& InLocation,
+                                                                    const FVector2D& InMousePosition, FVector& OutDrag,
+                                                                    FRotator& OutRotation, FVector& OutScale)
+{
+	// Compute a world space ray from the screen space mouse coordinates
+	FViewportCursorLocation MouseViewportRay(InView, InViewportClient, InMousePosition.X, InMousePosition.Y);
 
-					//NOW, find the rotation around the PlaneNormal to make the xaxis point at InDrag
-					OutRotation = FRotator::ZeroRotator;
+	FAbsoluteMovementParams Params;
+	Params.EyePos    = MouseViewportRay.GetOrigin();
+	Params.PixelDir  = MouseViewportRay.GetDirection();
+	Params.CameraDir = InView->GetViewDirection();
+	Params.Position  = InLocation;
+	//dampen by
+	Params.bMovementLockedToCamera = InViewportClient->IsShiftPressed();
+	Params.bPositionSnapping       = true;
 
-					OutRotation.Yaw = XYPlaneProjectedPosition.Rotation().Yaw - (EditorModeTools ? EditorModeTools->TranslateRotateXAxisAngle : 0 );
+	FMatrix InputCoordSystem = InViewportClient->GetWidgetCoordSystem();
 
-					if (bSnapEnabled)
-					{
-						FSnappingUtils::SnapRotatorToGrid( OutRotation );
-					}
+	Params.XAxis = InputCoordSystem.TransformVector(FVector(1, 0, 0));
+	Params.YAxis = InputCoordSystem.TransformVector(FVector(0, 1, 0));
+	Params.ZAxis = InputCoordSystem.TransformVector(FVector(0, 0, 1));
 
-					break;
-				}
-				default:
-					break;
-			}
-		}
+	switch (InViewportClient->GetWidgetMode())
+	{
+	case WM_Translate:
+		AbsoluteConvertMouseToAxis_Translate(InView, InputCoordSystem, Params, OutDrag);
+		break;
 
-		case WM_Rotate:
-		case WM_Scale:
-			break;
+	case WM_2D:
+		AbsoluteConvertMouseToAxis_WM_2D(InputCoordSystem, Params, OutDrag, OutRotation);
+		break;
+	case WM_TranslateRotateZ:
+		AbsoluteConvertMouseToAxis_TranslateRotateZ(InputCoordSystem, Params, OutDrag, OutRotation);
+		break;
+	case WM_Rotate:
+	case WM_Scale:
+	case WM_None:
+	case WM_Max:
+		break;
 	}
 }
 
@@ -1706,20 +716,20 @@ bool FWidget::AllowsAbsoluteRotationMovement(EWidgetMode WidgetMode, EAxisList::
 	return false;
 }
 
-/** 
+/**
  * Serializes the widget references so they don't get garbage collected.
  *
  * @param Ar	FArchive to serialize with
  */
-void FWidget::AddReferencedObjects( FReferenceCollector& Collector )
+void FWidget::AddReferencedObjects(FReferenceCollector& Collector)
 {
-	Collector.AddReferencedObject( AxisMaterialX );
-	Collector.AddReferencedObject( AxisMaterialY );
-	Collector.AddReferencedObject( AxisMaterialZ );
-	Collector.AddReferencedObject( OpaquePlaneMaterialXY );
-	Collector.AddReferencedObject( TransparentPlaneMaterialXY );
-	Collector.AddReferencedObject( GridMaterial );
-	Collector.AddReferencedObject( CurrentAxisMaterial );
+	Collector.AddReferencedObject(AxisMaterialX);
+	Collector.AddReferencedObject(AxisMaterialY);
+	Collector.AddReferencedObject(AxisMaterialZ);
+	Collector.AddReferencedObject(OpaquePlaneMaterialXY);
+	Collector.AddReferencedObject(TransparentPlaneMaterialXY);
+	Collector.AddReferencedObject(GridMaterial);
+	Collector.AddReferencedObject(CurrentAxisMaterial);
 }
 
 #define CAMERA_LOCK_DAMPING_FACTOR .1f
@@ -1732,12 +742,13 @@ void FWidget::AddReferencedObjects( FReferenceCollector& Collector )
 FVector FWidget::GetAbsoluteTranslationDelta(const FAbsoluteMovementParams& InParams)
 {
 	FPlane MovementPlane(InParams.Position, InParams.PlaneNormal);
-	FVector ProposedEndofEyeVector = InParams.EyePos + (InParams.PixelDir * (InParams.Position - InParams.EyePos).Size());
+	FVector ProposedEndofEyeVector =
+	    InParams.EyePos + (InParams.PixelDir * (InParams.Position - InParams.EyePos).Size());
 
 	//default to not moving
 	FVector RequestedPosition = InParams.Position;
 
-	float DotProductWithPlaneNormal = InParams.PixelDir|InParams.PlaneNormal;
+	float DotProductWithPlaneNormal = InParams.PixelDir | InParams.PlaneNormal;
 	//check to make sure we're not co-planar
 	if (FMath::Abs(DotProductWithPlaneNormal) > DELTA)
 	{
@@ -1755,8 +766,8 @@ FVector FWidget::GetAbsoluteTranslationDelta(const FAbsoluteMovementParams& InPa
 	DeltaPosition -= InitialOffset;
 
 	//remove the component along the normal we want to mute
-	float MovementAlongMutedAxis = DeltaPosition|InParams.NormalToRemove;
-	FVector OutDrag = DeltaPosition - (InParams.NormalToRemove*MovementAlongMutedAxis);
+	float MovementAlongMutedAxis = DeltaPosition | InParams.NormalToRemove;
+	FVector OutDrag              = DeltaPosition - (InParams.NormalToRemove * MovementAlongMutedAxis);
 
 	if (InParams.bMovementLockedToCamera)
 	{
@@ -1770,21 +781,24 @@ FVector FWidget::GetAbsoluteTranslationDelta(const FAbsoluteMovementParams& InPa
 	//the they requested position snapping and we're not moving with the camera
 	if (InParams.bPositionSnapping && !InParams.bMovementLockedToCamera && bSnapEnabled)
 	{
-		FVector MovementAlongAxis = FVector(OutDrag|InParams.XAxis, OutDrag|InParams.YAxis, OutDrag|InParams.ZAxis);
+		FVector MovementAlongAxis =
+		    FVector(OutDrag | InParams.XAxis, OutDrag | InParams.YAxis, OutDrag | InParams.ZAxis);
 		//translation (either xy plane or z)
-		FSnappingUtils::SnapPointToGrid( MovementAlongAxis, FVector(GEditor->GetGridSize(),GEditor->GetGridSize(),GEditor->GetGridSize()) );
-		OutDrag = MovementAlongAxis.X*InParams.XAxis + MovementAlongAxis.Y*InParams.YAxis + MovementAlongAxis.Z*InParams.ZAxis;
+		FSnappingUtils::SnapPointToGrid(
+		    MovementAlongAxis, FVector(GEditor->GetGridSize(), GEditor->GetGridSize(), GEditor->GetGridSize()));
+		OutDrag = MovementAlongAxis.X * InParams.XAxis + MovementAlongAxis.Y * InParams.YAxis +
+		    MovementAlongAxis.Z * InParams.ZAxis;
 	}
 
-	//get the distance from the original position to the new proposed position 
+	//get the distance from the original position to the new proposed position
 	FVector DeltaFromStart = InParams.Position + OutDrag - InitialTranslationPosition;
 
 	//Get the vector from the eye to the proposed new position (to make sure it's not behind the camera
-	FVector EyeToNewPosition = (InParams.Position + OutDrag) - InParams.EyePos;
-	float BehindTheCameraDotProduct = EyeToNewPosition|InParams.CameraDir;
+	FVector EyeToNewPosition        = (InParams.Position + OutDrag) - InParams.EyePos;
+	float BehindTheCameraDotProduct = EyeToNewPosition | InParams.CameraDir;
 
 	//Don't let the requested position go behind the camera
-	if ( BehindTheCameraDotProduct <= 0 )
+	if (BehindTheCameraDotProduct <= 0)
 	{
 		OutDrag = OutDrag.ZeroVector;
 	}
@@ -1799,26 +813,25 @@ FVector FWidget::GetAbsoluteTranslationInitialOffset(const FVector& InNewPositio
 	if (!bAbsoluteTranslationInitialOffsetCached)
 	{
 		bAbsoluteTranslationInitialOffsetCached = true;
-		InitialTranslationOffset = InNewPosition - InCurrentPosition;
-		InitialTranslationPosition = InCurrentPosition;
+		InitialTranslationOffset                = InNewPosition - InCurrentPosition;
+		InitialTranslationPosition              = InCurrentPosition;
 	}
 	return InitialTranslationOffset;
 }
 
 
-
 /**
  * Returns true if we're in Local Space editing mode
  */
 bool FWidget::IsRotationLocalSpace() const
 {
-	return ( CustomCoordSystemSpace == COORD_Local );
+	return (CustomCoordSystemSpace == COORD_Local);
 }
 
 void FWidget::UpdateDeltaRotation()
 {
 	TotalDeltaRotation += CurrentDeltaRotation;
-	if ( (TotalDeltaRotation <= -360.f) || (TotalDeltaRotation >= 360.f) )
+	if ((TotalDeltaRotation <= -360.f) || (TotalDeltaRotation >= 360.f))
 	{
 		TotalDeltaRotation = FRotator::ClampAxis(TotalDeltaRotation);
 	}
@@ -1832,477 +845,63 @@ float FWidget::GetDeltaRotation() const
 	return TotalDeltaRotation;
 }
 
-
-/**
- * If actively dragging, draws a ring representing the potential rotation of the selected objects, snap ticks, and "delta" markers
- * If not actively dragging, draws a quarter ring representing the closest quadrant to the camera
- * @param View - Information about the scene/camera/etc
- * @param PDI - Drawing interface
- * @param InAxis - Enumeration of axis to rotate about
- * @param InLocation - The Origin of the widget
- * @param Axis0 - The Axis that describes a 0 degree rotation
- * @param Axis1 - The Axis that describes a 90 degree rotation
- * @param InDirectionToWidget - Direction from camera to the widget
- * @param InColor - The color associated with the axis of rotation
- * @param InScale - Multiplier to maintain a constant screen size for rendering the widget
- */
-void FWidget::DrawRotationArc(const FSceneView* View, FPrimitiveDrawInterface* PDI, EAxisList::Type InAxis, const FVector& InLocation, const FVector& Axis0, const FVector& Axis1, const FVector& InDirectionToWidget, const FColor& InColor, const float InScale, FVector2D& OutAxisDir)
-{
-	bool bIsPerspective = ( View->ViewMatrices.GetProjectionMatrix().M[3][3] < 1.0f );
-	bool bIsOrtho = !bIsPerspective;
-
-	//if we're in an ortho viewport and the ring is perpendicular to the camera (both Axis0 & Axis1 are perpendicular)
-	bIsOrthoDrawingFullRing |= bIsOrtho && (FMath::Abs(Axis0|InDirectionToWidget) < KINDA_SMALL_NUMBER) && (FMath::Abs(Axis1|InDirectionToWidget) < KINDA_SMALL_NUMBER);
-
-	FColor ArcColor = InColor;
-	ArcColor.A = LargeOuterAlpha;
-
-	if (bDragging || (bIsOrthoDrawingFullRing))
-	{
-		if ((CurrentAxis&InAxis) || (bIsOrthoDrawingFullRing))
-		{
-			bool bDrawingArcBallAlso = CurrentAxis == EAxisList::XYZ;
-
-			float DeltaRotation = GetDeltaRotation();
-			float AdjustedDeltaRotation = IsRotationLocalSpace() ? -DeltaRotation : DeltaRotation;
-			float AbsRotation = FRotator::ClampAxis(FMath::Abs(DeltaRotation));
-			float AngleOfChangeRadians (AbsRotation * PI / 180.f);
-
-			//always draw clockwise, so if we're negative we need to flip the angle
-			float StartAngle = AdjustedDeltaRotation < 0.0f ? -AngleOfChangeRadians : 0.0f;
-			float FilledAngle = AngleOfChangeRadians;
-
-			//the axis of rotation
-			FVector ZAxis = Axis0 ^ Axis1;
-
-			ArcColor.A = LargeOuterAlpha;
-			DrawPartialRotationArc(View, PDI, InAxis, InLocation,  Axis0, Axis1, StartAngle, StartAngle + FilledAngle, ArcColor, InScale, InDirectionToWidget);
-			ArcColor.A = SmallOuterAlpha;
-			DrawPartialRotationArc(View, PDI, InAxis, InLocation,  Axis0, Axis1, StartAngle + FilledAngle, StartAngle + 2*PI, ArcColor, InScale, InDirectionToWidget);
-
-			ArcColor = (CurrentAxis&InAxis && !bDrawingArcBallAlso) ? CurrentColor : ArcColor;
-			if (!bDrawingArcBallAlso)
-			{
-				//Hallow Arrow
-				ArcColor.A = 0;
-				DrawStartStopMarker(PDI, InLocation, Axis0, Axis1, 0, ArcColor, InScale);
-				//Filled Arrow
-				ArcColor.A = LargeOuterAlpha;
-				DrawStartStopMarker(PDI, InLocation, Axis0, Axis1, AdjustedDeltaRotation, ArcColor, InScale);
-				ArcColor.A = 255;
-
-
-				FVector SnapLocation = InLocation;
-
-				if (GetDefault<ULevelEditorViewportSettings>()->RotGridEnabled)
-				{
-					float DeltaAngle = GEditor->GetRotGridSize().Yaw;
-					//every 22.5 degrees
-					float TickMarker = 22.5f;
-					for (float Angle = 0; Angle < 360.f; Angle += DeltaAngle)
-					{
-						FVector GridAxis = Axis0.RotateAngleAxis(Angle, ZAxis);
-						float PercentSize = (FMath::Fmod(Angle, TickMarker) == 0) ? .75f : .25f;
-						if (FMath::Fmod(Angle, 90.f) != 0)
-						{
-							DrawSnapMarker(PDI, SnapLocation, GridAxis, FVector::ZeroVector, ArcColor, InScale, 0.0f, PercentSize);
-						}
-					}
-				}
-
-				//draw axis tick marks
-				FColor AxisColor = InColor;
-				//Rotate Colors to match Axis 0
-				Swap(AxisColor.R, AxisColor.G);
-				Swap(AxisColor.B, AxisColor.R);
-				AxisColor.A = (AdjustedDeltaRotation == 0) ? MAX_uint8 : LargeOuterAlpha;
-				DrawSnapMarker(PDI, SnapLocation, Axis0, Axis1, AxisColor, InScale, .25f);
-				AxisColor.A = (AdjustedDeltaRotation == 180.f) ? MAX_uint8 : LargeOuterAlpha;
-				DrawSnapMarker(PDI, SnapLocation, -Axis0, -Axis1, AxisColor, InScale, .25f);
-
-				//Rotate Colors to match Axis 1
-				Swap(AxisColor.R, AxisColor.G);
-				Swap(AxisColor.B, AxisColor.R);
-				AxisColor.A = (AdjustedDeltaRotation == 90.f) ? MAX_uint8 : LargeOuterAlpha;
-				DrawSnapMarker(PDI, SnapLocation, Axis1, -Axis0, AxisColor, InScale, .25f);
-				AxisColor.A = (AdjustedDeltaRotation == 270.f) ? MAX_uint8 : LargeOuterAlpha;
-				DrawSnapMarker(PDI, SnapLocation, -Axis1, Axis0, AxisColor, InScale, .25f);
-
-				if (bDragging)
-				{
-					float OffsetAngle = IsRotationLocalSpace() ? 0 : AdjustedDeltaRotation;
-
-					CacheRotationHUDText(View, PDI, InLocation, Axis0.RotateAngleAxis(OffsetAngle, ZAxis), Axis1.RotateAngleAxis(OffsetAngle, ZAxis), DeltaRotation, InScale);
-				}
-			}
-		}
-	}
-	else
-	{
-		//Reverse the axes based on camera view
-		bool bMirrorAxis0 = ((Axis0 | InDirectionToWidget) <= 0.0f);
-		bool bMirrorAxis1 = ((Axis1 | InDirectionToWidget) <= 0.0f);
-
-		FVector RenderAxis0 = bMirrorAxis0 ? Axis0 : -Axis0;
-		FVector RenderAxis1 = bMirrorAxis1 ? Axis1 : -Axis1;
-		float Direction = (bMirrorAxis0 ^ bMirrorAxis1) ? -1.0f : 1.0f;
-
-		DrawPartialRotationArc(View, PDI, InAxis, InLocation, RenderAxis0, RenderAxis1, 0, PI/2, ArcColor, InScale, InDirectionToWidget);
-
-		FVector2D Axis0ScreenLocation;
-		if (!View->ScreenToPixel(View->WorldToScreen(InLocation + RenderAxis0 * 64.0f), Axis0ScreenLocation))
-		{
-			Axis0ScreenLocation.X = Axis0ScreenLocation.Y = 0;
-		}
-
-		FVector2D Axis1ScreenLocation;
-		if (!View->ScreenToPixel(View->WorldToScreen(InLocation + RenderAxis1 * 64.0f), Axis1ScreenLocation))
-		{
-			Axis1ScreenLocation.X = Axis1ScreenLocation.Y = 0;
-		}
-
-		OutAxisDir = ((Axis1ScreenLocation - Axis0ScreenLocation) * Direction).GetSafeNormal();
-	}
-}
-
-/**
- * If actively dragging, draws a ring representing the potential rotation of the selected objects, snap ticks, and "delta" markers
- * If not actively dragging, draws a quarter ring representing the closest quadrant to the camera
- * @param View - Information about the scene/camera/etc
- * @param PDI - Drawing interface
- * @param InAxis - Enumeration of axis to rotate about
- * @param InLocation - The Origin of the widget
- * @param Axis0 - The Axis that describes a 0 degree rotation
- * @param Axis1 - The Axis that describes a 90 degree rotation
- * @param InStartAngle - The starting angle about (Axis0^Axis1) to render the arc, in radians
- * @param InEndAngle - The ending angle about (Axis0^Axis1) to render the arc, in radians
- * @param InColor - The color associated with the axis of rotation
- * @param InScale - Multiplier to maintain a constant screen size for rendering the widget
- */
-void FWidget::DrawPartialRotationArc(const FSceneView* View, FPrimitiveDrawInterface* PDI, EAxisList::Type InAxis, const FVector& InLocation, const FVector& Axis0, const FVector& Axis1, const float InStartAngle, const float InEndAngle, const FColor& InColor, const float InScale, const FVector& InDirectionToWidget )
-{
-	const float InnerRadius = (INNER_AXIS_CIRCLE_RADIUS * InScale) + GetDefault<ULevelEditorViewportSettings>()->TransformWidgetSizeAdjustment;
-	const float OuterRadius = (OUTER_AXIS_CIRCLE_RADIUS * InScale) + GetDefault<ULevelEditorViewportSettings>()->TransformWidgetSizeAdjustment;
-
-	bool bIsPerspective = ( View->ViewMatrices.GetProjectionMatrix().M[3][3] < 1.0f );
-	PDI->SetHitProxy( new HWidgetAxis( InAxis ) );
-	{
-		FThickArcParams OuterArcParams(PDI, InLocation, TransparentPlaneMaterialXY, InnerRadius, OuterRadius);
-		FColor OuterColor = ( (CurrentAxis&InAxis  && ( CurrentAxis != EAxisList::XYZ ) )? CurrentColor : InColor );
-		//Pass through alpha
-		OuterColor.A = InColor.A;
-		DrawThickArc(OuterArcParams, Axis0, Axis1, InStartAngle, InEndAngle, OuterColor, InDirectionToWidget, !bIsPerspective );
-	}
-	PDI->SetHitProxy( NULL );
-
-	const bool bIsHitProxyView = View->Family->EngineShowFlags.HitProxies;
-	if (bIsPerspective && !bIsHitProxyView && !PDI->IsHitTesting())
-	{
-		FThickArcParams InnerArcParams(PDI, InLocation, GridMaterial, 0.0f, InnerRadius);
-		FColor InnerColor = InColor;
-		//if something is selected and it's not this
-		InnerColor.A = ((CurrentAxis & InAxis) && !bDragging) ? LargeInnerAlpha : SmallInnerAlpha;
-		DrawThickArc(InnerArcParams, Axis0, Axis1, InStartAngle, InEndAngle, InnerColor, InDirectionToWidget, false );
-	}
-}
-
-/**
- * Renders a portion of an arc for the rotation widget
- * @param InParams - Material, Radii, etc
- * @param InStartAxis - Start of the arc, in radians
- * @param InEndAxis - End of the arc, in radians
- * @param InColor - Color to use for the arc
- */
-void FWidget::DrawThickArc (const FThickArcParams& InParams, const FVector& Axis0, const FVector& Axis1, const float InStartAngle, const float InEndAngle, const FColor& InColor, const FVector& InDirectionToWidget, bool bIsOrtho )
+uint32 FWidget::GetDominantAxisIndex(const FVector& InDiff, FEditorViewportClient* ViewportClient) const
 {
-	if (InColor.A == 0)
+	uint32 DominantIndex = 0;
+	if (FMath::Abs(InDiff.X) < FMath::Abs(InDiff.Y))
 	{
-		return;
+		DominantIndex = 1;
 	}
 
-	// Add more sides to the circle if we've been scaled up to keep the circle looking circular
-	// An extra side for every 5 extra unreal units seems to produce a nice result
-	const int32 CircleSides = (GetDefault<ULevelEditorViewportSettings>()->TransformWidgetSizeAdjustment > 0) 
-		? AXIS_CIRCLE_SIDES + (GetDefault<ULevelEditorViewportSettings>()->TransformWidgetSizeAdjustment / 5)
-		: AXIS_CIRCLE_SIDES;
-	const int32 NumPoints = FMath::TruncToInt(CircleSides * (InEndAngle-InStartAngle)/(PI/2)) + 1;
-
-	FColor TriangleColor = InColor;
-	FColor RingColor = InColor;
-	RingColor.A = MAX_uint8;
-
-	FVector ZAxis = Axis0 ^ Axis1;
-	FVector LastVertex;
-
-	FDynamicMeshBuilder MeshBuilder(InParams.PDI->View->GetFeatureLevel());
+	const int32 WidgetMode = ViewportClient->GetWidgetMode();
 
-	for (int32 RadiusIndex = 0; RadiusIndex < 2; ++RadiusIndex)
+	switch (WidgetMode)
 	{
-		float Radius = (RadiusIndex == 0) ? InParams.OuterRadius : InParams.InnerRadius;
-		float TCRadius = Radius / (float) InParams.OuterRadius;
-		//Compute vertices for base circle.
-		for(int32 VertexIndex = 0;VertexIndex <= NumPoints;VertexIndex++)
+	case WM_Translate:
+		switch (ViewportClient->ViewportType)
 		{
-			float Percent = VertexIndex/(float)NumPoints;
-			float Angle = FMath::Lerp(InStartAngle, InEndAngle, Percent);
-			float AngleDeg = FRotator::ClampAxis(Angle * 180.f / PI);
-
-			FVector VertexDir = Axis0.RotateAngleAxis(AngleDeg, ZAxis);
-			VertexDir.Normalize();
-
-			float TCAngle = Percent*(PI/2);
-			FVector2D TC(TCRadius*FMath::Cos(Angle), TCRadius*FMath::Sin(Angle));
-
-			const FVector VertexPosition = InParams.Position + VertexDir*Radius;
-			FVector Normal = VertexPosition - InParams.Position;
-			Normal.Normalize();
-
-			FDynamicMeshVertex MeshVertex;
-			MeshVertex.Position = VertexPosition;
-			MeshVertex.Color = TriangleColor;
-			MeshVertex.TextureCoordinate[0] = TC;
-
-			MeshVertex.SetTangents(
-				-ZAxis,
-				(-ZAxis) ^ Normal,
-				Normal
-				);
-
-			MeshBuilder.AddVertex(MeshVertex); //Add bottom vertex
-
-			// Push out the arc line borders so they dont z-fight with the mesh arcs
-			FVector StartLinePos = LastVertex;
-			FVector EndLinePos = VertexPosition;
-			if (VertexIndex != 0)
+		case LVT_OrthoXY:
+			if (CurrentAxis == EAxisList::X)
 			{
-				InParams.PDI->DrawLine(StartLinePos,EndLinePos,RingColor,SDPG_Foreground);
+				DominantIndex = 0;
 			}
-			LastVertex = VertexPosition;
-		}
-	}
-
-	//Add top/bottom triangles, in the style of a fan.
-	int32 InnerVertexStartIndex = NumPoints + 1;
-	for(int32 VertexIndex = 0; VertexIndex < NumPoints; VertexIndex++)
-	{
-		MeshBuilder.AddTriangle(VertexIndex, VertexIndex+1, InnerVertexStartIndex+VertexIndex);
-		MeshBuilder.AddTriangle(VertexIndex+1, InnerVertexStartIndex+VertexIndex+1, InnerVertexStartIndex+VertexIndex);
-	}
-
-	MeshBuilder.Draw(InParams.PDI, FMatrix::Identity, InParams.Material->GetRenderProxy(),SDPG_Foreground,0.f);
-}
-
-/**
- * Draws protractor like ticks where the rotation widget would snap too.
- * Also, used to draw the wider axis tick marks
- * @param PDI - Drawing interface
- * @param InLocation - The Origin of the widget
- * @param Axis0 - The Axis that describes a 0 degree rotation
- * @param Axis1 - The Axis that describes a 90 degree rotation
- * @param InAngle - The Angle to rotate about the axis of rotation, the vector (Axis0 ^ Axis1)
- * @param InColor - The color to use for line/poly drawing
- * @param InScale - Multiplier to maintain a constant screen size for rendering the widget
- * @param InWidthPercent - The percent of the distance between the outer ring and inner ring to use for tangential thickness
- * @param InPercentSize - The percent of the distance between the outer ring and inner ring to use for radial distance
- */
-void FWidget::DrawSnapMarker(FPrimitiveDrawInterface* PDI, const FVector& InLocation, const FVector& Axis0, const FVector& Axis1, const FColor& InColor, const float InScale, const float InWidthPercent, const float InPercentSize)
-{
-	const float InnerDistance = (INNER_AXIS_CIRCLE_RADIUS * InScale) + GetDefault<ULevelEditorViewportSettings>()->TransformWidgetSizeAdjustment;
-	const float OuterDistance = (OUTER_AXIS_CIRCLE_RADIUS * InScale) + GetDefault<ULevelEditorViewportSettings>()->TransformWidgetSizeAdjustment;
-	const float MaxMarkerHeight = OuterDistance - InnerDistance;
-	const float MarkerWidth = MaxMarkerHeight*InWidthPercent;
-	const float MarkerHeight = MaxMarkerHeight*InPercentSize;
-
-	FVector Vertices[4];
-	Vertices[0] = InLocation + (OuterDistance)*Axis0 - (MarkerWidth*.5)*Axis1;
-	Vertices[1] = Vertices[0] + (MarkerWidth)*Axis1;
-	Vertices[2] = InLocation + (OuterDistance-MarkerHeight)*Axis0 - (MarkerWidth*.5)*Axis1;
-	Vertices[3] = Vertices[2] + (MarkerWidth)*Axis1;
-
-	//draw at least one line
-	PDI->DrawLine(Vertices[0], Vertices[2], InColor, SDPG_Foreground);
-
-	//if there should be thickness, draw the other lines
-	if (InWidthPercent > 0.0f)
-	{
-		PDI->DrawLine(Vertices[0], Vertices[1], InColor, SDPG_Foreground);
-		PDI->DrawLine(Vertices[1], Vertices[3], InColor, SDPG_Foreground);
-		PDI->DrawLine(Vertices[2], Vertices[3], InColor, SDPG_Foreground);
-
-		//fill in the box
-		FDynamicMeshBuilder MeshBuilder(PDI->View->GetFeatureLevel());
-
-		for(int32 VertexIndex = 0;VertexIndex < 4; VertexIndex++)
-		{
-			FDynamicMeshVertex MeshVertex;
-			MeshVertex.Position = Vertices[VertexIndex];
-			MeshVertex.Color = InColor;
-			MeshVertex.TextureCoordinate[0] = FVector2D(0.0f, 0.0f);
-			MeshVertex.SetTangents(
-				Axis0,
-				Axis1,
-				(Axis0) ^ Axis1
-				);
-			MeshBuilder.AddVertex(MeshVertex); //Add bottom vertex
-		}
-
-		MeshBuilder.AddTriangle(0, 1, 2);
-		MeshBuilder.AddTriangle(1, 3, 2);
-		MeshBuilder.Draw(PDI, FMatrix::Identity, TransparentPlaneMaterialXY->GetRenderProxy(),SDPG_Foreground,0.f);
-	}
-}
-
-/**
- * Draw Start/Stop Marker to show delta rotations along the arc of rotation
- * @param PDI - Drawing interface
- * @param InLocation - The Origin of the widget
- * @param Axis0 - The Axis that describes a 0 degree rotation
- * @param Axis1 - The Axis that describes a 90 degree rotation
- * @param InAngle - The Angle to rotate about the axis of rotation, the vector (Axis0 ^ Axis1), units are degrees
- * @param InColor - The color to use for line/poly drawing
- * @param InScale - Multiplier to maintain a constant screen size for rendering the widget
- */
-void FWidget::DrawStartStopMarker(FPrimitiveDrawInterface* PDI, const FVector& InLocation, const FVector& Axis0, const FVector& Axis1, const float InAngle, const FColor& InColor, const float InScale)
-{
-	const float ArrowHeightPercent = .8f;
-	const float InnerDistance = (INNER_AXIS_CIRCLE_RADIUS * InScale) + GetDefault<ULevelEditorViewportSettings>()->TransformWidgetSizeAdjustment;
-	const float OuterDistance = (OUTER_AXIS_CIRCLE_RADIUS * InScale) + GetDefault<ULevelEditorViewportSettings>()->TransformWidgetSizeAdjustment;
-	const float RingHeight = OuterDistance - InnerDistance;
-	const float ArrowHeight = RingHeight*ArrowHeightPercent;
-	const float ThirtyDegrees = PI / 6.0f;
-	const float HalfArrowidth = ArrowHeight*FMath::Tan(ThirtyDegrees);
-
-	FVector ZAxis = Axis0 ^ Axis1;
-	FVector RotatedAxis0 = Axis0.RotateAngleAxis(InAngle, ZAxis);
-	FVector RotatedAxis1 = Axis1.RotateAngleAxis(InAngle, ZAxis);
-
-	FVector Vertices[3];
-	Vertices[0] = InLocation + (OuterDistance)*RotatedAxis0;
-	Vertices[1] = Vertices[0] + (ArrowHeight)*RotatedAxis0 - HalfArrowidth*RotatedAxis1;
-	Vertices[2] = Vertices[1] + (2*HalfArrowidth)*RotatedAxis1;
-
-	PDI->DrawLine(Vertices[0], Vertices[1], InColor, SDPG_Foreground);
-	PDI->DrawLine(Vertices[1], Vertices[2], InColor, SDPG_Foreground);
-	PDI->DrawLine(Vertices[0], Vertices[2], InColor, SDPG_Foreground);
-
-	if (InColor.A > 0)
-	{
-		//fill in the box
-		FDynamicMeshBuilder MeshBuilder(PDI->View->GetFeatureLevel());
-
-		for(int32 VertexIndex = 0;VertexIndex < 3; VertexIndex++)
-		{
-			FDynamicMeshVertex MeshVertex;
-			MeshVertex.Position = Vertices[VertexIndex];
-			MeshVertex.Color = InColor;
-			MeshVertex.TextureCoordinate[0] = FVector2D(0.0f, 0.0f);
-			MeshVertex.SetTangents(
-				RotatedAxis0,
-				RotatedAxis1,
-				(RotatedAxis0) ^ RotatedAxis1
-				);
-			MeshBuilder.AddVertex(MeshVertex); //Add bottom vertex
-		}
-
-		MeshBuilder.AddTriangle(0, 1, 2);
-		MeshBuilder.Draw(PDI, FMatrix::Identity, TransparentPlaneMaterialXY->GetRenderProxy(),SDPG_Foreground,0.f);
-	}
-}
-
-/**
- * Caches off HUD text to display after 3d rendering is complete
- * @param View - Information about the scene/camera/etc
- * @param PDI - Drawing interface
- * @param InLocation - The Origin of the widget
- * @param Axis0 - The Axis that describes a 0 degree rotation
- * @param Axis1 - The Axis that describes a 90 degree rotation
- * @param AngleOfAngle - angle we've rotated so far (in degrees)
- */
-void FWidget::CacheRotationHUDText(const FSceneView* View, FPrimitiveDrawInterface* PDI, const FVector& InLocation, const FVector& Axis0, const FVector& Axis1, const float AngleOfChange, const float InScale)
-{
-	const float TextDistance = (ROTATION_TEXT_RADIUS * InScale) + GetDefault<ULevelEditorViewportSettings>()->TransformWidgetSizeAdjustment;
-
-	FVector AxisVectors[4] = { Axis0, Axis1, -Axis0, -Axis1};
-
-	for (int i = 0 ; i < 4; ++i)
-	{
-		FVector PotentialTextPosition = InLocation + (TextDistance)*AxisVectors[i];
-		if(View->ScreenToPixel(View->WorldToScreen(PotentialTextPosition), HUDInfoPos))
-		{
-			if (FMath::IsWithin<float>(HUDInfoPos.X, 0, View->UnscaledViewRect.Width()) && FMath::IsWithin<float>(HUDInfoPos.Y, 0, View->UnscaledViewRect.Height()))
+			else if (CurrentAxis == EAxisList::Y)
 			{
-				//only valid screen locations get a valid string
-				HUDString = FString::Printf(TEXT("%3.2f"), AngleOfChange);
-				break;
+				DominantIndex = 1;
 			}
-		}
-	}
-}
-
-uint32 FWidget::GetDominantAxisIndex( const FVector& InDiff, FEditorViewportClient* ViewportClient ) const
-{
-	uint32 DominantIndex = 0;
-	if( FMath::Abs(InDiff.X) < FMath::Abs(InDiff.Y) )
-	{
-		DominantIndex = 1;
-	}
-
-	const int32 WidgetMode = ViewportClient->GetWidgetMode();
-
-	switch( WidgetMode )
-	{
-		case WM_Translate:
-			switch( ViewportClient->ViewportType )
+			break;
+		case LVT_OrthoXZ:
+			if (CurrentAxis == EAxisList::X)
+			{
+				DominantIndex = 0;
+			}
+			else if (CurrentAxis == EAxisList::Z)
+			{
+				DominantIndex = 1;
+			}
+			break;
+		case LVT_OrthoYZ:
+			if (CurrentAxis == EAxisList::Y)
+			{
+				DominantIndex = 0;
+			}
+			else if (CurrentAxis == EAxisList::Z)
 			{
-				case LVT_OrthoXY:
-					if( CurrentAxis == EAxisList::X )
-					{
-						DominantIndex = 0;
-					}
-					else if( CurrentAxis == EAxisList::Y )
-					{
-						DominantIndex = 1;
-					}
-					break;
-				case LVT_OrthoXZ:
-					if( CurrentAxis == EAxisList::X )
-					{
-						DominantIndex = 0;
-					}
-					else if( CurrentAxis == EAxisList::Z )
-					{
-						DominantIndex = 1;
-					}
-					break;
-				case LVT_OrthoYZ:
-					if( CurrentAxis == EAxisList::Y )
-					{
-						DominantIndex = 0;
-					}
-					else if( CurrentAxis == EAxisList::Z )
-					{
-						DominantIndex = 1;
-					}
-					break;
-				default:
-					break;
+				DominantIndex = 1;
 			}
 			break;
 		default:
 			break;
+		}
+		break;
+	default:
+		break;
 	}
 
 	return DominantIndex;
 }
 
-
-EAxisList::Type FWidget::GetAxisToDraw( EWidgetMode WidgetMode ) const
-{
-	return EditorModeTools ? EditorModeTools->GetWidgetAxisToDraw( WidgetMode ) : EAxisList::All;
-}
-
 bool FWidget::IsWidgetDisabled() const
 {
 	return EditorModeTools ? (EditorModeTools->IsDefaultModeActive() && GEditor->HasLockedActors()) : false;
 }
-
diff --git a/Engine/Source/Editor/UnrealEd/Private/UnrealWidgetRender.cpp b/Engine/Source/Editor/UnrealEd/Private/UnrealWidgetRender.cpp
new file mode 100644
--- /dev/null
+++ b/Engine/Source/Editor/UnrealEd/Private/UnrealWidgetRender.cpp
@@ -0,0 +1,1452 @@
+// Copyright Epic Games, Inc. All Rights Reserved.
+
+#include "UnrealWidget.h"
+#include "Materials/Material.h"
+#include "CanvasItem.h"
+#include "Settings/LevelEditorViewportSettings.h"
+#include "Materials/MaterialInstanceDynamic.h"
+#include "EditorViewportClient.h"
+#include "EdMode.h"
+#include "EditorModeManager.h"
+#include "SnappingUtils.h"
+#include "DynamicMeshBuilder.h"
+#include "CanvasTypes.h"
+
+static const float AXIS_LENGTH = 35.0f;
+static const float TRANSLATE_ROTATE_AXIS_CIRCLE_RADIUS = 20.0f;
+static const float TWOD_AXIS_CIRCLE_RADIUS = 10.0f;
+static const float INNER_AXIS_CIRCLE_RADIUS = 48.0f;
+static const float OUTER_AXIS_CIRCLE_RADIUS = 56.0f;
+static const float ROTATION_TEXT_RADIUS = 75.0f;
+static const int32 AXIS_CIRCLE_SIDES = 24;
+static const float ARCALL_RELATIVE_INNER_SIZE = 0.75f;
+static const float AXIS_LENGTH_SCALE_OFFSET = 5.0f;
+
+/*
+ *  Simple struct used to create and group data related to the current window's / viewport's space,
+ *  orientation, and scale.
+ */
+struct FSpaceDescriptor
+{
+    /*
+     *  Creates a new FSpaceDescriptor.
+     *
+     *  @param View         The virtual view for the space.
+     *  @param Viewport     The real viewport for the space.
+     *  @param InLocation   The location of the camera in the virtual space.
+     */
+    FSpaceDescriptor(const FSceneView* View, const FEditorViewportClient* Viewport, const FVector& InLocation) :
+        bIsPerspective(View->ViewMatrices.GetProjectionMatrix().M[3][3] < 1.0f),
+        bIsLocalSpace(Viewport->GetWidgetCoordSystemSpace() == COORD_Local),
+        bIsOrthoXY(!bIsPerspective && FMath::Abs(View->ViewMatrices.GetViewMatrix().M[2][2]) > 0.0f),
+        bIsOrthoXZ(!bIsPerspective && FMath::Abs(View->ViewMatrices.GetViewMatrix().M[1][2]) > 0.0f),
+        bIsOrthoYZ(!bIsPerspective && FMath::Abs(View->ViewMatrices.GetViewMatrix().M[0][2]) > 0.0f),
+        UniformScale(View->WorldToScreen(InLocation).W * (4.0f / View->UnscaledViewRect.Width() / View->ViewMatrices.GetProjectionMatrix().M[0][0])),
+        Scale(CreateScale())
+    {
+    }
+
+    // Wether or not the view is perspective.
+    const bool bIsPerspective;
+
+    // Whether or not the view is in local space.
+    const bool bIsLocalSpace;
+
+    // Whether or not the view is orthogonal to the XY plane.
+    const bool bIsOrthoXY;
+
+    // Whether or not the view is orthogonal to the XZ plane.
+    const bool bIsOrthoXZ;
+
+    // Whether or not the view is orthogonal to the YZ plane.
+    const bool bIsOrthoYZ;
+
+    // The uniform scale for the space.
+    const float UniformScale;
+
+    // The scale vector for the space based on orientation.
+    const FVector Scale;
+
+    /*
+     *  Used to determine whether or not the X axis should be drawn.
+     *
+     *  @param AxisToDraw   The desired axis to draw.
+     *
+     *  @return True if the axis should be drawn. False otherwise.
+     */
+    bool ShouldDrawAxisX(const EAxisList::Type AxisToDraw)
+    {
+        return ShouldDrawAxis(EAxisList::X, AxisToDraw, bIsOrthoYZ);
+    }
+
+    /*
+     *  Used to determine whether or not the Y axis should be drawn.
+     *
+     *  @param AxisToDraw   The desired axis to draw.
+     *
+     *  @return True if the axis should be drawn. False otherwise.
+     */
+    bool ShouldDrawAxisY(const EAxisList::Type AxisToDraw)
+    {
+        return ShouldDrawAxis(EAxisList::Y, AxisToDraw, bIsOrthoXZ);
+    }
+
+    /*
+     *  Used to determine whether or not the Z axis should be drawn.
+     *
+     *  @param AxisToDraw   The desired axis to draw.
+     *
+     *  @return True if the axis should be drawn. False otherwise.
+     */
+    bool ShouldDrawAxisZ(const EAxisList::Type AxisToDraw)
+    {
+        return ShouldDrawAxis(EAxisList::Z, AxisToDraw, bIsOrthoXY);
+    }
+
+private:
+
+    /*
+     *  Creates a space scale vector from the determined orientation and uniform scale.
+     *
+     *  @return Space scale vector.
+     */
+    FVector CreateScale()
+    {
+        if (bIsOrthoXY)
+        {
+            return FVector(UniformScale, UniformScale, 1.0f);
+        }
+        else if (bIsOrthoXZ)
+        {
+            return FVector(UniformScale, 1.0f, UniformScale);
+        }
+        else if (bIsOrthoYZ)
+        {
+            return FVector(1.0f, UniformScale, UniformScale);
+        }
+        else
+        {
+            return FVector(UniformScale, UniformScale, UniformScale);
+        }
+    }
+
+    /*
+     *  Used to determine whether or not a specific axis should be drawn.
+     *
+     *  @param AxisToCheck  The axis to check.
+     *  @param AxisToDraw   The desired axis to draw.
+     *  @param bIsOrtho     Whether or not the axis to check is orthogonal to the viewing orientation.
+     *
+     *  @return True if the axis should be drawn. False otherwise.
+     */
+    bool ShouldDrawAxis(const EAxisList::Type AxisToCheck, const EAxisList::Type AxisToDraw, const bool bIsOrtho)
+    {
+        return (AxisToCheck & AxisToDraw) && (bIsPerspective || bIsLocalSpace || !bIsOrtho);
+    }
+};
+
+/**
+ * Renders any widget specific HUD text
+ * @param Canvas - Canvas to use for 2d rendering
+ */
+void FWidget::DrawHUD (FCanvas* Canvas)
+{
+	if (HUDString.Len())
+	{
+		const float DPIScale = Canvas->GetDPIScale();
+		int32 StringPosX = FMath::FloorToInt(HUDInfoPos.X/DPIScale);
+		int32 StringPosY = FMath::FloorToInt(HUDInfoPos.Y/DPIScale);
+
+		//measure string size
+		int32 StringSizeX, StringSizeY;
+		StringSize(GEngine->GetSmallFont(), StringSizeX, StringSizeY, *HUDString);
+		
+		//add some padding to the outside
+		const int32 Border = 5;
+		int32 FillMinX = StringPosX - Border - (StringSizeX>>1);
+		int32 FillMinY = StringPosY - Border;// - (StringSizeY>>1);
+		StringSizeX += 2*Border;
+		StringSizeY += 2*Border;
+
+		//mostly alpha'ed black
+		FCanvasTileItem TileItem( FVector2D( FillMinX, FillMinY), GWhiteTexture, FVector2D( StringSizeX, StringSizeY ), FLinearColor( 0.0f, 0.0f, 0.0f, .25f ) );
+		TileItem.BlendMode = SE_BLEND_Translucent;
+		Canvas->DrawItem( TileItem );
+		FCanvasTextItem TextItem( FVector2D( StringPosX, StringPosY), FText::FromString( HUDString ), GEngine->GetSmallFont(), FLinearColor::White );
+		TextItem.bCentreX = true;
+		Canvas->DrawItem( TextItem );	
+	}
+}
+
+void FWidget::Render( const FSceneView* View,FPrimitiveDrawInterface* PDI, FEditorViewportClient* ViewportClient )
+{
+	check(ViewportClient);
+
+
+	//reset HUD text
+	HUDString.Empty();
+
+	const bool bShowFlagsSupportsWidgetDrawing = View->Family->EngineShowFlags.ModeWidgets;
+	const bool bEditorModeToolsSupportsWidgetDrawing = EditorModeTools ? EditorModeTools->GetShowWidget() : true;
+	bool bDrawWidget;
+
+	// Because the movement routines use the widget axis to determine how to transform mouse movement into
+	// editor object movement, we need to still run through the Render routine even though widget drawing may be
+	// disabled.  So we keep a flag that is used to determine whether or not to actually render anything.  This way
+	// we can still update the widget axis' based on the Context's transform matrices, even though drawing is disabled.
+	if(bDefaultVisibility && bShowFlagsSupportsWidgetDrawing && bEditorModeToolsSupportsWidgetDrawing)
+	{
+		bDrawWidget = true;
+
+		// See if there is a custom coordinate system we should be using, only change it if we are drawing widgets.
+		CustomCoordSystem = ViewportClient->GetWidgetCoordSystem();
+	}
+	else
+	{
+		bDrawWidget = false;
+	}
+
+	CustomCoordSystemSpace = ViewportClient->GetWidgetCoordSystemSpace();
+
+	// If the current modes don't want to use the widget, don't draw it.
+	if( EditorModeTools && !EditorModeTools->UsesTransformWidget() )
+	{
+		CurrentAxis = EAxisList::None;
+		return;
+	}
+
+	FVector WidgetLocation = ViewportClient->GetWidgetLocation();
+	FVector2D NewOrigin;
+	if (View->ScreenToPixel(View->WorldToScreen(WidgetLocation), NewOrigin))
+	{
+		// Only update the viewport-space origin if the position was in front of the camera
+		Origin = NewOrigin;
+	}
+
+	switch( ViewportClient->GetWidgetMode() )
+	{
+		case WM_Translate:
+			Render_Translate(View, PDI, ViewportClient, WidgetLocation, bDrawWidget);
+			break;
+
+		case WM_Rotate:
+			Render_Rotate(View, PDI, ViewportClient, WidgetLocation, bDrawWidget);
+			break;
+
+		case WM_Scale:
+			Render_Scale(View, PDI, ViewportClient, WidgetLocation, bDrawWidget);
+			break;
+
+		case WM_TranslateRotateZ:
+			Render_TranslateRotateZ(View, PDI, ViewportClient, WidgetLocation, bDrawWidget);
+			break;
+
+		case WM_2D:
+			Render_2D(View, PDI, ViewportClient, WidgetLocation, bDrawWidget);
+			break;
+
+		default:
+			break;
+	}
+};
+
+/**
+ * Draws an arrow head line for a specific axis.
+ */
+void FWidget::Render_Axis(const FSceneView* View, FPrimitiveDrawInterface* PDI, EAxisList::Type InAxis, FMatrix& InMatrix, UMaterialInterface* InMaterial, const FLinearColor& InColor, FVector2D& OutAxisDir, const FVector& InScale, bool bDrawWidget, bool bCubeHead, float AxisLengthOffset)
+{
+	FMatrix AxisRotation = FMatrix::Identity;
+	if( InAxis == EAxisList::Y )
+	{
+		AxisRotation = FRotationMatrix::MakeFromXZ(FVector(0, 1, 0), FVector(0, 0, 1));
+	}
+	else if( InAxis == EAxisList::Z )
+	{
+		AxisRotation = FRotationMatrix::MakeFromXY(FVector(0, 0, 1), FVector(0, 1, 0));
+	}
+
+	FMatrix ArrowToWorld = AxisRotation * InMatrix;
+
+	// The scale that is passed in potentially leaves one component with a scale of 1, if that happens
+	// we need to extract the inform scale and use it to construct the scale that transforms the primitives
+	float UniformScale = InScale.GetMax() > 1.0f ? InScale.GetMax() : InScale.GetMin() < 1.0f ? InScale.GetMin() : 1.0f;
+	// After the primitives have been scaled and transformed, we apply this inverse scale that flattens the dimension
+	// that was scaled up to prevent it from intersecting with the near plane.  In perspective this won't have any effect,
+	// but in the ortho viewports it will prevent scaling in the direction of the camera and thus intersecting the near plane.
+	FVector FlattenScale = FVector(InScale.Component(0) == 1.0f ? 1.0f / UniformScale : 1.0f, InScale.Component(1) == 1.0f ? 1.0f / UniformScale : 1.0f, InScale.Component(2) == 1.0f ? 1.0f / UniformScale : 1.0f);
+
+	FScaleMatrix Scale(UniformScale);
+	ArrowToWorld = Scale * ArrowToWorld;
+
+	if( bDrawWidget )
+	{
+		const bool bDisabled = EditorModeTools ? (EditorModeTools->IsDefaultModeActive() && GEditor->HasLockedActors() ) : false;
+		PDI->SetHitProxy( new HWidgetAxis( InAxis, bDisabled) );
+
+		const float AxisLength = AXIS_LENGTH + GetDefault<ULevelEditorViewportSettings>()->TransformWidgetSizeAdjustment - (AxisLengthOffset * 2);
+		const float HalfHeight = AxisLength/2.0f;
+		const float CylinderRadius = 1.2f;
+		const FVector Offset(0, 0, HalfHeight + AxisLengthOffset);
+
+		switch( InAxis )
+		{
+			case EAxisList::X:
+			{
+				DrawCylinder(PDI, ( Scale * FRotationMatrix(FRotator(-90, 0.f, 0)) * InMatrix ) * FScaleMatrix(FlattenScale), Offset, FVector(1, 0, 0), FVector(0, 1, 0), FVector(0, 0, 1), CylinderRadius, HalfHeight, 16, InMaterial->GetRenderProxy(), SDPG_Foreground);
+				break;
+			}
+			case EAxisList::Y:
+			{
+				DrawCylinder(PDI, (Scale * FRotationMatrix(FRotator(0, 0, 90)) * InMatrix)* FScaleMatrix(FlattenScale), Offset, FVector(1, 0, 0), FVector(0, 1, 0), FVector(0, 0, 1), CylinderRadius, HalfHeight, 16, InMaterial->GetRenderProxy(), SDPG_Foreground );
+				break;
+			}
+			case EAxisList::Z:
+			{
+				DrawCylinder(PDI, ( Scale * InMatrix ) * FScaleMatrix(FlattenScale), Offset, FVector(1, 0, 0), FVector(0, 1, 0), FVector(0, 0, 1), CylinderRadius, HalfHeight, 16, InMaterial->GetRenderProxy(), SDPG_Foreground);
+				break;
+			}
+		}
+
+		if ( bCubeHead )
+		{
+			const float CubeHeadOffset = 3.0f;
+			FVector RootPos(AxisLength + CubeHeadOffset + AxisLengthOffset, 0, 0);
+
+			Render_Cube(PDI, (FTranslationMatrix(RootPos) * ArrowToWorld) * FScaleMatrix(FlattenScale), InMaterial, FVector(4.0f));
+		}
+		else
+		{
+			const float ConeHeadOffset = 12.0f;
+			FVector RootPos(AxisLength + ConeHeadOffset + AxisLengthOffset, 0, 0);
+
+			float Angle = FMath::DegreesToRadians( PI * 5 );
+			DrawCone(PDI, ( FScaleMatrix(-13) * FTranslationMatrix(RootPos) * ArrowToWorld ) * FScaleMatrix(FlattenScale), Angle, Angle, 32, false, FColor::White, InMaterial->GetRenderProxy(), SDPG_Foreground);
+		}
+	
+		PDI->SetHitProxy( NULL );
+	}
+
+	FVector2D NewOrigin;
+	FVector2D AxisEnd;
+	const FVector AxisEndWorld = ArrowToWorld.TransformPosition(FVector(64, 0, 0));
+	const FVector WidgetOrigin = InMatrix.GetOrigin();
+
+	if (View->ScreenToPixel(View->WorldToScreen(WidgetOrigin), NewOrigin) &&
+		View->ScreenToPixel(View->WorldToScreen(AxisEndWorld), AxisEnd))
+	{
+		// If both the origin and the axis endpoint are in front of the camera, trivially calculate the viewport space axis direction
+		OutAxisDir = (AxisEnd - NewOrigin).GetSafeNormal();
+	}
+	else
+	{
+		// If either the origin or axis endpoint are behind the camera, translate the entire widget in front of the camera in the view direction before performing the
+		// viewport space calculation
+		const FMatrix InvViewMatrix = View->ViewMatrices.GetInvViewMatrix();
+		const FVector ViewLocation = InvViewMatrix.GetOrigin();
+		const FVector ViewDirection = InvViewMatrix.GetUnitAxis(EAxis::Z);
+		const FVector Offset = ViewDirection * (FVector::DotProduct(ViewLocation - WidgetOrigin, ViewDirection) + 100.0f);
+		const FVector AdjustedWidgetOrigin = WidgetOrigin + Offset;
+		const FVector AdjustedWidgetAxisEnd = AxisEndWorld + Offset;
+
+		if (View->ScreenToPixel(View->WorldToScreen(AdjustedWidgetOrigin), NewOrigin) &&
+			View->ScreenToPixel(View->WorldToScreen(AdjustedWidgetAxisEnd), AxisEnd))
+		{
+			OutAxisDir = -(AxisEnd - NewOrigin).GetSafeNormal();
+		}
+	}
+}
+
+void FWidget::Render_Cube( FPrimitiveDrawInterface* PDI, const FMatrix& InMatrix, const UMaterialInterface* InMaterial, const FVector& InScale )
+{
+	const FMatrix CubeToWorld = FScaleMatrix(InScale) * InMatrix;
+	DrawBox( PDI, CubeToWorld, FVector(1,1,1), InMaterial->GetRenderProxy(), SDPG_Foreground );
+}
+
+void DrawCornerHelper( FPrimitiveDrawInterface* PDI, const FMatrix& LocalToWorld, const FVector& Length, float Thickness, const FMaterialRenderProxy* MaterialRenderProxy,uint8 DepthPriorityGroup  )
+{
+	const float TH = Thickness;
+
+	float TX = Length.X/2;
+	float TY = Length.Y/2;
+	float TZ = Length.Z/2;
+
+	FDynamicMeshBuilder MeshBuilder(PDI->View->GetFeatureLevel());
+
+	// Top
+	{
+		int32 VertexIndices[4];
+		VertexIndices[0] = MeshBuilder.AddVertex( FVector(-TX, -TY, +TZ), FVector2D::ZeroVector, FVector(1,0,0), FVector(0,1,0), FVector(0,0,1), FColor::White );
+		VertexIndices[1] = MeshBuilder.AddVertex( FVector(-TX, +TY, +TZ), FVector2D::ZeroVector, FVector(1,0,0), FVector(0,1,0), FVector(0,0,1), FColor::White );
+		VertexIndices[2] = MeshBuilder.AddVertex( FVector(+TX, +TY, +TZ), FVector2D::ZeroVector, FVector(1,0,0), FVector(0,1,0), FVector(0,0,1), FColor::White );
+		VertexIndices[3] = MeshBuilder.AddVertex( FVector(+TX, -TY, +TZ), FVector2D::ZeroVector, FVector(1,0,0), FVector(0,1,0), FVector(0,0,1), FColor::White );
+
+		MeshBuilder.AddTriangle(VertexIndices[0],VertexIndices[1],VertexIndices[2]);
+		MeshBuilder.AddTriangle(VertexIndices[0],VertexIndices[2],VertexIndices[3]);
+	}
+
+	//Left
+	{
+		int32 VertexIndices[4];
+		VertexIndices[0] = MeshBuilder.AddVertex( FVector(-TX,  -TY, TZ-TH),	FVector2D::ZeroVector, FVector(0,0,1), FVector(0,1,0), FVector(-1,0,0), FColor::White );
+		VertexIndices[1] = MeshBuilder.AddVertex( FVector(-TX, -TY, TZ),		FVector2D::ZeroVector, FVector(0,0,1), FVector(0,1,0), FVector(-1,0,0), FColor::White );
+		VertexIndices[2] = MeshBuilder.AddVertex( FVector(-TX, +TY, TZ),		FVector2D::ZeroVector, FVector(0,0,1), FVector(0,1,0), FVector(-1,0,0), FColor::White );
+		VertexIndices[3] = MeshBuilder.AddVertex( FVector(-TX, +TY, TZ-TH),		FVector2D::ZeroVector, FVector(0,0,1), FVector(0,1,0), FVector(-1,0,0), FColor::White );
+
+
+		MeshBuilder.AddTriangle(VertexIndices[0],VertexIndices[1],VertexIndices[2]);
+		MeshBuilder.AddTriangle(VertexIndices[0],VertexIndices[2],VertexIndices[3]);
+	}
+
+	// Front
+	{
+		int32 VertexIndices[5];
+		VertexIndices[0] = MeshBuilder.AddVertex( FVector(-TX,	+TY, TZ-TH),	FVector2D::ZeroVector, FVector(1,0,0), FVector(0,0,-1), FVector(0,1,0), FColor::White );
+		VertexIndices[1] = MeshBuilder.AddVertex( FVector(-TX,	+TY, +TZ  ),	FVector2D::ZeroVector, FVector(1,0,0), FVector(0,0,-1), FVector(0,1,0), FColor::White );
+		VertexIndices[2] = MeshBuilder.AddVertex( FVector(+TX-TH, +TY, +TX  ),	FVector2D::ZeroVector, FVector(1,0,0), FVector(0,0,-1), FVector(0,1,0), FColor::White );
+		VertexIndices[3] = MeshBuilder.AddVertex( FVector(+TX,	+TY, +TZ  ),	FVector2D::ZeroVector, FVector(1,0,0), FVector(0,0,-1), FVector(0,1,0), FColor::White );
+		VertexIndices[4] = MeshBuilder.AddVertex( FVector(+TX-TH, +TY, TZ-TH),	FVector2D::ZeroVector, FVector(1,0,0), FVector(0,0,-1), FVector(0,1,0), FColor::White );
+
+		MeshBuilder.AddTriangle(VertexIndices[0],VertexIndices[1],VertexIndices[2]);
+		MeshBuilder.AddTriangle(VertexIndices[0],VertexIndices[2],VertexIndices[4]);
+		MeshBuilder.AddTriangle(VertexIndices[4],VertexIndices[2],VertexIndices[3]);
+	}
+
+	// Back
+	{
+		int32 VertexIndices[5];
+		VertexIndices[0] = MeshBuilder.AddVertex( FVector(-TX,	-TY, TZ-TH),	FVector2D::ZeroVector, FVector(1,0,0), FVector(0,0,1), FVector(0,-1,0), FColor::White );
+		VertexIndices[1] = MeshBuilder.AddVertex( FVector(-TX,	-TY, +TZ),		FVector2D::ZeroVector, FVector(1,0,0), FVector(0,0,1), FVector(0,-1,0), FColor::White );
+		VertexIndices[2] = MeshBuilder.AddVertex( FVector(+TX-TH, -TY, +TX),	FVector2D::ZeroVector, FVector(1,0,0), FVector(0,0,1), FVector(0,-1,0), FColor::White );
+		VertexIndices[3] = MeshBuilder.AddVertex( FVector(+TX,	-TY, +TZ),		FVector2D::ZeroVector, FVector(1,0,0), FVector(0,0,1), FVector(0,-1,0), FColor::White );
+		VertexIndices[4] = MeshBuilder.AddVertex( FVector(+TX-TH, -TY, TZ-TH),	FVector2D::ZeroVector, FVector(1,0,0), FVector(0,0,1), FVector(0,-1,0), FColor::White );
+
+		MeshBuilder.AddTriangle(VertexIndices[0],VertexIndices[1],VertexIndices[2]);
+		MeshBuilder.AddTriangle(VertexIndices[0],VertexIndices[2],VertexIndices[4]);
+		MeshBuilder.AddTriangle(VertexIndices[4],VertexIndices[2],VertexIndices[3]);
+	}
+	// Bottom
+	{
+		int32 VertexIndices[4];
+		VertexIndices[0] = MeshBuilder.AddVertex( FVector(-TX, -TY, TZ-TH),		FVector2D::ZeroVector, FVector(1,0,0), FVector(0,0,-1), FVector(0,0,1), FColor::White );
+		VertexIndices[1] = MeshBuilder.AddVertex( FVector(-TX, +TY, TZ-TH),		FVector2D::ZeroVector, FVector(1,0,0), FVector(0,0,-1), FVector(0,0,1), FColor::White );
+		VertexIndices[2] = MeshBuilder.AddVertex( FVector(+TX-TH, +TY, TZ-TH),	FVector2D::ZeroVector, FVector(1,0,0), FVector(0,0,-1), FVector(0,0,1), FColor::White );
+		VertexIndices[3] = MeshBuilder.AddVertex( FVector(+TX-TH, -TY, TZ-TH),	FVector2D::ZeroVector, FVector(1,0,0), FVector(0,0,-1), FVector(0,0,1), FColor::White );
+
+		MeshBuilder.AddTriangle(VertexIndices[0],VertexIndices[1],VertexIndices[2]);
+		MeshBuilder.AddTriangle(VertexIndices[0],VertexIndices[2],VertexIndices[3]);
+	}
+	MeshBuilder.Draw(PDI,LocalToWorld,MaterialRenderProxy,DepthPriorityGroup,0.f);
+}
+
+void DrawDualAxis( FPrimitiveDrawInterface* PDI, const FMatrix& BoxToWorld,const FVector& Length, float Thickness, const FMaterialRenderProxy* AxisMat,const FMaterialRenderProxy* Axis2Mat )
+{
+	DrawCornerHelper(PDI, BoxToWorld, Length, Thickness, Axis2Mat, SDPG_Foreground);
+	DrawCornerHelper(PDI, FScaleMatrix(FVector(-1, 1, 1)) * FRotationMatrix(FRotator(-90, 0, 0)) * BoxToWorld, Length, Thickness, AxisMat, SDPG_Foreground);
+}
+/**
+ * Draws the translation widget.
+ */
+void FWidget::Render_Translate( const FSceneView* View, FPrimitiveDrawInterface* PDI, FEditorViewportClient* ViewportClient, const FVector& InLocation, bool bDrawWidget )
+{
+	// Figure out axis colors
+	const FLinearColor& XColor = ( CurrentAxis&EAxisList::X ? (FLinearColor)CurrentColor : AxisColorX );
+	const FLinearColor& YColor = ( CurrentAxis&EAxisList::Y ? (FLinearColor)CurrentColor : AxisColorY );
+	const FLinearColor& ZColor = ( CurrentAxis&EAxisList::Z ? (FLinearColor)CurrentColor : AxisColorZ );
+	FColor CurrentScreenColor = ( CurrentAxis & EAxisList::Screen ? CurrentColor : ScreenSpaceColor );
+
+	// Figure out axis matrices
+	FMatrix WidgetMatrix = CustomCoordSystem * FTranslationMatrix( InLocation );
+	const EAxisList::Type DrawAxis = GetAxisToDraw( ViewportClient->GetWidgetMode() );
+	const bool bDisabled = IsWidgetDisabled();
+
+    FSpaceDescriptor Space(View, ViewportClient, InLocation);
+
+	// Draw the axis lines with arrow heads
+	if( Space.ShouldDrawAxisX(DrawAxis) )
+	{
+		UMaterialInstanceDynamic* XMaterial = ( CurrentAxis&EAxisList::X ? CurrentAxisMaterial : AxisMaterialX );
+		Render_Axis( View, PDI, EAxisList::X, WidgetMatrix, XMaterial, XColor, XAxisDir, Space.Scale, bDrawWidget );
+	}
+
+	if( Space.ShouldDrawAxisY(DrawAxis) )
+	{
+		UMaterialInstanceDynamic* YMaterial = ( CurrentAxis&EAxisList::Y ? CurrentAxisMaterial : AxisMaterialY );
+		Render_Axis( View, PDI, EAxisList::Y, WidgetMatrix, YMaterial, YColor, YAxisDir, Space.Scale, bDrawWidget );
+	}
+
+	if( Space.ShouldDrawAxisZ(DrawAxis) )
+	{
+		UMaterialInstanceDynamic* ZMaterial = ( CurrentAxis&EAxisList::Z ? CurrentAxisMaterial : AxisMaterialZ );
+		Render_Axis( View, PDI, EAxisList::Z, WidgetMatrix, ZMaterial, ZColor, ZAxisDir, Space.Scale, bDrawWidget );
+	}
+
+	// Draw the grabbers
+	if( bDrawWidget )
+	{
+		FVector CornerPos = FVector(7, 0, 7) * Space.UniformScale;
+		FVector AxisSize = FVector(12, 1.2, 12) * Space.UniformScale;
+		float CornerLength = 1.2f * Space.UniformScale;
+
+		// After the primitives have been scaled and transformed, we apply this inverse scale that flattens the dimension
+		// that was scaled up to prevent it from intersecting with the near plane.  In perspective this won't have any effect,
+		// but in the ortho viewports it will prevent scaling in the direction of the camera and thus intersecting the near plane.
+		FVector FlattenScale = FVector(Space.Scale.Component(0) == 1.0f ? 1.0f / Space.UniformScale : 1.0f, Space.Scale.Component(1) == 1.0f ? 1.0f / Space.UniformScale : 1.0f, Space.Scale.Component(2) == 1.0f ? 1.0f / Space.UniformScale : 1.0f);
+
+		if (Space.bIsPerspective || Space.bIsLocalSpace || Space.bIsOrthoXY)
+		{
+			if( (DrawAxis&EAxisList::XY) == EAxisList::XY )							// Top
+			{
+				UMaterialInstanceDynamic* XMaterial = ( (CurrentAxis&EAxisList::XY) == EAxisList::XY ? CurrentAxisMaterial : AxisMaterialX );
+				UMaterialInstanceDynamic* YMaterial = ( (CurrentAxis&EAxisList::XY) == EAxisList::XY? CurrentAxisMaterial : AxisMaterialY );
+
+				PDI->SetHitProxy( new HWidgetAxis(EAxisList::XY, bDisabled) );
+				{
+					DrawDualAxis(PDI, ( FTranslationMatrix(CornerPos) * FRotationMatrix(FRotator(0, 0, 90)) * WidgetMatrix ) * FScaleMatrix(FlattenScale), AxisSize, CornerLength, XMaterial->GetRenderProxy(), YMaterial->GetRenderProxy());
+				}
+				PDI->SetHitProxy( NULL );
+			}
+		}
+
+		if (Space.bIsPerspective || Space.bIsLocalSpace || Space.bIsOrthoXZ)		// Front
+		{
+			if( (DrawAxis&EAxisList::XZ) == EAxisList::XZ ) 
+			{
+				UMaterialInstanceDynamic* XMaterial = ( (CurrentAxis&EAxisList::XZ) == EAxisList::XZ ? CurrentAxisMaterial : AxisMaterialX );
+				UMaterialInstanceDynamic* ZMaterial = ( (CurrentAxis&EAxisList::XZ) == EAxisList::XZ ? CurrentAxisMaterial : AxisMaterialZ );
+
+				PDI->SetHitProxy( new HWidgetAxis(EAxisList::XZ, bDisabled) );
+				{
+					DrawDualAxis(PDI, (FTranslationMatrix(CornerPos) * WidgetMatrix) * FScaleMatrix(FlattenScale), AxisSize, CornerLength, XMaterial->GetRenderProxy(), ZMaterial->GetRenderProxy() );
+				}
+				PDI->SetHitProxy( NULL );
+			}
+		}
+
+		if( Space.bIsPerspective || Space.bIsLocalSpace || Space.bIsOrthoYZ )		// Side
+		{
+			if( (DrawAxis&EAxisList::YZ) == EAxisList::YZ ) 
+			{
+				UMaterialInstanceDynamic* YMaterial = ( (CurrentAxis&EAxisList::YZ) == EAxisList::YZ ? CurrentAxisMaterial : AxisMaterialY );
+				UMaterialInstanceDynamic* ZMaterial = ( (CurrentAxis&EAxisList::YZ) == EAxisList::YZ ? CurrentAxisMaterial : AxisMaterialZ );
+
+				PDI->SetHitProxy( new HWidgetAxis(EAxisList::YZ, bDisabled) );
+				{
+					DrawDualAxis(PDI, (FTranslationMatrix(CornerPos) * FRotationMatrix(FRotator(0, 90, 0)) * WidgetMatrix) * FScaleMatrix(FlattenScale), AxisSize, CornerLength, YMaterial->GetRenderProxy(), ZMaterial->GetRenderProxy() );
+				}
+				PDI->SetHitProxy( NULL );
+			}
+		}
+	}
+
+	// Draw screen-space movement handle (circle)
+	if( bDrawWidget && ( DrawAxis & EAxisList::Screen ) && Space.bIsPerspective )
+	{
+		PDI->SetHitProxy( new HWidgetAxis(EAxisList::Screen, bDisabled) );
+		const FVector CameraXAxis = View->ViewMatrices.GetViewMatrix().GetColumn(0);
+		const FVector CameraYAxis = View->ViewMatrices.GetViewMatrix().GetColumn(1);
+		const FVector CameraZAxis = View->ViewMatrices.GetViewMatrix().GetColumn(2);
+
+		UMaterialInstanceDynamic* XYZMaterial = ( CurrentAxis&EAxisList::Screen) ? CurrentAxisMaterial : OpaquePlaneMaterialXY;
+		DrawSphere( PDI, InLocation, FRotator::ZeroRotator, 4.0f * Space.Scale, 10, 5, XYZMaterial->GetRenderProxy(), SDPG_Foreground );
+
+		PDI->SetHitProxy( NULL );
+	}
+}
+
+uint8 LargeInnerAlpha = 0x3f;
+uint8 SmallInnerAlpha = 0x0f;
+uint8 LargeOuterAlpha = 0x7f;
+uint8 SmallOuterAlpha = 0x0f;
+
+void FWidget::DrawColoredSphere(FPrimitiveDrawInterface* PDI, const FVector& Center, const FRotator& Orientation, FColor Color, const FVector& Radii, int32 NumSides, int32 NumRings, const FMaterialRenderProxy* MaterialRenderProxy, uint8 DepthPriority, bool bDisableBackfaceCulling)
+{
+	// Use a mesh builder to draw the sphere.
+	FDynamicMeshBuilder MeshBuilder(PDI->View->GetFeatureLevel());
+	{
+		// The first/last arc are on top of each other.
+		int32 NumVerts = (NumSides + 1) * (NumRings + 1);
+		FDynamicMeshVertex* Verts = (FDynamicMeshVertex*)FMemory::Malloc(NumVerts * sizeof(FDynamicMeshVertex));
+
+		// Calculate verts for one arc
+		FDynamicMeshVertex* ArcVerts = (FDynamicMeshVertex*)FMemory::Malloc((NumRings + 1) * sizeof(FDynamicMeshVertex));
+
+		for (int32 i = 0; i < NumRings + 1; i++)
+		{
+			FDynamicMeshVertex* ArcVert = &ArcVerts[i];
+
+			float angle = ((float)i / NumRings) * PI;
+
+			ArcVert->Color = Color;
+			// Note- unit sphere, so position always has mag of one. We can just use it for normal!			
+			ArcVert->Position.X = 0.0f;
+			ArcVert->Position.Y = FMath::Sin(angle);
+			ArcVert->Position.Z = FMath::Cos(angle);
+
+			ArcVert->SetTangents(
+				FVector(1, 0, 0),
+				FVector(0.0f, -ArcVert->Position.Z, ArcVert->Position.Y),
+				ArcVert->Position
+			);
+
+			ArcVert->TextureCoordinate[0].X = 0.0f;
+			ArcVert->TextureCoordinate[0].Y = ((float)i / NumRings);
+		}
+
+		// Then rotate this arc NumSides+1 times.
+		for (int32 s = 0; s < NumSides + 1; s++)
+		{
+			FRotator ArcRotator(0, 360.f * (float)s / NumSides, 0);
+			FRotationMatrix ArcRot(ArcRotator);
+			float XTexCoord = ((float)s / NumSides);
+
+			for (int32 v = 0; v < NumRings + 1; v++)
+			{
+				int32 VIx = (NumRings + 1)*s + v;
+				Verts[VIx].Color = Color;
+				Verts[VIx].Position = ArcRot.TransformPosition(ArcVerts[v].Position);
+
+				Verts[VIx].SetTangents(
+					ArcRot.TransformVector(ArcVerts[v].TangentX.ToFVector()),
+					ArcRot.TransformVector(ArcVerts[v].GetTangentY()),
+					ArcRot.TransformVector(ArcVerts[v].TangentZ.ToFVector())
+				);
+
+				Verts[VIx].TextureCoordinate[0].X = XTexCoord;
+				Verts[VIx].TextureCoordinate[0].Y = ArcVerts[v].TextureCoordinate[0].Y;
+			}
+		}
+
+		// Add all of the vertices we generated to the mesh builder.
+		for (int32 VertIdx = 0; VertIdx < NumVerts; VertIdx++)
+		{
+			MeshBuilder.AddVertex(Verts[VertIdx]);
+		}
+
+		// Add all of the triangles we generated to the mesh builder.
+		for (int32 s = 0; s < NumSides; s++)
+		{
+			int32 a0start = (s + 0) * (NumRings + 1);
+			int32 a1start = (s + 1) * (NumRings + 1);
+
+			for (int32 r = 0; r < NumRings; r++)
+			{
+				MeshBuilder.AddTriangle(a0start + r + 0, a1start + r + 0, a0start + r + 1);
+				MeshBuilder.AddTriangle(a1start + r + 0, a1start + r + 1, a0start + r + 1);
+			}
+		}
+
+		// Free our local copy of verts and arc verts
+		FMemory::Free(Verts);
+		FMemory::Free(ArcVerts);
+	}
+	MeshBuilder.Draw(PDI, FScaleMatrix(Radii) * FRotationMatrix(Orientation) * FTranslationMatrix(Center), MaterialRenderProxy, DepthPriority, bDisableBackfaceCulling);
+}
+/**
+ * Draws the rotation widget.
+ */
+void FWidget::Render_Rotate( const FSceneView* View,FPrimitiveDrawInterface* PDI, FEditorViewportClient* ViewportClient, const FVector& InLocation, bool bDrawWidget )
+{
+	float Scale = View->WorldToScreen(InLocation).W * (4.0f / View->UnscaledViewRect.Width() / View->ViewMatrices.GetProjectionMatrix().M[0][0]);
+
+	//get the axes 
+	FVector XAxis = CustomCoordSystem.TransformVector(FVector(1, 0, 0));
+	FVector YAxis = CustomCoordSystem.TransformVector(FVector(0, 1, 0));
+	FVector ZAxis = CustomCoordSystem.TransformVector(FVector(0, 0, 1));
+
+	EAxisList::Type DrawAxis = GetAxisToDraw(ViewportClient->GetWidgetMode());
+
+	//if drawing arcball, and dragging make sure axis is showing local
+	if (DrawAxis&EAxisList::XYZ  && CurrentAxis == EAxisList::XYZ && bDragging)
+	{
+		FMatrix LocalCoordinate = ViewportClient->GetLocalCoordinateSystem();
+		XAxis = LocalCoordinate.TransformVector(FVector(1, 0, 0));
+		YAxis = LocalCoordinate.TransformVector(FVector(0, 1, 0));
+		ZAxis = LocalCoordinate.TransformVector(FVector(0, 0, 1));
+	}
+
+	FVector DirectionToWidget = View->IsPerspectiveProjection() ? (InLocation - View->ViewMatrices.GetViewOrigin()) : -View->GetViewDirection();
+	DirectionToWidget.Normalize();
+
+	// Draw a circle for each axis
+	if (bDrawWidget || bDragging)
+	{
+		bIsOrthoDrawingFullRing = false;
+		//now draw the arc segments
+		if (DrawAxis&EAxisList::X)
+		{
+			DrawRotationArc(View, PDI, EAxisList::X, InLocation, ZAxis, YAxis, DirectionToWidget, AxisColorX.ToFColor(true), Scale, XAxisDir);
+		}
+
+		if (DrawAxis&EAxisList::Y)
+		{
+			DrawRotationArc(View, PDI, EAxisList::Y, InLocation, XAxis, ZAxis, DirectionToWidget, AxisColorY.ToFColor(true), Scale, YAxisDir);
+		}
+
+		if (DrawAxis&EAxisList::Z)
+		{
+			DrawRotationArc(View, PDI, EAxisList::Z, InLocation, XAxis, YAxis, DirectionToWidget, AxisColorZ.ToFColor(true), Scale, ZAxisDir);
+		}
+
+		if (DrawAxis&EAxisList::XYZ  && (!bDragging || CurrentAxis == EAxisList::XYZ) && GetDefault<ULevelEditorViewportSettings>()->bAllowArcballRotate)
+		{
+			FVector Center = InLocation;
+			FRotator Orientation = FRotator::ZeroRotator;
+			const float InnerDistance = (INNER_AXIS_CIRCLE_RADIUS * ARCALL_RELATIVE_INNER_SIZE * Scale) + GetDefault<ULevelEditorViewportSettings>()->TransformWidgetSizeAdjustment;
+			FVector Radii(InnerDistance, InnerDistance, InnerDistance);
+			const bool bDisabled = IsWidgetDisabled();
+			FColor ArcColor = ArcBallColor;
+			ArcColor.A = (CurrentAxis==EAxisList::XYZ) ? ArcBallColor.A + 5 : ArcBallColor.A; //less transparent if selected
+
+			//set priority to foreground so axis get hit first
+			PDI->SetHitProxy(new HWidgetAxis(EAxisList::XYZ, bDisabled, HPP_Foreground));
+			{
+				DrawColoredSphere(PDI, InLocation, FRotator::ZeroRotator, ArcColor, Radii, 32, 24, TransparentPlaneMaterialXY->GetRenderProxy(), SDPG_Foreground, true);
+			}
+			PDI->SetHitProxy(NULL);
+		}
+		
+		if(DrawAxis&EAxisList::Screen && (!bDragging || CurrentAxis == EAxisList::Screen) && GetDefault<ULevelEditorViewportSettings>()->bAllowScreenRotate)
+		{
+			const bool bIsPerspective = (View->ViewMatrices.GetProjectionMatrix().M[3][3] < 1.0f);
+			const bool bIsOrtho = !bIsPerspective;
+
+			const FVector Axis0 = View->GetViewUp();
+			const FVector Axis1 = View->GetViewRight();
+			
+			const float OuterRadius = (OUTER_AXIS_CIRCLE_RADIUS * 1.25f * Scale) + GetDefault<ULevelEditorViewportSettings>()->TransformWidgetSizeAdjustment;
+			const float InnerRadius = ((OUTER_AXIS_CIRCLE_RADIUS -1) * 1.25f * Scale) + GetDefault<ULevelEditorViewportSettings>()->TransformWidgetSizeAdjustment;
+
+			FColor ArcColor = (CurrentAxis==EAxisList::Screen) ? CurrentColor : ScreenAxisColor.ToFColor(true);
+
+			const bool bDisabled = IsWidgetDisabled();
+
+			PDI->SetHitProxy(new HWidgetAxis(EAxisList::Screen, bDisabled));
+			{
+				FThickArcParams OuterArcParams(PDI, InLocation, TransparentPlaneMaterialXY, InnerRadius, OuterRadius);
+				//Pass through alpha
+				DrawThickArc(OuterArcParams, Axis0, Axis1, 0.0f, 2.0f * PI, ArcColor, DirectionToWidget, !bIsPerspective);
+			}
+			PDI->SetHitProxy(NULL);
+		}
+	}
+}
+
+/**
+ * Draws the scaling widget.
+ */
+void FWidget::Render_Scale( const FSceneView* View,FPrimitiveDrawInterface* PDI, FEditorViewportClient* ViewportClient, const FVector& InLocation, bool bDrawWidget)
+{
+	// Figure out axis colors
+	const FLinearColor& XColor = ( CurrentAxis&EAxisList::X ? (FLinearColor)CurrentColor : AxisColorX );
+	const FLinearColor& YColor = ( CurrentAxis&EAxisList::Y ? (FLinearColor)CurrentColor : AxisColorY );
+	const FLinearColor& ZColor = ( CurrentAxis&EAxisList::Z ? (FLinearColor)CurrentColor : AxisColorZ );
+	FColor CurrentScreenColor = ( CurrentAxis & EAxisList::Screen ? CurrentColor : ScreenSpaceColor );
+
+	// Figure out axis materials
+
+	UMaterialInstanceDynamic* XMaterial = ( CurrentAxis&EAxisList::X ? CurrentAxisMaterial : AxisMaterialX );
+	UMaterialInstanceDynamic* YMaterial = ( CurrentAxis&EAxisList::Y ? CurrentAxisMaterial : AxisMaterialY );
+	UMaterialInstanceDynamic* ZMaterial = ( CurrentAxis&EAxisList::Z ? CurrentAxisMaterial : AxisMaterialZ );
+	UMaterialInstanceDynamic* XYZMaterial = ( CurrentAxis&EAxisList::XYZ ? CurrentAxisMaterial : OpaquePlaneMaterialXY );
+
+	FMatrix WidgetMatrix = CustomCoordSystem * FTranslationMatrix( InLocation );
+	const EAxisList::Type DrawAxis = GetAxisToDraw( ViewportClient->GetWidgetMode() );
+    FSpaceDescriptor Space(View, ViewportClient, InLocation);
+
+	// Use a constant uniform scale for this widget since orthographic view for it is not supported.
+	const FVector UniformScale(Space.UniformScale);
+
+	// Draw the axis lines with cube heads	
+    if (Space.ShouldDrawAxisX(DrawAxis))
+    {
+        Render_Axis(View, PDI, EAxisList::X, WidgetMatrix, XMaterial, XColor, XAxisDir, UniformScale, bDrawWidget, true, AXIS_LENGTH_SCALE_OFFSET);
+    }
+
+    if (Space.ShouldDrawAxisY(DrawAxis))
+    {
+        Render_Axis(View, PDI, EAxisList::Y, WidgetMatrix, YMaterial, YColor, YAxisDir, UniformScale, bDrawWidget, true, AXIS_LENGTH_SCALE_OFFSET);
+    }
+
+    if (Space.ShouldDrawAxisZ(DrawAxis))
+    {
+        Render_Axis(View, PDI, EAxisList::Z, WidgetMatrix, ZMaterial, ZColor, ZAxisDir, UniformScale, bDrawWidget, true, AXIS_LENGTH_SCALE_OFFSET);
+    }
+
+	// Draw grabber handles and center cube
+	if ( bDrawWidget )
+	{
+		const bool bDisabled = IsWidgetDisabled();
+
+		// Grabber handles - since orthographic scale widgets are not supported, we should always draw grabber handles if we're drawing
+		// the corresponding axes.
+		if( (DrawAxis&(EAxisList::X|EAxisList::Y)) == (EAxisList::X|EAxisList::Y) )
+		{
+			PDI->SetHitProxy( new HWidgetAxis(EAxisList::XY, bDisabled) );
+			{
+				PDI->DrawLine( WidgetMatrix.TransformPosition(FVector(24,0,0) * UniformScale), WidgetMatrix.TransformPosition(FVector(12,12,0) * UniformScale), XColor, SDPG_Foreground );
+				PDI->DrawLine( WidgetMatrix.TransformPosition(FVector(12,12,0) * UniformScale), WidgetMatrix.TransformPosition(FVector(0,24,0) * UniformScale), YColor, SDPG_Foreground );
+			}
+			PDI->SetHitProxy( NULL );
+		}
+
+		if( (DrawAxis&(EAxisList::X|EAxisList::Z)) == (EAxisList::X|EAxisList::Z) )
+		{
+			PDI->SetHitProxy( new HWidgetAxis(EAxisList::XZ, bDisabled) );
+			{
+				PDI->DrawLine( WidgetMatrix.TransformPosition(FVector(24,0,0) * UniformScale), WidgetMatrix.TransformPosition(FVector(12,0,12) * UniformScale), XColor, SDPG_Foreground );
+				PDI->DrawLine( WidgetMatrix.TransformPosition(FVector(12,0,12) * UniformScale), WidgetMatrix.TransformPosition(FVector(0,0,24) * UniformScale), ZColor, SDPG_Foreground );
+			}
+			PDI->SetHitProxy( NULL );
+		}
+
+		if( (DrawAxis&(EAxisList::Y|EAxisList::Z)) == (EAxisList::Y|EAxisList::Z) )
+		{
+			PDI->SetHitProxy( new HWidgetAxis(EAxisList::YZ, bDisabled) );
+			{
+				PDI->DrawLine( WidgetMatrix.TransformPosition(FVector(0,24,0) * UniformScale), WidgetMatrix.TransformPosition(FVector(0,12,12) * UniformScale), YColor, SDPG_Foreground );
+				PDI->DrawLine( WidgetMatrix.TransformPosition(FVector(0,12,12) * UniformScale), WidgetMatrix.TransformPosition(FVector(0,0,24) * UniformScale), ZColor, SDPG_Foreground );
+			}
+			PDI->SetHitProxy( NULL );
+		}
+
+		// Center cube
+		if( (DrawAxis&(EAxisList::XYZ)) == EAxisList::XYZ )
+		{
+			PDI->SetHitProxy( new HWidgetAxis(EAxisList::XYZ, bDisabled) );
+
+			Render_Cube(PDI, WidgetMatrix, XYZMaterial, UniformScale * 4 );
+
+			PDI->SetHitProxy( NULL );
+		}
+	}
+}
+
+/**
+* Draws the Translate & Rotate Z widget.
+*/
+
+void FWidget::Render_TranslateRotateZ( const FSceneView* View, FPrimitiveDrawInterface* PDI, FEditorViewportClient* ViewportClient, const FVector& InLocation, bool bDrawWidget )
+{
+	// Figure out axis colors
+
+	FColor XYPlaneColor  = ( (CurrentAxis&EAxisList::XY) ==  EAxisList::XY) ? CurrentColor : PlaneColorXY;
+	FColor ZRotateColor  = ( ( CurrentAxis&EAxisList::ZRotation ) == EAxisList::ZRotation ) ? CurrentColor : AxisColorZ.ToFColor(true);
+	FColor XColor        = ( ( CurrentAxis&EAxisList::X ) == EAxisList::X ) ? CurrentColor : AxisColorX.ToFColor(true);
+	FColor YColor        = ( ( CurrentAxis&EAxisList::Y ) == EAxisList::Y && CurrentAxis != EAxisList::ZRotation ) ? CurrentColor : AxisColorY.ToFColor(true);
+	FColor ZColor        = ( ( CurrentAxis&EAxisList::Z ) == EAxisList::Z ) ? CurrentColor : AxisColorZ.ToFColor(true);
+
+	// Figure out axis materials
+	UMaterialInterface* ZRotateMaterial = (CurrentAxis&EAxisList::ZRotation) == EAxisList::ZRotation? CurrentAxisMaterial : AxisMaterialZ;
+	UMaterialInterface* XMaterial = CurrentAxis&EAxisList::X? CurrentAxisMaterial : AxisMaterialX;
+	UMaterialInterface* YMaterial = ( CurrentAxis&EAxisList::Y && CurrentAxis != EAxisList::ZRotation ) ? CurrentAxisMaterial : AxisMaterialY;
+	UMaterialInterface* ZMaterial = CurrentAxis&EAxisList::Z ? CurrentAxisMaterial : AxisMaterialZ;
+
+	// Figure out axis matrices
+	FMatrix AxisMatrix = CustomCoordSystem * FTranslationMatrix( InLocation );
+	EAxisList::Type DrawAxis = GetAxisToDraw( ViewportClient->GetWidgetMode() );
+
+	FSpaceDescriptor Space(View, ViewportClient, InLocation);
+
+	// Draw the grabbers
+	if( bDrawWidget )
+	{
+		// Draw the axis lines with arrow heads
+		if( DrawAxis&EAxisList::X && (Space.bIsPerspective || Space.bIsLocalSpace || View->ViewMatrices.GetViewMatrix().M[0][2] != -1.f) )
+		{
+			Render_Axis( View, PDI, EAxisList::X, AxisMatrix, XMaterial, XColor, XAxisDir, Space.Scale, bDrawWidget );
+		}
+
+		if( DrawAxis&EAxisList::Y && (Space.bIsPerspective || Space.bIsLocalSpace || View->ViewMatrices.GetViewMatrix().M[1][2] != -1.f) )
+		{
+			Render_Axis( View, PDI, EAxisList::Y, AxisMatrix, YMaterial, YColor, YAxisDir, Space.Scale, bDrawWidget );
+		}
+
+		if( DrawAxis&EAxisList::Z && (Space.bIsPerspective || Space.bIsLocalSpace || View->ViewMatrices.GetViewMatrix().M[0][1] != 1.f) )
+		{
+			Render_Axis( View, PDI, EAxisList::Z, AxisMatrix, ZMaterial, ZColor, ZAxisDir, Space.Scale, bDrawWidget );
+		}
+
+		const bool bDisabled = IsWidgetDisabled();
+
+		const float ScaledRadius = (TRANSLATE_ROTATE_AXIS_CIRCLE_RADIUS * Space.UniformScale) + GetDefault<ULevelEditorViewportSettings>()->TransformWidgetSizeAdjustment;
+
+		//ZRotation
+		if( DrawAxis&EAxisList::ZRotation && (Space.bIsPerspective || Space.bIsLocalSpace || View->ViewMatrices.GetViewMatrix().M[0][2] != -1.f) )
+		{
+			PDI->SetHitProxy( new HWidgetAxis(EAxisList::ZRotation, bDisabled) );
+			{
+				FVector XAxis = CustomCoordSystem.TransformPosition( FVector(1,0,0).RotateAngleAxis( (EditorModeTools ? EditorModeTools->TranslateRotateXAxisAngle : 0 ), FVector(0,0,1)) );
+				FVector YAxis = CustomCoordSystem.TransformPosition( FVector(0,1,0).RotateAngleAxis( (EditorModeTools ? EditorModeTools->TranslateRotateXAxisAngle : 0 ), FVector(0,0,1)) );
+				FVector BaseArrowPoint = InLocation + XAxis * ScaledRadius;
+				DrawFlatArrow(PDI, BaseArrowPoint, XAxis, YAxis, ZRotateColor, ScaledRadius, ScaledRadius*.5f, ZRotateMaterial->GetRenderProxy(), SDPG_Foreground);
+			}
+			PDI->SetHitProxy( NULL );
+		}
+
+		//XY Plane
+		if( Space.bIsPerspective || Space.bIsLocalSpace || View->ViewMatrices.GetViewMatrix().M[0][1] != 1.f )
+		{
+			if( (DrawAxis & EAxisList::XY) == EAxisList::XY ) 
+			{
+				// Add more sides to the circle if we've been scaled up to keep the circle looking circular
+				// An extra side for every 5 extra unreal units seems to produce a nice result
+				const int32 CircleSides = (GetDefault<ULevelEditorViewportSettings>()->TransformWidgetSizeAdjustment > 0) 
+					? AXIS_CIRCLE_SIDES + (GetDefault<ULevelEditorViewportSettings>()->TransformWidgetSizeAdjustment / 5)
+					: AXIS_CIRCLE_SIDES;
+
+				PDI->SetHitProxy( new HWidgetAxis(EAxisList::XY, bDisabled) );
+				{
+					DrawCircle( PDI, InLocation, CustomCoordSystem.TransformPosition( FVector(1,0,0) ), CustomCoordSystem.TransformPosition( FVector(0,1,0) ), XYPlaneColor, ScaledRadius, CircleSides, SDPG_Foreground );
+					XYPlaneColor.A = ((CurrentAxis&EAxisList::XY) == EAxisList::XY) ? 0x3f : 0x0f;	//make the disc transparent
+					DrawDisc  ( PDI, InLocation, CustomCoordSystem.TransformPosition( FVector(1,0,0) ), CustomCoordSystem.TransformPosition( FVector(0,1,0) ), XYPlaneColor, ScaledRadius, CircleSides, TransparentPlaneMaterialXY->GetRenderProxy(), SDPG_Foreground );
+				}
+				PDI->SetHitProxy( NULL );
+			}
+		}
+	}
+}
+
+/**
+* Draws the 2D widget.
+*/
+
+void FWidget::Render_2D(const FSceneView* View, FPrimitiveDrawInterface* PDI, FEditorViewportClient* ViewportClient, const FVector& InLocation, bool bDrawWidget)
+{
+	//////////////////////////////////////////////////////////////////////////
+	// Translation subwidget
+	//////////////////////////////////////////////////////////////////////////
+
+	// Figure out axis colors
+	const FLinearColor& XColor = (CurrentAxis&EAxisList::X ? (FLinearColor)CurrentColor : AxisColorX);
+	const FLinearColor& YColor = (CurrentAxis&EAxisList::Y ? (FLinearColor)CurrentColor : AxisColorY);
+	const FLinearColor& ZColor = (CurrentAxis&EAxisList::Z ? (FLinearColor)CurrentColor : AxisColorZ);
+	FColor CurrentScreenColor = (CurrentAxis & EAxisList::Screen ? CurrentColor : ScreenSpaceColor);
+
+	// Figure out axis matrices
+	FMatrix WidgetMatrix = CustomCoordSystem * FTranslationMatrix(InLocation);
+
+    FSpaceDescriptor Space(View, ViewportClient, InLocation);
+
+	EAxisList::Type DrawAxis = EAxisList::None;
+	if (Space.bIsOrthoXY)
+	{
+		DrawAxis = EAxisList::X;
+	}
+	else if (Space.bIsOrthoXZ)
+	{
+		DrawAxis = EAxisList::XZ;
+	}
+	else if (Space.bIsOrthoYZ)
+	{
+		DrawAxis = EAxisList::Z;
+	}
+	else if (Space.bIsPerspective)
+	{
+		// Find the best plane to move on
+		const FVector CameraZAxis = View->ViewMatrices.GetViewMatrix().GetColumn(2);
+		const FVector LargestAxis = CameraZAxis.GetAbs();
+		if (LargestAxis.X > LargestAxis.Y)
+		{
+			if (LargestAxis.Z > LargestAxis.X)
+			{
+				DrawAxis = EAxisList::X;
+			}
+			else
+			{
+				DrawAxis = EAxisList::Z;
+			}
+		}
+		else
+		{
+			if (LargestAxis.Y > LargestAxis.Z)
+			{
+				DrawAxis = EAxisList::XZ;
+			}
+			else
+			{
+				DrawAxis = EAxisList::X;
+			}
+		}
+	}
+
+	const bool bDisabled = IsWidgetDisabled();
+
+	// Radius
+	const float ScaledRadius = (TWOD_AXIS_CIRCLE_RADIUS * Space.UniformScale) + GetDefault<ULevelEditorViewportSettings>()->TransformWidgetSizeAdjustment;
+	const int32 CircleSides = (GetDefault<ULevelEditorViewportSettings>()->TransformWidgetSizeAdjustment > 0)
+		? AXIS_CIRCLE_SIDES + (GetDefault<ULevelEditorViewportSettings>()->TransformWidgetSizeAdjustment / 5)
+		: AXIS_CIRCLE_SIDES;
+
+	// Draw the grabbers
+	if (bDrawWidget)
+	{
+		FVector CornerPos = FVector(7, 0, 7) * Space.UniformScale;
+		FVector AxisSize = FVector(12, 1.2, 12) * Space.UniformScale;
+		float CornerLength = 1.2f * Space.UniformScale;
+
+		// Draw the axis lines with arrow heads
+		if (Space.ShouldDrawAxisX(DrawAxis))
+		{
+			UMaterialInstanceDynamic* XMaterial = (CurrentAxis&EAxisList::X ? CurrentAxisMaterial : AxisMaterialX);
+			Render_Axis(View, PDI, EAxisList::X, WidgetMatrix, XMaterial, XColor, XAxisDir, Space.Scale, bDrawWidget);
+		}
+
+		if (Space.ShouldDrawAxisY(DrawAxis))
+		{
+			UMaterialInstanceDynamic* YMaterial = (CurrentAxis&EAxisList::Y ? CurrentAxisMaterial : AxisMaterialY);
+			Render_Axis(View, PDI, EAxisList::Y, WidgetMatrix, YMaterial, YColor, YAxisDir, Space.Scale, bDrawWidget);
+		}
+
+		if (Space.ShouldDrawAxisZ(DrawAxis))
+		{
+			UMaterialInstanceDynamic* ZMaterial = (CurrentAxis&EAxisList::Z ? CurrentAxisMaterial : AxisMaterialZ);
+			Render_Axis(View, PDI, EAxisList::Z, WidgetMatrix, ZMaterial, ZColor, ZAxisDir, Space.Scale, bDrawWidget);
+		}
+
+		// After the primitives have been scaled and transformed, we apply this inverse scale that flattens the dimension
+		// that was scaled up to prevent it from intersecting with the near plane.  In perspective this won't have any effect,
+		// but in the ortho viewports it will prevent scaling in the direction of the camera and thus intersecting the near plane.
+		FVector FlattenScale = FVector(Space.Scale.Component(0) == 1.0f ? 1.0f / Space.UniformScale : 1.0f, Space.Scale.Component(1) == 1.0f ? 1.0f / Space.UniformScale : 1.0f, Space.Scale.Component(2) == 1.0f ? 1.0f / Space.UniformScale : 1.0f);
+		float ArrowRadius = ScaledRadius * 2.0f;
+		float ArrowStartRadius = ScaledRadius * 1.3f;
+
+		uint8 HoverAlpha = 0xff;
+		uint8 NormalAlpha = 0x2f;
+
+		if (((DrawAxis&EAxisList::XZ) == EAxisList::XZ) && (Space.bIsPerspective || Space.bIsLocalSpace || Space.bIsOrthoXZ)) // Front
+		{
+			uint8 Alpha = ((CurrentAxis&EAxisList::XZ) == EAxisList::XZ ? HoverAlpha : NormalAlpha);
+			PDI->SetHitProxy(new HWidgetAxis(EAxisList::XZ, bDisabled));
+			{
+				FColor Color = YColor.ToFColor(true);
+				DrawCircle(PDI, InLocation, CustomCoordSystem.TransformPosition(FVector(1, 0, 0)), CustomCoordSystem.TransformPosition(FVector(0, 0, 1)), Color, ScaledRadius, CircleSides, SDPG_Foreground);
+				Color.A = Alpha;
+				DrawDisc(PDI, InLocation, CustomCoordSystem.TransformPosition(FVector(1, 0, 0)), CustomCoordSystem.TransformPosition(FVector(0, 0, 1)), Color, ScaledRadius, CircleSides, TransparentPlaneMaterialXY->GetRenderProxy(), SDPG_Foreground);
+			}
+			PDI->SetHitProxy(NULL);
+
+			PDI->SetHitProxy(new HWidgetAxis(EAxisList::Rotate2D, bDisabled));
+			{
+				FColor Color = YColor.ToFColor(true);
+				Color.A = ((CurrentAxis&EAxisList::Rotate2D) == EAxisList::Rotate2D ? HoverAlpha : NormalAlpha);
+
+				FVector XAxis = CustomCoordSystem.TransformPosition(FVector(1, 0, 0).RotateAngleAxis((EditorModeTools ? EditorModeTools->TranslateRotate2DAngle : 0), FVector(0, -1, 0)));
+				FVector YAxis = CustomCoordSystem.TransformPosition(FVector(0, 0, 1).RotateAngleAxis((EditorModeTools ? EditorModeTools->TranslateRotate2DAngle : 0), FVector(0, -1, 0)));
+				FVector BaseArrowPoint = InLocation + XAxis * ArrowStartRadius;
+				DrawFlatArrow(PDI, BaseArrowPoint, XAxis, YAxis, Color, ArrowRadius, ArrowRadius*.5f, TransparentPlaneMaterialXY->GetRenderProxy(), SDPG_Foreground);
+			}
+			PDI->SetHitProxy(NULL);
+		}
+	}
+}
+
+/**
+ * Caches off HUD text to display after 3d rendering is complete
+ * @param View - Information about the scene/camera/etc
+ * @param PDI - Drawing interface
+ * @param InLocation - The Origin of the widget
+ * @param Axis0 - The Axis that describes a 0 degree rotation
+ * @param Axis1 - The Axis that describes a 90 degree rotation
+ * @param AngleOfAngle - angle we've rotated so far (in degrees)
+ */
+void FWidget::CacheRotationHUDText(const FSceneView* View, FPrimitiveDrawInterface* PDI, const FVector& InLocation, const FVector& Axis0, const FVector& Axis1, const float AngleOfChange, const float InScale)
+{
+	const float TextDistance = (ROTATION_TEXT_RADIUS * InScale) + GetDefault<ULevelEditorViewportSettings>()->TransformWidgetSizeAdjustment;
+
+	FVector AxisVectors[4] = { Axis0, Axis1, -Axis0, -Axis1};
+
+	for (int i = 0 ; i < 4; ++i)
+	{
+		FVector PotentialTextPosition = InLocation + (TextDistance)*AxisVectors[i];
+		if(View->ScreenToPixel(View->WorldToScreen(PotentialTextPosition), HUDInfoPos))
+		{
+			if (FMath::IsWithin<float>(HUDInfoPos.X, 0, View->UnscaledViewRect.Width()) && FMath::IsWithin<float>(HUDInfoPos.Y, 0, View->UnscaledViewRect.Height()))
+			{
+				//only valid screen locations get a valid string
+				HUDString = FString::Printf(TEXT("%3.2f"), AngleOfChange);
+				break;
+			}
+		}
+	}
+}
+
+/**
+ * If actively dragging, draws a ring representing the potential rotation of the selected objects, snap ticks, and "delta" markers
+ * If not actively dragging, draws a quarter ring representing the closest quadrant to the camera
+ * @param View - Information about the scene/camera/etc
+ * @param PDI - Drawing interface
+ * @param InAxis - Enumeration of axis to rotate about
+ * @param InLocation - The Origin of the widget
+ * @param Axis0 - The Axis that describes a 0 degree rotation
+ * @param Axis1 - The Axis that describes a 90 degree rotation
+ * @param InDirectionToWidget - Direction from camera to the widget
+ * @param InColor - The color associated with the axis of rotation
+ * @param InScale - Multiplier to maintain a constant screen size for rendering the widget
+ */
+void FWidget::DrawRotationArc(const FSceneView* View, FPrimitiveDrawInterface* PDI, EAxisList::Type InAxis, const FVector& InLocation, const FVector& Axis0, const FVector& Axis1, const FVector& InDirectionToWidget, const FColor& InColor, const float InScale, FVector2D& OutAxisDir)
+{
+	bool bIsPerspective = ( View->ViewMatrices.GetProjectionMatrix().M[3][3] < 1.0f );
+	bool bIsOrtho = !bIsPerspective;
+
+	//if we're in an ortho viewport and the ring is perpendicular to the camera (both Axis0 & Axis1 are perpendicular)
+	bIsOrthoDrawingFullRing |= bIsOrtho && (FMath::Abs(Axis0|InDirectionToWidget) < KINDA_SMALL_NUMBER) && (FMath::Abs(Axis1|InDirectionToWidget) < KINDA_SMALL_NUMBER);
+
+	FColor ArcColor = InColor;
+	ArcColor.A = LargeOuterAlpha;
+
+	if (bDragging || (bIsOrthoDrawingFullRing))
+	{
+		if ((CurrentAxis&InAxis) || (bIsOrthoDrawingFullRing))
+		{
+			bool bDrawingArcBallAlso = CurrentAxis == EAxisList::XYZ;
+
+			float DeltaRotation = GetDeltaRotation();
+			float AdjustedDeltaRotation = IsRotationLocalSpace() ? -DeltaRotation : DeltaRotation;
+			float AbsRotation = FRotator::ClampAxis(FMath::Abs(DeltaRotation));
+			float AngleOfChangeRadians (AbsRotation * PI / 180.f);
+
+			//always draw clockwise, so if we're negative we need to flip the angle
+			float StartAngle = AdjustedDeltaRotation < 0.0f ? -AngleOfChangeRadians : 0.0f;
+			float FilledAngle = AngleOfChangeRadians;
+
+			//the axis of rotation
+			FVector ZAxis = Axis0 ^ Axis1;
+
+			ArcColor.A = LargeOuterAlpha;
+			DrawPartialRotationArc(View, PDI, InAxis, InLocation,  Axis0, Axis1, StartAngle, StartAngle + FilledAngle, ArcColor, InScale, InDirectionToWidget);
+			ArcColor.A = SmallOuterAlpha;
+			DrawPartialRotationArc(View, PDI, InAxis, InLocation,  Axis0, Axis1, StartAngle + FilledAngle, StartAngle + 2*PI, ArcColor, InScale, InDirectionToWidget);
+
+			ArcColor = (CurrentAxis&InAxis && !bDrawingArcBallAlso) ? CurrentColor : ArcColor;
+			if (!bDrawingArcBallAlso)
+			{
+				//Hallow Arrow
+				ArcColor.A = 0;
+				DrawStartStopMarker(PDI, InLocation, Axis0, Axis1, 0, ArcColor, InScale);
+				//Filled Arrow
+				ArcColor.A = LargeOuterAlpha;
+				DrawStartStopMarker(PDI, InLocation, Axis0, Axis1, AdjustedDeltaRotation, ArcColor, InScale);
+				ArcColor.A = 255;
+
+
+				FVector SnapLocation = InLocation;
+
+				if (GetDefault<ULevelEditorViewportSettings>()->RotGridEnabled)
+				{
+					float DeltaAngle = GEditor->GetRotGridSize().Yaw;
+					//every 22.5 degrees
+					float TickMarker = 22.5f;
+					for (float Angle = 0; Angle < 360.f; Angle += DeltaAngle)
+					{
+						FVector GridAxis = Axis0.RotateAngleAxis(Angle, ZAxis);
+						float PercentSize = (FMath::Fmod(Angle, TickMarker) == 0) ? .75f : .25f;
+						if (FMath::Fmod(Angle, 90.f) != 0)
+						{
+							DrawSnapMarker(PDI, SnapLocation, GridAxis, FVector::ZeroVector, ArcColor, InScale, 0.0f, PercentSize);
+						}
+					}
+				}
+
+				//draw axis tick marks
+				FColor AxisColor = InColor;
+				//Rotate Colors to match Axis 0
+				Swap(AxisColor.R, AxisColor.G);
+				Swap(AxisColor.B, AxisColor.R);
+				AxisColor.A = (AdjustedDeltaRotation == 0) ? MAX_uint8 : LargeOuterAlpha;
+				DrawSnapMarker(PDI, SnapLocation, Axis0, Axis1, AxisColor, InScale, .25f);
+				AxisColor.A = (AdjustedDeltaRotation == 180.f) ? MAX_uint8 : LargeOuterAlpha;
+				DrawSnapMarker(PDI, SnapLocation, -Axis0, -Axis1, AxisColor, InScale, .25f);
+
+				//Rotate Colors to match Axis 1
+				Swap(AxisColor.R, AxisColor.G);
+				Swap(AxisColor.B, AxisColor.R);
+				AxisColor.A = (AdjustedDeltaRotation == 90.f) ? MAX_uint8 : LargeOuterAlpha;
+				DrawSnapMarker(PDI, SnapLocation, Axis1, -Axis0, AxisColor, InScale, .25f);
+				AxisColor.A = (AdjustedDeltaRotation == 270.f) ? MAX_uint8 : LargeOuterAlpha;
+				DrawSnapMarker(PDI, SnapLocation, -Axis1, Axis0, AxisColor, InScale, .25f);
+
+				if (bDragging)
+				{
+					float OffsetAngle = IsRotationLocalSpace() ? 0 : AdjustedDeltaRotation;
+
+					CacheRotationHUDText(View, PDI, InLocation, Axis0.RotateAngleAxis(OffsetAngle, ZAxis), Axis1.RotateAngleAxis(OffsetAngle, ZAxis), DeltaRotation, InScale);
+				}
+			}
+		}
+	}
+	else
+	{
+		//Reverse the axes based on camera view
+		bool bMirrorAxis0 = ((Axis0 | InDirectionToWidget) <= 0.0f);
+		bool bMirrorAxis1 = ((Axis1 | InDirectionToWidget) <= 0.0f);
+
+		FVector RenderAxis0 = bMirrorAxis0 ? Axis0 : -Axis0;
+		FVector RenderAxis1 = bMirrorAxis1 ? Axis1 : -Axis1;
+		float Direction = (bMirrorAxis0 ^ bMirrorAxis1) ? -1.0f : 1.0f;
+
+		DrawPartialRotationArc(View, PDI, InAxis, InLocation, RenderAxis0, RenderAxis1, 0, PI/2, ArcColor, InScale, InDirectionToWidget);
+
+		FVector2D Axis0ScreenLocation;
+		if (!View->ScreenToPixel(View->WorldToScreen(InLocation + RenderAxis0 * 64.0f), Axis0ScreenLocation))
+		{
+			Axis0ScreenLocation.X = Axis0ScreenLocation.Y = 0;
+		}
+
+		FVector2D Axis1ScreenLocation;
+		if (!View->ScreenToPixel(View->WorldToScreen(InLocation + RenderAxis1 * 64.0f), Axis1ScreenLocation))
+		{
+			Axis1ScreenLocation.X = Axis1ScreenLocation.Y = 0;
+		}
+
+		OutAxisDir = ((Axis1ScreenLocation - Axis0ScreenLocation) * Direction).GetSafeNormal();
+	}
+}
+
+/**
+ * If actively dragging, draws a ring representing the potential rotation of the selected objects, snap ticks, and "delta" markers
+ * If not actively dragging, draws a quarter ring representing the closest quadrant to the camera
+ * @param View - Information about the scene/camera/etc
+ * @param PDI - Drawing interface
+ * @param InAxis - Enumeration of axis to rotate about
+ * @param InLocation - The Origin of the widget
+ * @param Axis0 - The Axis that describes a 0 degree rotation
+ * @param Axis1 - The Axis that describes a 90 degree rotation
+ * @param InStartAngle - The starting angle about (Axis0^Axis1) to render the arc, in radians
+ * @param InEndAngle - The ending angle about (Axis0^Axis1) to render the arc, in radians
+ * @param InColor - The color associated with the axis of rotation
+ * @param InScale - Multiplier to maintain a constant screen size for rendering the widget
+ */
+void FWidget::DrawPartialRotationArc(const FSceneView* View, FPrimitiveDrawInterface* PDI, EAxisList::Type InAxis, const FVector& InLocation, const FVector& Axis0, const FVector& Axis1, const float InStartAngle, const float InEndAngle, const FColor& InColor, const float InScale, const FVector& InDirectionToWidget )
+{
+	const float InnerRadius = (INNER_AXIS_CIRCLE_RADIUS * InScale) + GetDefault<ULevelEditorViewportSettings>()->TransformWidgetSizeAdjustment;
+	const float OuterRadius = (OUTER_AXIS_CIRCLE_RADIUS * InScale) + GetDefault<ULevelEditorViewportSettings>()->TransformWidgetSizeAdjustment;
+
+	bool bIsPerspective = ( View->ViewMatrices.GetProjectionMatrix().M[3][3] < 1.0f );
+	PDI->SetHitProxy( new HWidgetAxis( InAxis ) );
+	{
+		FThickArcParams OuterArcParams(PDI, InLocation, TransparentPlaneMaterialXY, InnerRadius, OuterRadius);
+		FColor OuterColor = ( (CurrentAxis&InAxis  && ( CurrentAxis != EAxisList::XYZ ) )? CurrentColor : InColor );
+		//Pass through alpha
+		OuterColor.A = InColor.A;
+		DrawThickArc(OuterArcParams, Axis0, Axis1, InStartAngle, InEndAngle, OuterColor, InDirectionToWidget, !bIsPerspective );
+	}
+	PDI->SetHitProxy( NULL );
+
+	const bool bIsHitProxyView = View->Family->EngineShowFlags.HitProxies;
+	if (bIsPerspective && !bIsHitProxyView && !PDI->IsHitTesting())
+	{
+		FThickArcParams InnerArcParams(PDI, InLocation, GridMaterial, 0.0f, InnerRadius);
+		FColor InnerColor = InColor;
+		//if something is selected and it's not this
+		InnerColor.A = ((CurrentAxis & InAxis) && !bDragging) ? LargeInnerAlpha : SmallInnerAlpha;
+		DrawThickArc(InnerArcParams, Axis0, Axis1, InStartAngle, InEndAngle, InnerColor, InDirectionToWidget, false );
+	}
+}
+
+/**
+ * Renders a portion of an arc for the rotation widget
+ * @param InParams - Material, Radii, etc
+ * @param InStartAxis - Start of the arc, in radians
+ * @param InEndAxis - End of the arc, in radians
+ * @param InColor - Color to use for the arc
+ */
+void FWidget::DrawThickArc (const FThickArcParams& InParams, const FVector& Axis0, const FVector& Axis1, const float InStartAngle, const float InEndAngle, const FColor& InColor, const FVector& InDirectionToWidget, bool bIsOrtho )
+{
+	if (InColor.A == 0)
+	{
+		return;
+	}
+
+	// Add more sides to the circle if we've been scaled up to keep the circle looking circular
+	// An extra side for every 5 extra unreal units seems to produce a nice result
+	const int32 CircleSides = (GetDefault<ULevelEditorViewportSettings>()->TransformWidgetSizeAdjustment > 0) 
+		? AXIS_CIRCLE_SIDES + (GetDefault<ULevelEditorViewportSettings>()->TransformWidgetSizeAdjustment / 5)
+		: AXIS_CIRCLE_SIDES;
+	const int32 NumPoints = FMath::TruncToInt(CircleSides * (InEndAngle-InStartAngle)/(PI/2)) + 1;
+
+	FColor TriangleColor = InColor;
+	FColor RingColor = InColor;
+	RingColor.A = MAX_uint8;
+
+	FVector ZAxis = Axis0 ^ Axis1;
+	FVector LastVertex;
+
+	FDynamicMeshBuilder MeshBuilder(InParams.PDI->View->GetFeatureLevel());
+
+	for (int32 RadiusIndex = 0; RadiusIndex < 2; ++RadiusIndex)
+	{
+		float Radius = (RadiusIndex == 0) ? InParams.OuterRadius : InParams.InnerRadius;
+		float TCRadius = Radius / (float) InParams.OuterRadius;
+		//Compute vertices for base circle.
+		for(int32 VertexIndex = 0;VertexIndex <= NumPoints;VertexIndex++)
+		{
+			float Percent = VertexIndex/(float)NumPoints;
+			float Angle = FMath::Lerp(InStartAngle, InEndAngle, Percent);
+			float AngleDeg = FRotator::ClampAxis(Angle * 180.f / PI);
+
+			FVector VertexDir = Axis0.RotateAngleAxis(AngleDeg, ZAxis);
+			VertexDir.Normalize();
+
+			float TCAngle = Percent*(PI/2);
+			FVector2D TC(TCRadius*FMath::Cos(Angle), TCRadius*FMath::Sin(Angle));
+
+			const FVector VertexPosition = InParams.Position + VertexDir*Radius;
+			FVector Normal = VertexPosition - InParams.Position;
+			Normal.Normalize();
+
+			FDynamicMeshVertex MeshVertex;
+			MeshVertex.Position = VertexPosition;
+			MeshVertex.Color = TriangleColor;
+			MeshVertex.TextureCoordinate[0] = TC;
+
+			MeshVertex.SetTangents(
+				-ZAxis,
+				(-ZAxis) ^ Normal,
+				Normal
+				);
+
+			MeshBuilder.AddVertex(MeshVertex); //Add bottom vertex
+
+			// Push out the arc line borders so they dont z-fight with the mesh arcs
+			FVector StartLinePos = LastVertex;
+			FVector EndLinePos = VertexPosition;
+			if (VertexIndex != 0)
+			{
+				InParams.PDI->DrawLine(StartLinePos,EndLinePos,RingColor,SDPG_Foreground);
+			}
+			LastVertex = VertexPosition;
+		}
+	}
+
+	//Add top/bottom triangles, in the style of a fan.
+	int32 InnerVertexStartIndex = NumPoints + 1;
+	for(int32 VertexIndex = 0; VertexIndex < NumPoints; VertexIndex++)
+	{
+		MeshBuilder.AddTriangle(VertexIndex, VertexIndex+1, InnerVertexStartIndex+VertexIndex);
+		MeshBuilder.AddTriangle(VertexIndex+1, InnerVertexStartIndex+VertexIndex+1, InnerVertexStartIndex+VertexIndex);
+	}
+
+	MeshBuilder.Draw(InParams.PDI, FMatrix::Identity, InParams.Material->GetRenderProxy(),SDPG_Foreground,0.f);
+}
+
+/**
+ * Draws protractor like ticks where the rotation widget would snap too.
+ * Also, used to draw the wider axis tick marks
+ * @param PDI - Drawing interface
+ * @param InLocation - The Origin of the widget
+ * @param Axis0 - The Axis that describes a 0 degree rotation
+ * @param Axis1 - The Axis that describes a 90 degree rotation
+ * @param InAngle - The Angle to rotate about the axis of rotation, the vector (Axis0 ^ Axis1)
+ * @param InColor - The color to use for line/poly drawing
+ * @param InScale - Multiplier to maintain a constant screen size for rendering the widget
+ * @param InWidthPercent - The percent of the distance between the outer ring and inner ring to use for tangential thickness
+ * @param InPercentSize - The percent of the distance between the outer ring and inner ring to use for radial distance
+ */
+void FWidget::DrawSnapMarker(FPrimitiveDrawInterface* PDI, const FVector& InLocation, const FVector& Axis0, const FVector& Axis1, const FColor& InColor, const float InScale, const float InWidthPercent, const float InPercentSize)
+{
+	const float InnerDistance = (INNER_AXIS_CIRCLE_RADIUS * InScale) + GetDefault<ULevelEditorViewportSettings>()->TransformWidgetSizeAdjustment;
+	const float OuterDistance = (OUTER_AXIS_CIRCLE_RADIUS * InScale) + GetDefault<ULevelEditorViewportSettings>()->TransformWidgetSizeAdjustment;
+	const float MaxMarkerHeight = OuterDistance - InnerDistance;
+	const float MarkerWidth = MaxMarkerHeight*InWidthPercent;
+	const float MarkerHeight = MaxMarkerHeight*InPercentSize;
+
+	FVector Vertices[4];
+	Vertices[0] = InLocation + (OuterDistance)*Axis0 - (MarkerWidth*.5)*Axis1;
+	Vertices[1] = Vertices[0] + (MarkerWidth)*Axis1;
+	Vertices[2] = InLocation + (OuterDistance-MarkerHeight)*Axis0 - (MarkerWidth*.5)*Axis1;
+	Vertices[3] = Vertices[2] + (MarkerWidth)*Axis1;
+
+	//draw at least one line
+	PDI->DrawLine(Vertices[0], Vertices[2], InColor, SDPG_Foreground);
+
+	//if there should be thickness, draw the other lines
+	if (InWidthPercent > 0.0f)
+	{
+		PDI->DrawLine(Vertices[0], Vertices[1], InColor, SDPG_Foreground);
+		PDI->DrawLine(Vertices[1], Vertices[3], InColor, SDPG_Foreground);
+		PDI->DrawLine(Vertices[2], Vertices[3], InColor, SDPG_Foreground);
+
+		//fill in the box
+		FDynamicMeshBuilder MeshBuilder(PDI->View->GetFeatureLevel());
+
+		for(int32 VertexIndex = 0;VertexIndex < 4; VertexIndex++)
+		{
+			FDynamicMeshVertex MeshVertex;
+			MeshVertex.Position = Vertices[VertexIndex];
+			MeshVertex.Color = InColor;
+			MeshVertex.TextureCoordinate[0] = FVector2D(0.0f, 0.0f);
+			MeshVertex.SetTangents(
+				Axis0,
+				Axis1,
+				(Axis0) ^ Axis1
+				);
+			MeshBuilder.AddVertex(MeshVertex); //Add bottom vertex
+		}
+
+		MeshBuilder.AddTriangle(0, 1, 2);
+		MeshBuilder.AddTriangle(1, 3, 2);
+		MeshBuilder.Draw(PDI, FMatrix::Identity, TransparentPlaneMaterialXY->GetRenderProxy(),SDPG_Foreground,0.f);
+	}
+}
+
+/**
+ * Draw Start/Stop Marker to show delta rotations along the arc of rotation
+ * @param PDI - Drawing interface
+ * @param InLocation - The Origin of the widget
+ * @param Axis0 - The Axis that describes a 0 degree rotation
+ * @param Axis1 - The Axis that describes a 90 degree rotation
+ * @param InAngle - The Angle to rotate about the axis of rotation, the vector (Axis0 ^ Axis1), units are degrees
+ * @param InColor - The color to use for line/poly drawing
+ * @param InScale - Multiplier to maintain a constant screen size for rendering the widget
+ */
+void FWidget::DrawStartStopMarker(FPrimitiveDrawInterface* PDI, const FVector& InLocation, const FVector& Axis0, const FVector& Axis1, const float InAngle, const FColor& InColor, const float InScale)
+{
+	const float ArrowHeightPercent = .8f;
+	const float InnerDistance = (INNER_AXIS_CIRCLE_RADIUS * InScale) + GetDefault<ULevelEditorViewportSettings>()->TransformWidgetSizeAdjustment;
+	const float OuterDistance = (OUTER_AXIS_CIRCLE_RADIUS * InScale) + GetDefault<ULevelEditorViewportSettings>()->TransformWidgetSizeAdjustment;
+	const float RingHeight = OuterDistance - InnerDistance;
+	const float ArrowHeight = RingHeight*ArrowHeightPercent;
+	const float ThirtyDegrees = PI / 6.0f;
+	const float HalfArrowidth = ArrowHeight*FMath::Tan(ThirtyDegrees);
+
+	FVector ZAxis = Axis0 ^ Axis1;
+	FVector RotatedAxis0 = Axis0.RotateAngleAxis(InAngle, ZAxis);
+	FVector RotatedAxis1 = Axis1.RotateAngleAxis(InAngle, ZAxis);
+
+	FVector Vertices[3];
+	Vertices[0] = InLocation + (OuterDistance)*RotatedAxis0;
+	Vertices[1] = Vertices[0] + (ArrowHeight)*RotatedAxis0 - HalfArrowidth*RotatedAxis1;
+	Vertices[2] = Vertices[1] + (2*HalfArrowidth)*RotatedAxis1;
+
+	PDI->DrawLine(Vertices[0], Vertices[1], InColor, SDPG_Foreground);
+	PDI->DrawLine(Vertices[1], Vertices[2], InColor, SDPG_Foreground);
+	PDI->DrawLine(Vertices[0], Vertices[2], InColor, SDPG_Foreground);
+
+	if (InColor.A > 0)
+	{
+		//fill in the box
+		FDynamicMeshBuilder MeshBuilder(PDI->View->GetFeatureLevel());
+
+		for(int32 VertexIndex = 0;VertexIndex < 3; VertexIndex++)
+		{
+			FDynamicMeshVertex MeshVertex;
+			MeshVertex.Position = Vertices[VertexIndex];
+			MeshVertex.Color = InColor;
+			MeshVertex.TextureCoordinate[0] = FVector2D(0.0f, 0.0f);
+			MeshVertex.SetTangents(
+				RotatedAxis0,
+				RotatedAxis1,
+				(RotatedAxis0) ^ RotatedAxis1
+				);
+			MeshBuilder.AddVertex(MeshVertex); //Add bottom vertex
+		}
+
+		MeshBuilder.AddTriangle(0, 1, 2);
+		MeshBuilder.Draw(PDI, FMatrix::Identity, TransparentPlaneMaterialXY->GetRenderProxy(),SDPG_Foreground,0.f);
+	}
+}
+
+EAxisList::Type FWidget::GetAxisToDraw( EWidgetMode WidgetMode ) const
+{
+	return EditorModeTools ? EditorModeTools->GetWidgetAxisToDraw( WidgetMode ) : EAxisList::All;
+}
diff --git a/Engine/Source/Editor/UnrealEd/Public/UnrealWidget.h b/Engine/Source/Editor/UnrealEd/Public/UnrealWidget.h
--- a/Engine/Source/Editor/UnrealEd/Public/UnrealWidget.h
+++ b/Engine/Source/Editor/UnrealEd/Public/UnrealWidget.h
@@ -3,8 +3,8 @@
 #pragma once
 
 #include "CoreMinimal.h"
-#include "UObject/GCObject.h"
 #include "HitProxies.h"
+#include "UObject/GCObject.h"
 
 class FCanvas;
 class FEditorModeTools;
@@ -18,19 +18,18 @@ class FMaterialRenderProxy;
 /** Coordinate system identifiers. */
 enum ECoordSystem
 {
-	COORD_None	= -1,
+	COORD_None = -1,
 	COORD_World,
 	COORD_Local,
 	COORD_Max,
 };
 
-class FWidget
-	: public FGCObject
+class FWidget : public FGCObject
 {
 public:
 	enum EWidgetMode
 	{
-		WM_None			= -1,
+		WM_None = -1,
 		WM_Translate,
 		WM_TranslateRotateZ,
 		WM_2D,
@@ -50,64 +49,32 @@ public:
 	 * Renders any widget specific HUD text
 	 * @param Canvas - Canvas to use for 2d rendering
 	 */
-	void DrawHUD (FCanvas* Canvas);
+	void DrawHUD(FCanvas* Canvas);
 
 	// @param ViewportClient must not be 0
-	void Render( const FSceneView* View, FPrimitiveDrawInterface* PDI, FEditorViewportClient* ViewportClient );
-
-	// @param ViewportClient must not be 0
-	void RenderGrid( const FSceneView* View, FPrimitiveDrawInterface* PDI, FEditorViewportClient* ViewportClient );
-
-	/**
-	 * Draws an arrow head line for a specific axis.
-	 * @param	bCubeHead		[opt] If true, render a cube at the axis tips.  If false (the default), render a cone.
-	 */
-	void Render_Axis(const FSceneView* View, FPrimitiveDrawInterface* PDI, EAxisList::Type InAxis, FMatrix& InMatrix, UMaterialInterface* InMaterial, const FLinearColor& InColor, FVector2D& OutAxisDir, const FVector& InScale, bool bDrawWidget, bool bCubeHead=false, float AxisLengthOffset = 0);
-
-	/**
-	 * Draws a cube
-	 */
-	void Render_Cube( FPrimitiveDrawInterface* PDI, const FMatrix& InMatrix, const UMaterialInterface* InMaterial, const FVector& InScale );
-
-	/**
-	 * Draws the translation widget.
-	 */
-	UNREALED_API void Render_Translate( const FSceneView* View, FPrimitiveDrawInterface* PDI, FEditorViewportClient* ViewportClient,  const FVector& InLocation, bool bDrawWidget );
-
-	/**
-	 * Draws the rotation widget.
-	 */
-	UNREALED_API void Render_Rotate( const FSceneView* View, FPrimitiveDrawInterface* PDI, FEditorViewportClient* ViewportClient, const FVector& InLocation, bool bDrawWidget );
-
-	/**
-	 * Draws the scaling widget.
-	 */
-	void Render_Scale( const FSceneView* View, FPrimitiveDrawInterface* PDI, FEditorViewportClient* ViewportClient, const FVector& InLocation, bool bDrawWidget );
-
-	/**
-	* Draws the Translate & Rotate Z widget.
-	*/
-	void Render_TranslateRotateZ( const FSceneView* View, FPrimitiveDrawInterface* PDI, FEditorViewportClient* ViewportClient, const FVector& InLocation, bool bDrawWidget );
-
-	/**
-	* Draws the combined 2D widget.
-	*/
-	void Render_2D(const FSceneView* View, FPrimitiveDrawInterface* PDI, FEditorViewportClient* ViewportClient, const FVector& InLocation, bool bDrawWidget);
+	void Render(const FSceneView* View, FPrimitiveDrawInterface* PDI, FEditorViewportClient* ViewportClient);
 
+public:
 	/**
 	 * Converts mouse movement on the screen to widget axis movement/rotation.
 	 */
-	void ConvertMouseMovementToAxisMovement(FSceneView* InView,  FEditorViewportClient* InViewportClient, bool bInUsedDragModifier, FVector& InDiff, FVector& OutDrag, FRotator& OutRotation, FVector& OutScale );
+	void ConvertMouseMovementToAxisMovement(FSceneView* InView, FEditorViewportClient* InViewportClient,
+	                                        bool bInUsedDragModifier, FVector& InDiff, FVector& OutDrag,
+	                                        FRotator& OutRotation, FVector& OutScale);
 
 	/**
 	 * Absolute Translation conversion from mouse movement on the screen to widget axis movement/rotation.
 	 */
-	void AbsoluteTranslationConvertMouseMovementToAxisMovement(FSceneView* InView, FEditorViewportClient* InViewportClient, const FVector& InLocation, const FVector2D& InMousePosition, FVector& OutDrag, FRotator& OutRotation, FVector& OutScale );
+	void AbsoluteTranslationConvertMouseMovementToAxisMovement(FSceneView* InView,
+	                                                           FEditorViewportClient* InViewportClient,
+	                                                           const FVector& InLocation,
+	                                                           const FVector2D& InMousePosition, FVector& OutDrag,
+	                                                           FRotator& OutRotation, FVector& OutScale);
 
 	/** 
 	 * Grab the initial offset again first time input is captured
 	 */
-	void ResetInitialTranslationOffset (void)
+	void ResetInitialTranslationOffset(void)
 	{
 		bAbsoluteTranslationInitialOffsetCached = false;
 	}
@@ -159,7 +126,7 @@ public:
 	void SetDragStartPosition(const FVector2D& Position)
 	{
 		DragStartPos = Position;
-		LastDragPos = DragStartPos;
+		LastDragPos  = DragStartPos;
 	}
 
 	/**
@@ -173,16 +140,16 @@ public:
 	/**
 	 * Sets if we are currently engaging the widget in dragging
 	 */
-	void SetDragging (const bool InDragging) 
-	{ 
+	void SetDragging(const bool InDragging)
+	{
 		bDragging = InDragging;
 	}
 
 	/**
 	 * Sets if we are currently engaging the widget in dragging
 	 */
-	void SetSnapEnabled (const bool InSnapEnabled) 
-	{ 
+	void SetSnapEnabled(const bool InSnapEnabled)
+	{
 		bSnapEnabled = InSnapEnabled;
 	}
 
@@ -191,7 +158,7 @@ public:
 	 *
 	 * @param Ar	FArchive to serialize with
 	 */
-	virtual void AddReferencedObjects( FReferenceCollector& Collector ) override;
+	virtual void AddReferencedObjects(FReferenceCollector& Collector) override;
 	virtual FString GetReferencerName() const override
 	{
 		return "FWidget";
@@ -200,7 +167,7 @@ public:
 	/**
 	 * Gets the axis to draw based on the current widget mode
 	 */
-	EAxisList::Type GetAxisToDraw( EWidgetMode WidgetMode ) const;
+	EAxisList::Type GetAxisToDraw(EWidgetMode WidgetMode) const;
 
 	/** @return true if the widget is disabled */
 	bool IsWidgetDisabled() const;
@@ -209,12 +176,26 @@ public:
 	UNREALED_API void UpdateDeltaRotation();
 
 	/** Resets the total delta rotation back to zero */
-	void ResetDeltaRotation() { TotalDeltaRotation = 0; }
+	void ResetDeltaRotation()
+	{
+		TotalDeltaRotation = 0;
+	}
 
 	/** @return the rotation speed of the widget */
-	static float GetRotationSpeed() { return (2.f*(float)PI)/360.f; }
-private:
+	static float GetRotationSpeed()
+	{
+		return (2.f * (float)PI) / 360.f;
+	}
 
+private:
+	void ConvertMouseToAxis_Translate(FVector2D DragDir, FVector& InOutDelta, FVector& OutDrag) const;
+	void ConvertMouseToAxis_Rotate(FVector2D TangentDir, FVector2D DragDir, FSceneView* InView,
+	                               FEditorViewportClient* InViewportClient, FVector& InOutDelta, FRotator& OutRotation);
+	void ConvertMouseToAxis_Scale(FVector2D DragDir, FVector& InOutDelta, FVector& OutScale);
+	void ConvertMouseToAxis_TranslateRotateZ(FVector2D TangentDir, FVector2D DragDir, FVector& InOutDelta,
+	                                         FVector& OutDrag, FRotator& OutRotation);
+	void ConvertMouseToAxis_WM_2D(FVector2D TangentDir, FVector2D DragDir, FVector& InOutDelta, FVector& OutDrag,
+	                              FRotator& OutRotation);
 
 	struct FAbsoluteMovementParams
 	{
@@ -243,16 +224,46 @@ private:
 		bool bPositionSnapping;
 	};
 
+	void AbsoluteConvertMouseToAxis_Translate(FSceneView* InView, const FMatrix& InputCoordSystem,
+	                                          FAbsoluteMovementParams& InOutParams, FVector& OutDrag);
+	void AbsoluteConvertMouseToAxis_WM_2D(const FMatrix& InputCoordSystem, FAbsoluteMovementParams& InOutParams,
+	                                      FVector& OutDrag, FRotator& OutRotation);
+	void AbsoluteConvertMouseToAxis_TranslateRotateZ(const FMatrix& InputCoordSystem,
+	                                                 FAbsoluteMovementParams& InOutParams, FVector& OutDrag,
+	                                                 FRotator& OutRotation);
+
+	/**
+	 * Render helper functions
+	 */
+	void RenderGrid(const FSceneView* View, FPrimitiveDrawInterface* PDI, FEditorViewportClient* ViewportClient);
+	void Render_Axis(const FSceneView* View, FPrimitiveDrawInterface* PDI, EAxisList::Type InAxis, FMatrix& InMatrix,
+	                 UMaterialInterface* InMaterial, const FLinearColor& InColor, FVector2D& OutAxisDir,
+	                 const FVector& InScale, bool bDrawWidget, bool bCubeHead = false, float AxisLengthOffset = 0);
+	void Render_Cube(FPrimitiveDrawInterface* PDI, const FMatrix& InMatrix, const UMaterialInterface* InMaterial,
+	                 const FVector& InScale);
+	UNREALED_API void Render_Translate(const FSceneView* View, FPrimitiveDrawInterface* PDI,
+	                                   FEditorViewportClient* ViewportClient, const FVector& InLocation,
+	                                   bool bDrawWidget);
+	UNREALED_API void Render_Rotate(const FSceneView* View, FPrimitiveDrawInterface* PDI,
+	                                FEditorViewportClient* ViewportClient, const FVector& InLocation, bool bDrawWidget);
+	void Render_Scale(const FSceneView* View, FPrimitiveDrawInterface* PDI, FEditorViewportClient* ViewportClient,
+	                  const FVector& InLocation, bool bDrawWidget);
+	void Render_TranslateRotateZ(const FSceneView* View, FPrimitiveDrawInterface* PDI,
+	                             FEditorViewportClient* ViewportClient, const FVector& InLocation, bool bDrawWidget);
+	void Render_2D(const FSceneView* View, FPrimitiveDrawInterface* PDI, FEditorViewportClient* ViewportClient,
+	               const FVector& InLocation, bool bDrawWidget);
+
+
 	struct FThickArcParams
 	{
-		FThickArcParams(FPrimitiveDrawInterface* InPDI, const FVector& InPosition, UMaterialInterface* InMaterial, const float InInnerRadius, const float InOuterRadius)
-			: Position(InPosition)
-			, PDI(InPDI)
-			, Material(InMaterial)
-			, InnerRadius(InInnerRadius)
-			, OuterRadius(InOuterRadius)
-		{
-		}
+		FThickArcParams(FPrimitiveDrawInterface* InPDI, const FVector& InPosition, UMaterialInterface* InMaterial,
+		                const float InInnerRadius, const float InOuterRadius)
+		    : Position(InPosition)
+		    , PDI(InPDI)
+		    , Material(InMaterial)
+		    , InnerRadius(InInnerRadius)
+		    , OuterRadius(InOuterRadius)
+		{}
 
 		/** The current position of the widget */
 		FVector Position;
@@ -273,7 +284,7 @@ private:
 	 * @param InParams - Structure containing all the information needed for absolute movement
 	 * @return - The requested delta from the current position
 	 */
-	FVector GetAbsoluteTranslationDelta (const FAbsoluteMovementParams& InParams);
+	FVector GetAbsoluteTranslationDelta(const FAbsoluteMovementParams& InParams);
 	/**
 	 * Returns the offset from the initial selection point
 	 */
@@ -303,7 +314,10 @@ private:
 	 * @param InScale - Multiplier to maintain a constant screen size for rendering the widget
 	 * @param OutAxisDir - Viewport-space direction of rotation arc chord is placed here
 	 */
-	void DrawRotationArc(const FSceneView* View, FPrimitiveDrawInterface* PDI, EAxisList::Type InAxis, const FVector& InLocation, const FVector& Axis0, const FVector& Axis1, const FVector& InDirectionToWidget, const FColor& InColor, const float InScale, FVector2D& OutAxisEnd);
+	void DrawRotationArc(const FSceneView* View, FPrimitiveDrawInterface* PDI, EAxisList::Type InAxis,
+	                     const FVector& InLocation, const FVector& Axis0, const FVector& Axis1,
+	                     const FVector& InDirectionToWidget, const FColor& InColor, const float InScale,
+	                     FVector2D& OutAxisEnd);
 
 	/**
 	 * If actively dragging, draws a ring representing the potential rotation of the selected objects, snap ticks, and "delta" markers
@@ -319,7 +333,10 @@ private:
 	 * @param InColor - The color associated with the axis of rotation
 	 * @param InScale - Multiplier to maintain a constant screen size for rendering the widget
 	 */
-	void DrawPartialRotationArc(const FSceneView* View, FPrimitiveDrawInterface* PDI, EAxisList::Type InAxis, const FVector& InLocation, const FVector& Axis0, const FVector& Axis1, const float InStartAngle, const float InEndAngle, const FColor& InColor, const float InScale, const FVector& InDirectionToWidget);
+	void DrawPartialRotationArc(const FSceneView* View, FPrimitiveDrawInterface* PDI, EAxisList::Type InAxis,
+	                            const FVector& InLocation, const FVector& Axis0, const FVector& Axis1,
+	                            const float InStartAngle, const float InEndAngle, const FColor& InColor,
+	                            const float InScale, const FVector& InDirectionToWidget);
 
 	/**
 	 * Renders a portion of an arc for the rotation widget
@@ -328,7 +345,9 @@ private:
 	 * @param InEndAxis - End of the arc
 	 * @param InColor - Color to use for the arc
 	 */
-	void DrawThickArc (const FThickArcParams& InParams, const FVector& Axis0, const FVector& Axis1, const float InStartAngle, const float InEndAngle, const FColor& InColor, const FVector& InDirectionToWidget, bool bIsOrtho );
+	void DrawThickArc(const FThickArcParams& InParams, const FVector& Axis0, const FVector& Axis1,
+	                  const float InStartAngle, const float InEndAngle, const FColor& InColor,
+	                  const FVector& InDirectionToWidget, bool bIsOrtho);
 
 	/**
 	 * Draws protractor like ticks where the rotation widget would snap too.
@@ -343,7 +362,9 @@ private:
 	 * @param InWidthPercent - The percent of the distance between the outer ring and inner ring to use for tangential thickness
 	 * @param InPercentSize - The percent of the distance between the outer ring and inner ring to use for radial distance
 	 */
-	void DrawSnapMarker(FPrimitiveDrawInterface* PDI, const FVector& InLocation, const FVector& Axis0, const FVector& Axis1, const FColor& InColor, const float InScale, const float InWidthPercent=0.0f, const float InPercentSize = 1.0f);
+	void DrawSnapMarker(FPrimitiveDrawInterface* PDI, const FVector& InLocation, const FVector& Axis0,
+	                    const FVector& Axis1, const FColor& InColor, const float InScale,
+	                    const float InWidthPercent = 0.0f, const float InPercentSize = 1.0f);
 
 	/**
 	 * Draw Start/Stop Marker to show delta rotations along the arc of rotation
@@ -355,7 +376,8 @@ private:
 	 * @param InColor - The color to use for line/poly drawing
 	 * @param InScale - Multiplier to maintain a constant screen size for rendering the widget
 	 */
-	void 	DrawStartStopMarker(FPrimitiveDrawInterface* PDI, const FVector& InLocation, const FVector& Axis0, const FVector& Axis1, const float InAngle, const FColor& InColor, const float InScale);
+	void DrawStartStopMarker(FPrimitiveDrawInterface* PDI, const FVector& InLocation, const FVector& Axis0,
+	                         const FVector& Axis1, const float InAngle, const FColor& InColor, const float InScale);
 
 	/**
 	 * Caches off HUD text to display after 3d rendering is complete
@@ -367,7 +389,9 @@ private:
 	 * @param AngleOfAngle - angle we've rotated so far (in degrees)
 	 * @param InScale - Multiplier to maintain a constant screen size for rendering the widget
 	 */
-	void CacheRotationHUDText(const FSceneView* View, FPrimitiveDrawInterface* PDI, const FVector& InLocation, const FVector& Axis0, const FVector& Axis1, const float AngleOfChange, const float InScale);
+	void CacheRotationHUDText(const FSceneView* View, FPrimitiveDrawInterface* PDI, const FVector& InLocation,
+	                          const FVector& Axis0, const FVector& Axis1, const float AngleOfChange,
+	                          const float InScale);
 
 	/**
 	 * Gets the axis to use when converting mouse movement, accounting for Ortho views.
@@ -377,10 +401,13 @@ private:
 	 *
 	 * @return Index of the dominant axis.
 	 */
-	uint32 GetDominantAxisIndex( const FVector& InDiff, FEditorViewportClient* ViewportClient ) const;
+	uint32 GetDominantAxisIndex(const FVector& InDiff, FEditorViewportClient* ViewportClient) const;
 
 
-	void DrawColoredSphere(FPrimitiveDrawInterface* PDI, const FVector& Center, const FRotator& Orientation, FColor Color, const FVector& Radii, int32 NumSides, int32 NumRings, const FMaterialRenderProxy* MaterialRenderProxy, uint8 DepthPriority, bool bDisableBackfaceCulling);
+	void DrawColoredSphere(FPrimitiveDrawInterface* PDI, const FVector& Center, const FRotator& Orientation,
+	                       FColor Color, const FVector& Radii, int32 NumSides, int32 NumRings,
+	                       const FMaterialRenderProxy* MaterialRenderProxy, uint8 DepthPriority,
+	                       bool bDisableBackfaceCulling);
 
 	/** The axis currently being moused over */
 	EAxisList::Type CurrentAxis;
@@ -420,7 +447,7 @@ private:
 	 * An extra matrix to apply to the widget before drawing it (allows for local/custom coordinate systems).
 	 */
 	FMatrix CustomCoordSystem;
-	
+
 	/** The space of the custom coord system */
 	ECoordSystem CustomCoordSystemSpace;
 
@@ -456,15 +483,14 @@ private:
  */
 struct HWidgetAxis : public HHitProxy
 {
-	DECLARE_HIT_PROXY( UNREALED_API );
+	DECLARE_HIT_PROXY(UNREALED_API);
 
 	EAxisList::Type Axis;
-	uint32 bDisabled:1;
+	uint32 bDisabled : 1;
 
-	HWidgetAxis(EAxisList::Type InAxis, bool InbDisabled = false, EHitProxyPriority InHitProxy = HPP_UI):
-		HHitProxy(InHitProxy),
-		Axis(InAxis),
-		bDisabled(InbDisabled) {}
+	HWidgetAxis(EAxisList::Type InAxis, bool InbDisabled = false, EHitProxyPriority InHitProxy = HPP_UI)
+	    : HHitProxy(InHitProxy), Axis(InAxis), bDisabled(InbDisabled)
+	{}
 
 	virtual EMouseCursor::Type GetMouseCursor() override
 	{
