Added an apply to all button to CSV import, which will bulk apply the selected settings to new CSV files being imported.

diff --git a/Engine/Source/Editor/UnrealEd/Classes/Factories/CSVImportFactory.h b/Engine/Source/Editor/UnrealEd/Classes/Factories/CSVImportFactory.h
--- a/Engine/Source/Editor/UnrealEd/Classes/Factories/CSVImportFactory.h
+++ b/Engine/Source/Editor/UnrealEd/Classes/Factories/CSVImportFactory.h
@@ -47,6 +47,8 @@ struct FCSVImportSettings
 
 	UPROPERTY(BlueprintReadWrite, Category="Misc")
 	TEnumAsByte<ERichCurveInterpMode> ImportCurveInterpMode;
+
+	FString DataToImport;
 };
 
 UCLASS(hidecategories=Object)
@@ -59,8 +61,8 @@ class UNREALED_API UCSVImportFactory : public UFactory, public IImportSettingsPa
 	virtual UObject* FactoryCreateText(UClass* InClass, UObject* InParent, FName InName, EObjectFlags Flags, UObject* Context, const TCHAR* Type, const TCHAR*& Buffer, const TCHAR* BufferEnd, FFeedbackContext* Warn, bool& bOutOperationCanceled) override;
 	virtual bool DoesSupportClass(UClass * Class) override;
 	virtual bool FactoryCanImport(const FString& Filename) override;
-	virtual	IImportSettingsParser* GetImportSettingsParser() override { return this; }
-
+	virtual	IImportSettingsParser* GetImportSettingsParser() override;
+	virtual void CleanUp() override;
 
 	/* Reimport an object that was created based on a CSV */
 	EReimportResult::Type ReimportCSV(UObject* Obj);
@@ -71,20 +73,22 @@ class UNREALED_API UCSVImportFactory : public UFactory, public IImportSettingsPa
 	virtual void ParseFromJson(TSharedRef<class FJsonObject> ImportSettingsJson) override;
 
 protected:
-	virtual TArray<FString> DoImportDataTable(class UDataTable* TargetDataTable, const FString& DataToImport);
-	virtual TArray<FString> DoImportCurveTable(class UCurveTable* TargetCurveTable, const FString& DataToImport, const ERichCurveInterpMode ImportCurveInterpMode);
-	virtual TArray<FString> DoImportCurve(class UCurveBase* TargetCurve, const FString& DataToImport);
+	virtual TArray<FString> DoImportDataTable(const FCSVImportSettings& ImportSettings, class UDataTable* TargetDataTable);
+	virtual TArray<FString> DoImportCurveTable(const FCSVImportSettings& ImportSettings, class UCurveTable* TargetCurveTable);
+	virtual TArray<FString> DoImportCurve(const FCSVImportSettings& InImportSettings, class UCurveBase* TargetCurve);
 
 private:
 	/* Reimport object from the given path*/
 	EReimportResult::Type Reimport(UObject* Obj, const FString& Path);
 
+	bool bImportAll = false;
+
 public:
 	UPROPERTY(BlueprintReadWrite, Category="Automation")
 	FCSVImportSettings AutomatedImportSettings;
 
 	/** Temporary data table to use to display import options */
 	UPROPERTY()
-	UDataTable* TempImportDataTable;
+	UDataTable* DataTableImportOptions;
 };
 
diff --git a/Engine/Source/Editor/UnrealEd/Private/Factories/CSVImportFactory.cpp b/Engine/Source/Editor/UnrealEd/Private/Factories/CSVImportFactory.cpp
--- a/Engine/Source/Editor/UnrealEd/Private/Factories/CSVImportFactory.cpp
+++ b/Engine/Source/Editor/UnrealEd/Private/Factories/CSVImportFactory.cpp
@@ -98,6 +98,19 @@ bool UCSVImportFactory::FactoryCanImport(const FString& Filename)
 	return false;
 }
 
+IImportSettingsParser* UCSVImportFactory::GetImportSettingsParser()
+{
+	return this;
+}
+
+void UCSVImportFactory::CleanUp()
+{
+	Super::CleanUp();
+	
+	bImportAll = false;
+	DataTableImportOptions = nullptr;
+}
+
 UObject* UCSVImportFactory::FactoryCreateText(UClass* InClass, UObject* InParent, FName InName, EObjectFlags Flags, UObject* Context, const TCHAR* Type, const TCHAR*& Buffer, const TCHAR* BufferEnd, FFeedbackContext* Warn, bool& bOutOperationCanceled)
 {
 	GEditor->GetEditorSubsystem<UImportSubsystem>()->BroadcastAssetPreImport(this, InClass, InParent, InName, Type);
@@ -109,53 +122,53 @@ UObject* UCSVImportFactory::FactoryCreateText(UClass* InClass, UObject* InParent
 	UCurveTable* ExistingCurveTable = FindObject<UCurveTable>(InParent, *InName.ToString());
 	UCurveBase* ExistingCurve = FindObject<UCurveBase>(InParent, *InName.ToString());
 
-	// Save off information if so
-	bool bHaveInfo = false;
-	ERichCurveInterpMode ImportCurveInterpMode = RCIM_Linear;
-	ECSVImportType ImportType = ECSVImportType::ECSV_DataTable;
-	UClass* DataTableClass = UDataTable::StaticClass();
+	FCSVImportSettings ImportSettings;
 
-	// Clear our temp table
-	TempImportDataTable = nullptr;
+	bool bHaveInfo = false;
 
 	if (IsAutomatedImport())
 	{
-		ImportCurveInterpMode = AutomatedImportSettings.ImportCurveInterpMode;
-		ImportType = AutomatedImportSettings.ImportType;
-
-		TempImportDataTable = NewObject<UDataTable>(this, UDataTable::StaticClass(), InName, Flags);
-		TempImportDataTable->RowStruct = AutomatedImportSettings.ImportRowStruct;
+		ImportSettings = AutomatedImportSettings;
 
 		// For automated import to work a row struct must be specified for a datatable type or a curve type must be specified
-		bHaveInfo = TempImportDataTable->RowStruct != nullptr || ImportType != ECSVImportType::ECSV_DataTable;
+		bHaveInfo = ImportSettings.ImportRowStruct != nullptr || ImportSettings.ImportType != ECSVImportType::ECSV_DataTable;
 	}
 	else if (ExistingTable != nullptr)
 	{
-		ImportType = ECSVImportType::ECSV_DataTable;
-		TempImportDataTable = NewObject<UDataTable>(this, ExistingTable->GetClass(), InName, Flags);
-		TempImportDataTable->CopyImportOptions(ExistingTable);
+		ImportSettings.ImportType = ECSVImportType::ECSV_DataTable;
+		ImportSettings.ImportRowStruct = ExistingTable->RowStruct;
 		
 		bHaveInfo = true;
 	}
 	else if (ExistingCurveTable != nullptr)
 	{
-		ImportType = ECSVImportType::ECSV_CurveTable;
+		ImportSettings.ImportType = ECSVImportType::ECSV_CurveTable;
 		bHaveInfo = true;
 	}
 	else if (ExistingCurve != nullptr)
 	{
-		ImportType = ExistingCurve->IsA(UCurveFloat::StaticClass()) ? ECSVImportType::ECSV_CurveFloat : ECSVImportType::ECSV_CurveVector;
+		if (ExistingCurve->IsA(UCurveFloat::StaticClass()))
+		{
+			ImportSettings.ImportType = ECSVImportType::ECSV_CurveFloat;
+		}
+		else if (ExistingCurve->IsA(UCurveLinearColor::StaticClass()))
+		{
+			ImportSettings.ImportType = ECSVImportType::ECSV_CurveLinearColor;
+		}
+		else
+		{
+			ImportSettings.ImportType = ECSVImportType::ECSV_CurveVector;
+		}
 		bHaveInfo = true;
 	}
-
-	bool bDoImport = true;
-
-	if (!TempImportDataTable)
+	else if (bImportAll)
 	{
-		// Create an empty one
-		TempImportDataTable = NewObject<UDataTable>(this, UDataTable::StaticClass(), InName, Flags);
+		ImportSettings = AutomatedImportSettings;
+		bHaveInfo = true;
 	}
 
+	bool bDoImport = true;
+
 	// If we do not have the info we need, pop up window to ask for things
 	if (!bHaveInfo && !IsAutomatedImport())
 	{
@@ -180,25 +193,34 @@ UObject* UCSVImportFactory::FactoryCreateText(UClass* InClass, UObject* InParent
 			ParentFullPath = InParent->GetPathName();
 		}
 
+		DataTableImportOptions = NewObject<UDataTable>(this, UDataTable::StaticClass(), InName, Flags);
+
 		Window->SetContent
 		(
 			SAssignNew(ImportOptionsWindow, SCSVImportOptions)
 				.WidgetWindow(Window)
 				.FullPath(FText::FromString(ParentFullPath))
-				.TempImportDataTable(TempImportDataTable)
+				.TempImportDataTable(DataTableImportOptions)
 		);
 
 		FSlateApplication::Get().AddModalWindow(Window, ParentWindow, false);
 
-		ImportType = ImportOptionsWindow->GetSelectedImportType();
-		TempImportDataTable->RowStruct = ImportOptionsWindow->GetSelectedRowStruct();
-		ImportCurveInterpMode = ImportOptionsWindow->GetSelectedCurveIterpMode();
+		ImportSettings.ImportType = ImportOptionsWindow->GetSelectedImportType();
+		ImportSettings.ImportRowStruct = ImportOptionsWindow->GetSelectedRowStruct();
+		ImportSettings.ImportCurveInterpMode = ImportOptionsWindow->GetSelectedCurveIterpMode();
+
 		bDoImport = ImportOptionsWindow->ShouldImport();
+		bImportAll = ImportOptionsWindow->ShouldImportAll();
 		bOutOperationCanceled = !bDoImport;
+
+		if (bImportAll)
+		{
+			AutomatedImportSettings = ImportSettings;
+		}
 	}
 	else if (!bHaveInfo && IsAutomatedImport())
 	{
-		if (ImportType == ECSVImportType::ECSV_DataTable && !TempImportDataTable->RowStruct)
+		if (ImportSettings.ImportType == ECSVImportType::ECSV_DataTable && !ImportSettings.ImportRowStruct)
 		{
 			UE_LOG(LogCSVImportFactory, Error, TEXT("A Data table row type must be specified in the import settings json file for automated import"));
 		}
@@ -209,23 +231,24 @@ UObject* UCSVImportFactory::FactoryCreateText(UClass* InClass, UObject* InParent
 	if (bDoImport)
 	{
 		// Convert buffer to an FString (will this be slow with big tables?)
-		FString String;
-		//const int32 BufferSize = BufferEnd - Buffer;
-		//appBufferToString( String, Buffer, BufferSize );
 		int32 NumChars = (BufferEnd - Buffer);
-		TArray<TCHAR>& StringChars = String.GetCharArray();
+		TArray<TCHAR>& StringChars = ImportSettings.DataToImport.GetCharArray();
 		StringChars.AddUninitialized(NumChars+1);
 		FMemory::Memcpy(StringChars.GetData(), Buffer, NumChars*sizeof(TCHAR));
 		StringChars.Last() = 0;
 
 		TArray<FString> Problems;
 
-		if (ImportType == ECSVImportType::ECSV_DataTable)
+		if (ImportSettings.ImportType == ECSVImportType::ECSV_DataTable)
 		{
+			UClass* DataTableClass = UDataTable::StaticClass();
+			UDataTable* TempImportDataTable = NewObject<UDataTable>(this, DataTableClass, InName, Flags);
+
 			// If there is an existing table, need to call this to free data memory before recreating object
 			UDataTable::FOnDataTableChanged OldOnDataTableChanged;
 			if (ExistingTable != nullptr)
 			{
+				TempImportDataTable->CopyImportOptions(ExistingTable);
 				OldOnDataTableChanged = MoveTemp(ExistingTable->OnDataTableChanged());
 				ExistingTable->OnDataTableChanged().Clear();
 				DataTableClass = ExistingTable->GetClass();
@@ -234,12 +257,13 @@ UObject* UCSVImportFactory::FactoryCreateText(UClass* InClass, UObject* InParent
 
 			// Create/reset table
 			UDataTable* NewTable = NewObject<UDataTable>(InParent, DataTableClass, InName, Flags);
+			TempImportDataTable->RowStruct = ImportSettings.ImportRowStruct;
 			
 			NewTable->CopyImportOptions(TempImportDataTable);
 			NewTable->AssetImportData->Update(CurrentFilename);
 
 			// Go ahead and create table from string
-			Problems = DoImportDataTable(NewTable, String);
+			Problems = DoImportDataTable(ImportSettings, NewTable);
 
 			// hook delegates back up and inform listeners of changes
 			NewTable->OnDataTableChanged() = MoveTemp(OldOnDataTableChanged);
@@ -249,7 +273,7 @@ UObject* UCSVImportFactory::FactoryCreateText(UClass* InClass, UObject* InParent
 			UE_LOG(LogCSVImportFactory, Log, TEXT("Imported DataTable '%s' - %d Problems"), *InName.ToString(), Problems.Num());
 			NewAsset = NewTable;
 		}
-		else if (ImportType == ECSVImportType::ECSV_CurveTable)
+		else if (ImportSettings.ImportType == ECSVImportType::ECSV_CurveTable)
 		{
 			UClass* CurveTableClass = UCurveTable::StaticClass();
 
@@ -268,7 +292,7 @@ UObject* UCSVImportFactory::FactoryCreateText(UClass* InClass, UObject* InParent
 			NewTable->AssetImportData->Update(CurrentFilename);
 
 			// Go ahead and create table from string
-			Problems = DoImportCurveTable(NewTable, String, ImportCurveInterpMode);
+			Problems = DoImportCurveTable(ImportSettings, NewTable);
 
 			// hook delegates back up and inform listeners of changes
 			NewTable->OnCurveTableChanged() = MoveTemp(OldOnCurveTableChanged);
@@ -278,14 +302,11 @@ UObject* UCSVImportFactory::FactoryCreateText(UClass* InClass, UObject* InParent
 			UE_LOG(LogCSVImportFactory, Log, TEXT("Imported CurveTable '%s' - %d Problems"), *InName.ToString(), Problems.Num());
 			NewAsset = NewTable;
 		}
-		else if (ImportType == ECSVImportType::ECSV_CurveFloat || ImportType == ECSVImportType::ECSV_CurveVector || ImportType == ECSVImportType::ECSV_CurveLinearColor)
+		else if (ImportSettings.ImportType == ECSVImportType::ECSV_CurveFloat || ImportSettings.ImportType == ECSVImportType::ECSV_CurveVector || ImportSettings.ImportType == ECSVImportType::ECSV_CurveLinearColor)
 		{
-			UClass* CurveClass = ExistingCurve ? ExistingCurve->GetClass() : GetCurveClass( ImportType );
-
 			// Create/reset curve
-			UCurveBase* NewCurve = NewObject<UCurveBase>(InParent, CurveClass, InName, Flags);
-
-			Problems = DoImportCurve(NewCurve, String);
+			UCurveBase* NewCurve = NewObject<UCurveBase>(InParent, GetCurveClass(ImportSettings.ImportType), InName, Flags);
+			Problems = DoImportCurve(ImportSettings, NewCurve);
 
 			UE_LOG(LogCSVImportFactory, Log, TEXT("Imported Curve '%s' - %d Problems"), *InName.ToString(), Problems.Num());
 			NewCurve->AssetImportData->Update(CurrentFilename);
@@ -364,31 +385,31 @@ EReimportResult::Type UCSVImportFactory::Reimport(UObject* Obj, const FString& P
 	return EReimportResult::Failed;
 }
 
-TArray<FString> UCSVImportFactory::DoImportDataTable(UDataTable* TargetDataTable, const FString& DataToImport)
+TArray<FString> UCSVImportFactory::DoImportDataTable(const FCSVImportSettings& InImportSettings, UDataTable* TargetDataTable)
 {
 	// Are we importing JSON data?
 	const bool bIsJSON = CurrentFilename.EndsWith(TEXT(".json"));
 	if (bIsJSON)
 	{
-		return TargetDataTable->CreateTableFromJSONString(DataToImport);
+		return TargetDataTable->CreateTableFromJSONString(InImportSettings.DataToImport);
 	}
 
-	return TargetDataTable->CreateTableFromCSVString(DataToImport);
+	return TargetDataTable->CreateTableFromCSVString(InImportSettings.DataToImport);
 }
 
-TArray<FString> UCSVImportFactory::DoImportCurveTable(UCurveTable* TargetCurveTable, const FString& DataToImport, const ERichCurveInterpMode InImportCurveInterpMode)
+TArray<FString> UCSVImportFactory::DoImportCurveTable(const FCSVImportSettings& InImportSettings, UCurveTable* TargetCurveTable)
 {
 	// Are we importing JSON data?
 	const bool bIsJSON = CurrentFilename.EndsWith(TEXT(".json"));
 	if (bIsJSON)
 	{
-		return TargetCurveTable->CreateTableFromJSONString(DataToImport, InImportCurveInterpMode);
+		return TargetCurveTable->CreateTableFromJSONString(InImportSettings.DataToImport, InImportSettings.ImportCurveInterpMode);
 	}
 
-	return TargetCurveTable->CreateTableFromCSVString(DataToImport, InImportCurveInterpMode);
+	return TargetCurveTable->CreateTableFromCSVString(InImportSettings.DataToImport, InImportSettings.ImportCurveInterpMode);
 }
 
-TArray<FString> UCSVImportFactory::DoImportCurve(UCurveBase* TargetCurve, const FString& DataToImport)
+TArray<FString> UCSVImportFactory::DoImportCurve(const FCSVImportSettings& InImportSettings, UCurveBase* TargetCurve)
 {
 	// Are we importing JSON data?
 	const bool bIsJSON = CurrentFilename.EndsWith(TEXT(".json"));
@@ -399,7 +420,7 @@ TArray<FString> UCSVImportFactory::DoImportCurve(UCurveBase* TargetCurve, const
 		return Result;
 	}
 
-	return TargetCurve->CreateCurveFromCSVString(DataToImport);
+	return TargetCurve->CreateCurveFromCSVString(InImportSettings.DataToImport);
 }
 
 //////////////////////////////////////////////////////////////////////////
diff --git a/Engine/Source/Editor/UnrealEd/Private/SCSVImportOptions.cpp b/Engine/Source/Editor/UnrealEd/Private/SCSVImportOptions.cpp
--- a/Engine/Source/Editor/UnrealEd/Private/SCSVImportOptions.cpp
+++ b/Engine/Source/Editor/UnrealEd/Private/SCSVImportOptions.cpp
@@ -27,6 +27,7 @@ void SCSVImportOptions::Construct(const FArguments& InArgs)
 	ImportTypes.Add(MakeShareable(new ECSVImportType(ECSVImportType::ECSV_CurveTable)));
 	ImportTypes.Add(MakeShareable(new ECSVImportType(ECSVImportType::ECSV_CurveFloat)));
 	ImportTypes.Add(MakeShareable(new ECSVImportType(ECSVImportType::ECSV_CurveVector)));
+	ImportTypes.Add(MakeShareable(new ECSVImportType(ECSVImportType::ECSV_CurveLinearColor)));
 
 	// Create properties view
 	FPropertyEditorModule & EditModule = FModuleManager::Get().GetModuleChecked<FPropertyEditorModule>("PropertyEditor");
@@ -154,7 +155,7 @@ void SCSVImportOptions::Construct(const FArguments& InArgs)
 					PropertyView.ToSharedRef()
 				]
 			]
-			// Ok/Cancel
+			// Apply/Apply to All/Cancel
 			+SVerticalBox::Slot()
 			.AutoHeight()
 			[
@@ -164,13 +165,22 @@ void SCSVImportOptions::Construct(const FArguments& InArgs)
 				.Padding(2)
 				[
 					SNew(SButton)
-					.Text(LOCTEXT("OK", "OK"))
+					.Text(LOCTEXT("Import", "Apply"))
 					.OnClicked( this, &SCSVImportOptions::OnImport )
 					.IsEnabled( this, &SCSVImportOptions::CanImport )
 				]
 				+SHorizontalBox::Slot()
 				.AutoWidth()
 				.Padding(2)
+				[
+					SNew(SButton)
+					.Text(LOCTEXT("ImportAll", "Apply to All"))
+					.OnClicked( this, &SCSVImportOptions::OnImportAll )
+					.IsEnabled( this, &SCSVImportOptions::CanImport )
+				]
+				+SHorizontalBox::Slot()
+				.AutoWidth()
+				.Padding(2)
 				[
 					SNew(SButton)
 					.Text(LOCTEXT("Cancel", "Cancel"))
@@ -197,7 +207,12 @@ void SCSVImportOptions::Construct(const FArguments& InArgs)
 	/** If we should import */
 bool SCSVImportOptions::ShouldImport()
 {
-	return ((SelectedStruct != nullptr) || GetSelectedImportType() != ECSVImportType::ECSV_DataTable) && bImport;
+	return ((SelectedStruct != nullptr) || GetSelectedImportType() != ECSVImportType::ECSV_DataTable) && (UserDlgResponse != ECSVImportOptionDlgResponse::Cancel);
+}
+
+bool SCSVImportOptions::ShouldImportAll()
+{
+	return UserDlgResponse == ECSVImportOptionDlgResponse::ImportAll;
 }
 
 /** Get the row struct we selected */
@@ -254,6 +269,10 @@ FString SCSVImportOptions::GetImportTypeText(TSharedPtr<ECSVImportType> Type) co
 	{
 		EnumString = TEXT("Vector Curve");
 	}
+	else if (*Type == ECSVImportType::ECSV_CurveLinearColor)
+	{
+		EnumString = TEXT("Linear Color Curve");
+	}
 	return EnumString;
 }
 
@@ -312,17 +331,8 @@ TSharedRef<SWidget> SCSVImportOptions::MakeCurveTypeWidget(CurveInterpModePtr In
 /** Called when 'OK' button is pressed */
 FReply SCSVImportOptions::OnImport()
 {
-	SelectedImportType = *ImportTypeCombo->GetSelectedItem();
-	if (CurveInterpCombo->GetSelectedItem().IsValid())
-	{
-		SelectedCurveInterpMode = *CurveInterpCombo->GetSelectedItem();
-	}
-	bImport = true;
-	if (WidgetWindow.IsValid())
-	{
-		WidgetWindow.Pin()->RequestDestroyWindow();
-	}
-	return FReply::Handled();
+	UserDlgResponse = ECSVImportOptionDlgResponse::Import;
+	return HandleImport();
 }
 
 bool SCSVImportOptions::CanImport() const
@@ -349,7 +359,7 @@ bool SCSVImportOptions::CanImport() const
 /** Called when 'Cancel' button is pressed */
 FReply SCSVImportOptions::OnCancel()
 {
-	bImport = false;
+	UserDlgResponse = ECSVImportOptionDlgResponse::Cancel;
 	if (WidgetWindow.IsValid())
 	{
 		WidgetWindow.Pin()->RequestDestroyWindow();
@@ -374,4 +384,24 @@ FText SCSVImportOptions::GetSelectedCurveTypeText() const
 		: FText::GetEmpty();
 }
 
+FReply SCSVImportOptions::HandleImport()
+{
+	SelectedImportType = *ImportTypeCombo->GetSelectedItem();
+	if (CurveInterpCombo->GetSelectedItem().IsValid())
+	{
+		SelectedCurveInterpMode = *CurveInterpCombo->GetSelectedItem();
+	}
+	if (WidgetWindow.IsValid())
+	{
+		WidgetWindow.Pin()->RequestDestroyWindow();
+	}
+	return FReply::Handled();
+}
+
+FReply SCSVImportOptions::OnImportAll()
+{
+	UserDlgResponse = ECSVImportOptionDlgResponse::ImportAll;
+	return HandleImport();
+}
+
 #undef LOCTEXT_NAMESPACE
diff --git a/Engine/Source/Editor/UnrealEd/Public/SCSVImportOptions.h b/Engine/Source/Editor/UnrealEd/Public/SCSVImportOptions.h
--- a/Engine/Source/Editor/UnrealEd/Public/SCSVImportOptions.h
+++ b/Engine/Source/Editor/UnrealEd/Public/SCSVImportOptions.h
@@ -15,6 +15,13 @@
 #include "Widgets/SCompoundWidget.h"
 #include "Widgets/Input/SComboBox.h"
 
+enum class ECSVImportOptionDlgResponse : uint8
+{
+	Import,
+	ImportAll,
+	Cancel
+};
+
 class UNREALED_API SCSVImportOptions : public SCompoundWidget
 {
 private:
@@ -34,7 +41,7 @@ public:
 	SLATE_END_ARGS()
 
 	SCSVImportOptions()
-		: bImport(false)
+		: UserDlgResponse(ECSVImportOptionDlgResponse::Cancel)
 		, SelectedImportType(ECSVImportType::ECSV_DataTable)
 		, SelectedStruct(nullptr)
 		, TempImportDataTable(nullptr)
@@ -45,6 +52,9 @@ public:
 	/** If we should import */
 	bool ShouldImport();
 
+	/** If the current settings should be applied to all items being imported */
+	bool ShouldImportAll();
+
 	/** Get the row struct we selected */
 	UScriptStruct* GetSelectedRowStruct();
 
@@ -79,7 +89,7 @@ public:
 	/** Called to create a widget for each curve interpolation enum */
 	TSharedRef<SWidget> MakeCurveTypeWidget(CurveInterpModePtr InterpMode);
 
-	/** Called when 'OK' button is pressed */
+	/** Called when 'Apply' button is pressed */
 	FReply OnImport();
 
 	/** Do we have all of the data we need to import this asset? */
@@ -93,8 +103,11 @@ public:
 	FText GetSelectedCurveTypeText() const;
 
 private:
+	FReply HandleImport();
+	FReply OnImportAll();
+
 	/** Whether we should go ahead with import */
-	uint8										bImport : 1;
+	ECSVImportOptionDlgResponse					UserDlgResponse : ECSVImportOptionDlgResponse::Cancel;
 
 	/** Window that owns us */
 	TWeakPtr< SWindow >							WidgetWindow;
