Fix Linux compile errors
#rnx

diff --git a/Engine/Source/Editor/UnrealEd/Public/SCSVImportOptions.h b/Engine/Source/Editor/UnrealEd/Public/SCSVImportOptions.h
--- a/Engine/Source/Editor/UnrealEd/Public/SCSVImportOptions.h
+++ b/Engine/Source/Editor/UnrealEd/Public/SCSVImportOptions.h
@@ -107,7 +107,7 @@ private:
 	FReply OnImportAll();
 
 	/** Whether we should go ahead with import */
-	ECSVImportOptionDlgResponse					UserDlgResponse : ECSVImportOptionDlgResponse::Cancel;
+	ECSVImportOptionDlgResponse					UserDlgResponse;
 
 	/** Window that owns us */
 	TWeakPtr< SWindow >							WidgetWindow;
