Fix multiply defined constants used by FWidget

#rnx

diff --git a/Engine/Source/Editor/UnrealEd/Private/UnrealWidget.cpp b/Engine/Source/Editor/UnrealEd/Private/UnrealWidget.cpp
--- a/Engine/Source/Editor/UnrealEd/Private/UnrealWidget.cpp
+++ b/Engine/Source/Editor/UnrealEd/Private/UnrealWidget.cpp
@@ -15,7 +15,15 @@
 
 IMPLEMENT_HIT_PROXY(HWidgetAxis, HHitProxy);
 
-static const float INNER_AXIS_CIRCLE_RADIUS = 48.0f;
+constexpr float FWidget::AXIS_LENGTH;
+constexpr float FWidget::TRANSLATE_ROTATE_AXIS_CIRCLE_RADIUS;
+constexpr float FWidget::TWOD_AXIS_CIRCLE_RADIUS;
+constexpr float FWidget::INNER_AXIS_CIRCLE_RADIUS;
+constexpr float FWidget::OUTER_AXIS_CIRCLE_RADIUS;
+constexpr float FWidget::ROTATION_TEXT_RADIUS;
+constexpr int32 FWidget::AXIS_CIRCLE_SIDES;
+constexpr float FWidget::ARCALL_RELATIVE_INNER_SIZE;
+constexpr float FWidget::AXIS_LENGTH_SCALE_OFFSET;
 
 FWidget::FWidget()
 {
diff --git a/Engine/Source/Editor/UnrealEd/Private/UnrealWidgetRender.cpp b/Engine/Source/Editor/UnrealEd/Private/UnrealWidgetRender.cpp
--- a/Engine/Source/Editor/UnrealEd/Private/UnrealWidgetRender.cpp
+++ b/Engine/Source/Editor/UnrealEd/Private/UnrealWidgetRender.cpp
@@ -12,16 +12,6 @@
 #include "DynamicMeshBuilder.h"
 #include "CanvasTypes.h"
 
-static const float AXIS_LENGTH = 35.0f;
-static const float TRANSLATE_ROTATE_AXIS_CIRCLE_RADIUS = 20.0f;
-static const float TWOD_AXIS_CIRCLE_RADIUS = 10.0f;
-static const float INNER_AXIS_CIRCLE_RADIUS = 48.0f;
-static const float OUTER_AXIS_CIRCLE_RADIUS = 56.0f;
-static const float ROTATION_TEXT_RADIUS = 75.0f;
-static const int32 AXIS_CIRCLE_SIDES = 24;
-static const float ARCALL_RELATIVE_INNER_SIZE = 0.75f;
-static const float AXIS_LENGTH_SCALE_OFFSET = 5.0f;
-
 /*
  *  Simple struct used to create and group data related to the current window's / viewport's space,
  *  orientation, and scale.
diff --git a/Engine/Source/Editor/UnrealEd/Public/UnrealWidget.h b/Engine/Source/Editor/UnrealEd/Public/UnrealWidget.h
--- a/Engine/Source/Editor/UnrealEd/Public/UnrealWidget.h
+++ b/Engine/Source/Editor/UnrealEd/Public/UnrealWidget.h
@@ -27,6 +27,16 @@ enum ECoordSystem
 class FWidget : public FGCObject
 {
 public:
+	constexpr static float AXIS_LENGTH                         = 35.0f;
+	constexpr static float TRANSLATE_ROTATE_AXIS_CIRCLE_RADIUS = 20.0f;
+	constexpr static float TWOD_AXIS_CIRCLE_RADIUS             = 10.0f;
+	constexpr static float INNER_AXIS_CIRCLE_RADIUS            = 48.0f;
+	constexpr static float OUTER_AXIS_CIRCLE_RADIUS            = 56.0f;
+	constexpr static float ROTATION_TEXT_RADIUS                = 75.0f;
+	constexpr static int32 AXIS_CIRCLE_SIDES                   = 24;
+	constexpr static float ARCALL_RELATIVE_INNER_SIZE          = 0.75f;
+	constexpr static float AXIS_LENGTH_SCALE_OFFSET            = 5.0f;
+
 	enum EWidgetMode
 	{
 		WM_None = -1,
