Fix bad merge

diff --git a/Engine/Source/Editor/UnrealEd/Private/FileHelpers.cpp b/Engine/Source/Editor/UnrealEd/Private/FileHelpers.cpp
--- a/Engine/Source/Editor/UnrealEd/Private/FileHelpers.cpp
+++ b/Engine/Source/Editor/UnrealEd/Private/FileHelpers.cpp
@@ -4291,7 +4291,6 @@ static bool InternalCheckoutAndSavePackages(const TArray<UPackage*>& PackagesToS
 			for (UPackage* Package : PackagesToSave)
 			{
 				// List unsaved packages that were not checked out
-				if (Package->HasAnyPackageFlags(PKG_NewlyCreated) && !PackagesCheckedOut.Contains(Package))
 				if (!PackagesCheckedOut.Contains(Package))
 				{
 					PackagesToMarkForAdd.Add(Package);
