Detect duplicated level built data GUIDs and sanitize them. This was causing multiple nasty crashes in VLM building code.

diff --git a/Engine/Source/Editor/UnrealEd/Private/StaticLightingSystem/StaticLightingSystem.cpp b/Engine/Source/Editor/UnrealEd/Private/StaticLightingSystem/StaticLightingSystem.cpp
--- a/Engine/Source/Editor/UnrealEd/Private/StaticLightingSystem/StaticLightingSystem.cpp
+++ b/Engine/Source/Editor/UnrealEd/Private/StaticLightingSystem/StaticLightingSystem.cpp
@@ -522,6 +522,32 @@ bool FStaticLightingSystem::BeginLightmassProcess()
 			Component->VisibilityId = INDEX_NONE;
 		}
 
+		{
+			TMap<FGuid, ULevel*> LevelGuids;
+
+			for (int32 LevelIndex = 0; LevelIndex < World->GetNumLevels(); LevelIndex++)
+			{
+				ULevel* Level = World->GetLevel(LevelIndex);
+
+				if (ShouldOperateOnLevel(Level) && Options.ShouldBuildLightingForLevel(Level))
+				{
+					if (LevelGuids.Contains(Level->LevelBuildDataId))
+					{
+						FMessageLog("LightingResults").Warning()
+							->AddToken(FUObjectToken::Create(LevelGuids[Level->LevelBuildDataId]->GetOuter()))
+							->AddToken(FTextToken::Create(LOCTEXT("LightmassError_DuplicatedLevelGuids1", "has the same level built data GUID as")))
+							->AddToken(FUObjectToken::Create(Level->GetOuter()))
+							->AddToken(FTextToken::Create(LOCTEXT("LightmassError_DuplicatedLevelGuids2", ". A new GUID is assigned to the later one. All previously built lighting is invalidated and the level needs to be resaved.")));
+
+						Level->LevelBuildDataId = FGuid::NewGuid();
+						Level->MarkPackageDirty();
+					}
+
+					LevelGuids.Add(Level->LevelBuildDataId, Level);
+				}
+			}
+		}
+
 		FString SkippedLevels;
 		for ( int32 LevelIndex=0; LevelIndex < World->GetNumLevels(); LevelIndex++ )
 		{
