Make sure WillNeverCacheCookedPlatformDataAgain is called after Save is finished, rather than just after caching. This was causing a sporadic issue where we would try to reload USoundWave::RawData in USoundWave::Serialize during save after clearing it in WillNeverCacheCookedPlatformDataAgain. The rarity of this issue was likely caused becaues this would only happen if we hit a DDC miss for the compressed audio during save, and we place the compressed data in the DDC during load.

diff --git a/Engine/Source/Editor/UnrealEd/Private/CookOnTheFlyServer.cpp b/Engine/Source/Editor/UnrealEd/Private/CookOnTheFlyServer.cpp
--- a/Engine/Source/Editor/UnrealEd/Private/CookOnTheFlyServer.cpp
+++ b/Engine/Source/Editor/UnrealEd/Private/CookOnTheFlyServer.cpp
@@ -3047,19 +3047,6 @@ bool UCookOnTheFlyServer::FinishPackageCacheForCookedPlatformData(UPackage* Pack
 		}
 	}
 
-	if (CurrentCookMode == ECookMode::CookByTheBook)
-	{
-		// For each object for which data is cached we can call FinishedCookedPlatformDataCache
-		// we can only safely call this when we are finished caching the object completely.
-		// this doesn't ever happen for cook in editor or cook on the fly mode
-		for (UObject* Obj : CurrentReentryData.CachedObjectsInOuter)
-		{
-			check(!IsCookingInEditor());
-			// this might be run multiple times for a single object
-			Obj->WillNeverCacheCookedPlatformDataAgain();
-		}
-	}
-
 	// all these objects have finished so release their async begincache back to the pool
 	for (const auto& FinishedCached : CurrentReentryData.BeginCacheCallCount )
 	{
@@ -3496,6 +3483,14 @@ void UCookOnTheFlyServer::SaveCookedPackages(
 				for (UObject* Object : ObjectsInPackage)
 				{
 					Object->ClearAllCachedCookedPlatformData();
+					if (CurrentCookMode == ECookMode::CookByTheBook)
+					{
+						// For each object for which data is cached we can call FinishedCookedPlatformDataCache
+						// we can only safely call this when we are finished caching the object completely.
+						// this doesn't ever happen for cook in editor or cook on the fly mode
+						// this might be run multiple times for a single object
+						Object->WillNeverCacheCookedPlatformDataAgain();
+					}
 				}
 			}
 
