Set temporary AssetRegistry for populating PackageNameCache while cooking DLC

diff --git a/Engine/Source/Editor/UnrealEd/Private/CookOnTheFlyServer.cpp b/Engine/Source/Editor/UnrealEd/Private/CookOnTheFlyServer.cpp
--- a/Engine/Source/Editor/UnrealEd/Private/CookOnTheFlyServer.cpp
+++ b/Engine/Source/Editor/UnrealEd/Private/CookOnTheFlyServer.cpp
@@ -7716,6 +7716,13 @@ void UCookOnTheFlyServer::StartCookByTheBook( const FCookByTheBookStartupOptions
 	
 	if ( IsCookingDLC() )
 	{
+		IAssetRegistry* CacheAssetRegistry = PackageNameCache->GetAssetRegistry();
+		if (CacheAssetRegistry == nullptr)
+		{
+			UE_LOG(LogCook, Log, TEXT("Temporarily Replacing PackageNameCache Asset Registry with the CookOnTheFlyServer's AssetRegistry to initialise Cache"));
+			PackageNameCache->SetAssetRegistry(AssetRegistry);
+		}
+
 		// if we are cooking dlc we must be based on a release version cook
 		check( !BasedOnReleaseVersion.IsEmpty() );
 
@@ -7765,6 +7772,8 @@ void UCookOnTheFlyServer::StartCookByTheBook( const FCookByTheBookStartupOptions
 			}
 			CookByTheBookOptions->BasedOnReleaseCookedPackages.Add(PlatformName, MoveTemp(PackageList));
 		}
+
+		PackageNameCache->SetAssetRegistry(CacheAssetRegistry);
 	}
 	
 	// don't resave the global shader map files in dlc
@@ -8249,7 +8258,7 @@ bool UCookOnTheFlyServer::GetAllPackageFilenamesFromAssetRegistry( const FString
 
 		for (const TPair<FName, const FAssetData*>& RegistryData : RegistryDataMap)
 		{
-			int32 AddedIndex = PackageNames.Add(RegistryData.Value->ObjectPath);
+			int32 AddedIndex = PackageNames.Add(RegistryData.Value->PackageName);
 			if (PackageNameCache->Contains(PackageNames.Last()))
 			{
 				OutPackageFilenames[AddedIndex] = PackageNameCache->GetCachedStandardPackageFileFName(PackageNames.Last());
@@ -8273,7 +8282,7 @@ bool UCookOnTheFlyServer::GetAllPackageFilenamesFromAssetRegistry( const FString
 				FName StandardFileFName;
 				if (!PackageNameCache->CalculateCacheData(PackageName, PackageFilename, StandardFilename, StandardFileFName))
 				{
-					UE_LOG(LogCook, Warning, TEXT("Could not resolve package %s from %s"), *PackageName.ToString(), *AssetRegistryPath);
+					//UE_LOG(LogCook, Warning, TEXT("Could not resolve package %s from %s"), *PackageName.ToString(), *AssetRegistryPath);
 				}
 
 				OutPackageFilenames[AssetIndex] = StandardFileFName;
diff --git a/Engine/Source/Editor/UnrealEd/Private/Cooker/PackageNameCache.h b/Engine/Source/Editor/UnrealEd/Private/Cooker/PackageNameCache.h
--- a/Engine/Source/Editor/UnrealEd/Private/Cooker/PackageNameCache.h
+++ b/Engine/Source/Editor/UnrealEd/Private/Cooker/PackageNameCache.h
@@ -51,7 +51,10 @@ struct FPackageNameCache
 	bool			CalculateCacheData(FName PackageName, FString& OutPackageFilename, FString& OutStandardFilename, FName& OutStandardFileFName) const;
 
 	bool			Contains(FName PackageName) const;
-	
+
+	IAssetRegistry* GetAssetRegistry() const;
+	void SetAssetRegistry(IAssetRegistry* NewRegistry);
+
 private:
 	bool DoesPackageExist(const FName& PackageName, FString* OutFilename) const;
 	const FCachedPackageFilename& Cache(const FName& PackageName) const;
@@ -154,6 +157,16 @@ bool FPackageNameCache::Contains(FName PackageName) const
 	return PackageFilenameCache.Contains(PackageName);
 }
 
+IAssetRegistry* FPackageNameCache::GetAssetRegistry() const
+{
+	return AssetRegistry;
+}
+
+void FPackageNameCache::SetAssetRegistry(IAssetRegistry* NewRegistry)
+{
+	 AssetRegistry = NewRegistry;
+}
+
 const FCachedPackageFilename& FPackageNameCache::Cache(const FName& PackageName) const
 {
 	check(IsInGameThread());
