Re-add log removed by accident

diff --git a/Engine/Source/Editor/UnrealEd/Private/CookOnTheFlyServer.cpp b/Engine/Source/Editor/UnrealEd/Private/CookOnTheFlyServer.cpp
--- a/Engine/Source/Editor/UnrealEd/Private/CookOnTheFlyServer.cpp
+++ b/Engine/Source/Editor/UnrealEd/Private/CookOnTheFlyServer.cpp
@@ -8282,7 +8282,7 @@ bool UCookOnTheFlyServer::GetAllPackageFilenamesFromAssetRegistry( const FString
 				FName StandardFileFName;
 				if (!PackageNameCache->CalculateCacheData(PackageName, PackageFilename, StandardFilename, StandardFileFName))
 				{
-					//UE_LOG(LogCook, Warning, TEXT("Could not resolve package %s from %s"), *PackageName.ToString(), *AssetRegistryPath);
+					UE_LOG(LogCook, Warning, TEXT("Could not resolve package %s from %s"), *PackageName.ToString(), *AssetRegistryPath);
 				}
 
 				OutPackageFilenames[AssetIndex] = StandardFileFName;
