Fix crash when alternating cooking platforms (UE-92633).

- Ideological counterpart of CL 12872233 for UE-91944 from Dev-Core for 4.25

diff --git a/Engine/Source/Editor/UnrealEd/Private/CookOnTheFlyServer.cpp b/Engine/Source/Editor/UnrealEd/Private/CookOnTheFlyServer.cpp
--- a/Engine/Source/Editor/UnrealEd/Private/CookOnTheFlyServer.cpp
+++ b/Engine/Source/Editor/UnrealEd/Private/CookOnTheFlyServer.cpp
@@ -3947,13 +3947,19 @@ void UCookOnTheFlyServer::MarkPackageDirtyForCooker( UPackage *Package )
 			{
 				TArray<const ITargetPlatform*> CookedPlatforms;
 				// if we have already cooked this package and we have made changes then recook ;)
-				if (PackageTracker->CookedPackages.GetCookedPlatforms(PackageFFileName, CookedPlatforms) )
+				PackageTracker->CookedPackages.GetCookedPlatforms(PackageFFileName, CookedPlatforms);
+				bool bHadCookedPackages = CookedPlatforms.Num() != 0;
+				PackageTracker->CookedPackages.RemoveFile(PackageFFileName);
+				if (bHadCookedPackages)
 				{
 					if (IsCookByTheBookRunning())
 					{
 						// if this package was previously cooked and we are doing a cook by the book 
 						// we need to recook this package before finishing cook by the book
-						PackageTracker->EnqueueUniqueCookRequest(FFilePlatformRequest(PackageFFileName, CookedPlatforms));
+
+						// Note - SessionPlatforms here instead of CookedPlatforms. The lock is not needed 
+						// because we have checked IsCookByTheBook and therefore we are single threaded.
+						PackageTracker->EnqueueUniqueCookRequest(FFilePlatformRequest(PackageFFileName, PlatformManager->GetSessionPlatforms()));
 					}
 					else
 					{
@@ -7491,7 +7497,6 @@ void UCookOnTheFlyServer::StartCookByTheBook( const FCookByTheBookStartupOptions
 	//force precache objects to refresh themselves before cooking anything
 	LastUpdateTick = INT_MAX;
 
-	CookByTheBookOptions->bRunning = true;
 	CookByTheBookOptions->bCancel = false;
 	CookByTheBookOptions->CookTime = 0.0f;
 	CookByTheBookOptions->CookStartTime = FPlatformTime::Seconds();
@@ -7518,6 +7523,9 @@ void UCookOnTheFlyServer::StartCookByTheBook( const FCookByTheBookStartupOptions
 	PlatformManager->SelectSessionPlatforms(TargetPlatforms, PackageTracker);
 	check(PlatformManager->GetSessionPlatforms().Num() == TargetPlatforms.Num());
 
+	// We want to set bRunning = true as early as possible, but it implies that session platforms have been selected so this is the earliest point we can set it
+	CookByTheBookOptions->bRunning = true;
+
 	RefreshPlatformAssetRegistries(TargetPlatforms);
 
 	if (CurrentCookMode == ECookMode::CookByTheBook)
