Change FPlatformMemoryStats.AvailableVirtual  on windows to be clamped to the remaining system virtual memory available, which is limited by the size of the PageFile.
Add CookSettings MemoryMinFreeVirtual to garbage collect when AvailableVirtual is too low.
Normalize the names of the 4 CookSettings.Memory settings variables.
Change boolean expression of the 4 memory values to be (freevirtual OR freephysical) AND (usedvirtual OR usedphysical), which seems like the most useful way to combine them.
#rnx

diff --git a/Engine/Config/BaseEditor.ini b/Engine/Config/BaseEditor.ini
--- a/Engine/Config/BaseEditor.ini
+++ b/Engine/Config/BaseEditor.ini
@@ -361,9 +361,9 @@ MaxMemoryAllowance=16384
 [CookSettings]
 PackagesPerGC=0
 IdleTimeToGC=20
-MaxMemoryAllowance=16384
+MemoryMaxUsedPhysical=16384
+MemoryMinFreePhysical=1024
 MaxConcurrentShaderJobs=8192
-MinReservedMemory=1024
 MaxDiffsToLog=5
 IgnoreHeaderDiffs=True
 
diff --git a/Engine/Source/Editor/UnrealEd/Classes/Commandlets/CookCommandlet.h b/Engine/Source/Editor/UnrealEd/Classes/Commandlets/CookCommandlet.h
--- a/Engine/Source/Editor/UnrealEd/Classes/Commandlets/CookCommandlet.h
+++ b/Engine/Source/Editor/UnrealEd/Classes/Commandlets/CookCommandlet.h
@@ -73,9 +73,6 @@ class UCookCommandlet
 	/** Cooks specified list of files */
 	bool CookByTheBook(const TArray<ITargetPlatform*>& Platforms, TArray<FString>& FilesInPath);
 
-	/** See if the cooker has exceeded max memory allowance in this case the cooker should force a garbage collection */
-	bool HasExceededMaxMemory(uint64 MaxMemoryAllowance) const;
-
 	/**	Process deferred commands */
 	void ProcessDeferredCommands();
 
diff --git a/Engine/Source/Editor/UnrealEd/Classes/CookOnTheSide/CookOnTheFlyServer.h b/Engine/Source/Editor/UnrealEd/Classes/CookOnTheSide/CookOnTheFlyServer.h
--- a/Engine/Source/Editor/UnrealEd/Classes/CookOnTheSide/CookOnTheFlyServer.h
+++ b/Engine/Source/Editor/UnrealEd/Classes/CookOnTheSide/CookOnTheFlyServer.h
@@ -139,12 +139,11 @@ private:
 	uint32 PackagesPerGC;
 	/** Amount of time that is allowed to be idle before forcing a garbage collect. Set to 0 to never force GC due to idle time */
 	double IdleTimeToGC;
-	/** Max memory the cooker should use before forcing a gc */
-	uint64 MaxMemoryAllowance;
-	/** Min memory before the cooker should partial gc */
-	uint64 MinMemoryBeforeGC;
-	/** If we have less then this much memory free then finish current task and kick off gc */
-	uint64 MinFreeMemory;
+	// Memory Limits for when to Collect Garbage
+	uint64 MemoryMaxUsedVirtual;
+	uint64 MemoryMaxUsedPhysical;
+	uint64 MemoryMinFreeVirtual;
+	uint64 MemoryMinFreePhysical;
 	/** Max number of packages to save before we partial gc */
 	int32 MaxNumPackagesBeforePartialGC;
 	/** Max number of concurrent shader jobs reducing this too low will increase cook time */
@@ -495,6 +494,7 @@ public:
 	double GetIdleTimeToGC() const;
 
 	/** Returns the configured amount of memory allowed before forcing a GC */
+	UE_DEPRECATED(4.26, "Use HasExceededMaxMemory instead")
 	uint64 GetMaxMemoryAllowance() const;
 
 	/** Mark package as keep around for the cooker (don't GC) */
diff --git a/Engine/Source/Editor/UnrealEd/Private/Commandlets/CookCommandlet.cpp b/Engine/Source/Editor/UnrealEd/Private/Commandlets/CookCommandlet.cpp
--- a/Engine/Source/Editor/UnrealEd/Private/Commandlets/CookCommandlet.cpp
+++ b/Engine/Source/Editor/UnrealEd/Private/Commandlets/CookCommandlet.cpp
@@ -903,7 +903,6 @@ bool UCookCommandlet::CookByTheBook( const TArray<ITargetPlatform*>& Platforms,
 
 		const uint32 PackagesPerGC = CookOnTheFlyServer->GetPackagesPerGC();
 		const double IdleTimeToGC = CookOnTheFlyServer->GetIdleTimeToGC();
-		const uint64 MaxMemoryAllowance = CookOnTheFlyServer->GetMaxMemoryAllowance();
 		const uint32 PackagesPerPartialGC = CookOnTheFlyServer->GetPackagesPerPartialGC();
 
 		double LastCookActionTime = FPlatformTime::Seconds();
@@ -1081,19 +1080,6 @@ bool UCookCommandlet::CookByTheBook( const TArray<ITargetPlatform*>& Platforms,
 	return true;
 }
 
-bool UCookCommandlet::HasExceededMaxMemory(uint64 MaxMemoryAllowance) const
-{
-	const FPlatformMemoryStats MemStats = FPlatformMemory::GetStats();
-
-	uint64 UsedMemory = MemStats.UsedPhysical;
-	if ( (UsedMemory >= MaxMemoryAllowance) && 
-		(MaxMemoryAllowance > 0u) )
-	{
-		return true;
-	}
-	return false;
-}
-
 void UCookCommandlet::ProcessDeferredCommands()
 {
 #if PLATFORM_MAC
diff --git a/Engine/Source/Editor/UnrealEd/Private/CookOnTheFlyServer.cpp b/Engine/Source/Editor/UnrealEd/Private/CookOnTheFlyServer.cpp
--- a/Engine/Source/Editor/UnrealEd/Private/CookOnTheFlyServer.cpp
+++ b/Engine/Source/Editor/UnrealEd/Private/CookOnTheFlyServer.cpp
@@ -3756,38 +3756,46 @@ bool UCookOnTheFlyServer::HasExceededMaxMemory() const
 
 	TStringBuilder<256> TriggerMessages;
 	const FPlatformMemoryStats MemStats = FPlatformMemory::GetStats();
-	if (MinFreeMemory != 0)
+	if (MemoryMinFreeVirtual > 0 || MemoryMinFreePhysical > 0)
 	{
 		++ActiveTriggers;
-		if (MemStats.AvailablePhysical < MinFreeMemory)
+		bool bFired = false;
+		if (MemoryMinFreeVirtual > 0 && MemStats.AvailableVirtual < MemoryMinFreeVirtual)
 		{
-			++FiredTriggers;
-			TriggerMessages.Appendf(TEXT("\n  CookSettings.MinFreeMemory: Available physical memory %dMiB is less than %dMiB."),
-				static_cast<uint32>(MemStats.AvailablePhysical / 1024 / 1024), static_cast<uint32>(MinFreeMemory / 1024 / 1024));
+			TriggerMessages.Appendf(TEXT("\n  CookSettings.MemoryMinFreeVirtual: Available virtual memory %dMiB is less than %dMiB."),
+				static_cast<uint32>(MemStats.AvailableVirtual / 1024 / 1024), static_cast<uint32>(MemoryMinFreeVirtual / 1024 / 1024));
+			bFired = true;
 		}
-	}
-
-	// don't gc if we haven't reached our min gc level yet
-	if (MinMemoryBeforeGC > 0)
-	{
-		++ActiveTriggers;
-		if (MemStats.UsedVirtual >= MinMemoryBeforeGC)
+		if (MemoryMinFreePhysical > 0 && MemStats.AvailablePhysical < MemoryMinFreePhysical)
+		{
+			TriggerMessages.Appendf(TEXT("\n  CookSettings.MemoryMinFreePhysical: Available physical memory %dMiB is less than %dMiB."),
+				static_cast<uint32>(MemStats.AvailablePhysical / 1024 / 1024), static_cast<uint32>(MemoryMinFreePhysical / 1024 / 1024));
+			bFired = true;
+		}
+		if (bFired)
 		{
-			TriggerMessages.Appendf(TEXT("\n  CookSettings.MinMemoryBeforeGC: Used virtual memory %dMiB is greater than %dMiB."),
-				static_cast<uint32>(MemStats.UsedVirtual / 1024 / 1024), static_cast<uint32>(MinMemoryBeforeGC / 1024 / 1024));
 			++FiredTriggers;
 		}
 	}
 
-	if (MaxMemoryAllowance > 0u)
+	if (MemoryMaxUsedVirtual > 0 || MemoryMaxUsedPhysical > 0)
 	{
 		++ActiveTriggers;
-		//uint64 UsedMemory = MemStats.UsedVirtual; 
-		uint64 UsedMemory = MemStats.UsedPhysical; //should this be used virtual?
-		if (UsedMemory >= MaxMemoryAllowance)
+		bool bFired = false;
+		if (MemoryMaxUsedVirtual > 0 && MemStats.UsedVirtual >= MemoryMaxUsedVirtual)
+		{
+			TriggerMessages.Appendf(TEXT("\n  CookSettings.MemoryMaxUsedVirtual: Used virtual memory %dMiB is greater than %dMiB."),
+				static_cast<uint32>(MemStats.UsedVirtual / 1024 / 1024), static_cast<uint32>(MemoryMaxUsedVirtual / 1024 / 1024));
+			bFired = true;
+		}
+		if (MemoryMaxUsedPhysical > 0 && MemStats.UsedPhysical >= MemoryMaxUsedPhysical)
+		{
+			TriggerMessages.Appendf(TEXT("\n  CookSettings.MemoryMaxUsedPhysical: Used physical memory %dMiB is greater than %dMiB."),
+				static_cast<uint32>(MemStats.UsedPhysical / 1024 / 1024), static_cast<uint32>(MemoryMaxUsedPhysical / 1024 / 1024));
+			bFired = true;
+		}
+		if (bFired)
 		{
-			TriggerMessages.Appendf(TEXT("\n  CookSettings.MaxMemoryAllowance: Used physical memory %dMiB is greater than %dMiB."),
-				static_cast<uint32>(MemStats.UsedPhysical / 1024 / 1024), static_cast<uint32>(MaxMemoryAllowance / 1024 / 1024));
 			++FiredTriggers;
 		}
 	}
@@ -4024,7 +4032,7 @@ double UCookOnTheFlyServer::GetIdleTimeToGC() const
 
 uint64 UCookOnTheFlyServer::GetMaxMemoryAllowance() const
 {
-	return MaxMemoryAllowance;
+	return MemoryMaxUsedPhysical;
 }
 
 const TArray<FName>& UCookOnTheFlyServer::GetFullPackageDependencies(const FName& PackageName ) const
@@ -4663,43 +4671,62 @@ void UCookOnTheFlyServer::Initialize( ECookMode::Type DesiredCookMode, ECookInit
 	IdleTimeToGC = 20.0;
 	GConfig->GetDouble( TEXT("CookSettings"), TEXT("IdleTimeToGC"), IdleTimeToGC, GEditorIni );
 
-	int32 MaxMemoryAllowanceInMB = 8 * 1024;
-	GConfig->GetInt( TEXT("CookSettings"), TEXT("MaxMemoryAllowance"), MaxMemoryAllowanceInMB, GEditorIni );
-	MaxMemoryAllowanceInMB = FMath::Max(MaxMemoryAllowanceInMB, 0);
-	MaxMemoryAllowance = MaxMemoryAllowanceInMB * 1024LL * 1024LL;
+	auto ReadMemorySetting = [](const TCHAR* SettingName, uint64& TargetVariable)
+	{
+		int32 ValueInMB = 0;
+		if (GConfig->GetInt(TEXT("CookSettings"), SettingName, ValueInMB, GEditorIni))
+		{
+			ValueInMB = FMath::Max(ValueInMB, 0);
+			TargetVariable = ValueInMB * 1024LL * 1024LL;
+			return true;
+		}
+		return false;
+	};
+	MemoryMaxUsedVirtual = 0;
+	MemoryMaxUsedPhysical = 0;
+	MemoryMinFreeVirtual = 0;
+	MemoryMinFreePhysical = 0;
+	ReadMemorySetting(TEXT("MemoryMaxUsedVirtual"), MemoryMaxUsedVirtual);
+	ReadMemorySetting(TEXT("MemoryMaxUsedPhysical"), MemoryMaxUsedPhysical);
+	ReadMemorySetting(TEXT("MemoryMinFreeVirtual"), MemoryMinFreeVirtual);
+	ReadMemorySetting(TEXT("MemoryMinFreePhysical"), MemoryMinFreePhysical);
+
+	uint64 MaxMemoryAllowance;
+	if (ReadMemorySetting(TEXT("MaxMemoryAllowance"), MaxMemoryAllowance))
+	{ 
+		UE_LOG(LogCook, Warning, TEXT("CookSettings.MaxMemoryAllowance is deprecated. Use CookSettings.MemoryMaxUsedPhysical instead."));
+		MemoryMaxUsedPhysical = MaxMemoryAllowance;
+	}
+	uint64 MinMemoryBeforeGC;
+	if (ReadMemorySetting(TEXT("MinMemoryBeforeGC"), MinMemoryBeforeGC))
+	{
+		UE_LOG(LogCook, Warning, TEXT("CookSettings.MinMemoryBeforeGC is deprecated. Use CookSettings.MemoryMaxUsedVirtual instead."));
+		MemoryMaxUsedVirtual = MinMemoryBeforeGC;
+	}
+	uint64 MinFreeMemory;
+	if (ReadMemorySetting(TEXT("MinFreeMemory"), MinFreeMemory))
+	{
+		UE_LOG(LogCook, Warning, TEXT("CookSettings.MinFreeMemory is deprecated. Use CookSettings.MemoryMinFreePhysical instead."));
+		MemoryMinFreePhysical = MinFreeMemory;
+	}
+	uint64 MinReservedMemory;
+	if (ReadMemorySetting(TEXT("MinReservedMemory"), MinReservedMemory))
+	{
+		UE_LOG(LogCook, Warning, TEXT("CookSettings.MinReservedMemory is deprecated. Use CookSettings.MemoryMinFreePhysical instead."));
+		MemoryMinFreePhysical = MinReservedMemory;
+	}
 	
-	int32 MinMemoryBeforeGCInMB = 0; // 6 * 1024;
-	GConfig->GetInt(TEXT("CookSettings"), TEXT("MinMemoryBeforeGC"), MinMemoryBeforeGCInMB, GEditorIni);
-	MinMemoryBeforeGCInMB = FMath::Max(MinMemoryBeforeGCInMB, 0);
-	MinMemoryBeforeGC = MinMemoryBeforeGCInMB * 1024LL * 1024LL;
-	MinMemoryBeforeGC = FMath::Min(MaxMemoryAllowance, MinMemoryBeforeGC);
-
 	MinFreeUObjectIndicesBeforeGC = 100000;
 	GConfig->GetInt(TEXT("CookSettings"), TEXT("MinFreeUObjectIndicesBeforeGC"), MinFreeUObjectIndicesBeforeGC, GEditorIni);
 	MinFreeUObjectIndicesBeforeGC = FMath::Max(MinFreeUObjectIndicesBeforeGC, 0);
 
-	int32 MinFreeMemoryInMB = 0;
-	GConfig->GetInt(TEXT("CookSettings"), TEXT("MinFreeMemory"), MinFreeMemoryInMB, GEditorIni);
-	MinFreeMemoryInMB = FMath::Max(MinFreeMemoryInMB, 0);
-	MinFreeMemory = MinFreeMemoryInMB * 1024LL * 1024LL;
-
-	// check the amount of OS memory and use that number minus the reserved memory number
-	int32 MinReservedMemoryInMB = 0;
-	GConfig->GetInt(TEXT("CookSettings"), TEXT("MinReservedMemory"), MinReservedMemoryInMB, GEditorIni);
-	MinReservedMemoryInMB = FMath::Max(MinReservedMemoryInMB, 0);
-	int64 MinReservedMemory = MinReservedMemoryInMB * 1024LL * 1024LL;
-	if ( MinReservedMemory )
-	{
-		int64 TotalRam = FPlatformMemory::GetPhysicalGBRam() * 1024LL * 1024LL * 1024LL;
-		MaxMemoryAllowance = FMath::Min<int64>( MaxMemoryAllowance, TotalRam - MinReservedMemory );
-	}
-
 	MaxNumPackagesBeforePartialGC = 400;
 	GConfig->GetInt(TEXT("CookSettings"), TEXT("MaxNumPackagesBeforePartialGC"), MaxNumPackagesBeforePartialGC, GEditorIni);
 	
 	GConfig->GetArray(TEXT("CookSettings"), TEXT("ConfigSettingBlacklist"), ConfigSettingBlacklist, GEditorIni);
 
-	UE_LOG(LogCook, Display, TEXT("Max memory allowance for cook %dmb min free memory %dmb"), MaxMemoryAllowanceInMB, MinFreeMemoryInMB);
+	UE_LOG(LogCook, Display, TEXT("CookSettings for Memory: MemoryMaxUsedVirtual %dMiB, MemoryMaxUsedPhysical %dMiB, MemoryMinFreeVirtual %dMiB, MemoryMinFreePhysical %dMiB"),
+		MemoryMaxUsedVirtual / 1024 / 1024, MemoryMaxUsedPhysical / 1024 / 1024, MemoryMinFreeVirtual / 1024 / 1024, MemoryMinFreePhysical / 1024 / 1024);
 
 
 	{
diff --git a/Engine/Source/Runtime/Core/Private/Windows/WindowsPlatformMemory.cpp b/Engine/Source/Runtime/Core/Private/Windows/WindowsPlatformMemory.cpp
--- a/Engine/Source/Runtime/Core/Private/Windows/WindowsPlatformMemory.cpp
+++ b/Engine/Source/Runtime/Core/Private/Windows/WindowsPlatformMemory.cpp
@@ -227,7 +227,11 @@ FPlatformMemoryStats FWindowsPlatformMemory::GetStats()
 	MemoryStats.TotalPhysical = MemoryStatusEx.ullTotalPhys;
 	MemoryStats.AvailablePhysical = MemoryStatusEx.ullAvailPhys;
 	MemoryStats.AvailableVirtual = MemoryStatusEx.ullAvailVirtual;
-	
+
+	// On Windows, Virtual Memory is limited per process to the address space (e.g. 47 bits (128Tb)), but is additionally limited by the sum of used virtual memory across all processes
+	// must be less than PhysicalMemory plus the Virtual Memory Page Size. The remaining virtual memory space given this system-wide limit is stored in ullAvailPageSize
+	MemoryStats.AvailableVirtual = FMath::Min(MemoryStats.AvailableVirtual, MemoryStatusEx.ullAvailPageFile);
+
 	MemoryStats.UsedPhysical = ProcessMemoryCounters.WorkingSetSize;
 	MemoryStats.PeakUsedPhysical = ProcessMemoryCounters.PeakWorkingSetSize;
 	MemoryStats.UsedVirtual = ProcessMemoryCounters.PagefileUsage;
