AssetRegistryDiff: Bumping the "Unable to find class name" to a normal log line instead of a warning

diff --git a/Engine/Source/Editor/UnrealEd/Private/Commandlets/DiffAssetRegistriesCommandlet.cpp b/Engine/Source/Editor/UnrealEd/Private/Commandlets/DiffAssetRegistriesCommandlet.cpp
--- a/Engine/Source/Editor/UnrealEd/Private/Commandlets/DiffAssetRegistriesCommandlet.cpp
+++ b/Engine/Source/Editor/UnrealEd/Private/Commandlets/DiffAssetRegistriesCommandlet.cpp
@@ -542,7 +542,7 @@ FName UDiffAssetRegistriesCommandlet::GetClassName(FAssetRegistryState& InRegist
 
 		if (NewName == NAME_None)
 		{
-			UE_LOG(LogDiffAssets, Warning, TEXT("Unable to find class type of asset %s"), *InAssetPath.ToString());
+			UE_LOG(LogDiffAssets, Log, TEXT("Unable to find class type of asset %s"), *InAssetPath.ToString());
 		}
 		AssetPathToClassName.Add(InAssetPath) = NewName;
 	}
