Fix crash after canceling a cook and changing a platform (UE-92633).

diff --git a/Engine/Source/Editor/UnrealEd/Private/CookOnTheFlyServer.cpp b/Engine/Source/Editor/UnrealEd/Private/CookOnTheFlyServer.cpp
--- a/Engine/Source/Editor/UnrealEd/Private/CookOnTheFlyServer.cpp
+++ b/Engine/Source/Editor/UnrealEd/Private/CookOnTheFlyServer.cpp
@@ -7937,9 +7937,17 @@ void UCookOnTheFlyServer::StartCookByTheBook( const FCookByTheBookStartupOptions
 	// this is to support canceling cooks from the editor
 	// this is required to make sure that the cooker is in a good state after cancel occurs
 	// if too many packages are being recooked after resume then we may need to figure out a different way to do this
-	for ( const FFilePlatformRequest& PreviousRequest : CookByTheBookOptions->PreviousCookRequests )
+	if (IsCookingInEditor())
 	{
-		PackageTracker->EnqueueUniqueCookRequest( PreviousRequest );
+		for (const FFilePlatformRequest& PreviousRequest : CookByTheBookOptions->PreviousCookRequests)
+		{
+			// do not queue previous requests that targeted a different platform
+			const TArray<const ITargetPlatform*>& PreviousPlatforms = PreviousRequest.GetPlatforms();
+			if (TargetPlatforms.Num() == 1 && PreviousPlatforms.Num() == TargetPlatforms.Num() && PreviousPlatforms[0] == TargetPlatforms[0])
+			{
+				PackageTracker->EnqueueUniqueCookRequest(PreviousRequest);
+			}
+		}
 	}
 	CookByTheBookOptions->PreviousCookRequests.Empty();
 
