Make Cooker Analytics parameters configurable through Engine ini files. Note that overriding the AppId/ApiKey will utilise the modern datarouter URL, while leaving it at the default will keep the legacy URL

diff --git a/Engine/Source/Editor/UnrealEd/Private/Commandlets/CookCommandlet.cpp b/Engine/Source/Editor/UnrealEd/Private/Commandlets/CookCommandlet.cpp
--- a/Engine/Source/Editor/UnrealEd/Private/Commandlets/CookCommandlet.cpp
+++ b/Engine/Source/Editor/UnrealEd/Private/Commandlets/CookCommandlet.cpp
@@ -95,38 +95,51 @@ namespace DetailedCookStats
 		#undef ADD_COOK_STAT_FLT
 	});
 
-	/** Gathers the cook stats registered with the FCookStatsManager delegate and logs them as a CSV. */
 	static void LogCookStats(const FString& CookCmdLine)
 	{
-		// Optionally create an analytics provider to send stats to for central collection.
-		if (GIsBuildMachine || FParse::Param(FCommandLine::Get(), TEXT("SendCookAnalytics")))
+		bool bSendCookAnalytics = false;
+		GConfig->GetBool(TEXT("CookAnalytics"), TEXT("SendAnalytics"), bSendCookAnalytics, GEngineIni);
+
+		if (GIsBuildMachine || FParse::Param(FCommandLine::Get(), TEXT("SendCookAnalytics")) || bSendCookAnalytics)
 		{
 			FString APIServerET;
-			// This value is set by an INI private to Epic.
 			if (GConfig->GetString(TEXT("CookAnalytics"), TEXT("APIServer"), APIServerET, GEngineIni))
 			{
-				TSharedPtr<IAnalyticsProviderET> CookAnalytics = FAnalyticsET::Get().CreateAnalyticsProvider(FAnalyticsET::Config(TEXT("Cook"), APIServerET, FString(), true));
+				FString AppId(TEXT("Cook"));
+ 				bool bUseLegacyCookProtocol = !GConfig->GetString(TEXT("CookAnalytics"), TEXT("AppId"), AppId, GEngineIni);
+
+				// Optionally create an analytics provider to send stats to for central collection.
+				TSharedPtr<IAnalyticsProviderET> CookAnalytics = FAnalyticsET::Get().CreateAnalyticsProvider(FAnalyticsET::Config(AppId, APIServerET, FString(), bUseLegacyCookProtocol));
 				if (CookAnalytics.IsValid())
 				{
-					// start the session
-					CookAnalytics->SetUserID(FString(FPlatformProcess::ComputerName()) + FString(TEXT("\\")) + FString(FPlatformProcess::UserName(false)));
+					CookAnalytics->SetUserID(FString::Printf(TEXT("%s\\%s"), FPlatformProcess::ComputerName(), FPlatformProcess::UserName(false)));
 					CookAnalytics->StartSession(MakeAnalyticsEventAttributeArray(
 						TEXT("Project"), CookProject,
 						TEXT("CmdLine"), CookCmdLine,
 						TEXT("IsBuildMachine"), GIsBuildMachine,
 						TEXT("TargetPlatforms"), TargetPlatforms
 					));
+
+					TArray<FString> CookStatsToSend;
+					const bool bUseWhitelist = GConfig->GetArray(TEXT("CookAnalytics"), TEXT("CookStats"), CookStatsToSend, GEngineIni) > 0;
 					// Sends each cook stat to the analytics provider.
-					auto SendCookStatsToAnalytics = [CookAnalytics](const FString& StatName, const TArray<FCookStatsManager::StringKeyValue>& StatAttributes)
+					auto SendCookStatsToAnalytics = [CookAnalytics, &CookStatsToSend, bUseWhitelist](const FString& StatName, const TArray<FCookStatsManager::StringKeyValue>& StatAttributes)
 					{
-						// convert all stats directly to an analytics event
-						TArray<FAnalyticsEventAttribute> StatAttrs;
-						StatAttrs.Reset(StatAttributes.Num());
-						for (const auto& Attr : StatAttributes)
+						if (!bUseWhitelist || CookStatsToSend.Contains(StatName))
+						{
+							// convert filtered stats directly to an analytics event
+							TArray<FAnalyticsEventAttribute> StatAttrs;
+							StatAttrs.Reset(StatAttributes.Num());
+							for (const auto& Attr : StatAttributes)
+							{
+								StatAttrs.Emplace(Attr.Key, Attr.Value);
+							}
+							CookAnalytics->RecordEvent(StatName, MoveTemp(StatAttrs));
+						}
+						else
 						{
-							StatAttrs.Emplace(Attr.Key, Attr.Value);
+							UE_LOG(LogCookCommandlet, Verbose, TEXT("[%s] not present on cook analytics whitelist"), *StatName);
 						}
-						CookAnalytics->RecordEvent(StatName, MoveTemp(StatAttrs));
 					};
 					FCookStatsManager::LogCookStats(SendCookStatsToAnalytics);
 				}
