Prevent loading of streaming level with CreateNewStreamingLevel when the type is abstract

diff --git a/Engine/Source/Editor/UnrealEd/Private/EditorLevelUtils.cpp b/Engine/Source/Editor/UnrealEd/Private/EditorLevelUtils.cpp
--- a/Engine/Source/Editor/UnrealEd/Private/EditorLevelUtils.cpp
+++ b/Engine/Source/Editor/UnrealEd/Private/EditorLevelUtils.cpp
@@ -400,8 +400,8 @@ ULevelStreaming* UEditorLevelUtils::AddLevelToWorld_Internal(UWorld* InWorld, co
 	}
 	else
 	{
-		// If the selected class is still NULL, abort the operation.
-		if (LevelStreamingClass == nullptr)
+		// If the selected class is still NULL or the selected class is abstract, abort the operation.
+		if (LevelStreamingClass == nullptr || LevelStreamingClass->HasAnyClassFlags(CLASS_Abstract))
 		{
 			return nullptr;
 		}
