Sequencer: Set frame rate on export

diff --git a/Engine/Source/Editor/UnrealEd/Private/Fbx/FbxMainExport.cpp b/Engine/Source/Editor/UnrealEd/Private/Fbx/FbxMainExport.cpp
--- a/Engine/Source/Editor/UnrealEd/Private/Fbx/FbxMainExport.cpp
+++ b/Engine/Source/Editor/UnrealEd/Private/Fbx/FbxMainExport.cpp
@@ -1689,6 +1689,21 @@ bool FFbxExporter::ExportLevelSequence(UMovieScene* MovieScene, const TArray<FGu
 		return false;
 	}
 
+	double FrameRate = MovieScene->GetDisplayRate().AsDecimal();
+	FbxTime::EMode TimeMode = FbxTime::ConvertFrameRateToTimeMode(FrameRate);
+
+	// Unknown frame rates, (ie. 12 fps) return eDefaultMode, so override with custom so that the custom frame rate takes effect
+	if (TimeMode == FbxTime::eDefaultMode)
+	{
+		TimeMode = FbxTime::eCustom;
+	}
+
+	Scene->GetGlobalSettings().SetTimeMode(TimeMode);
+	if (TimeMode == FbxTime::eCustom)
+	{
+		Scene->GetGlobalSettings().SetCustomFrameRate(FrameRate);
+	}
+
 	for (const FMovieSceneBinding& MovieSceneBinding : MovieScene->GetBindings())
 	{
 		// If there are specific bindings to export, export those only
