CIS Fixes

diff --git a/Engine/Source/Editor/UnrealEd/Private/Subsystems/StatusBarSubsystem.cpp b/Engine/Source/Editor/UnrealEd/Private/Subsystems/StatusBarSubsystem.cpp
--- a/Engine/Source/Editor/UnrealEd/Private/Subsystems/StatusBarSubsystem.cpp
+++ b/Engine/Source/Editor/UnrealEd/Private/Subsystems/StatusBarSubsystem.cpp
@@ -10,6 +10,9 @@
 #include "ISourceControlProvider.h"
 #include "ISourceControlModule.h"
 #include "SourceControlWindows.h"
+#include "Framework/Commands/Commands.h"
+#include "Widgets/Input/SComboButton.h"
+#include "EditorStyleSet.h"
 
 #define LOCTEXT_NAMESPACE "StatusBar"
 
diff --git a/Engine/Source/Editor/UnrealEd/Public/Subsystems/StatusBarSubsystem.h b/Engine/Source/Editor/UnrealEd/Public/Subsystems/StatusBarSubsystem.h
--- a/Engine/Source/Editor/UnrealEd/Public/Subsystems/StatusBarSubsystem.h
+++ b/Engine/Source/Editor/UnrealEd/Public/Subsystems/StatusBarSubsystem.h
@@ -7,6 +7,7 @@
 
 #include "StatusBarSubsystem.generated.h"
 
+class SWidget;
 class SStatusBar;
 
 UCLASS()
