Fix HAL/PlatformFileManager.h non-portable casing CIS issues

diff --git a/Engine/Plugins/Developer/Concert/ConcertSync/ConcertSyncClient/Source/ConcertSyncClient/Private/ConcertClientPackageManager.cpp b/Engine/Plugins/Developer/Concert/ConcertSync/ConcertSyncClient/Source/ConcertSyncClient/Private/ConcertClientPackageManager.cpp
--- a/Engine/Plugins/Developer/Concert/ConcertSync/ConcertSyncClient/Source/ConcertSyncClient/Private/ConcertClientPackageManager.cpp
+++ b/Engine/Plugins/Developer/Concert/ConcertSync/ConcertSyncClient/Source/ConcertSyncClient/Private/ConcertClientPackageManager.cpp
@@ -19,7 +19,7 @@
 #include "Misc/FileHelper.h"
 #include "Misc/PackageName.h"
 #include "HAL/FileManager.h"
-#include "HAL/PlatformFilemanager.h"
+#include "HAL/PlatformFileManager.h"
 
 #if WITH_EDITOR
 	#include "Editor.h"
diff --git a/Engine/Plugins/Developer/Concert/ConcertSync/ConcertSyncClient/Source/ConcertSyncClient/Private/ConcertClientWorkspace.cpp b/Engine/Plugins/Developer/Concert/ConcertSync/ConcertSyncClient/Source/ConcertSyncClient/Private/ConcertClientWorkspace.cpp
--- a/Engine/Plugins/Developer/Concert/ConcertSync/ConcertSyncClient/Source/ConcertSyncClient/Private/ConcertClientWorkspace.cpp
+++ b/Engine/Plugins/Developer/Concert/ConcertSync/ConcertSyncClient/Source/ConcertSyncClient/Private/ConcertClientWorkspace.cpp
@@ -26,7 +26,7 @@
 #include "UObject/StructOnScope.h"
 #include "Misc/PackageName.h"
 #include "HAL/FileManager.h"
-#include "HAL/PlatformFilemanager.h"
+#include "HAL/PlatformFileManager.h"
 #include "Misc/App.h"
 #include "Misc/Paths.h"
 #include "Misc/FileHelper.h"
diff --git a/Engine/Source/Editor/UnrealEd/Private/Kismet2/DebuggerCommands.cpp b/Engine/Source/Editor/UnrealEd/Private/Kismet2/DebuggerCommands.cpp
--- a/Engine/Source/Editor/UnrealEd/Private/Kismet2/DebuggerCommands.cpp
+++ b/Engine/Source/Editor/UnrealEd/Private/Kismet2/DebuggerCommands.cpp
@@ -63,7 +63,7 @@
 #include "IAndroidDeviceDetectionModule.h"
 #include "IAndroidDeviceDetection.h"
 #include "CookerSettings.h"
-#include "HAL/PlatformFilemanager.h"
+#include "HAL/PlatformFileManager.h"
 #include "SourceControlHelpers.h"
 #include "ISourceControlModule.h"
 #include "Framework/Notifications/NotificationManager.h"
@@ -610,28 +610,35 @@ void FPlayWorldCommands::AddPIEPreviewDeviceActions(const FPlayWorldCommands &Co
 
 void FPlayWorldCommands::BuildToolbar(FToolMenuSection& InSection, bool bIncludeLaunchButtonAndOptions)
 {
+	FToolMenuEntry PlayMenuEntry =
+		FToolMenuEntry::InitToolBarButton(
+			FPlayWorldCommands::Get().RepeatLastPlay,
+			LOCTEXT("RepeatLastPlay", "Play"),
+			TAttribute< FText >::Create(TAttribute< FText >::FGetter::CreateStatic(&FInternalPlayWorldCommandCallbacks::GetRepeatLastPlayToolTip)),
+			TAttribute< FSlateIcon >::Create(TAttribute< FSlateIcon >::FGetter::CreateStatic(&FInternalPlayWorldCommandCallbacks::GetRepeatLastPlayIcon)),
+			FName(TEXT("LevelToolbarPlay")));
+
+	PlayMenuEntry.StyleNameOverride = FName("PlayToolbar");
+
 	// Play
-	InSection.AddEntry(FToolMenuEntry::InitToolBarButton(
-		FPlayWorldCommands::Get().RepeatLastPlay,
-		LOCTEXT("RepeatLastPlay", "Play"),
-		TAttribute< FText >::Create( TAttribute< FText >::FGetter::CreateStatic( &FInternalPlayWorldCommandCallbacks::GetRepeatLastPlayToolTip ) ),
-		TAttribute< FSlateIcon >::Create(TAttribute< FSlateIcon >::FGetter::CreateStatic(&FInternalPlayWorldCommandCallbacks::GetRepeatLastPlayIcon)),
-		FName(TEXT("LevelToolbarPlay"))
-	));
+	InSection.AddEntry(PlayMenuEntry);
 
 	// Play combo box
 	FUIAction SpecialPIEOptionsMenuAction;
 	SpecialPIEOptionsMenuAction.IsActionVisibleDelegate = FIsActionButtonVisible::CreateStatic(&FInternalPlayWorldCommandCallbacks::CanShowNonPlayWorldOnlyActions);
 
-	InSection.AddEntry(FToolMenuEntry::InitComboButton(
+
+	FToolMenuEntry PlayMenuCombo = FToolMenuEntry::InitComboButton(
 		"PlayCombo",
 		SpecialPIEOptionsMenuAction,
 		FOnGetContent::CreateStatic(&GeneratePlayMenuContent, GlobalPlayWorldActions.ToSharedRef()),
 		LOCTEXT("PlayCombo_Label", "Active Play Mode"),
 		LOCTEXT("PIEComboToolTip", "Change Play Mode and Play Settings"),
 		FSlateIcon(FEditorStyle::GetStyleSetName(), "LevelEditor.RepeatLastPlay"),
-		true
-	));
+		true);
+
+	PlayMenuCombo.StyleNameOverride = FName("PlayToolbar");
+	InSection.AddEntry(PlayMenuCombo);
 
 	if (bIncludeLaunchButtonAndOptions)
 	{
diff --git a/Engine/Source/Runtime/Core/Private/GenericPlatform/GenericPlatformMisc.cpp b/Engine/Source/Runtime/Core/Private/GenericPlatform/GenericPlatformMisc.cpp
--- a/Engine/Source/Runtime/Core/Private/GenericPlatform/GenericPlatformMisc.cpp
+++ b/Engine/Source/Runtime/Core/Private/GenericPlatform/GenericPlatformMisc.cpp
@@ -2,7 +2,7 @@
 
 #include "GenericPlatform/GenericPlatformMisc.h"
 #include "Misc/AssertionMacros.h"
-#include "HAL/PlatformFilemanager.h"
+#include "HAL/PlatformFileManager.h"
 #include "HAL/CriticalSection.h"
 #include "Misc/ScopeRWLock.h"
 #include "Math/UnrealMathUtility.h"
@@ -1451,19 +1451,14 @@ FString FGenericPlatformMisc::LoadTextFileFromPlatformPackage(const FString& Rel
 {
 	FString Path = RootDir() / RelativePath;
 	FString Result;
-	if (FFileHelper::LoadFileToString(Result, &IPlatformFile::GetPlatformPhysical(), *Path))
-	{
-		return Result;
-	}
-
-	Result.Empty();
+	FFileHelper::LoadFileToString(Result, *Path);
 	return Result;
 }
 
 bool FGenericPlatformMisc::FileExistsInPlatformPackage(const FString& RelativePath)
 {
 	FString Path = RootDir() / RelativePath;
-	return IPlatformFile::GetPlatformPhysical().FileExists(*Path);
+	return IFileManager::Get().FileExists(*Path);
 }
 
 void FGenericPlatformMisc::TearDown()
