More CIS fixing

diff --git a/Engine/Source/Editor/UnrealEd/Private/Subsystems/StatusBarSubsystem.cpp b/Engine/Source/Editor/UnrealEd/Private/Subsystems/StatusBarSubsystem.cpp
--- a/Engine/Source/Editor/UnrealEd/Private/Subsystems/StatusBarSubsystem.cpp
+++ b/Engine/Source/Editor/UnrealEd/Private/Subsystems/StatusBarSubsystem.cpp
@@ -13,6 +13,7 @@
 #include "Framework/Commands/Commands.h"
 #include "Widgets/Input/SComboButton.h"
 #include "EditorStyleSet.h"
+#include "FileHelpers.h"
 
 #define LOCTEXT_NAMESPACE "StatusBar"
 
@@ -492,3 +493,5 @@ TSharedRef<SWidget> UStatusBarSubsystem::MakeStatusBarWidget(FName StatusBarName
 
 	return StatusBar;
 }
+
+#undef LOCTEXT_NAMESPACE
\ No newline at end of file
