Status bar updates
- Hooked up console edit box
- Hooked up active tool messages to modes (the mode specific messages need a lot of improvement
- Added status bars to standalone asset editors

diff --git a/Engine/Plugins/Experimental/PythonScriptPlugin/Source/PythonScriptPlugin/Private/PythonScriptPlugin.cpp b/Engine/Plugins/Experimental/PythonScriptPlugin/Source/PythonScriptPlugin/Private/PythonScriptPlugin.cpp
--- a/Engine/Plugins/Experimental/PythonScriptPlugin/Source/PythonScriptPlugin/Private/PythonScriptPlugin.cpp
+++ b/Engine/Plugins/Experimental/PythonScriptPlugin/Source/PythonScriptPlugin/Private/PythonScriptPlugin.cpp
@@ -112,7 +112,7 @@ FText FPythonCommandExecutor::GetDescription() const
 
 FText FPythonCommandExecutor::GetHintText() const
 {
-	return LOCTEXT("PythonCommandExecutorHintText", "Enter Python script or a filename to execute");
+	return LOCTEXT("PythonCommandExecutorHintText", "Enter Python script or a filename");
 }
 
 void FPythonCommandExecutor::GetAutoCompleteSuggestions(const TCHAR* Input, TArray<FString>& Out)
@@ -173,7 +173,7 @@ FText FPythonREPLCommandExecutor::GetDescription() const
 
 FText FPythonREPLCommandExecutor::GetHintText() const
 {
-	return LOCTEXT("PythonREPLCommandExecutorHintText", "Enter a single Python statement");
+	return LOCTEXT("PythonREPLCommandExecutorHintText", "Enter a Python statement");
 }
 
 void FPythonREPLCommandExecutor::GetAutoCompleteSuggestions(const TCHAR* Input, TArray<FString>& Out)
diff --git a/Engine/Source/Developer/OutputLog/Private/OutputLogModule.cpp b/Engine/Source/Developer/OutputLog/Private/OutputLogModule.cpp
--- a/Engine/Source/Developer/OutputLog/Private/OutputLogModule.cpp
+++ b/Engine/Source/Developer/OutputLog/Private/OutputLogModule.cpp
@@ -19,7 +19,7 @@
 #include "Editor.h"
 #endif
 
-IMPLEMENT_MODULE( FOutputLogModule, OutputLog );
+IMPLEMENT_MODULE(FOutputLogModule, OutputLog);
 
 namespace OutputLogModule
 {
@@ -41,7 +41,7 @@ public:
 	~FOutputLogHistory()
 	{
 		// At shutdown, GLog may already be null
-		if( GLog != NULL )
+		if (GLog != NULL)
 		{
 			GLog->RemoveOutputDevice(this);
 		}
@@ -55,7 +55,7 @@ public:
 
 protected:
 
-	virtual void Serialize( const TCHAR* V, ELogVerbosity::Type Verbosity, const class FName& Category ) override
+	virtual void Serialize(const TCHAR* V, ELogVerbosity::Type Verbosity, const class FName& Category) override
 	{
 		// Capture all incoming messages and store them in history
 		SOutputLog::CreateLogMessages(V, Verbosity, Category, Messages);
@@ -73,23 +73,23 @@ static TSharedPtr<FOutputLogHistory> OutputLogHistory;
 /** Our global active output log */
 static TWeakPtr<SOutputLog> OutputLog;
 
-TSharedRef<SDockTab> SpawnOutputLog( const FSpawnTabArgs& Args )
+TSharedRef<SDockTab> SpawnOutputLog(const FSpawnTabArgs& Args)
 {
 	return SNew(SDockTab)
 		.Icon(FEditorStyle::GetBrush("Log.TabIcon"))
-		.TabRole( ETabRole::NomadTab )
-		.Label( NSLOCTEXT("OutputLog", "TabTitle", "Output Log") )
+		.TabRole(ETabRole::NomadTab)
+		.Label(NSLOCTEXT("OutputLog", "TabTitle", "Output Log"))
 		[
-			SAssignNew(OutputLog, SOutputLog).Messages( OutputLogHistory->GetMessages() )
+			SAssignNew(OutputLog, SOutputLog).Messages(OutputLogHistory->GetMessages())
 		];
 }
 
-TSharedRef<SDockTab> SpawnDeviceOutputLog( const FSpawnTabArgs& Args )
+TSharedRef<SDockTab> SpawnDeviceOutputLog(const FSpawnTabArgs& Args)
 {
 	return SNew(SDockTab)
 		.Icon(FEditorStyle::GetBrush("Log.TabIcon"))
-		.TabRole( ETabRole::NomadTab )
-		.Label( NSLOCTEXT("OutputLog", "DeviceTabTitle", "Device Output Log") )
+		.TabRole(ETabRole::NomadTab)
+		.Label(NSLOCTEXT("OutputLog", "DeviceTabTitle", "Device Output Log"))
 		[
 			SNew(SDeviceOutputLog)
 		];
@@ -97,17 +97,17 @@ TSharedRef<SDockTab> SpawnDeviceOutputLog( const FSpawnTabArgs& Args )
 
 void FOutputLogModule::StartupModule()
 {
-	FGlobalTabmanager::Get()->RegisterNomadTabSpawner(OutputLogModule::OutputLogTabName, FOnSpawnTab::CreateStatic( &SpawnOutputLog ) )
+	FGlobalTabmanager::Get()->RegisterNomadTabSpawner(OutputLogModule::OutputLogTabName, FOnSpawnTab::CreateStatic(&SpawnOutputLog))
 		.SetDisplayName(NSLOCTEXT("UnrealEditor", "OutputLogTab", "Output Log"))
 		.SetTooltipText(NSLOCTEXT("UnrealEditor", "OutputLogTooltipText", "Open the Output Log tab."))
-		.SetGroup( WorkspaceMenu::GetMenuStructure().GetDeveloperToolsLogCategory() )
-		.SetIcon( FSlateIcon(FEditorStyle::GetStyleSetName(), "Log.TabIcon") );
+		.SetGroup(WorkspaceMenu::GetMenuStructure().GetDeveloperToolsLogCategory())
+		.SetIcon(FSlateIcon(FEditorStyle::GetStyleSetName(), "Log.TabIcon"));
 
-	FGlobalTabmanager::Get()->RegisterNomadTabSpawner(OutputLogModule::DeviceOutputLogTabName, FOnSpawnTab::CreateStatic( &SpawnDeviceOutputLog ) )
+	FGlobalTabmanager::Get()->RegisterNomadTabSpawner(OutputLogModule::DeviceOutputLogTabName, FOnSpawnTab::CreateStatic(&SpawnDeviceOutputLog))
 		.SetDisplayName(NSLOCTEXT("UnrealEditor", "DeviceOutputLogTab", "Device Output Log"))
 		.SetTooltipText(NSLOCTEXT("UnrealEditor", "DeviceOutputLogTooltipText", "Open the Device Output Log tab."))
-		.SetGroup( WorkspaceMenu::GetMenuStructure().GetDeveloperToolsLogCategory() )
-		.SetIcon( FSlateIcon(FEditorStyle::GetStyleSetName(), "Log.TabIcon") );
+		.SetGroup(WorkspaceMenu::GetMenuStructure().GetDeveloperToolsLogCategory())
+		.SetIcon(FSlateIcon(FEditorStyle::GetStyleSetName(), "Log.TabIcon"));
 
 #if WITH_EDITOR
 	FEditorDelegates::BeginPIE.AddRaw(this, &FOutputLogModule::ClearOnPIE);
@@ -131,19 +131,22 @@ void FOutputLogModule::ShutdownModule()
 	OutputLogHistory.Reset();
 }
 
-TSharedRef< SWidget > FOutputLogModule::MakeConsoleInputBox( TSharedPtr< SMultiLineEditableTextBox >& OutExposedEditableTextBox ) const
+TSharedRef< SWidget > FOutputLogModule::MakeConsoleInputBox(TSharedPtr<SMultiLineEditableTextBox>& OutExposedEditableTextBox, const FSimpleDelegate& OnCloseConsole) const
 {
-	TSharedRef< SConsoleInputBox > NewConsoleInputBox = SNew( SConsoleInputBox );
+	TSharedRef< SConsoleInputBox > NewConsoleInputBox =
+		SNew(SConsoleInputBox)
+		.OnCloseConsole(OnCloseConsole);
+
 	OutExposedEditableTextBox = NewConsoleInputBox->GetEditableTextBox();
 	return NewConsoleInputBox;
 }
 
-void FOutputLogModule::ToggleDebugConsoleForWindow( const TSharedRef< SWindow >& Window, const EDebugConsoleStyle::Type InStyle, const FDebugConsoleDelegates& DebugConsoleDelegates )
+void FOutputLogModule::ToggleDebugConsoleForWindow(const TSharedRef<SWindow>& Window, const EDebugConsoleStyle::Type InStyle, const FDebugConsoleDelegates& DebugConsoleDelegates)
 {
 	bool bShouldOpen = true;
 	// Close an existing console box, if there is one
-	TSharedPtr< SWidget > PinnedDebugConsole( DebugConsole.Pin() );
-	if( PinnedDebugConsole.IsValid() )
+	TSharedPtr< SWidget > PinnedDebugConsole(DebugConsole.Pin());
+	if (PinnedDebugConsole.IsValid())
 	{
 		// If the console is already open close it unless it is in a different window.  In that case reopen it on that window
 		bShouldOpen = false;
@@ -154,13 +157,13 @@ void FOutputLogModule::ToggleDebugConsoleForWindow( const TSharedRef< SWindow >&
 			DebugConsole.Reset();
 		}
 
-		if( WindowForExistingConsole != Window )
+		if (WindowForExistingConsole != Window)
 		{
 			// Console is being opened on another window
 			bShouldOpen = true;
 		}
 	}
-	
+
 	TSharedPtr<SDockTab> ActiveTab = FGlobalTabmanager::Get()->GetActiveTab();
 	if (ActiveTab.IsValid() && ActiveTab->GetLayoutIdentifier() == FTabId(OutputLogModule::OutputLogTabName))
 	{
@@ -168,20 +171,20 @@ void FOutputLogModule::ToggleDebugConsoleForWindow( const TSharedRef< SWindow >&
 		bShouldOpen = false;
 	}
 
-	if( bShouldOpen )
+	if (bShouldOpen)
 	{
 		const EDebugConsoleStyle::Type DebugConsoleStyle = InStyle;
-		TSharedRef< SDebugConsole > DebugConsoleRef = SNew( SDebugConsole, DebugConsoleStyle, this, &DebugConsoleDelegates );
+		TSharedRef< SDebugConsole > DebugConsoleRef = SNew(SDebugConsole, DebugConsoleStyle, this, &DebugConsoleDelegates);
 		DebugConsole = DebugConsoleRef;
 
 		const int32 MaximumZOrder = MAX_int32;
-		Window->AddOverlaySlot( MaximumZOrder )
-		.VAlign(VAlign_Bottom)
-		.HAlign(HAlign_Center)
-		.Padding( 10.0f )
-		[
-			DebugConsoleRef
-		];
+		Window->AddOverlaySlot(MaximumZOrder)
+			.VAlign(VAlign_Bottom)
+			.HAlign(HAlign_Center)
+			.Padding(10.0f)
+			[
+				DebugConsoleRef
+			];
 
 		// Force keyboard focus
 		DebugConsoleRef->SetFocusToEditableText();
@@ -190,14 +193,14 @@ void FOutputLogModule::ToggleDebugConsoleForWindow( const TSharedRef< SWindow >&
 
 void FOutputLogModule::CloseDebugConsole()
 {
-	TSharedPtr< SWidget > PinnedDebugConsole( DebugConsole.Pin() );
+	TSharedPtr< SWidget > PinnedDebugConsole(DebugConsole.Pin());
 
-	if( PinnedDebugConsole.IsValid() )
+	if (PinnedDebugConsole.IsValid())
 	{
 		TSharedPtr< SWindow > WindowForExistingConsole = FSlateApplication::Get().FindWidgetWindow(PinnedDebugConsole.ToSharedRef());
 		if (WindowForExistingConsole.IsValid())
 		{
-			WindowForExistingConsole->RemoveOverlaySlot( PinnedDebugConsole.ToSharedRef() );
+			WindowForExistingConsole->RemoveOverlaySlot(PinnedDebugConsole.ToSharedRef());
 			DebugConsole.Reset();
 		}
 	}
diff --git a/Engine/Source/Developer/OutputLog/Private/SOutputLog.cpp b/Engine/Source/Developer/OutputLog/Private/SOutputLog.cpp
--- a/Engine/Source/Developer/OutputLog/Private/SOutputLog.cpp
+++ b/Engine/Source/Developer/OutputLog/Private/SOutputLog.cpp
@@ -82,21 +82,35 @@ void SConsoleInputBox::Construct(const FArguments& InArgs)
 		.Placement( InArgs._SuggestionListPlacement )
 		[
 			SNew(SHorizontalBox)
-
 			+SHorizontalBox::Slot()
 			.AutoWidth()
 			.Padding(FMargin(0.0f, 0.0f, 4.0f, 0.0f))
 			[
 				SNew(SComboButton)
 				.IsEnabled(this, &SConsoleInputBox::IsCommandExecutorMenuEnabled)
-				.ComboButtonStyle(FEditorStyle::Get(), "GenericFilters.ComboButtonStyle")
-				.ForegroundColor(FLinearColor::White)
+				.ComboButtonStyle(FAppStyle::Get(), "DebugConsole.ComboButton")
 				.ContentPadding(0)
 				.OnGetMenuContent(this, &SConsoleInputBox::GetCommandExecutorMenuContent)
 				.ButtonContent()
 				[
-					SNew(STextBlock)
-					.Text(this, &SConsoleInputBox::GetActiveCommandExecutorDisplayName)
+					SNew(SHorizontalBox)
+					+SHorizontalBox::Slot()
+					.Padding(2.0f)
+					.HAlign(HAlign_Left)
+					.VAlign(VAlign_Center)
+					.AutoWidth()
+					[
+						SNew(SImage)
+						.ColorAndOpacity(FSlateColor::UseForeground())
+						.Image(FAppStyle::Get().GetBrush("DebugConsole.Icon"))
+					]
+					+ SHorizontalBox::Slot()
+					.VAlign(VAlign_Center)
+					.Padding(2.0f)
+					[
+						SNew(STextBlock)
+						.Text(this, &SConsoleInputBox::GetActiveCommandExecutorDisplayName)
+					]
 				]
 			]
 
diff --git a/Engine/Source/Developer/OutputLog/Public/OutputLogModule.h b/Engine/Source/Developer/OutputLog/Public/OutputLogModule.h
--- a/Engine/Source/Developer/OutputLog/Public/OutputLogModule.h
+++ b/Engine/Source/Developer/OutputLog/Public/OutputLogModule.h
@@ -38,11 +38,11 @@ public:
 	virtual void ShutdownModule();
 
 	/** Generates a console input box widget.  Remember, this widget will become invalid if the
-	    output log DLL is unloaded on the fly. */
-	virtual TSharedRef< SWidget > MakeConsoleInputBox( TSharedPtr< SMultiLineEditableTextBox >& OutExposedEditableTextBox ) const;
+		output log DLL is unloaded on the fly. */
+	virtual TSharedRef<SWidget> MakeConsoleInputBox(TSharedPtr<SMultiLineEditableTextBox>& OutExposedEditableTextBox, const FSimpleDelegate& OnCloseConsole) const;
 
 	/** Opens a debug console in the specified window, if not already open */
-	virtual void ToggleDebugConsoleForWindow( const TSharedRef< SWindow >& Window, const EDebugConsoleStyle::Type InStyle, const FDebugConsoleDelegates& DebugConsoleDelegates );
+	virtual void ToggleDebugConsoleForWindow(const TSharedRef<SWindow>& Window, const EDebugConsoleStyle::Type InStyle, const FDebugConsoleDelegates& DebugConsoleDelegates);
 
 	/** Closes the debug console for the specified window */
 	virtual void CloseDebugConsole();
diff --git a/Engine/Source/Editor/EditorStyle/Private/StarshipStyle.cpp b/Engine/Source/Editor/EditorStyle/Private/StarshipStyle.cpp
--- a/Engine/Source/Editor/EditorStyle/Private/StarshipStyle.cpp
+++ b/Engine/Source/Editor/EditorStyle/Private/StarshipStyle.cpp
@@ -930,7 +930,19 @@ void FStarshipEditorStyle::FStyle::SetupGeneralStyles()
 			.SetBackgroundColor( LogColor_Background )
 			);
 
-		Set("DebugConsole.Background", new BOX_BRUSH("Old/Menu_Background", FMargin(8.0f / 64.0f)));
+		Set("DebugConsole.Background", new FSlateNoResource());
+
+		const FButtonStyle DebugConsoleButton = FButtonStyle(FStarshipCoreStyle::GetCoreStyle().GetWidgetStyle<FButtonStyle>("NoBorder"))
+			.SetNormalForeground(FStyleColors::Foreground)
+			.SetNormalPadding(FMargin(2, 2, 2, 2))
+			.SetPressedPadding(FMargin(2, 3, 2, 1));
+
+		const FComboButtonStyle DebugConsoleComboButton = FComboButtonStyle(FStarshipCoreStyle::GetCoreStyle().GetWidgetStyle<FComboButtonStyle>("ComboButton"))
+			.SetDownArrowImage(FSlateNoResource())
+			.SetButtonStyle(DebugConsoleButton);
+
+		Set("DebugConsole.ComboButton", DebugConsoleComboButton);
+		Set("DebugConsole.Icon", new CORE_IMAGE_BRUSH_SVG("Starship/Common/layout-header-body", Icon16x16));
 
 		Set("OutputLog.OpenSourceLocation", new IMAGE_BRUSH("Icons/icon_Asset_Open_Source_Location_16x", Icon16x16));
 		Set("OutputLog.OpenInExternalEditor", new IMAGE_BRUSH("Icons/icon_Asset_Open_In_External_Editor_16x", Icon16x16));
@@ -7902,6 +7914,7 @@ void FStarshipEditorStyle::FStyle::SetupStatusBarStyle()
 	Set("StatusBar.ContentBrowserDown", new CORE_IMAGE_BRUSH_SVG("Starship/Common/chevron-down", Icon16x16));
 	Set("StatusBar.Height", 32.0f);
 	Set("StatusBar.Background", new FSlateColorBrush(FStyleColors::Background));
+	Set("StatusBar.HelpIcon", new CORE_IMAGE_BRUSH_SVG("Starship/Common/help", Icon16x16, FStyleColors::Foreground));
 
 	const FButtonStyle SlimToolBarButton = FButtonStyle(FStarshipCoreStyle::GetCoreStyle().GetWidgetStyle<FButtonStyle>("SlimToolBar.Button"));
 
@@ -7943,6 +7956,14 @@ void FStarshipEditorStyle::FStyle::SetupStatusBarStyle()
 
 	Set("StatusBarToolBar.SimpleComboButton", FStarshipCoreStyle::GetCoreStyle().GetWidgetStyle<FComboButtonStyle>("StatusBarToolBar.SimpleComboButton"));
 	Set("StatusBarToolBar.IconSize", Icon16x16);
+
+	Set("StatusBar.Message.MessageText",
+		FTextBlockStyle(FAppStyle::Get().GetWidgetStyle<FTextBlockStyle>("NormalText"))
+		.SetColorAndOpacity(FStyleColors::Foreground));
+
+	Set("StatusBar.Message.HintText",
+		FTextBlockStyle(FAppStyle::Get().GetWidgetStyle<FTextBlockStyle>("NormalText"))
+		.SetColorAndOpacity(FStyleColors::Primary));
 }
 
 #undef IMAGE_BRUSH
diff --git a/Engine/Source/Editor/LevelEditor/Private/SLevelEditor.cpp b/Engine/Source/Editor/LevelEditor/Private/SLevelEditor.cpp
--- a/Engine/Source/Editor/LevelEditor/Private/SLevelEditor.cpp
+++ b/Engine/Source/Editor/LevelEditor/Private/SLevelEditor.cpp
@@ -217,14 +217,16 @@ void SLevelEditor::Initialize( const TSharedRef<SDockTab>& OwnerTab, const TShar
 			FLevelEditorToolBar::MakeLevelEditorToolBar(LevelEditorCommands.ToSharedRef(), SharedThis(this))
 		]
 		+SVerticalBox::Slot()
+		.Padding(0.0f, 1.0f, 0.0f, 0.0f)
 		.FillHeight( 1.0f )
 		[
 			ContentArea
 		]
 		+SVerticalBox::Slot()
+		.Padding(0.0f, 1.0f, 0.0f, 0.0f)
 		.AutoHeight()
 		[
-			GEditor->GetEditorSubsystem<UStatusBarSubsystem>()->MakeStatusBarWidget(TEXT("LevelEditor.StatusBar"))
+			GEditor->GetEditorSubsystem<UStatusBarSubsystem>()->MakeStatusBarWidget(TEXT("LevelEditor.StatusBar"), OwnerTab)
 		]
 	];
 	
diff --git a/Engine/Source/Editor/LevelEditor/Private/SLevelEditorToolBox.cpp b/Engine/Source/Editor/LevelEditor/Private/SLevelEditorToolBox.cpp
--- a/Engine/Source/Editor/LevelEditor/Private/SLevelEditorToolBox.cpp
+++ b/Engine/Source/Editor/LevelEditor/Private/SLevelEditorToolBox.cpp
@@ -17,6 +17,7 @@
 #include "EdMode.h"
 #include "Widgets/Docking/SDockTab.h"
 #include "Widgets/Layout/SExpandableArea.h"
+#include "Subsystems/StatusBarSubsystem.h"
 
 #define LOCTEXT_NAMESPACE "SLevelEditorToolBox"
 
@@ -162,6 +163,8 @@ EVisibility SLevelEditorToolBox::GetNoToolSelectedTextVisibility() const
 
 void SLevelEditorToolBox::UpdateInlineContent(const TSharedPtr<IToolkit>& Toolkit, TSharedPtr<SWidget> InlineContent) 
 {
+	static const FName LevelEditorStatusBarName = "LevelEditor.StatusBar";
+
 	if (Toolkit.IsValid())
 	{
 		if(Toolkit->GetEditorMode() || Toolkit->GetScriptableEditorMode())
@@ -195,11 +198,21 @@ void SLevelEditorToolBox::UpdateInlineContent(const TSharedPtr<IToolkit>& Toolki
 					.Font(FEditorStyle::Get().GetFontStyle("EditorModesPanel.CategoryFontStyle"))
 				]
 			);
-		}
 
+			if (StatusBarMessageHandle.IsValid())
+			{
+				GEditor->GetEditorSubsystem<UStatusBarSubsystem>()->PopStatusBarMessage(LevelEditorStatusBarName, StatusBarMessageHandle);
+			}
+
+			StatusBarMessageHandle = GEditor->GetEditorSubsystem<UStatusBarSubsystem>()->PushStatusBarMessage(LevelEditorStatusBarName, TAttribute<FText>::Create(
+				TAttribute<FText>::FGetter::CreateSP(ModeToolkit.Get(), &FModeToolkit::GetActiveToolDisplayName))
+			);
+		}
 	}
 	else
 	{
+		GEditor->GetEditorSubsystem<UStatusBarSubsystem>()->PopStatusBarMessage(LevelEditorStatusBarName, StatusBarMessageHandle);
+		StatusBarMessageHandle.Reset();
 
 		TabName = NSLOCTEXT("LevelEditor", "ToolsTabTitle", "Toolbox");
 		TabIcon = FEditorStyle::Get().GetBrush("LevelEditor.Tabs.Modes");
diff --git a/Engine/Source/Editor/LevelEditor/Private/SLevelEditorToolBox.h b/Engine/Source/Editor/LevelEditor/Private/SLevelEditorToolBox.h
--- a/Engine/Source/Editor/LevelEditor/Private/SLevelEditorToolBox.h
+++ b/Engine/Source/Editor/LevelEditor/Private/SLevelEditorToolBox.h
@@ -11,6 +11,7 @@
 #include "Editor.h"
 #include "ILevelEditor.h"
 #include "Misc/NotifyHook.h"
+#include "Subsystems/StatusBarSubsystem.h"
 
 class FExtender;
 class SBorder;
@@ -80,4 +81,7 @@ private:
 
 	/** The icon that should be displayed in the parent tab */
 	const FSlateBrush* TabIcon;
+
+	/** Handle to status bar messages for removing them later */
+	FStatusBarMessageHandle StatusBarMessageHandle;
 };
diff --git a/Engine/Source/Editor/UnrealEd/Private/Subsystems/StatusBarSubsystem.cpp b/Engine/Source/Editor/UnrealEd/Private/Subsystems/StatusBarSubsystem.cpp
--- a/Engine/Source/Editor/UnrealEd/Private/Subsystems/StatusBarSubsystem.cpp
+++ b/Engine/Source/Editor/UnrealEd/Private/Subsystems/StatusBarSubsystem.cpp
@@ -12,12 +12,24 @@
 #include "SourceControlWindows.h"
 #include "Framework/Commands/Commands.h"
 #include "Widgets/Images/SImage.h"
-#include "Widgets/Input/SComboButton.h"
 #include "Widgets/Input/SButton.h"
 #include "Widgets/Layout/SBorder.h"
 #include "Widgets/Layout/SBox.h"
 #include "EditorStyleSet.h"
 #include "FileHelpers.h"
+#include "Widgets/Input/SComboButton.h"
+#include "Widgets/Images/SImage.h"
+#include "Widgets/Text/SRichTextBlock.h"
+#include "Widgets/Layout/SBox.h"
+#include "Widgets/Layout/SBorder.h"
+#include "Widgets/Input/SButton.h"
+#include "Widgets/Input/SMultiLineEditableTextBox.h"
+#include "Framework/Docking/TabManager.h"
+#include "Widgets/Docking/SDockTab.h"
+#include "Framework/Application/SlateApplication.h"
+#include "OutputLogModule.h"
+#include "ContentBrowserModule.h"
+#include "IContentBrowserSingleton.h"
 
 #define LOCTEXT_NAMESPACE "StatusBar"
 
@@ -279,20 +291,8 @@ struct FSourceControlMenuHelpers
 		}
 	}
 
-	static FReply OnSourceControlStatusButtonClicked(TWeakPtr<SMenuAnchor> MenuAnchor)
-	{
-		if (TSharedPtr<SMenuAnchor> MenuAnchorPinned = MenuAnchor.Pin())
-		{
-			MenuAnchorPinned->SetIsOpen(!MenuAnchorPinned->IsOpen());
-		}
-
-		return FReply::Handled();
-	}
-
 	static TSharedRef<SWidget> MakeSourceControlStatusWidget()
 	{
-		TSharedPtr<SMenuAnchor> MenuAnchor;
-
 		return
 			SNew(SComboButton)
 			.ContentPadding(FMargin(6.0f,0.0f))
@@ -332,17 +332,31 @@ FSourceControlMenuHelpers::EQueryState FSourceControlMenuHelpers::QueryState = F
 
 class SStatusBar : public SCompoundWidget
 {
-	SLATE_BEGIN_ARGS(SStatusBar)
+	SLATE_BEGIN_ARGS(SStatusBar)	
 	{}
+		SLATE_EVENT(FSimpleDelegate, OnConsoleClosed)
+
+		SLATE_EVENT(FOnGetContent, OnGetContentBrowser)
+
 	SLATE_END_ARGS()
 
 public:
-	void Construct(const FArguments& InArgs)
+	virtual bool SupportsKeyboardFocus() const { return false; }
+
+	void Construct(const FArguments& InArgs, FName InStatusBarName, const TSharedRef<SDockTab> InParentTab)
 	{
+		StatusBarName = InStatusBarName;
+		ParentTab = InParentTab;
+
 		UpArrow = FAppStyle::Get().GetBrush("StatusBar.ContentBrowserUp");
 		DownArrow = FAppStyle::Get().GetBrush("StatusBar.ContentBrowserDown");
 
 		const FSlateBrush* StatusBarBackground = FAppStyle::Get().GetBrush("StatusBar.Background");
+
+		GetContentBrowserDelegate = InArgs._OnGetContentBrowser;
+
+		FSlateApplication::Get().OnFocusChanging().AddSP(this, &SStatusBar::OnGlobalFocusChanging);
+
 		ChildSlot
 		[
 			SNew(SBox)
@@ -358,42 +372,23 @@ public:
 					.BorderImage(StatusBarBackground)
 					.VAlign(VAlign_Center)
 					[			
-						SNew(SButton)
-						.ButtonStyle(&FAppStyle::Get().GetWidgetStyle<FButtonStyle>("StatusBar.StatusBarButton"))
-						[
-							SNew(SHorizontalBox)
-							+ SHorizontalBox::Slot()
-							.HAlign(HAlign_Left)
-							.VAlign(VAlign_Center)
-							.Padding(2.0f)
-							.AutoWidth()
-							[
-								SNew(SImage)
-								.ColorAndOpacity(FSlateColor::UseForeground())
-								.Image(this, &SStatusBar::GetContentBrowserExpandArrowImage)
-							]
-							+ SHorizontalBox::Slot()
-							.Padding(2.0f)
-							.HAlign(HAlign_Left)
-							.VAlign(VAlign_Center)
-							.AutoWidth()
-							[
-								SNew(SImage)
-								.ColorAndOpacity(FSlateColor::UseForeground())
-								.Image(FAppStyle::Get().GetBrush("StatusBar.ContentBrowserIcon"))
-							]
-							+ SHorizontalBox::Slot()
-							.VAlign(VAlign_Center)
-							.Padding(2.0f)
-							[
-								SNew(STextBlock)
-								.TextStyle(&FAppStyle::Get().GetWidgetStyle<FTextBlockStyle>("NormalText"))
-								.Text(LOCTEXT("StatusBar_ContentBrowserButton", "Content Browser"))
-							]
-						]
+						MakeContentBrowserWidget()
 					]
 				] 
 				+ SHorizontalBox::Slot()
+				.AutoWidth()
+				.Padding(1.0f, 0.0f)
+				[
+					SNew(SBorder)
+					.Padding(0.0f)
+					.BorderImage(StatusBarBackground)
+					.VAlign(VAlign_Center)
+					.Padding(FMargin(6.0f, 0.0f))
+					[
+						MakeDebugConsoleWidget(InArgs._OnConsoleClosed)
+					]
+				]
+				+ SHorizontalBox::Slot()
 				.FillWidth(1.0f)
 				.Padding(1.0f, 0.0f)
 				[
@@ -401,8 +396,9 @@ public:
 					.Padding(0.0f)
 					.BorderImage(StatusBarBackground)
 					.VAlign(VAlign_Center)
+					.Padding(FMargin(6.0f, 0.0f))
 					[
-						SNullWidget::NullWidget
+						MakeStatusMessageWidget()
 					]
 				]
 				+ SHorizontalBox::Slot()
@@ -422,12 +418,124 @@ public:
 		];
 	}
 
+	void PushMessage(FStatusBarMessageHandle InHandle, const TAttribute<FText>& InMessage, const TAttribute<FText>& InHintText)
+	{
+		MessageStack.Emplace(InMessage, InHintText, InHandle);
+	}
+
+	void PopMessage(FStatusBarMessageHandle& InHandle)
+	{
+		if (InHandle.IsValid() && MessageStack.Num() > 0)
+		{
+			MessageStack.RemoveAll([InHandle](const FStatusBarMessage& Message)
+				{
+					return Message.Handle == InHandle;
+				});
+		}
+	}
+
+	void ClearAllMessages()
+	{
+		MessageStack.Empty();
+	}
+
+	EVisibility GetHelpIconVisibility() const
+	{
+		if (MessageStack.Num() > 0)
+		{
+			const FStatusBarMessage& MessageData = MessageStack.Top();
+
+			const FText& Message = MessageData.MessageText.Get();
+			const FText& HintText = MessageData.HintText.Get();
+
+			return (!Message.IsEmpty() || !HintText.IsEmpty()) ? EVisibility::SelfHitTestInvisible : EVisibility::Collapsed;
+		}
+
+		return EVisibility::Collapsed;
+
+	}
+
+	TSharedPtr<SDockTab> GetParentTab() const { return ParentTab.Pin(); }
+
+	void FocusDebugConsole()
+	{
+		FSlateApplication::Get().SetKeyboardFocus(ConsoleEditBox, EFocusCause::SetDirectly);
+	}
+
+	bool IsDebugConsoleFocused() const
+	{
+		return ConsoleEditBox->HasKeyboardFocus();
+	}
+
+	void OnGlobalFocusChanging(const FFocusEvent& FocusEvent, const FWeakWidgetPath& OldFocusedWidgetPath, const TSharedPtr<SWidget>& OldFocusedWidget, const FWidgetPath& NewFocusedWidgetPath, const TSharedPtr<SWidget>& NewFocusedWidget)
+	{
+		//if (ContentBrowserOverlayContent.IsValid() && !NewFocusedWidgetPath.ContainsWidget(ContentBrowserOverlayContent.ToSharedRef()))
+		//{
+		//	DismissContentBrowser();
+		//}
+	}
+
 private:
 	const FSlateBrush* GetContentBrowserExpandArrowImage() const
 	{
 		return DownArrow;
 	}
 
+	FText GetStatusBarMessage() const
+	{
+		FText FullMessage;
+		if (MessageStack.Num() > 0)
+		{
+			const FStatusBarMessage& MessageData = MessageStack.Top();
+
+			const FText& Message = MessageData.MessageText.Get();
+			const FText& HintText = MessageData.HintText.Get();
+
+			FullMessage = HintText.IsEmpty() ? Message : FText::Format(LOCTEXT("StatusBarMessageFormat", "{0} <StatusBar.Message.InHintText>{1}</>"), Message, HintText);
+		}
+
+		return FullMessage;
+	}
+
+	TSharedRef<SWidget> MakeContentBrowserWidget()
+	{
+		return
+			SNew(SButton)
+			.ButtonStyle(&FAppStyle::Get().GetWidgetStyle<FButtonStyle>("StatusBar.StatusBarButton"))
+			.OnClicked(this, &SStatusBar::OnContentBrowserButtonClicked)
+			[
+				SNew(SHorizontalBox)
+				+ SHorizontalBox::Slot()
+				.HAlign(HAlign_Left)
+				.VAlign(VAlign_Center)
+				.Padding(2.0f)
+				.AutoWidth()
+				[
+					SNew(SImage)
+					.ColorAndOpacity(FSlateColor::UseForeground())
+					.Image(this, &SStatusBar::GetContentBrowserExpandArrowImage)
+				]
+				+ SHorizontalBox::Slot()
+				.Padding(2.0f)
+				.HAlign(HAlign_Left)
+				.VAlign(VAlign_Center)
+				.AutoWidth()
+				[
+					SNew(SImage)
+					.ColorAndOpacity(FSlateColor::UseForeground())
+					.Image(FAppStyle::Get().GetBrush("StatusBar.ContentBrowserIcon"))
+				]
+				+ SHorizontalBox::Slot()
+				.VAlign(VAlign_Center)
+				.Padding(2.0f)
+				[
+					SNew(STextBlock)
+					.TextStyle(&FAppStyle::Get().GetWidgetStyle<FTextBlockStyle>("NormalText"))
+					.Text(LOCTEXT("StatusBar_ContentBrowserButton", "Content Browser"))
+				]
+			];	
+	}
+
 	TSharedRef<SWidget> MakeStatusBarToolBarWidget()
 	{
 		RegisterStatusBarMenu();
@@ -435,30 +543,96 @@ private:
 		FToolMenuContext MenuContext;
 		RegisterSourceControlStatus();
 
-
 		return UToolMenus::Get()->GenerateWidget("StatusBar.ToolBar", MenuContext);
 	}
 
+	TSharedRef<SWidget> MakeDebugConsoleWidget(FSimpleDelegate OnConsoleClosed)
+	{
+		FOutputLogModule& OutputLogModule = FModuleManager::LoadModuleChecked<FOutputLogModule>(TEXT("OutputLog"));
+
+		return
+			SNew(SBox)
+			.WidthOverride(350.f)
+			[
+				OutputLogModule.MakeConsoleInputBox(ConsoleEditBox, OnConsoleClosed)
+			];
+	}
+
+	TSharedRef<SWidget> MakeStatusMessageWidget()
+	{
+		return 
+			SNew(SHorizontalBox)
+			+ SHorizontalBox::Slot()
+			.AutoWidth()
+			.VAlign(VAlign_Center)
+			.HAlign(HAlign_Center)
+			[
+				SNew(SImage)
+				.ColorAndOpacity(FSlateColor::UseForeground())
+				.Image(FAppStyle::Get().GetBrush("StatusBar.HelpIcon"))
+				.Visibility(this, &SStatusBar::GetHelpIconVisibility)
+			]
+			+ SHorizontalBox::Slot()
+			.VAlign(VAlign_Center)
+			.Padding(5.0f, 0.0f)
+			[
+				SNew(SRichTextBlock)
+				.TextStyle(&FAppStyle::Get().GetWidgetStyle<FTextBlockStyle>("StatusBar.Message.MessageText"))
+				.Text(this, &SStatusBar::GetStatusBarMessage)
+				.DecoratorStyleSet(&FAppStyle::Get())
+			];
+	}
+
+	FReply OnContentBrowserButtonClicked()
+	{
+		/*if(!ContentBrowserOverlayContent.IsValid())
+		{
+			TSharedPtr<SWindow> Window = FSlateApplication::Get().FindWidgetWindow(SharedThis(this));
+
+			Window->AddOverlaySlot()
+				.VAlign(VAlign_Bottom)
+				.Padding(FMargin(0, 0, 0, GetTickSpaceGeometry().GetLocalSize().Y + 1))
+				[
+					SAssignNew(ContentBrowserOverlayContent, SBox)
+					.HeightOverride(Window->GetSizeInScreen().Y * 0.3f)
+					[
+						GetContentBrowserDelegate.Execute()
+					]
+				];
+		}*/
+
+		return FReply::Handled();
+	}
+
+	void DismissContentBrowser()
+	{
+		if (ContentBrowserOverlayContent.IsValid())
+		{
+			TSharedPtr<SWindow> Window = FSlateApplication::Get().FindWidgetWindow(SharedThis(this));
+
+			Window->RemoveOverlaySlot(ContentBrowserOverlayContent.ToSharedRef());
+
+			ContentBrowserOverlayContent.Reset();
+		}
+	}
+
 	void RegisterStatusBarMenu()
 	{
-		static const FName StatusBarName("StatusBar.ToolBar");
+		static const FName StatusBarToolBarName("StatusBar.ToolBar");
 		UToolMenus* ToolMenus = UToolMenus::Get();
-		if (ToolMenus->IsMenuRegistered(StatusBarName))
+		if (ToolMenus->IsMenuRegistered(StatusBarToolBarName))
 		{
 			return;
 		}
 
-		UToolMenu* ToolBar = ToolMenus->RegisterMenu(StatusBarName, NAME_None, EMultiBoxType::SlimHorizontalToolBar);
+		UToolMenu* ToolBar = ToolMenus->RegisterMenu(StatusBarToolBarName, NAME_None, EMultiBoxType::SlimHorizontalToolBar);
 		ToolBar->StyleName = "StatusBarToolBar";
-
 	}
 
 	void RegisterSourceControlStatus()
 	{
-
 		// Source Control preferences
 		FSourceControlMenuHelpers::CheckSourceControlStatus();
-
 		{
 			UToolMenu* SourceControlMenu = UToolMenus::Get()->ExtendMenu("StatusBar.ToolBar");
 			FToolMenuSection& Section = SourceControlMenu->FindOrAddSection("SourceControl");
@@ -473,11 +647,31 @@ private:
 				));
 		}
 	}
+
+	struct FStatusBarMessage
+	{
+		FStatusBarMessage(const TAttribute<FText>& InMessageText, const TAttribute<FText>& InHintText, FStatusBarMessageHandle InHandle)
+			: MessageText(InMessageText)
+			, HintText(InHintText)
+			, Handle(InHandle)
+		{}
+
+		TAttribute<FText> MessageText;
+		TAttribute<FText> HintText;
+		FStatusBarMessageHandle Handle;
+	};
 private:
+	TArray<FStatusBarMessage> MessageStack;
+	TSharedPtr<SMultiLineEditableTextBox> ConsoleEditBox;
+	TWeakPtr<SDockTab> ParentTab;
+	TSharedPtr<SWidget> ContentBrowserOverlayContent;
+	FOnGetContent GetContentBrowserDelegate;
 	const FSlateBrush* UpArrow;
 	const FSlateBrush* DownArrow;
+	FName StatusBarName;
 };
 
+int32 UStatusBarSubsystem::HandleCounter = 0;
 
 void UStatusBarSubsystem::Initialize(FSubsystemCollectionBase& Collection)
 {
@@ -489,13 +683,118 @@ void UStatusBarSubsystem::Deinitialize()
 	FSourceControlCommands::Unregister();
 }
 
-TSharedRef<SWidget> UStatusBarSubsystem::MakeStatusBarWidget(FName StatusBarName)
+bool UStatusBarSubsystem::FocusDebugConsole(TSharedRef<SWindow> ParentWindow)
 {
-	TSharedRef<SStatusBar> StatusBar = SNew(SStatusBar);
+	bool bFocusedSuccessfully = false;
+
+	for (auto StatusBar : StatusBars)
+	{
+		if (TSharedPtr<SStatusBar> StatusBarPinned = StatusBar.Value.Pin())
+		{
+			TSharedPtr<SDockTab> ParentTab = StatusBarPinned->GetParentTab();
+			if (ParentTab && ParentTab->IsForeground() && ParentTab->GetParentWindow() == ParentWindow)
+			{
+				// Cache off the previously focused widget so we can restore focus if the user hits the focus key again
+				PreviousKeyboardFocusedWidget = FSlateApplication::Get().GetKeyboardFocusedWidget();
+				StatusBarPinned->FocusDebugConsole();
+
+				bFocusedSuccessfully = true;
+				break;
+			}
+		}
+	}
+
+	return bFocusedSuccessfully;
+}
+
+TSharedRef<SWidget> UStatusBarSubsystem::MakeStatusBarWidget(FName StatusBarName, const TSharedRef<SDockTab>& InParentTab)
+{
+	TSharedRef<SStatusBar> StatusBar =
+		SNew(SStatusBar, StatusBarName, InParentTab)
+		.OnConsoleClosed_UObject(this, &UStatusBarSubsystem::OnDebugConsoleClosed)
+		.OnGetContentBrowser_UObject(this, &UStatusBarSubsystem::OnGetContentBrowser);
+
+	// Clean up stale status bars
+	for (auto It = StatusBars.CreateIterator(); It; ++It)
+	{
+		if (!It.Value().IsValid())
+		{
+			It.RemoveCurrent();
+		}
+	}
 
 	StatusBars.Add(StatusBarName, StatusBar);
 
 	return StatusBar;
 }
 
+FStatusBarMessageHandle UStatusBarSubsystem::PushStatusBarMessage(FName StatusBarName, const TAttribute<FText>& InMessage, const TAttribute<FText>& InHintText)
+{
+	if (TSharedPtr<SStatusBar> StatusBar = GetStatusBar(StatusBarName))
+	{
+		FStatusBarMessageHandle NewHandle(++HandleCounter);
+
+		StatusBar->PushMessage(NewHandle, InMessage, InHintText);
+
+		return NewHandle;
+	}
+
+	return FStatusBarMessageHandle();
+}
+
+FStatusBarMessageHandle UStatusBarSubsystem::PushStatusBarMessage(FName StatusBarName, const TAttribute<FText>& InMessage)
+{
+	return PushStatusBarMessage(StatusBarName, InMessage, TAttribute<FText>());
+}
+
+void UStatusBarSubsystem::PopStatusBarMessage(FName StatusBarName, FStatusBarMessageHandle InHandle)
+{
+	if (TSharedPtr<SStatusBar> StatusBar = GetStatusBar(StatusBarName))
+	{
+		StatusBar->PopMessage(InHandle);
+	}
+}
+
+void UStatusBarSubsystem::ClearStatusBarMessages(FName StatusBarName)
+{
+	if (TSharedPtr<SStatusBar> StatusBar = GetStatusBar(StatusBarName))
+	{
+		StatusBar->ClearAllMessages();
+	}
+}
+
+void UStatusBarSubsystem::OnDebugConsoleClosed()
+{
+	if (PreviousKeyboardFocusedWidget.IsValid())
+	{
+		FSlateApplication::Get().SetKeyboardFocus(PreviousKeyboardFocusedWidget.Pin());
+		PreviousKeyboardFocusedWidget.Reset();
+	}
+}
+
+void UStatusBarSubsystem::CreateContentBrowserIfNeeded()
+{
+	if(!StatusBarContentBrowser.IsValid())
+	{
+		IContentBrowserSingleton& ContentBrowserSingleton = FModuleManager::Get().LoadModuleChecked<FContentBrowserModule>("ContentBrowser").Get();;
+
+		FContentBrowserConfig Config;
+		Config.bCanSetAsPrimaryBrowser = false;
+
+		StatusBarContentBrowser = ContentBrowserSingleton.CreateContentBrowser("StatusBarContentBrowser", nullptr, &Config);
+	}
+}
+
+TSharedPtr<SStatusBar> UStatusBarSubsystem::GetStatusBar(FName StatusBarName) const
+{
+	return StatusBars.FindRef(StatusBarName).Pin();
+}
+
+TSharedRef<SWidget> UStatusBarSubsystem::OnGetContentBrowser()
+{
+	CreateContentBrowserIfNeeded();
+
+	return StatusBarContentBrowser.ToSharedRef();
+}
+
 #undef LOCTEXT_NAMESPACE
\ No newline at end of file
diff --git a/Engine/Source/Editor/UnrealEd/Private/Toolkits/GlobalEditorCommonCommands.cpp b/Engine/Source/Editor/UnrealEd/Private/Toolkits/GlobalEditorCommonCommands.cpp
--- a/Engine/Source/Editor/UnrealEd/Private/Toolkits/GlobalEditorCommonCommands.cpp
+++ b/Engine/Source/Editor/UnrealEd/Private/Toolkits/GlobalEditorCommonCommands.cpp
@@ -14,6 +14,7 @@
 #include "Widgets/Docking/SDockTab.h"
 #include "Toolkits/SGlobalOpenAssetDialog.h"
 #include "Toolkits/SGlobalTabSwitchingDialog.h"
+#include "Subsystems/StatusBarSubsystem.h"
 
 #define LOCTEXT_NAMESPACE "GlobalEditorCommonCommands"
 
@@ -131,13 +132,17 @@ void FGlobalEditorCommonCommands::OnSummonedConsoleCommandBox()
 	if( ParentWindow.IsValid() && ParentWindow->GetType() == EWindowType::Normal )
 	{
 		TSharedRef<SWindow> WindowRef = ParentWindow.ToSharedRef();
-		FOutputLogModule& OutputLogModule = FModuleManager::LoadModuleChecked< FOutputLogModule >(TEXT("OutputLog"));
+		FOutputLogModule& OutputLogModule = FModuleManager::LoadModuleChecked<FOutputLogModule>(TEXT("OutputLog"));
 
-		FDebugConsoleDelegates Delegates;
-		Delegates.OnConsoleCommandExecuted = FSimpleDelegate::CreateStatic( &CloseDebugConsole );
-		Delegates.OnCloseConsole = FSimpleDelegate::CreateStatic( &CloseDebugConsole );
+		if (!GEditor->GetEditorSubsystem<UStatusBarSubsystem>()->FocusDebugConsole(ParentWindow.ToSharedRef()))
+		{
+			// A status bar was not found, pop open a floating window instead
+			FDebugConsoleDelegates Delegates;
+			Delegates.OnConsoleCommandExecuted = FSimpleDelegate::CreateStatic(&CloseDebugConsole);
+			Delegates.OnCloseConsole = FSimpleDelegate::CreateStatic(&CloseDebugConsole);
 
-		OutputLogModule.ToggleDebugConsoleForWindow(WindowRef, EDebugConsoleStyle::Compact, Delegates);
+			OutputLogModule.ToggleDebugConsoleForWindow(WindowRef, EDebugConsoleStyle::Compact, Delegates);
+		}
 	}
 }
 
diff --git a/Engine/Source/Editor/UnrealEd/Private/Toolkits/SStandaloneAssetEditorToolkitHost.cpp b/Engine/Source/Editor/UnrealEd/Private/Toolkits/SStandaloneAssetEditorToolkitHost.cpp
--- a/Engine/Source/Editor/UnrealEd/Private/Toolkits/SStandaloneAssetEditorToolkitHost.cpp
+++ b/Engine/Source/Editor/UnrealEd/Private/Toolkits/SStandaloneAssetEditorToolkitHost.cpp
@@ -12,8 +12,9 @@
 #include "Interfaces/IMainFrameModule.h"
 #include "Widgets/Docking/SDockTab.h"
 #include "UObject/Package.h"
-
+#include "Subsystems/StatusBarSubsystem.h"
 #include "ToolMenus.h"
+#include "IIntroTutorials.h"
 
 #define LOCTEXT_NAMESPACE "StandaloneAssetEditorToolkit"
 
@@ -58,8 +59,34 @@ void SStandaloneAssetEditorToolkitHost::SetupInitialContent( const TSharedRef<FT
 
 	HostTabPtr = InHostTab;
 
+	MenuOverlayWidgetContent = SNew(SBox);
+
 	RestoreFromLayout(DefaultLayout);
 	GenerateMenus(bCreateDefaultStandaloneMenu);
+
+	if (InHostTab)
+	{
+		IIntroTutorials& IntroTutorials = FModuleManager::LoadModuleChecked<IIntroTutorials>(TEXT("IntroTutorials"));
+		TSharedRef<SWidget> TutorialWidget = IntroTutorials.CreateTutorialsWidget(HostedAssetEditorToolkit->GetToolkitContextFName(), InHostTab->GetParentWindow());
+
+		InHostTab->SetRightContent(
+			SNew(SHorizontalBox)
+			+SHorizontalBox::Slot()
+			.VAlign(VAlign_Center)
+			.Padding(8.0f, 0.0f, 8.0f, 0.0f)
+			[
+				MenuOverlayWidgetContent.ToSharedRef()
+			]
+			+ SHorizontalBox::Slot()
+			.AutoWidth()
+			.Padding(8.0f, 0.0f, 8.0f, 0.0f)
+			.VAlign(VAlign_Center)
+			[
+				TutorialWidget
+			]
+		);
+	}
+
 }
 
 void SStandaloneAssetEditorToolkitHost::CreateDefaultStandaloneMenuBar(UToolMenu* MenuBar)
@@ -135,49 +162,22 @@ void SStandaloneAssetEditorToolkitHost::RestoreFromLayout( const TSharedRef<FTab
 	TSharedPtr<SWidget> RestoredUI = MyTabManager->RestoreFrom( NewLayout, ParentWindow );
 
 	checkf(RestoredUI.IsValid(), TEXT("The layout must have a primary dock area") );
-	
-	MenuOverlayWidgetContent.Reset();
-	MenuWidgetContent.Reset();
+
 	this->ChildSlot
 	[
-		SNew( SVerticalBox )
-			// Menu bar area
-			+SVerticalBox::Slot()
-				.AutoHeight()
-				[
-					SNew(SOverlay)
-					// The menu bar
-					+SOverlay::Slot()
-					[
-						SAssignNew(MenuWidgetContent, SBorder)
-						.Padding(0)
-						.BorderImage(FEditorStyle::GetBrush("NoBorder"))
-						[
-							DefaultMenuWidget.ToSharedRef()
-						]
-					]
-					// The menu bar overlay
-					+SOverlay::Slot()
-					.HAlign(HAlign_Right)
-					[
-						SNew(SHorizontalBox)
-						+ SHorizontalBox::Slot()
-						[
-							SAssignNew(MenuOverlayWidgetContent, SBorder)
-							.Padding(0)
-							.BorderImage(FEditorStyle::GetBrush("NoBorder"))
-						]				
-					]
-				]
-			// Viewport/Document/Docking area
-			+SVerticalBox::Slot()
-				.Padding( 1.0f )
-				
-				// Fills all leftover space
-				.FillHeight( 1.0f )
-				[
-					RestoredUI.ToSharedRef()
-				]
+		SNew(SVerticalBox)
+		+ SVerticalBox::Slot()
+		.Padding(0.0f, 1.0f, 0.0f, 0.0f)
+		.FillHeight(1.0f)
+		[
+			RestoredUI.ToSharedRef()
+		]
+		+ SVerticalBox::Slot()
+		.Padding(0.0f, 1.0f, 0.0f, 0.0f)
+		.AutoHeight()
+		[
+			GEditor->GetEditorSubsystem<UStatusBarSubsystem>()->MakeStatusBarWidget(GetMenuName(), HostTab)
+		]
 	];
 }
 
@@ -208,8 +208,6 @@ void SStandaloneAssetEditorToolkitHost::GenerateMenus(bool bForceCreateMenu)
 		HostedAssetEditorToolkit->InitToolMenuContext(ToolMenuContext);
 		IMainFrameModule& MainFrameModule = FModuleManager::LoadModuleChecked<IMainFrameModule>( "MainFrame" );
 		DefaultMenuWidget = MainFrameModule.MakeMainMenu( MyTabManager, AssetEditorMenuName, ToolMenuContext );
-
-		MenuWidgetContent->SetContent(DefaultMenuWidget.ToSharedRef());
 	}
 }
 
@@ -365,7 +363,7 @@ void SStandaloneAssetEditorToolkitHost::OnTabClosed(TSharedRef<SDockTab> TabClos
 {
 	check(TabClosed == HostTabPtr.Pin());
 
-	MyTabManager->SetMenuMultiBox(nullptr);
+	MyTabManager->SetMenuMultiBox(nullptr, nullptr);
 	
 	if(HostedAssetEditorToolkit.IsValid())
 	{
diff --git a/Engine/Source/Editor/UnrealEd/Public/Subsystems/StatusBarSubsystem.h b/Engine/Source/Editor/UnrealEd/Public/Subsystems/StatusBarSubsystem.h
--- a/Engine/Source/Editor/UnrealEd/Public/Subsystems/StatusBarSubsystem.h
+++ b/Engine/Source/Editor/UnrealEd/Public/Subsystems/StatusBarSubsystem.h
@@ -9,6 +9,41 @@
 
 class SWidget;
 class SStatusBar;
+class SWindow;
+class SDockTab;
+
+template<typename ObjectType> 
+class TAttribute;
+
+struct FStatusBarMessageHandle
+{
+	friend class UStatusBarSubsystem;
+
+	FStatusBarMessageHandle()
+		: Id(INDEX_NONE)
+	{}
+
+	bool IsValid()
+	{
+		return Id != INDEX_NONE;
+	}
+
+	void Reset()
+	{
+		Id = INDEX_NONE;
+	}
+
+	bool operator==(const FStatusBarMessageHandle& OtherHandle) const
+	{
+		return Id == OtherHandle.Id;
+	}
+private:
+	FStatusBarMessageHandle(int32 InId)
+		: Id(InId)
+	{}
+
+	int32 Id;
+};
 
 UCLASS()
 class UNREALED_API UStatusBarSubsystem : public UEditorSubsystem
@@ -27,12 +62,58 @@ public:
 	 */
 	virtual void Deinitialize() override final;
 
+	/**
+	 * Focuses the debug console on the status bar for status bar residing in the passed in parent window 
+	 *
+	 * @param ParentWindow The parent window of the status bar 
+	 * @return true of a status bar debug console was found and focused 
+	 */
+	bool FocusDebugConsole(TSharedRef<SWindow> ParentWindow);
+
 	/** 
 	 * Creates a new instance of a status bar widget
 	 *
 	 * @param StatusBarName	The name of the status bar for updating it later.
 	 */
-	TSharedRef<SWidget> MakeStatusBarWidget(FName StatusBarName);
+	TSharedRef<SWidget> MakeStatusBarWidget(FName StatusBarName, const TSharedRef<SDockTab>& InParentTab);
+
+	/** 
+	 * Pushes a new status bar message
+	 *
+	 * @param StatusBarName	The name of the status bar to push messages to
+	 * @param InMessage		The message to display
+	 * @param InHintText	Optional hint text message.  This message will be highlighted to make it stand out
+	 * @return	A handle to the message for clearing it later
+	 */
+	FStatusBarMessageHandle PushStatusBarMessage(FName StatusBarName, const TAttribute<FText>& InMessage, const TAttribute<FText>& InHintText);
+	FStatusBarMessageHandle PushStatusBarMessage(FName StatusBarName, const TAttribute<FText>& InMessage);
+
+	/**
+	 * Removes a message from the status bar.  When messages are removed the previous message on the stack (if any) is displayed
+	 *
+	 * @param StatusBarName	The name of the status bar to remove from
+	 * @param InHandle		Handle to the status bar message to remove
+	 */
+	void PopStatusBarMessage(FName StatusBarName, FStatusBarMessageHandle InHandle);
+
+	/**
+	 * Removes all message from the status bar
+	 *
+	 * @param StatusBarName	The name of the status bar to remove from
+	 */
+	void ClearStatusBarMessages(FName StatusBarName);
+
+private:
+	void OnDebugConsoleClosed();
+	void CreateContentBrowserIfNeeded();
+	TSharedPtr<SStatusBar> GetStatusBar(FName StatusBarName) const;
+	TSharedRef<SWidget> OnGetContentBrowser();
+
 private:
 	TMap<FName, TWeakPtr<SStatusBar>> StatusBars;
+	TWeakPtr<SWidget> PreviousKeyboardFocusedWidget;
+
+	/** The floating content browser that is opened via the content browser button in the status bar */
+	TSharedPtr<SWidget> StatusBarContentBrowser;
+	static int32 HandleCounter;
 };
