Viewport related styling changes
- Overhauled in viewport toolbar
- Updated active viewport outline and hidden tab dock icons

diff --git a/Engine/Source/Editor/Cascade/Private/SCascadePreviewToolbar.cpp b/Engine/Source/Editor/Cascade/Private/SCascadePreviewToolbar.cpp
--- a/Engine/Source/Editor/Cascade/Private/SCascadePreviewToolbar.cpp
+++ b/Engine/Source/Editor/Cascade/Private/SCascadePreviewToolbar.cpp
@@ -20,8 +20,6 @@ void SCascadePreviewViewportToolBar::Construct(const FArguments& InArgs)
 	[
 		SNew(SBorder)
 		.BorderImage(FEditorStyle::GetBrush("NoBorder"))
-		// Color and opacity is changed based on whether or not the mouse cursor is hovering over the toolbar area
-		.ColorAndOpacity(this, &SViewportToolBar::OnGetColorAndOpacity)
 		.ForegroundColor(FEditorStyle::GetSlateColor("DefaultForeground"))
 		[
 			SNew(SHorizontalBox)
diff --git a/Engine/Source/Editor/EditorStyle/Private/StarshipStyle.cpp b/Engine/Source/Editor/EditorStyle/Private/StarshipStyle.cpp
--- a/Engine/Source/Editor/EditorStyle/Private/StarshipStyle.cpp
+++ b/Engine/Source/Editor/EditorStyle/Private/StarshipStyle.cpp
@@ -2283,7 +2283,82 @@ void FStarshipEditorStyle::FStyle::SetupSequencerStyles()
 
 void FStarshipEditorStyle::FStyle::SetupViewportStyles()
 {
-	// Viewport ToolbarBar
+	{
+		FToolBarStyle ViewportToolbarStyle = FStarshipCoreStyle::GetCoreStyle().GetWidgetStyle<FToolBarStyle>("SlimToolBar");
+
+		FButtonStyle ViewportMenuButton = FButtonStyle(FStarshipCoreStyle::GetCoreStyle().GetWidgetStyle<FButtonStyle>("NoBorder"))
+			.SetNormalForeground(FStyleColors::Foreground)
+			.SetHoveredForeground(FStyleColors::ForegroundHover)
+			.SetPressedForeground(FStyleColors::ForegroundHover)
+			.SetDisabledForeground(FStyleColors::Foreground)
+			.SetNormalPadding(FMargin(0, 2, 0, 2))
+			.SetPressedPadding(FMargin(0, 3, 0, 1));
+
+
+		Set("EditorViewportToolBar.Button", ViewportMenuButton);
+
+
+		FCheckBoxStyle MaximizeRestoreButton = FCheckBoxStyle(ViewportToolbarStyle.ToggleButton)
+			.SetForegroundColor(FStyleColors::Foreground)
+			.SetPressedForegroundColor(FStyleColors::ForegroundHover)
+			.SetHoveredForegroundColor(FStyleColors::ForegroundHover)
+			.SetCheckedForegroundColor(FStyleColors::Foreground)
+			.SetCheckedPressedForegroundColor(FStyleColors::ForegroundHover)
+			.SetCheckedHoveredForegroundColor(FStyleColors::ForegroundHover);
+
+		Set("EditorViewportToolBar.MaximizeRestoreButton", MaximizeRestoreButton);
+
+		ViewportToolbarStyle
+			.SetBackground(FSlateNoResource())
+			.SetIconSize(Icon16x16)
+			.SetLabelPadding(FMargin(0))
+			.SetComboButtonPadding(FMargin(0))
+			.SetBlockPadding(FMargin(2.0f,0.0f))
+			.SetIndentedBlockPadding(FMargin(0))
+			.SetButtonPadding(FMargin(0))
+			.SetSubMenuIndicator(FSlateNoResource())
+			.SetCheckBoxPadding(FMargin(3.0f, 0.0f))
+			.SetButtonStyle(ViewportMenuButton)
+			.SetSeparatorBrush(FSlateNoResource())
+			.SetSeparatorPadding(FMargin(4.0f, 0.0f))
+			.SetExpandBrush(IMAGE_BRUSH("Icons/toolbar_expand_16x", Icon8x8));
+
+		Set("EditorViewportToolBar", ViewportToolbarStyle);
+
+		FButtonStyle ViewportMenuWarningButton = FButtonStyle(FStarshipCoreStyle::GetCoreStyle().GetWidgetStyle<FButtonStyle>("NoBorder"))
+			.SetNormalForeground(FStyleColors::AccentYellow)
+			.SetHoveredForeground(FStyleColors::ForegroundHover)
+			.SetPressedForeground(FStyleColors::ForegroundHover)
+			.SetDisabledForeground(FStyleColors::AccentYellow)
+			.SetNormalPadding(FMargin(0, 2, 0, 2))
+			.SetPressedPadding(FMargin(0, 3, 0, 1));
+
+		Set("EditorViewportToolBar.WarningButton", ViewportMenuWarningButton);
+
+		FLinearColor ToolbarBackgroundColor = FStyleColors::Foldout;
+		ToolbarBackgroundColor.A = .75f;
+
+		Set("EditorViewportToolBar.Background", new FSlateColorBrush(ToolbarBackgroundColor));
+		Set("EditorViewportToolBar.OptionsDropdown", new CORE_IMAGE_BRUSH_SVG("Starship/Common/hamburger", FVector2D(13,8)));
+
+		Set("EditorViewportToolBar.Font", FStyleFonts::Get().Normal);
+
+		Set("EditorViewportToolBar.MenuButton", FButtonStyle(Button)
+			.SetNormal(BOX_BRUSH("Common/SmallRoundedButton", FMargin(7.f / 16.f), FLinearColor(1, 1, 1, 0.75f)))
+			.SetHovered(BOX_BRUSH("Common/SmallRoundedButton", FMargin(7.f / 16.f), FLinearColor(1, 1, 1, 1.0f)))
+			.SetPressed(BOX_BRUSH("Common/SmallRoundedButton", FMargin(7.f / 16.f)))
+		);
+
+		
+		Set("EditorViewportToolBar.Button", HoverHintOnly);
+
+		Set("EditorViewportToolBar.MenuDropdown", new IMAGE_BRUSH("Common/ComboArrow", Icon8x8));
+		Set("EditorViewportToolBar.Maximize.Normal", new IMAGE_BRUSH_SVG("Starship/EditorViewport/maximize", Icon16x16));
+		Set("EditorViewportToolBar.Maximize.Checked", new IMAGE_BRUSH_SVG("Starship/EditorViewport/restore", Icon16x16));
+		Set("EditorViewportToolBar.RestoreFromImmersive.Normal", new IMAGE_BRUSH("Icons/icon_RestoreFromImmersive_16px", Icon16x16));
+	}
+
+	// Legacy Viewport ToolbarBar
 	{
 		FToolBarStyle ViewportToolbarStyle =
 			FToolBarStyle()
@@ -2292,7 +2367,7 @@ void FStarshipEditorStyle::FStyle::SetupViewportStyles()
 			.SetSubMenuIndicator(IMAGE_BRUSH("Common/SubmenuArrow", Icon8x8))
 			.SetComboButtonPadding(FMargin(0))
 			.SetButtonPadding(FMargin(0))
-			.SetCheckBoxPadding(FMargin(0))
+			.SetCheckBoxPadding(FMargin(4))
 			.SetSeparatorBrush(BOX_BRUSH("Old/Button", 8.0f / 32.0f, FLinearColor::Transparent))
 			.SetSeparatorPadding(FMargin(1.f, 0.f, 0.f, 0.f))
 			.SetIconSize(Icon16x16)
@@ -4040,29 +4115,28 @@ void FStarshipEditorStyle::FStyle::SetupLevelEditorStyle()
 		Set("MainFrame.AboutUnrealEd", new IMAGE_BRUSH("Icons/Help/icon_Help_unreal_16x", Icon16x16));
 		Set("MainFrame.CreditsUnrealEd", new IMAGE_BRUSH("Icons/Help/icon_Help_credits_16x", Icon16x16));
 
-		const FLinearColor IconColor = FLinearColor::Black;
-		Set( "EditorViewport.TranslateMode", new IMAGE_BRUSH( "Icons/icon_translateb_16x", Icon16x16 ) );
-		Set( "EditorViewport.TranslateMode.Small", new IMAGE_BRUSH( "Icons/icon_translateb_16x", Icon16x16 ) );
-		Set( "EditorViewport.RotateMode", new IMAGE_BRUSH( "Icons/icon_rotateb_16x", Icon16x16 ) );
-		Set( "EditorViewport.RotateMode.Small", new IMAGE_BRUSH( "Icons/icon_rotateb_16x", Icon16x16 ) );
-		Set( "EditorViewport.ScaleMode", new IMAGE_BRUSH( "Icons/icon_scaleb_16x", Icon16x16 ) );
-		Set( "EditorViewport.ScaleMode.Small", new IMAGE_BRUSH( "Icons/icon_scaleb_16x", Icon16x16 ) );
+		Set( "EditorViewport.TranslateMode", new IMAGE_BRUSH_SVG( "Starship/EditorViewport/translate", Icon16x16 ) );
+		Set( "EditorViewport.RotateMode", new IMAGE_BRUSH_SVG("Starship/EditorViewport/rotate", Icon16x16 ) );
+		Set( "EditorViewport.ScaleMode", new IMAGE_BRUSH_SVG( "Starship/EditorViewport/scale", Icon16x16 ) );
+
 		Set( "EditorViewport.TranslateRotateMode", new IMAGE_BRUSH( "Icons/icon_translate_rotate_40x", Icon20x20 ) );
-		Set( "EditorViewport.TranslateRotateMode.Small", new IMAGE_BRUSH( "Icons/icon_translate_rotate_40x", Icon20x20 ) );
 		Set( "EditorViewport.TranslateRotate2DMode", new IMAGE_BRUSH("Icons/icon_translate_rotate_2d_40x", Icon20x20));
-		Set( "EditorViewport.TranslateRotate2DMode.Small", new IMAGE_BRUSH("Icons/icon_translate_rotate_2d_40x", Icon20x20));
-		Set("EditorViewport.ToggleRealTime", new IMAGE_BRUSH("Icons/icon_MatEd_Realtime_40x", Icon40x40));
-		Set( "EditorViewport.ToggleRealTime.Small", new IMAGE_BRUSH( "Icons/icon_MatEd_Realtime_40x", Icon20x20 ) );
-		Set( "EditorViewport.LocationGridSnap", new IMAGE_BRUSH( "Old/LevelEditor/LocationGridSnap", Icon14x14, IconColor) );
-		Set( "EditorViewport.RotationGridSnap", new IMAGE_BRUSH( "Old/LevelEditor/RotationGridSnap", Icon14x14, IconColor ) );
-		Set( "EditorViewport.Layer2DSnap", new IMAGE_BRUSH("Old/LevelEditor/Layer2DSnap", Icon14x14, IconColor));
-		Set( "EditorViewport.ScaleGridSnap", new IMAGE_BRUSH( "Old/LevelEditor/ScaleGridSnap", Icon14x14, IconColor ) );
-		Set( "EditorViewport.ToggleSurfaceSnapping", new IMAGE_BRUSH( "Icons/icon_surface_snapping_14px", Icon14x14 ) );
-		Set( "EditorViewport.RelativeCoordinateSystem_Local", new IMAGE_BRUSH( "Icons/icon_axis_local_16px", Icon16x16, IconColor ) );
-		Set( "EditorViewport.RelativeCoordinateSystem_Local.Small", new IMAGE_BRUSH( "Icons/icon_axis_local_16px", Icon16x16, IconColor ) );
-		Set( "EditorViewport.RelativeCoordinateSystem_World", new IMAGE_BRUSH( "Icons/icon_axis_world_16px", Icon16x16, IconColor ) );
-		Set( "EditorViewport.RelativeCoordinateSystem_World.Small", new IMAGE_BRUSH( "Icons/icon_axis_world_16px", Icon16x16, IconColor ) );
-		Set( "EditorViewport.CamSpeedSetting", new IMAGE_BRUSH( "Icons/icon_CameraSpeed_24x16px", FVector2D( 24, 16 ), IconColor ) );
+
+		Set( "EditorViewport.ToggleRealTime", new IMAGE_BRUSH("Icons/icon_MatEd_Realtime_40x", Icon40x40));
+
+		Set( "EditorViewport.LocationGridSnap", new IMAGE_BRUSH_SVG("Starship/EditorViewport/location-grid-snap", Icon16x16));
+		Set( "EditorViewport.RotationGridSnap", new IMAGE_BRUSH_SVG("Starship/EditorViewport/rotation-grid-snap", Icon16x16));
+
+		Set( "EditorViewport.Layer2DSnap", new IMAGE_BRUSH("Old/LevelEditor/Layer2DSnap", Icon14x14));
+
+		Set("EditorViewport.ScaleGridSnap", new IMAGE_BRUSH( "Old/LevelEditor/ScaleGridSnap", Icon14x14 ) );
+		Set("EditorViewport.ToggleSurfaceSnapping", new IMAGE_BRUSH_SVG( "Starship/EditorViewport/surface-snap", Icon16x16 ) );
+		Set("EditorViewport.ToggleSurfaceSnapping", new IMAGE_BRUSH_SVG("Starship/EditorViewport/surface-snap", Icon16x16));
+
+		Set( "EditorViewport.RelativeCoordinateSystem_Local", new IMAGE_BRUSH_SVG( "Starship/EditorViewport/transform-local", Icon16x16 ) );
+		Set( "EditorViewport.RelativeCoordinateSystem_World", new IMAGE_BRUSH_SVG( "Starship/EditorViewport/transform-world", Icon16x16 ) );
+
+		Set( "EditorViewport.CamSpeedSetting", new IMAGE_BRUSH_SVG( "Starship/EditorViewport/camera", Icon16x16) );
 		
 		Set( "EditorViewport.LitMode", new IMAGE_BRUSH( "Icons/icon_ViewMode_Lit_16px", Icon16x16 ) );
 		Set( "EditorViewport.UnlitMode", new IMAGE_BRUSH( "Icons/icon_ViewMode_Unlit_16px", Icon16x16 ) );
@@ -4396,7 +4470,7 @@ void FStarshipEditorStyle::FStyle::SetupLevelEditorStyle()
 	// Level viewport 
 #if WITH_EDITOR || (IS_PROGRAM && WITH_UNREAL_DEVELOPER_TOOLS)
 	{
-		Set( "LevelViewport.ActiveViewportBorder", new BORDER_BRUSH( "Old/White", FMargin(1), SelectionColor ) ); 
+		Set("LevelViewport.ActiveViewportBorder", new CORE_BORDER_BRUSH("Starship/Common/White", FMargin(1), FStyleColors::Primary));
 		Set( "LevelViewport.NoViewportBorder", new FSlateNoResource() );
 		Set( "LevelViewport.DebugBorder", new BOX_BRUSH( "Old/Window/ViewportDebugBorder", 0.8f, FLinearColor(.7,0,0,.5) ) );
 		Set( "LevelViewport.BlackBackground", new FSlateColorBrush( FLinearColor::Black ) ); 
@@ -4423,41 +4497,6 @@ void FStarshipEditorStyle::FStyle::SetupLevelEditorStyle()
 		Set( "LevelEditor.HideSelected", new IMAGE_BRUSH( "Old/SelectionDetails/HideSelected", FVector2D(32,32) ) );
 	}
 
-	// Level viewport toolbar
-	{
-		Set( "EditorViewportToolBar.Font", DEFAULT_FONT( "Bold", 9 ) );
-
-		Set( "EditorViewportToolBar.MenuButton", FButtonStyle(Button)
-			.SetNormal(BOX_BRUSH( "Common/SmallRoundedButton", FMargin(7.f/16.f), FLinearColor(1,1,1,0.75f)))
-			.SetHovered(BOX_BRUSH( "Common/SmallRoundedButton", FMargin(7.f/16.f), FLinearColor(1,1,1, 1.0f)))
-			.SetPressed(BOX_BRUSH( "Common/SmallRoundedButton", FMargin(7.f/16.f) ))
-		);
-
-		Set("EditorViewportToolBar.MenuButtonWarning", FButtonStyle(Button)
-			.SetNormal(BOX_BRUSH("Common/SmallRoundedButton", FMargin(7.f / 16.f), LogColor_Warning_LinearRef))
-			.SetHovered(BOX_BRUSH("Common/SmallRoundedButton", FMargin(7.f / 16.f), *LogColor_Warning_LinearRef*FLinearColor(1.1, 1.1, 1.1, 1)))
-			.SetPressed(BOX_BRUSH("Common/SmallRoundedButton", FMargin(7.f / 16.f), LogColor_Warning_LinearRef))
-		);
-
-		Set( "EditorViewportToolBar.Button", HoverHintOnly );
-
-		/* Set style structure for "EditorViewportToolBar.Button" ... */
-		const FCheckBoxStyle EditorViewportToolBarButton = FCheckBoxStyle()
-			.SetCheckBoxType(ESlateCheckBoxType::ToggleButton)
-			.SetUncheckedImage( FSlateNoResource() )
-			.SetUncheckedPressedImage( BOX_BRUSH( "Old/LevelViewportToolBar/MenuButton_Pressed", 4.0f/16.0f ) )
-			.SetUncheckedHoveredImage( BOX_BRUSH( "Old/Border", 4.0f/16.0f ) )
-			.SetCheckedImage( FSlateNoResource() )
-			.SetCheckedHoveredImage( BOX_BRUSH( "Old/Border", 4.0f/16.0f ) )
-			.SetCheckedPressedImage( BOX_BRUSH( "Old/LevelViewportToolBar/MenuButton_Pressed", 4.0f/16.0f ) );
-		/* ... and set new style */
-		Set( "LevelViewportToolBar.CheckBoxButton", EditorViewportToolBarButton );
-
-		Set( "EditorViewportToolBar.MenuDropdown", new IMAGE_BRUSH( "Common/ComboArrow", Icon8x8 ) );
-		Set( "LevelViewportToolBar.Maximize.Normal", new IMAGE_BRUSH( "Old/LevelViewportToolBar/Maximized_Unchecked", Icon16x16 ) );
-		Set( "LevelViewportToolBar.Maximize.Checked", new IMAGE_BRUSH( "Old/LevelViewportToolBar/Maximized_Checked", Icon16x16 ) );
-		Set( "LevelViewportToolBar.RestoreFromImmersive.Normal", new IMAGE_BRUSH( "Icons/icon_RestoreFromImmersive_16px", Icon16x16 ) );
-	}
 	// Show flags menus
 	{
 		Set( "ShowFlagsMenu.AntiAliasing", new IMAGE_BRUSH( "Icons/icon_ShowAnti-aliasing_16x", Icon16x16 ) );
diff --git a/Engine/Source/Editor/Kismet/Private/SSCSEditorViewport.cpp b/Engine/Source/Editor/Kismet/Private/SSCSEditorViewport.cpp
--- a/Engine/Source/Editor/Kismet/Private/SSCSEditorViewport.cpp
+++ b/Engine/Source/Editor/Kismet/Private/SSCSEditorViewport.cpp
@@ -31,50 +31,46 @@ public:
 	{
 		EditorViewport = InArgs._EditorViewport;
 
-		static const FName DefaultForegroundName("DefaultForeground");
+		const FMargin ToolbarSlotPadding(4.0f, 1.0f);
 
 		this->ChildSlot
 		[
 			SNew(SBorder)
-			.BorderImage(FEditorStyle::GetBrush("NoBorder"))
-			.ColorAndOpacity(this, &SViewportToolBar::OnGetColorAndOpacity)
-			.ForegroundColor(FEditorStyle::GetSlateColor(DefaultForegroundName))
+			.BorderImage(FAppStyle::Get().GetBrush("EditorViewportToolBar.Background"))
+			.Cursor(EMouseCursor::Default)
 			[
 				SNew(SHorizontalBox)
 				+SHorizontalBox::Slot()
 				.AutoWidth()
-				.Padding(2.0f, 2.0f)
+				.Padding(ToolbarSlotPadding)
 				[
 					SNew(SEditorViewportToolbarMenu)
 					.ParentToolBar(SharedThis(this))
 					.Cursor(EMouseCursor::Default)
-					.Image("EditorViewportToolBar.MenuDropdown")
+					.Image("EditorViewportToolBar.OptionsDropdown")
 					.OnGetMenuContent(this, &SSCSEditorViewportToolBar::GeneratePreviewMenu)
 				]
 				+ SHorizontalBox::Slot()
 				.AutoWidth()
-				.Padding(2.0f, 2.0f)
+				.Padding(ToolbarSlotPadding)
 				[
 					SNew( SEditorViewportToolbarMenu )
 					.ParentToolBar( SharedThis( this ) )
-					.Cursor( EMouseCursor::Default )
 					.Label(this, &SSCSEditorViewportToolBar::GetCameraMenuLabel)
-					.LabelIcon(this, &SSCSEditorViewportToolBar::GetCameraMenuLabelIcon)
 					.OnGetMenuContent(this, &SSCSEditorViewportToolBar::GenerateCameraMenu)
 				]
 				+ SHorizontalBox::Slot()
 				.AutoWidth()
-				.Padding(2.0f, 2.0f)
+				.Padding(ToolbarSlotPadding)
 				[
 					SNew( SEditorViewportToolbarMenu )
 					.ParentToolBar( SharedThis( this ) )
 					.Cursor( EMouseCursor::Default )
 					.Label(this, &SSCSEditorViewportToolBar::GetViewMenuLabel)
-					.LabelIcon(this, &SSCSEditorViewportToolBar::GetViewMenuLabelIcon)
 					.OnGetMenuContent(this, &SSCSEditorViewportToolBar::GenerateViewMenu)
 				]
 				+ SHorizontalBox::Slot()
-				.Padding( 3.0f, 1.0f )
+				.Padding(ToolbarSlotPadding)
 				.HAlign( HAlign_Right )
 				[
 					SNew(STransformViewportToolBar)
@@ -119,16 +115,6 @@ public:
 		return NSLOCTEXT("BlueprintEditor", "CameraMenuTitle_Default", "Camera");
 	}
 
-	const FSlateBrush* GetCameraMenuLabelIcon() const
-	{
-		if(EditorViewport.IsValid())
-		{
-			return GetCameraMenuLabelIconFromViewportType( EditorViewport.Pin()->GetViewportClient()->GetViewportType() );
-		}
-
-		return FEditorStyle::GetBrush(NAME_None);
-	}
-
 	TSharedRef<SWidget> GenerateCameraMenu() const
 	{
 		TSharedPtr<const FUICommandList> CommandList = EditorViewport.IsValid()? EditorViewport.Pin()->GetCommandList(): nullptr;
@@ -175,35 +161,6 @@ public:
 		return Label;
 	}
 
-	const FSlateBrush* GetViewMenuLabelIcon() const
-	{
-		static FName LitModeIconName("EditorViewport.LitMode");
-		static FName UnlitModeIconName("EditorViewport.UnlitMode");
-		static FName WireframeModeIconName("EditorViewport.WireframeMode");
-
-		FName Icon = NAME_None;
-
-		if (EditorViewport.IsValid())
-		{
-			switch (EditorViewport.Pin()->GetViewportClient()->GetViewMode())
-			{
-			case VMI_Lit:
-				Icon = LitModeIconName;
-				break;
-
-			case VMI_Unlit:
-				Icon = UnlitModeIconName;
-				break;
-
-			case VMI_BrushWireframe:
-				Icon = WireframeModeIconName;
-				break;
-			}
-		}
-
-		return FEditorStyle::GetBrush(Icon);
-	}
-
 	TSharedRef<SWidget> GenerateViewMenu() const
 	{
 		TSharedPtr<const FUICommandList> CommandList = EditorViewport.IsValid() ? EditorViewport.Pin()->GetCommandList() : nullptr;
diff --git a/Engine/Source/Editor/LevelEditor/Private/SActorPilotViewportToolbar.h b/Engine/Source/Editor/LevelEditor/Private/SActorPilotViewportToolbar.h
--- a/Engine/Source/Editor/LevelEditor/Private/SActorPilotViewportToolbar.h
+++ b/Engine/Source/Editor/LevelEditor/Private/SActorPilotViewportToolbar.h
@@ -116,8 +116,6 @@ public:
 			SNew( SBorder )
 			.BorderImage( FEditorStyle::GetBrush("NoBorder") )
 			.Padding(FMargin(4.f, 0.f))
-			// Color and opacity is changed based on whether or not the mouse cursor is hovering over the toolbar area
-			.ColorAndOpacity( this, &SViewportToolBar::OnGetColorAndOpacity )
 			[
 				ToolbarBuilder.MakeWidget()
 			]
diff --git a/Engine/Source/Editor/LevelEditor/Private/SLevelViewport.cpp b/Engine/Source/Editor/LevelEditor/Private/SLevelViewport.cpp
--- a/Engine/Source/Editor/LevelEditor/Private/SLevelViewport.cpp
+++ b/Engine/Source/Editor/LevelEditor/Private/SLevelViewport.cpp
@@ -73,7 +73,6 @@
 #include "EditorLevelUtils.h"
 #include "Engine/LevelStreaming.h"
 #include "Editor/WorldBrowser/Public/WorldBrowserModule.h"
-#include "ToolMenus.h"
 
 static const FName LevelEditorName("LevelEditor");
 
@@ -2909,10 +2908,8 @@ void SActorPreview::Construct( const FArguments& InArgs )
 			[
 				SNew( SBorder )
 					.Visibility(BorderVisibility)
-
 					.Padding( 16.0f )
 					.BorderImage( FEditorStyle::GetBrush( "UniformShadow_Tint" ) )
-
 					.BorderBackgroundColor( this, &SActorPreview::GetBorderColorAndOpacity )
 					.ColorAndOpacity( this, &SActorPreview::GetColorAndOpacity )
 
@@ -3935,7 +3932,6 @@ void SLevelViewport::OnSetViewportConfiguration(FName ConfigurationName)
 			GCurrentLevelEditingViewportClient = nullptr;
 			ViewportTabPinned->SetViewportConfiguration(ConfigurationName);
 			FSlateApplication::Get().DismissAllMenus();
-			UToolMenus::Get()->CleanupStaleWidgetsNextTick(true);
 		}
 	}
 }
diff --git a/Engine/Source/Editor/LevelEditor/Private/SLevelViewportToolBar.cpp b/Engine/Source/Editor/LevelEditor/Private/SLevelViewportToolBar.cpp
--- a/Engine/Source/Editor/LevelEditor/Private/SLevelViewportToolBar.cpp
+++ b/Engine/Source/Editor/LevelEditor/Private/SLevelViewportToolBar.cpp
@@ -37,6 +37,7 @@
 #include "SScalabilitySettings.h"
 #include "Editor/EditorPerformanceSettings.h"
 #include "SEditorViewportViewMenuContext.h"
+#include "ToolMenu.h"
 
 #define LOCTEXT_NAMESPACE "LevelViewportToolBar"
 
@@ -159,48 +160,39 @@ void SLevelViewportToolBar::Construct( const FArguments& InArgs )
 	Viewport = InArgs._Viewport;
 	TSharedRef<SLevelViewport> ViewportRef = Viewport.Pin().ToSharedRef();
 
-	FLevelEditorModule& LevelEditorModule = FModuleManager::GetModuleChecked<FLevelEditorModule>( "LevelEditor");
-
-	const FMargin ToolbarSlotPadding( 2.0f, 2.0f );
-	const FMargin ToolbarButtonPadding( 2.0f, 0.0f );
+	FLevelEditorModule& LevelEditorModule = FModuleManager::GetModuleChecked<FLevelEditorModule>(TEXT("LevelEditor"));
 
-	static const FName DefaultForegroundName("DefaultForeground");
+	const FMargin ToolbarSlotPadding(4.0f, 1.0f);
+	const FMargin ToolbarButtonPadding(4.0f, 0.0f);
 
 	ChildSlot
 	[
-		SNew( SBorder )
-		.BorderImage( FEditorStyle::GetBrush("NoBorder") )
-		// Color and opacity is changed based on whether or not the mouse cursor is hovering over the toolbar area
-		.ColorAndOpacity( this, &SViewportToolBar::OnGetColorAndOpacity )
-		.ForegroundColor( FEditorStyle::GetSlateColor(DefaultForegroundName) )
+		SNew(SBorder)
+		.BorderImage(FAppStyle::Get().GetBrush("EditorViewportToolBar.Background"))
+		.Cursor(EMouseCursor::Default)
 		[
 			SNew( SHorizontalBox )
 			+ SHorizontalBox::Slot()
 			.AutoWidth()
-			.Padding( ToolbarSlotPadding )
+			.Padding(ToolbarSlotPadding)
 			[
 				SNew( SEditorViewportToolbarMenu )
 				.ParentToolBar( SharedThis( this ) )
-				.Cursor( EMouseCursor::Default )
-				.Image( "EditorViewportToolBar.MenuDropdown" )
+				.Image("EditorViewportToolBar.OptionsDropdown")
 				.AddMetaData<FTagMetaData>(FTagMetaData(TEXT("EditorViewportToolBar.MenuDropdown")))
 				.OnGetMenuContent( this, &SLevelViewportToolBar::GenerateOptionsMenu )
 			]
-
 			+ SHorizontalBox::Slot()
 			[
 				SNew( SHorizontalBox )
 				.Visibility(Viewport.Pin().Get(), &SLevelViewport::GetFullToolbarVisibility)
-
 				+ SHorizontalBox::Slot()
 				.AutoWidth()
 				.Padding( ToolbarSlotPadding )
 				[
 					SNew( SEditorViewportToolbarMenu )
 					.ParentToolBar( SharedThis( this ) )
-					.Cursor( EMouseCursor::Default )
 					.Label( this, &SLevelViewportToolBar::GetCameraMenuLabel )
-					.LabelIcon( this, &SLevelViewportToolBar::GetCameraMenuLabelIcon )
 					.AddMetaData<FTagMetaData>(FTagMetaData(TEXT("EditorViewportToolBar.CameraMenu")))
 					.OnGetMenuContent( this, &SLevelViewportToolBar::GenerateCameraMenu ) 
 				]
@@ -209,8 +201,7 @@ void SLevelViewportToolBar::Construct( const FArguments& InArgs )
 				.Padding( ToolbarSlotPadding )
 				[
 					SNew( SLevelEditorViewportViewMenu, ViewportRef, SharedThis(this) )
-					.Cursor( EMouseCursor::Default )
-					.MenuExtenders( GetViewMenuExtender() )
+					.MenuExtenders(GetViewMenuExtender())
 					.AddMetaData<FTagMetaData>(FTagMetaData(TEXT("ViewMenuButton")))
 				]
 				+ SHorizontalBox::Slot()
@@ -219,7 +210,6 @@ void SLevelViewportToolBar::Construct( const FArguments& InArgs )
 				[
 					SNew( SEditorViewportToolbarMenu )
 					.Label( LOCTEXT("ShowMenuTitle", "Show") )
-					.Cursor( EMouseCursor::Default )
 					.ParentToolBar( SharedThis( this ) )
 					.AddMetaData<FTagMetaData>(FTagMetaData(TEXT("EditorViewportToolBar.ShowMenu")))
 					.OnGetMenuContent( this, &SLevelViewportToolBar::GenerateShowMenu ) 
@@ -228,10 +218,8 @@ void SLevelViewportToolBar::Construct( const FArguments& InArgs )
 				.AutoWidth()
 				.Padding( ToolbarSlotPadding )
 				[
-
 					SNew( SEditorViewportToolbarMenu )
 					.Label( this, &SLevelViewportToolBar::GetViewModeOptionsMenuLabel )
-					.Cursor( EMouseCursor::Default )
 					.ParentToolBar( SharedThis( this ) )
 					.Visibility( this, &SLevelViewportToolBar::GetViewModeOptionsVisibility )
 					.AddMetaData<FTagMetaData>(FTagMetaData(TEXT("EditorViewportToolBar.ViewModeOptions")))
@@ -243,7 +231,6 @@ void SLevelViewportToolBar::Construct( const FArguments& InArgs )
 				[
 					SNew( SEditorViewportToolbarMenu )
 					.ParentToolBar( SharedThis( this ) )
-					.Cursor( EMouseCursor::Default )
 					.Label( this, &SLevelViewportToolBar::GetDevicePreviewMenuLabel )
 					.LabelIcon( this, &SLevelViewportToolBar::GetDevicePreviewMenuLabelIcon )
 					.OnGetMenuContent( this, &SLevelViewportToolBar::GenerateDevicePreviewMenu )
@@ -256,18 +243,16 @@ void SLevelViewportToolBar::Construct( const FArguments& InArgs )
 				[
 					// Button to show that realtime is off
 					SNew(SEditorViewportToolBarButton)	
-					.Cursor(EMouseCursor::Default)
 					.ButtonType(EUserInterfaceActionType::Button)
-					.ButtonStyle(&FEditorStyle::Get().GetWidgetStyle<FButtonStyle>("EditorViewportToolBar.MenuButtonWarning"))
+					.ButtonStyle(&FAppStyle::Get().GetWidgetStyle<FButtonStyle>("EditorViewportToolBar.WarningButton"))
 					.OnClicked(this, &SLevelViewportToolBar::OnRealtimeWarningClicked)
 					.Visibility(this, &SLevelViewportToolBar::GetRealtimeWarningVisibility)
 					.ToolTipText(LOCTEXT("RealtimeOff_ToolTip", "This viewport is not updating in realtime.  Click to turn on realtime mode."))
 					.Content()
 					[
 						SNew(STextBlock)
-						.Font(FEditorStyle::GetFontStyle("EditorViewportToolBar.Font"))
-						.Text(LOCTEXT("RealtimeOff", "Realtime: Off"))
-						.ColorAndOpacity(FLinearColor::Black)
+						.TextStyle(&FAppStyle::Get().GetWidgetStyle<FTextBlockStyle>("SmallText"))
+						.Text(LOCTEXT("RealtimeOff", "Realtime Off"))
 					]
 				]
 				+ SHorizontalBox::Slot()
@@ -277,16 +262,15 @@ void SLevelViewportToolBar::Construct( const FArguments& InArgs )
 					// Button to show scalability warnings
 					SNew(SEditorViewportToolbarMenu)
 					.ParentToolBar(SharedThis(this))
-					.Cursor(EMouseCursor::Default)
 					.Label(this, &SLevelViewportToolBar::GetScalabilityWarningLabel)
-					.MenuStyle(FEditorStyle::Get(), "EditorViewportToolBar.MenuButtonWarning")
+					.MenuStyle(&FAppStyle::Get().GetWidgetStyle<FButtonStyle>("EditorViewportToolBar.WarningButton"))
 					.OnGetMenuContent(this, &SLevelViewportToolBar::GetScalabilityWarningMenuContent)
 					.Visibility(this, &SLevelViewportToolBar::GetScalabilityWarningVisibility)
 					.ToolTipText(LOCTEXT("ScalabilityWarning_ToolTip", "Non-default scalability settings could be affecting what is shown in this viewport.\nFor example you may experience lower visual quality, reduced particle counts, and other artifacts that don't match what the scene would look like when running outside of the editor. Click to make changes."))
 				]
 				+ SHorizontalBox::Slot()
-				.Padding( ToolbarSlotPadding )
-				.HAlign( HAlign_Right )
+				.Padding(ToolbarSlotPadding)
+				.HAlign(HAlign_Right)
 				[
 					SNew(STransformViewportToolBar)
 					.Viewport(ViewportRef)
@@ -295,33 +279,32 @@ void SLevelViewportToolBar::Construct( const FArguments& InArgs )
 					.Visibility(ViewportRef, &SLevelViewport::GetTransformToolbarVisibility)
 				]
 				+ SHorizontalBox::Slot()
-				.HAlign( HAlign_Right )
+				.HAlign(HAlign_Right)
 				.AutoWidth()
-				.Padding( ToolbarButtonPadding )
+				.Padding(ToolbarButtonPadding)
 				[
 					//The Maximize/Minimize button is only displayed when not in Immersive mode.
-					SNew( SEditorViewportToolBarButton )
-					.Cursor( EMouseCursor::Default )
-					.ButtonType( EUserInterfaceActionType::ToggleButton )
-					.IsChecked( ViewportRef, &SLevelViewport::IsMaximized )
-					.OnClicked( ViewportRef, &SLevelViewport::OnToggleMaximize )
-					.Visibility( ViewportRef, &SLevelViewport::GetMaximizeToggleVisibility )
-					.Image( "LevelViewportToolBar.Maximize" )
-					.ToolTipText( LOCTEXT("Maximize_ToolTip", "Maximizes or restores this viewport") )
+					SNew(SEditorViewportToolBarButton)
+					.ButtonType(EUserInterfaceActionType::ToggleButton)
+					.CheckBoxStyle(&FAppStyle::Get().GetWidgetStyle<FCheckBoxStyle>("EditorViewportToolBar.MaximizeRestoreButton"))
+					.IsChecked(ViewportRef, &SLevelViewport::IsMaximized)
+					.OnClicked(ViewportRef, &SLevelViewport::OnToggleMaximize)
+					.Visibility(ViewportRef, &SLevelViewport::GetMaximizeToggleVisibility)
+					.Image("EditorViewportToolBar.Maximize")
+					.ToolTipText(LOCTEXT("Maximize_ToolTip", "Maximizes or restores this viewport"))
 				]
 				+ SHorizontalBox::Slot()
-				.HAlign( HAlign_Right )
+				.HAlign(HAlign_Right)
 				.AutoWidth()
-				.Padding( ToolbarButtonPadding )
+				.Padding(ToolbarButtonPadding)
 				[
 					//The Restore from Immersive' button is only displayed when the editor is in Immersive mode.
-					SNew( SEditorViewportToolBarButton )
-					.Cursor( EMouseCursor::Default )
-					.ButtonType( EUserInterfaceActionType::Button )
-					.OnClicked( ViewportRef, &SLevelViewport::OnToggleMaximize )
-					.Visibility( ViewportRef, &SLevelViewport::GetCloseImmersiveButtonVisibility )
-					.Image( "LevelViewportToolBar.RestoreFromImmersive.Normal" )
-					.ToolTipText( LOCTEXT("RestoreFromImmersive_ToolTip", "Restore from Immersive") )
+					SNew(SEditorViewportToolBarButton)
+					.ButtonType(EUserInterfaceActionType::Button)
+					.OnClicked(ViewportRef, &SLevelViewport::OnToggleMaximize)
+					.Visibility(ViewportRef, &SLevelViewport::GetCloseImmersiveButtonVisibility)
+					.Image("EditorViewportToolBar.RestoreFromImmersive.Normal")
+					.ToolTipText(LOCTEXT("RestoreFromImmersive_ToolTip", "Restore from Immersive"))
 				]
 			]
 		]
@@ -390,18 +373,7 @@ const FSlateBrush* SLevelViewportToolBar::GetDevicePreviewMenuLabelIcon() const
 		return FEditorStyle::GetOptionalBrush( PlatformIcon );
 	}
 
-	return NULL;
-}
-
-const FSlateBrush* SLevelViewportToolBar::GetCameraMenuLabelIcon() const
-{
-	TSharedPtr< SLevelViewport > PinnedViewport( Viewport.Pin() );
-	if( PinnedViewport.IsValid() )
-	{
-		return GetCameraMenuLabelIconFromViewportType(PinnedViewport->GetLevelViewportClient().ViewportType);
-	}
-
-	return FEditorStyle::GetBrush(NAME_None);
+	return nullptr;
 }
 
 bool SLevelViewportToolBar::IsCurrentLevelViewport() const
@@ -1609,7 +1581,7 @@ EVisibility SLevelViewportToolBar::GetRealtimeWarningVisibility() const
 {
 	FLevelEditorViewportClient& ViewportClient = Viewport.Pin()->GetLevelViewportClient();
 	// If the viewport is not realtime and there is no override then realtime is off
-	return !ViewportClient.IsRealtime() && !ViewportClient.IsRealtimeOverrideSet() ? EVisibility::Visible : EVisibility::Collapsed;
+	return !ViewportClient.IsRealtime() && !ViewportClient.IsRealtimeOverrideSet() && ViewportClient.IsPerspective() ? EVisibility::Visible : EVisibility::Collapsed;
 }
 
 FText SLevelViewportToolBar::GetScalabilityWarningLabel() const
diff --git a/Engine/Source/Editor/LevelEditor/Private/SLevelViewportToolBar.h b/Engine/Source/Editor/LevelEditor/Private/SLevelViewportToolBar.h
--- a/Engine/Source/Editor/LevelEditor/Private/SLevelViewportToolBar.h
+++ b/Engine/Source/Editor/LevelEditor/Private/SLevelViewportToolBar.h
@@ -44,12 +44,6 @@ private:
 	 */
 	FText GetCameraMenuLabel() const;
 
-	/**
-	 * Returns the label icon for the "Camera" tool bar menu, which changes depending on the viewport type
-	 *
-	 * @return	Label icon to use for this menu label
-	 */
-	const FSlateBrush* GetCameraMenuLabelIcon() const;
 
 	/**
 	 * Returns the label for the "View" tool bar menu, which changes depending on viewport show flags
diff --git a/Engine/Source/Editor/MaterialEditor/Private/SMaterialEditorViewportToolBar.cpp b/Engine/Source/Editor/MaterialEditor/Private/SMaterialEditorViewportToolBar.cpp
--- a/Engine/Source/Editor/MaterialEditor/Private/SMaterialEditorViewportToolBar.cpp
+++ b/Engine/Source/Editor/MaterialEditor/Private/SMaterialEditorViewportToolBar.cpp
@@ -37,8 +37,6 @@ void SMaterialEditorViewportPreviewShapeToolBar::Construct(const FArguments& InA
 	[
 		SNew(SBorder)
 		.BorderImage(FEditorStyle::GetBrush("NoBorder"))
-		// Color and opacity is changed based on whether or not the mouse cursor is hovering over the toolbar area
-		.ColorAndOpacity(this, &SViewportToolBar::OnGetColorAndOpacity)
 		.ForegroundColor(FEditorStyle::GetSlateColor(DefaultForegroundName))
 		.HAlign(HAlign_Right)
 		[
diff --git a/Engine/Source/Editor/Persona/Private/SAnimViewportToolBar.cpp b/Engine/Source/Editor/Persona/Private/SAnimViewportToolBar.cpp
--- a/Engine/Source/Editor/Persona/Private/SAnimViewportToolBar.cpp
+++ b/Engine/Source/Editor/Persona/Private/SAnimViewportToolBar.cpp
@@ -179,11 +179,8 @@ void SAnimViewportToolBar::Construct(const FArguments& InArgs, TSharedPtr<class
 		Extenders.Add(MakeShared<FExtender>());
 	}
 
-	const FMargin ToolbarSlotPadding(2.0f, 2.0f);
-	const FMargin ToolbarButtonPadding(2.0f, 0.0f);
-
-	static const FName DefaultForegroundName("DefaultForeground");
-
+	const FMargin ToolbarSlotPadding(4.0f, 1.0f);
+	const FMargin ToolbarButtonPadding(4.0f, 0.0f);
 
 	TSharedRef<SHorizontalBox> LeftToolbar = SNew(SHorizontalBox)
 		+SHorizontalBox::Slot()
@@ -194,7 +191,7 @@ void SAnimViewportToolBar::Construct(const FArguments& InArgs, TSharedPtr<class
 			.ToolTipText(LOCTEXT("ViewMenuTooltip", "View Options.\nShift-clicking items will 'pin' them to the toolbar."))
 			.ParentToolBar(SharedThis(this))
 			.Cursor(EMouseCursor::Default)
-			.Image("EditorViewportToolBar.MenuDropdown")
+			.Image("EditorViewportToolBar.OptionsDropdown")
 			.AddMetaData<FTagMetaData>(FTagMetaData(TEXT("EditorViewportToolBar.MenuDropdown")))
 			.OnGetMenuContent(this, &SAnimViewportToolBar::GenerateViewMenu)
 		]
@@ -207,7 +204,6 @@ void SAnimViewportToolBar::Construct(const FArguments& InArgs, TSharedPtr<class
 			.ParentToolBar(SharedThis(this))
 			.Cursor(EMouseCursor::Default)
 			.Label(this, &SAnimViewportToolBar::GetCameraMenuLabel)
-			.LabelIcon(this, &SAnimViewportToolBar::GetCameraMenuLabelIcon)
 			.AddMetaData<FTagMetaData>(FTagMetaData(TEXT("EditorViewportToolBar.CameraMenu")))
 			.OnGetMenuContent(this, &SAnimViewportToolBar::GenerateViewportTypeMenu)
 		]
@@ -296,32 +292,30 @@ void SAnimViewportToolBar::Construct(const FArguments& InArgs, TSharedPtr<class
 
 	ChildSlot
 	[
-		SNew( SBorder )
-		.BorderImage( FEditorStyle::GetBrush("NoBorder") )
-		// Color and opacity is changed based on whether or not the mouse cursor is hovering over the toolbar area
-		.ColorAndOpacity( this, &SViewportToolBar::OnGetColorAndOpacity )
-		.ForegroundColor( FEditorStyle::GetSlateColor(DefaultForegroundName) )
+		SNew( SVerticalBox )
+		+ SVerticalBox::Slot()
+		.AutoHeight()
 		[
-			SNew( SVerticalBox )
-			+ SVerticalBox::Slot()
-			.AutoHeight()
+			SNew(SBorder)
+			.BorderImage(FAppStyle::Get().GetBrush("EditorViewportToolBar.Background"))
+			.Cursor(EMouseCursor::Default)
 			[
 				LeftToolbar
 			]
-			+SVerticalBox::Slot()
-			.AutoHeight()
-			[
-				PinnedCommands.ToSharedRef()
-			]
-			+SVerticalBox::Slot()
-			.Padding(FMargin(4.0f, 3.0f, 0.0f, 0.0f))
-			[
-				// Display text (e.g., item being previewed)
-				SNew(SRichTextBlock)
-				.DecoratorStyleSet(&FEditorStyle::Get())
-				.Text(InViewport.Get(), &SAnimationEditorViewportTabBody::GetDisplayString)
-				.TextStyle(&FEditorStyle::Get().GetWidgetStyle<FTextBlockStyle>("AnimViewport.MessageText"))
-			]
+		]
+		+SVerticalBox::Slot()
+		.AutoHeight()
+		[
+			PinnedCommands.ToSharedRef()
+		]
+		+SVerticalBox::Slot()
+		.Padding(FMargin(4.0f, 3.0f, 0.0f, 0.0f))
+		[
+			// Display text (e.g., item being previewed)
+			SNew(SRichTextBlock)
+			.DecoratorStyleSet(&FEditorStyle::Get())
+			.Text(InViewport.Get(), &SAnimationEditorViewportTabBody::GetDisplayString)
+			.TextStyle(&FEditorStyle::Get().GetWidgetStyle<FTextBlockStyle>("AnimViewport.MessageText"))
 		]
 	];
 	
@@ -1070,17 +1064,6 @@ FText SAnimViewportToolBar::GetCameraMenuLabel() const
 	return LOCTEXT("Viewport_Default", "Camera");
 }
 
-const FSlateBrush* SAnimViewportToolBar::GetCameraMenuLabelIcon() const
-{
-	TSharedPtr< SAnimationEditorViewportTabBody > PinnedViewport(Viewport.Pin());
-	if (PinnedViewport.IsValid())
-	{
-		return GetCameraMenuLabelIconFromViewportType( PinnedViewport->GetLevelViewportClient().ViewportType );
-	}
-
-	return FEditorStyle::GetBrush(NAME_None);
-}
-
 TOptional<float> SAnimViewportToolBar::OnGetFOVValue() const
 {
 	if(Viewport.IsValid())
diff --git a/Engine/Source/Editor/Persona/Private/SAnimViewportToolBar.h b/Engine/Source/Editor/Persona/Private/SAnimViewportToolBar.h
--- a/Engine/Source/Editor/Persona/Private/SAnimViewportToolBar.h
+++ b/Engine/Source/Editor/Persona/Private/SAnimViewportToolBar.h
@@ -132,7 +132,6 @@ private:
 	 * @return	Label to use for this Menu
 	 */
 	FText GetCameraMenuLabel() const;
-	const FSlateBrush* GetCameraMenuLabelIcon() const;
 
 	/** Called by the FOV slider in the perspective viewport to get the FOV value */
 	TOptional<float> OnGetFOVValue() const;
diff --git a/Engine/Source/Editor/UMGEditor/Private/Designer/SDesignerView.cpp b/Engine/Source/Editor/UMGEditor/Private/Designer/SDesignerView.cpp
--- a/Engine/Source/Editor/UMGEditor/Private/Designer/SDesignerView.cpp
+++ b/Engine/Source/Editor/UMGEditor/Private/Designer/SDesignerView.cpp
@@ -629,7 +629,7 @@ TSharedRef<SWidget> SDesignerView::CreateOverlayUI()
 				.ButtonStyle(&ToolBarStyle.ButtonStyle)
 				.ToolTipText(LOCTEXT("ZoomToFit_ToolTip", "Zoom To Fit"))
 				.OnClicked(this, &SDesignerView::HandleZoomToFitClicked)
-				.ContentPadding(FEditorStyle::Get().GetMargin("LegacyViewportMenu.SToolBarButtonBlock.Button.Padding"))
+				.ContentPadding(ToolBarStyle.ButtonPadding)
 				[
 					SNew(SImage)
 					.Image(FEditorStyle::GetBrush("UMGEditor.ZoomToFit"))
@@ -644,7 +644,7 @@ TSharedRef<SWidget> SDesignerView::CreateOverlayUI()
 				.ButtonStyle(&ToolBarStyle.ButtonStyle)
 				.ToolTipText(LOCTEXT("SwapAspectRatio_ToolTip", "Switch between Landscape and Portrait"))
 				.OnClicked(this, &SDesignerView::HandleSwapAspectRatioClicked)
-				.ContentPadding(FEditorStyle::Get().GetMargin("LegacyViewportMenu.SToolBarButtonBlock.Button.Padding"))
+				.ContentPadding(ToolBarStyle.ButtonPadding)
 				.IsEnabled(this, &SDesignerView::GetAspectRatioSwitchEnabled)
 				[
 					SNew(SImage)
@@ -659,7 +659,7 @@ TSharedRef<SWidget> SDesignerView::CreateOverlayUI()
 				.ButtonStyle(&ToolBarStyle.ButtonStyle)
 				.ToolTipText(LOCTEXT("Mirror_ToolTip", "Flip the current safe zones"))
 				.OnClicked(this, &SDesignerView::HandleFlipSafeZonesClicked)
-				.ContentPadding(FEditorStyle::Get().GetMargin("LegacyViewportMenu.SToolBarButtonBlock.Button.Padding"))	
+				.ContentPadding(ToolBarStyle.ButtonPadding)
 				.IsEnabled(this, &SDesignerView::GetFlipDeviceEnabled)
 				[
 					SNew(SImage)
@@ -676,12 +676,12 @@ TSharedRef<SWidget> SDesignerView::CreateOverlayUI()
 			.ButtonStyle(&ToolBarStyle.ButtonStyle)
 			.ForegroundColor(FLinearColor::Black)
 			.OnGetMenuContent(this, &SDesignerView::GetResolutionsMenu)
-			.ContentPadding(FEditorStyle::Get().GetMargin("LegacyViewportMenu.SToolBarButtonBlock.Button.Padding"))
+			.ContentPadding(ToolBarStyle.ButtonPadding)
 			.ButtonContent()
 			[
 				SNew(STextBlock)
 				.Text(LOCTEXT("ScreenSize", "Screen Size"))
-				.TextStyle(FEditorStyle::Get(), "LegacyViewportMenu.Label")
+				.TextStyle(&ToolBarStyle.LabelStyle)
 			]
 		]
 
@@ -694,12 +694,12 @@ TSharedRef<SWidget> SDesignerView::CreateOverlayUI()
 			.ButtonStyle(&ToolBarStyle.ButtonStyle)
 			.ForegroundColor(FLinearColor::Black)
 			.OnGetMenuContent(this, &SDesignerView::GetScreenSizingFillMenu)
-			.ContentPadding(FEditorStyle::Get().GetMargin("LegacyViewportMenu.SToolBarButtonBlock.Button.Padding"))
+			.ContentPadding(ToolBarStyle.ButtonPadding)
 			.ButtonContent()
 			[
 				SNew(STextBlock)
 				.Text(this, &SDesignerView::GetScreenSizingFillText)
-				.TextStyle(FEditorStyle::Get(), "LegacyViewportMenu.Label")
+				.TextStyle(&ToolBarStyle.LabelStyle)
 			]
 		]
 
diff --git a/Engine/Source/Editor/UnrealEd/Private/SCommonEditorViewportToolbarBase.cpp b/Engine/Source/Editor/UnrealEd/Private/SCommonEditorViewportToolbarBase.cpp
--- a/Engine/Source/Editor/UnrealEd/Private/SCommonEditorViewportToolbarBase.cpp
+++ b/Engine/Source/Editor/UnrealEd/Private/SCommonEditorViewportToolbarBase.cpp
@@ -32,25 +32,16 @@ void SCommonEditorViewportToolbarBase::Construct(const FArguments& InArgs, TShar
  	TSharedRef<SEditorViewport> ViewportRef = GetInfoProvider().GetViewportWidget();
 	TSharedPtr<SHorizontalBox> MainBoxPtr;
 
-	const FMargin ToolbarSlotPadding( 2.0f, 2.0f );
-	const FMargin ToolbarButtonPadding( 2.0f, 0.0f );
-
-	static const FName DefaultForegroundName("DefaultForeground");
+	const FMargin ToolbarSlotPadding(4.0f, 1.0f);
+	const FMargin ToolbarButtonPadding(4.0f, 0.0f);
 
 	ChildSlot
 	[
 		SNew( SBorder )
-		.BorderImage( FEditorStyle::GetBrush("NoBorder") )
-		// Color and opacity is changed based on whether or not the mouse cursor is hovering over the toolbar area
-		.ColorAndOpacity( this, &SViewportToolBar::OnGetColorAndOpacity )
-		.ForegroundColor( FEditorStyle::GetSlateColor(DefaultForegroundName) )
+		.BorderImage(FAppStyle::Get().GetBrush("EditorViewportToolBar.Background"))
+		.Cursor(EMouseCursor::Default)
 		[
-			SNew( SVerticalBox )
-			+ SVerticalBox::Slot()
-			.AutoHeight()
-			[
-				SAssignNew( MainBoxPtr, SHorizontalBox )
-			]
+			SAssignNew( MainBoxPtr, SHorizontalBox )
 		]
 	];
 
@@ -62,7 +53,7 @@ void SCommonEditorViewportToolbarBase::Construct(const FArguments& InArgs, TShar
 			SNew(SEditorViewportToolbarMenu)
 			.ParentToolBar(SharedThis(this))
 			.Cursor(EMouseCursor::Default)
-			.Image("EditorViewportToolBar.MenuDropdown")
+			.Image("EditorViewportToolBar.OptionsDropdown")
 			.OnGetMenuContent(this, &SCommonEditorViewportToolbarBase::GenerateOptionsMenu)
 		];
 
@@ -75,7 +66,6 @@ void SCommonEditorViewportToolbarBase::Construct(const FArguments& InArgs, TShar
 			.ParentToolBar(SharedThis(this))
 			.Cursor(EMouseCursor::Default)
 			.Label(this, &SCommonEditorViewportToolbarBase::GetCameraMenuLabel)
-			.LabelIcon(this, &SCommonEditorViewportToolbarBase::GetCameraMenuLabelIcon)
 			.OnGetMenuContent(this, &SCommonEditorViewportToolbarBase::GenerateCameraMenu)
 		];
 
@@ -109,16 +99,15 @@ void SCommonEditorViewportToolbarBase::Construct(const FArguments& InArgs, TShar
 				SNew(SEditorViewportToolBarButton)
 				.Cursor(EMouseCursor::Default)
 				.ButtonType(EUserInterfaceActionType::Button)
-				.ButtonStyle(&FEditorStyle::Get().GetWidgetStyle<FButtonStyle>("EditorViewportToolBar.MenuButtonWarning"))
+				.ButtonStyle(&FAppStyle::Get().GetWidgetStyle<FButtonStyle>("EditorViewportToolBar.WarningButton"))
 				.OnClicked(this, &SCommonEditorViewportToolbarBase::OnRealtimeWarningClicked)
 				.Visibility(this, &SCommonEditorViewportToolbarBase::GetRealtimeWarningVisibility)
 				.ToolTipText(LOCTEXT("RealtimeOff_ToolTip", "This viewport is not updating in realtime.  Click to turn on realtime mode."))
 				.Content()
 				[
 					SNew(STextBlock)
-					.Font(FEditorStyle::GetFontStyle("EditorViewportToolBar.Font"))
-					.Text(LOCTEXT("RealtimeOff", "Realtime: Off"))
-					.ColorAndOpacity(FLinearColor::Black)
+					.TextStyle(&FAppStyle::Get().GetWidgetStyle<FTextBlockStyle>("SmallText"))
+					.Text(LOCTEXT("RealtimeOff", "Realtime Off"))
 				]
 			];
 	}
@@ -142,13 +131,12 @@ void SCommonEditorViewportToolbarBase::Construct(const FArguments& InArgs, TShar
 			// Button to show scalability warnings
 			SNew(SEditorViewportToolbarMenu)
 			.ParentToolBar(SharedThis(this))
-			.Cursor(EMouseCursor::Default)
 			.Label(this, &SCommonEditorViewportToolbarBase::GetScalabilityWarningLabel)
 			.MenuStyle(FEditorStyle::Get(), "EditorViewportToolBar.MenuButtonWarning")
 			.OnGetMenuContent(this, &SCommonEditorViewportToolbarBase::GetScalabilityWarningMenuContent)
 			.Visibility(this, &SCommonEditorViewportToolbarBase::GetScalabilityWarningVisibility)
 			.ToolTipText(LOCTEXT("ScalabilityWarning_ToolTip", "Non-default scalability settings could be affecting what is shown in this viewport.\nFor example you may experience lower visual quality, reduced particle counts, and other artifacts that don't match what the scene would look like when running outside of the editor. Click to make changes."))
-			];
+		];
 
 	// Add optional toolbar slots to be added by child classes inherited from this common viewport toolbar
 	ExtendLeftAlignedToolbarSlots(MainBoxPtr, SharedThis(this));
@@ -173,10 +161,6 @@ FText SCommonEditorViewportToolbarBase::GetCameraMenuLabel() const
 	return GetCameraMenuLabelFromViewportType( GetViewportClient().GetViewportType() );
 }
 
-const FSlateBrush* SCommonEditorViewportToolbarBase::GetCameraMenuLabelIcon() const
-{
-	return GetCameraMenuLabelIconFromViewportType( GetViewportClient().GetViewportType() );
-}
 
 EVisibility SCommonEditorViewportToolbarBase::GetViewModeOptionsVisibility() const
 {
diff --git a/Engine/Source/Editor/UnrealEd/Private/SEditorViewport.cpp b/Engine/Source/Editor/UnrealEd/Private/SEditorViewport.cpp
--- a/Engine/Source/Editor/UnrealEd/Private/SEditorViewport.cpp
+++ b/Engine/Source/Editor/UnrealEd/Private/SEditorViewport.cpp
@@ -58,15 +58,6 @@ void SEditorViewport::Construct( const FArguments& InArgs )
 			.ViewportSize(InArgs._ViewportSize)
 			[
 				SAssignNew(ViewportOverlay, SOverlay)
-				+ SOverlay::Slot()
-				[
-					SNew(SBorder)
-					.BorderImage(this, &SEditorViewport::OnGetViewportBorderBrush)
-					.BorderBackgroundColor(this, &SEditorViewport::OnGetViewportBorderColorAndOpacity)
-					.Visibility(this, &SEditorViewport::OnGetViewportContentVisibility)
-					.Padding(0.0f)
-					.ShowEffectWhenDisabled(false)
-				]
 			]
 		]
 	];
@@ -95,7 +86,7 @@ void SEditorViewport::Construct( const FArguments& InArgs )
 
 	TSharedPtr<SWidget> ViewportToolbar = MakeViewportToolbar();
 
-	if( ViewportToolbar.IsValid() )
+	if(ViewportToolbar.IsValid())
 	{
 		ViewportOverlay->AddSlot()
 			.VAlign(VAlign_Top)
@@ -104,6 +95,16 @@ void SEditorViewport::Construct( const FArguments& InArgs )
 			];
 	}
 
+	ViewportOverlay->AddSlot()
+	[
+		SNew(SBorder)
+		.BorderImage(this, &SEditorViewport::OnGetViewportBorderBrush)
+		.BorderBackgroundColor(this, &SEditorViewport::OnGetViewportBorderColorAndOpacity)
+		.Visibility(this, &SEditorViewport::GetActiveBorderVisibility)
+		.Padding(0.0f)
+		.ShowEffectWhenDisabled(false)
+	];
+
 	PopulateViewportOverlays(ViewportOverlay.ToSharedRef());
 }
 
@@ -705,6 +706,18 @@ EActiveTimerReturnType SEditorViewport::EnsureTick( double InCurrentTime, float
 	return bShouldContinue ? EActiveTimerReturnType::Continue : EActiveTimerReturnType::Stop;
 }
 
+EVisibility SEditorViewport::GetActiveBorderVisibility() const
+{
+	EVisibility BaseVisibility = OnGetViewportContentVisibility();
+	if (BaseVisibility != EVisibility::Collapsed)
+	{
+		// The active border should never be hit testable as it overlays viewport UI but is for display purposes only
+		return EVisibility::HitTestInvisible;
+	}
+
+	return BaseVisibility;
+}
+
 ///////////////////////////////////////////////////////////////////////////////
 // begin feature level control functions block
 ///////////////////////////////////////////////////////////////////////////////
diff --git a/Engine/Source/Editor/UnrealEd/Private/SEditorViewportToolBarButton.cpp b/Engine/Source/Editor/UnrealEd/Private/SEditorViewportToolBarButton.cpp
--- a/Engine/Source/Editor/UnrealEd/Private/SEditorViewportToolBarButton.cpp
+++ b/Engine/Source/Editor/UnrealEd/Private/SEditorViewportToolBarButton.cpp
@@ -35,7 +35,6 @@ void SEditorViewportToolBarButton::Construct( const FArguments& Declaration)
 			.OnClicked( OnClickedDelegate )
 			.HAlign( HAlign_Center )
 			.VAlign( VAlign_Center )
-			.ForegroundColor( FSlateColor::UseForeground() )
 			[
 				// If we have a content override use it instead of the default image
 				bContentOverride
@@ -70,6 +69,7 @@ void SEditorViewportToolBarButton::Construct( const FArguments& Declaration)
 					[
 						SNew( SImage )
 						.Image( this, &SEditorViewportToolBarButton::OnGetButtonImage )
+						.ColorAndOpacity(FSlateColor::UseForeground())
 					])
 			];
 	}
diff --git a/Engine/Source/Editor/UnrealEd/Private/SEditorViewportToolBarMenu.cpp b/Engine/Source/Editor/UnrealEd/Private/SEditorViewportToolBarMenu.cpp
--- a/Engine/Source/Editor/UnrealEd/Private/SEditorViewportToolBarMenu.cpp
+++ b/Engine/Source/Editor/UnrealEd/Private/SEditorViewportToolBarMenu.cpp
@@ -23,21 +23,18 @@ void SEditorViewportToolbarMenu::Construct( const FArguments& Declaration )
 	checkf(ParentToolBar.IsValid(), TEXT("The parent toolbar must be specified") );
 
 	TSharedPtr<SWidget> ButtonContent;
-	
-	const float MenuIconSize = 16.0f;
 
 	// Create the content for the button.  We always use an image over a label if it is valid
 	if( ImageName != NAME_None )
 	{
 		ButtonContent = 
 			SNew(SBox)
-			.HeightOverride(MenuIconSize)
 			.VAlign(VAlign_Center)
 			.HAlign(HAlign_Center)
 			[
-				SNew( SImage )	
-					.Image( ImageBrush )
-					.ColorAndOpacity( FSlateColor::UseForeground() )
+				SNew(SImage)	
+				.Image(ImageBrush)
+				.ColorAndOpacity(FSlateColor::UseForeground())
 			];
 	}
 	else
@@ -45,36 +42,27 @@ void SEditorViewportToolbarMenu::Construct( const FArguments& Declaration )
 		if( LabelIconBrush.IsBound() || LabelIconBrush.Get() != NULL )
 		{
 			// Label with an icon to the left
-			
 			ButtonContent = 
-				SNew( SHorizontalBox )
+				SNew(SHorizontalBox)
 				+SHorizontalBox::Slot()
+				.VAlign(VAlign_Center)
+				.HAlign(HAlign_Center)
 				.AutoWidth()
-				.Padding( 0.0f, 0.0f, 3.0f, 0.0f )
+				.Padding(5.0f, 0.0f)
 				[
-					SNew( SBox )
-					.Visibility( this, &SEditorViewportToolbarMenu::GetLabelIconVisibility )
-					.WidthOverride( MenuIconSize )
-					.HeightOverride( MenuIconSize )
-					.VAlign(VAlign_Center)
-					.HAlign(HAlign_Center)
-					[
-						SNew( SImage )
-							.Image( LabelIconBrush )
-					]
+					SNew( SImage )
+					.Visibility(this, &SEditorViewportToolbarMenu::GetLabelIconVisibility)
+					.Image(LabelIconBrush)
+					.ColorAndOpacity(FSlateColor::UseForeground())
 				]
 				+SHorizontalBox::Slot()
-				.FillWidth( 1.0f )
-				.VAlign( VAlign_Center )
+				.Padding(2.0f,1.0f, 0.0f, 0.0f)
+				.AutoWidth()
+				.VAlign(VAlign_Center)
 				[
-					SNew(SBox)
-					.HeightOverride(MenuIconSize)
-					.VAlign(VAlign_Center)
-					[
-						SNew( STextBlock )
-							.Font( FEditorStyle::GetFontStyle("EditorViewportToolBar.Font") )
-							.Text( Label )
-					]
+					SNew(STextBlock)
+					.TextStyle(&FAppStyle::Get().GetWidgetStyle<FTextBlockStyle>("SmallText"))
+					.Text(Label)
 				];
 		}
 		else
@@ -82,12 +70,11 @@ void SEditorViewportToolbarMenu::Construct( const FArguments& Declaration )
 			// Just the label text, no icon
 			ButtonContent = 
 				SNew(SBox)
-				.HeightOverride(MenuIconSize)
 				.VAlign(VAlign_Center)
 				[
 					SNew( STextBlock )
-						.Font( FEditorStyle::GetFontStyle("EditorViewportToolBar.Font") )
-						.Text( Label )
+					.TextStyle(&FAppStyle::Get().GetWidgetStyle<FTextBlockStyle>("SmallText"))
+					.Text(Label)
 				];
 		}
 	}
@@ -100,10 +87,11 @@ void SEditorViewportToolbarMenu::Construct( const FArguments& Declaration )
 			SNew( SButton )
 			// Allows users to drag with the mouse to select options after opening the menu */
 			.ClickMethod( EButtonClickMethod::MouseDown )
-			.ContentPadding( FMargin( 5.0f, 2.0f ) )
-			.VAlign( VAlign_Center )
+			.ContentPadding( FMargin( 2.0f,0.0f ) )
+			.VAlign(VAlign_Center)
 			.ButtonStyle(Declaration._MenuStyle)
 			.OnClicked( this, &SEditorViewportToolbarMenu::OnMenuClicked )
+			.ForegroundColor(Declaration._ForegroundColor)
 			[
 				ButtonContent.ToSharedRef()
 			]
@@ -132,6 +120,8 @@ FReply SEditorViewportToolbarMenu::OnMenuClicked()
 
 void SEditorViewportToolbarMenu::OnMouseEnter( const FGeometry& MyGeometry, const FPointerEvent& MouseEvent )
 {
+	SCompoundWidget::OnMouseEnter(MyGeometry, MouseEvent);
+
 	// See if there is another menu on the same tool bar already open
 	TWeakPtr<SMenuAnchor> OpenedMenu = ParentToolBar.Pin()->GetOpenMenu();
 	if( OpenedMenu.IsValid() && OpenedMenu.Pin()->IsOpen() && OpenedMenu.Pin() != MenuAnchor )
@@ -151,4 +141,9 @@ EVisibility SEditorViewportToolbarMenu::GetLabelIconVisibility() const
 TWeakPtr<class SViewportToolBar> SEditorViewportToolbarMenu::GetParentToolBar() const
 {
 	return ParentToolBar;
-}
\ No newline at end of file
+}
+
+bool SEditorViewportToolbarMenu::IsMenuOpen() const
+{
+	return MenuAnchor->IsOpen();
+}
diff --git a/Engine/Source/Editor/UnrealEd/Private/SEditorViewportViewMenu.cpp b/Engine/Source/Editor/UnrealEd/Private/SEditorViewportViewMenu.cpp
--- a/Engine/Source/Editor/UnrealEd/Private/SEditorViewportViewMenu.cpp
+++ b/Engine/Source/Editor/UnrealEd/Private/SEditorViewportViewMenu.cpp
@@ -26,7 +26,6 @@ void SEditorViewportViewMenu::Construct( const FArguments& InArgs, TSharedRef<SE
 			.ParentToolBar( InParentToolBar)
 			.Cursor( EMouseCursor::Default )
 			.Label(this, &SEditorViewportViewMenu::GetViewMenuLabel)
-			.LabelIcon(this, &SEditorViewportViewMenu::GetViewMenuLabelIcon)
 			.OnGetMenuContent( this, &SEditorViewportViewMenu::GenerateViewMenuContent )
 	);
 }
@@ -55,145 +54,6 @@ FText SEditorViewportViewMenu::GetViewMenuLabel() const
 	return Label;
 }
 
-const FSlateBrush* SEditorViewportViewMenu::GetViewMenuLabelIcon() const
-{
-	FName Icon = NAME_None;
-	TSharedPtr< SEditorViewport > PinnedViewport = Viewport.Pin();
-	if( PinnedViewport.IsValid() )
-	{
-		static FName WireframeIcon( "EditorViewport.WireframeMode" );
-		static FName UnlitIcon( "EditorViewport.UnlitMode" );
-		static FName LitIcon( "EditorViewport.LitMode" );
-		static FName DetailLightingIcon("EditorViewport.DetailLightingMode");
-		static FName LightingOnlyIcon("EditorViewport.LightingOnlyMode");
-		static FName LightComplexityIcon("EditorViewport.LightComplexityMode");
-		static FName ShaderComplexityIcon("EditorViewport.ShaderComplexityMode");
-		static FName QuadOverdrawIcon("EditorViewport.QuadOverdrawMode");
-		static FName ShaderComplexityWithQuadOverdrawIcon("EditorViewport.ShaderCOmplexityWithQuadOverdrawMode");
-		static FName PrimitiveDistanceAccuracyIcon("EditorViewport.TexStreamAccPrimitiveDistanceMode");
-		static FName MeshUVDensityAccuracyIcon("EditorViewport.TexStreamAccMeshUVDensityMode");
-		static FName MaterialTextureScaleAccuracyIcon("EditorViewport.TexStreamAccMaterialTextureScaleMode");
-		static FName RequiredTextureResolutionIcon("EditorViewport.RequiredTextureResolutionMode");
-		static FName LightOverlapIcon("EditorViewport.StationaryLightOverlapMode");
-		static FName LightmapDensityIcon("EditorViewport.LightmapDensityMode");
-		static FName ReflectionModeIcon("EditorViewport.ReflectionOverrideMode");
-		static FName LODColorationIcon("EditorViewport.LODColorationMode");
-		static FName VisualizeBufferIcon("EditorViewport.VisualizeBufferMode");
-		static FName CollisionPawnIcon("EditorViewport.CollisionPawn");
-		static FName CollisionVisibilityIcon("EditorViewport.CollisionVisibility");
-
-		switch( PinnedViewport->GetViewportClient()->GetViewMode() )
-		{
-			case VMI_BrushWireframe:
-				Icon = WireframeIcon;
-				break;
-
-			case VMI_Wireframe:
-				Icon = WireframeIcon;
-				break;
-
-			case VMI_Unlit:
-				Icon = UnlitIcon;
-				break;
-
-			case VMI_Lit:
-				Icon = LitIcon;
-				break;
-
-#if RHI_RAYTRACING
-			// #dxr_todo: UE-72550 use special icons for ray tracing view modes
-			case VMI_RayTracingDebug:
-				Icon = LitIcon; 
-				break;
-
-			case VMI_PathTracing:
-				Icon = LitIcon; 
-				break;
-#endif
-			case VMI_Lit_DetailLighting:
-				Icon = DetailLightingIcon;
-				break;
-
-			case VMI_LightingOnly:
-				Icon = LightingOnlyIcon;
-				break;
-
-			case VMI_LightComplexity:
-				Icon = LightComplexityIcon;
-				break;
-
-			case VMI_ShaderComplexity:
-				Icon = ShaderComplexityIcon;
-				break;
-
-			case VMI_QuadOverdraw:
-				Icon = QuadOverdrawIcon;
-				break;
-
-			case VMI_ShaderComplexityWithQuadOverdraw:
-				Icon = ShaderComplexityWithQuadOverdrawIcon;
-				break;
-
-			case VMI_PrimitiveDistanceAccuracy:
-				Icon = PrimitiveDistanceAccuracyIcon;
-				break;
-
-			case VMI_MeshUVDensityAccuracy:
-				Icon = MeshUVDensityAccuracyIcon;
-				break;
-
-			case VMI_MaterialTextureScaleAccuracy:
-				Icon = MaterialTextureScaleAccuracyIcon;
-				break;
-
-			case VMI_RequiredTextureResolution:
-				Icon = RequiredTextureResolutionIcon;
-				break;
-
-			case VMI_StationaryLightOverlap:
-				Icon = LightOverlapIcon;
-				break;
-
-			case VMI_LightmapDensity:
-				Icon = LightmapDensityIcon;
-				break;
-
-			case VMI_ReflectionOverride:
-				Icon = ReflectionModeIcon;
-				break;
-
-			case VMI_VisualizeBuffer:
-				Icon = VisualizeBufferIcon;
-				break;
-
-			case VMI_CollisionPawn:
-				Icon = CollisionPawnIcon;
-				break;
-
-			case VMI_CollisionVisibility:
-				Icon = CollisionVisibilityIcon;
-				break;
-
-			case VMI_LitLightmapDensity:
-				break;
-
-			case VMI_LODColoration:
-				Icon = LODColorationIcon;
-				break;
-
-			case VMI_HLODColoration:
-				Icon = LODColorationIcon;
-				break;
-
-			case VMI_GroupLODColoration:
-				Icon = LODColorationIcon;
-				break;
-		}
-	}
-
-	return FEditorStyle::GetBrush(Icon);
-}
-
 void SEditorViewportViewMenu::RegisterMenus() const
 {
 	if (!UToolMenus::Get()->IsMenuRegistered(BaseMenuName))
diff --git a/Engine/Source/Editor/UnrealEd/Private/STransformViewportToolbar.cpp b/Engine/Source/Editor/UnrealEd/Private/STransformViewportToolbar.cpp
--- a/Engine/Source/Editor/UnrealEd/Private/STransformViewportToolbar.cpp
+++ b/Engine/Source/Editor/UnrealEd/Private/STransformViewportToolbar.cpp
@@ -30,17 +30,12 @@
 #include "Settings/EditorProjectSettings.h"
 #include "LevelEditor.h"
 #include "LevelEditorActions.h"
+#include "Styling/ToolBarStyle.h"
+#include "SEditorViewportToolBarMenu.h"
+#include "Widgets/Layout/SSpacer.h"
 
 #define LOCTEXT_NAMESPACE "TransformToolBar"
 
-namespace TransformViewportToolbarDefs
-{
-	/** Size of the arrow shown on SGridSnapSettings Menu button */
-	const float DownArrowSize = 4.0f;
-
-	/** Size of the icon displayed on the toggle button of SGridSnapSettings */
-	const float ToggleImageScale = 16.0f;
-}
 
 void STransformViewportToolBar::Construct( const FArguments& InArgs )
 {
@@ -55,14 +50,27 @@ void STransformViewportToolBar::Construct( const FArguments& InArgs )
 	SViewportToolBar::Construct(SViewportToolBar::FArguments());
 }
 
-TSharedRef< SWidget > STransformViewportToolBar::MakeSurfaceSnappingButton( FName ToolBarStyle )
+TSharedRef< SWidget > STransformViewportToolBar::MakeSurfaceSnappingButton()
+{
+	check(!SurfaceSnappingMenu.IsValid());
+	SurfaceSnappingMenu = 
+		SNew(SEditorViewportToolbarMenu)
+		.ParentToolBar(SharedThis(this))
+		.Image("EditorViewport.ToggleSurfaceSnapping")
+		.OnGetMenuContent(this, &STransformViewportToolBar::GenerateSurfaceSnappingMenu)
+		.ForegroundColor(this, &STransformViewportToolBar::GetSurfaceSnappingForegroundColor);
+
+	return SurfaceSnappingMenu.ToSharedRef();
+}
+
+TSharedRef<SWidget> STransformViewportToolBar::GenerateSurfaceSnappingMenu()
 {
-	auto IsSnappingEnabled = []{
+	auto IsSnappingEnabled = [] {
 		return GetDefault<ULevelEditorViewportSettings>()->SnapToSurface.bEnabled;
 	};
 
 	const bool bCloseAfterSelection = true;
-	FMenuBuilder MenuBuilder( bCloseAfterSelection, CommandList );
+	FMenuBuilder MenuBuilder(bCloseAfterSelection, CommandList);
 
 	MenuBuilder.AddMenuEntry(FEditorViewportCommands::Get().SurfaceSnapping);
 
@@ -73,17 +81,16 @@ TSharedRef< SWidget > STransformViewportToolBar::MakeSurfaceSnappingButton( FNam
 			LOCTEXT("SnapToSurfaceSettings_RotationTip", "When checked, snapping an object to a surface will also rotate the object to align to the surface normal"),
 			FSlateIcon(),
 			FUIAction(
-				FExecuteAction::CreateStatic( []{
+				FExecuteAction::CreateStatic([] {
 					auto& Settings = GetMutableDefault<ULevelEditorViewportSettings>()->SnapToSurface;
 					Settings.bSnapRotation = !Settings.bSnapRotation;
-				}),
-				FCanExecuteAction::CreateStatic( IsSnappingEnabled ),
-				FIsActionChecked::CreateStatic( []{
+					}),
+				FCanExecuteAction::CreateStatic(IsSnappingEnabled),
+				FIsActionChecked::CreateStatic([] {
 					const auto& Settings = GetDefault<ULevelEditorViewportSettings>()->SnapToSurface;
 					return Settings.bSnapRotation;
-				})
-			),
-			NAME_None,
+					})),
+				NAME_None,
 			EUserInterfaceActionType::RadioButton);
 
 		MenuBuilder.AddWidget(
@@ -94,29 +101,28 @@ TSharedRef< SWidget > STransformViewportToolBar::MakeSurfaceSnappingButton( FNam
 			.AutoWidth()
 			.Padding(FMargin(0.f, 0.f, 5.f, 0.f))
 			[
-				SNew( STextBlock )
-				.Text( LOCTEXT("SnapToSurfaceSettings_Offset", "Surface Offset") )
+				SNew(STextBlock)
+				.Text(LOCTEXT("SnapToSurfaceSettings_Offset", "Surface Offset"))
 			]
-			+SHorizontalBox::Slot()
-			.VAlign( VAlign_Bottom )
-			.FillWidth( 1.f )
+			+ SHorizontalBox::Slot()
+			.VAlign(VAlign_Bottom)
+			.FillWidth(1.f)
 			[
-				SNew( SNumericEntryBox<float> )
-				.Value( 
-					TAttribute<TOptional<float>>::Create(TAttribute<TOptional<float>>::FGetter::CreateStatic([]{
+				SNew(SNumericEntryBox<float>)
+				.Value(
+					TAttribute<TOptional<float>>::Create(TAttribute<TOptional<float>>::FGetter::CreateStatic([] {
 						const auto& Settings = GetDefault<ULevelEditorViewportSettings>()->SnapToSurface;
 						return TOptional<float>(Settings.SnapOffsetExtent);
-					}))
+						}))
 				)
 				.OnValueChanged(
-					SNumericEntryBox<float>::FOnValueChanged::CreateStatic([](float Val){
-						GetMutableDefault<ULevelEditorViewportSettings>()->SnapToSurface.SnapOffsetExtent = Val;
-					})
-				)
-				.MinValue(0.f)
-				.MaxValue(HALF_WORLD_MAX)
-				.MaxSliderValue(1000.f) // 'Sensible' range for the slider (10m)
-				.AllowSpin(true)
+				SNumericEntryBox<float>::FOnValueChanged::CreateStatic([](float Val) {
+					GetMutableDefault<ULevelEditorViewportSettings>()->SnapToSurface.SnapOffsetExtent = Val;
+					}))
+					.MinValue(0.f)
+					.MaxValue(HALF_WORLD_MAX)
+					.MaxSliderValue(1000.f) // 'Sensible' range for the slider (10m)
+					.AllowSpin(true)
 			],
 			FText::GetEmpty()
 		);
@@ -124,82 +130,44 @@ TSharedRef< SWidget > STransformViewportToolBar::MakeSurfaceSnappingButton( FNam
 	}
 	MenuBuilder.EndSection();
 
-	// Have to use a custom widget here to make the checkbox work with the subsequent widget :(
-	return 
-		SNew(SHorizontalBox)
+	return MenuBuilder.MakeWidget();
+}
 
-		+SHorizontalBox::Slot()
-		.AutoWidth()
-		[
-			SNew(SCheckBox)
-			.Cursor( EMouseCursor::Default )
-			.Style( FEditorStyle::Get(), EMultiBlockLocation::ToName(FEditorStyle::Join(ToolBarStyle, ".ToggleButton"), EMultiBlockLocation::Start) )
-			.Padding( 0 )
-			.ToolTipText( LOCTEXT("SurfaceSnappingCheckboxDescription", "Open editor surface snapping options") )
-			.IsChecked_Static( []{ return GetDefault<ULevelEditorViewportSettings>()->SnapToSurface.bEnabled ? ECheckBoxState::Checked : ECheckBoxState::Unchecked; })
-			.Content()
-			[
-				SNew( SComboButton )
-				.ButtonStyle( FEditorStyle::Get(), "HoverHintOnly" )
-				.HasDownArrow( false )
-				.ContentPadding( 0 )
-				.ButtonContent()
-				[
-					SNew( SVerticalBox )
-
-					+ SVerticalBox::Slot()
-					.AutoHeight()
-					.Padding(FMargin(5.f, 2.f, 3.f, 0.f))
-					[
-						SNew( SBox )
-						.WidthOverride( TransformViewportToolbarDefs::ToggleImageScale )
-						.HeightOverride( TransformViewportToolbarDefs::ToggleImageScale )
-						.HAlign( HAlign_Center )
-						.VAlign( VAlign_Center )
-						[
-							SNew( SImage )
-							.Image( FEditorStyle::GetBrush("EditorViewport.ToggleSurfaceSnapping") )
-						]
-					]
-
-					+ SVerticalBox::Slot()
-					.AutoHeight()
-					.HAlign( HAlign_Center )
-					.Padding(FMargin(0.f, 0.f, 0.f, 3.f))
-					[
-						SNew( SBox )
-						.WidthOverride( TransformViewportToolbarDefs::DownArrowSize )
-						.HeightOverride( TransformViewportToolbarDefs::DownArrowSize )
-						[
-							SNew(SImage)
-							.Image(FEditorStyle::GetBrush("ComboButton.Arrow"))
-							.ColorAndOpacity(FLinearColor::Black)
-						]
-					]
-				]
-				.MenuContent()
-				[
-					MenuBuilder.MakeWidget()
-				]
-			]
-		]
+FSlateColor STransformViewportToolBar::GetSurfaceSnappingForegroundColor() const
+{
+	static const FCheckBoxStyle& ViewportToolbarCheckStyle = FAppStyle::Get().GetWidgetStyle<FToolBarStyle>("EditorViewportToolBar").ToggleButton;
 
-		+SHorizontalBox::Slot()
-		.AutoWidth()
-		[
-			SNew(SBorder)
-			.Padding(FMargin(1.0f, 0.0f, 0.0f, 0.0f))
-			.BorderImage(FEditorStyle::GetDefaultBrush())
-			.BorderBackgroundColor(FLinearColor::Black)
-		];
+	const bool bSurfaceSnappingEnabled = GetDefault<ULevelEditorViewportSettings>()->SnapToSurface.bEnabled;
+
+	bool bShouldAppearHovered = SurfaceSnappingMenu->IsHovered() || SurfaceSnappingMenu->IsMenuOpen();
+	// Hovered and checked
+	if (bShouldAppearHovered && bSurfaceSnappingEnabled)
+	{
+		return ViewportToolbarCheckStyle.CheckedHoveredForeground;
+	}
+	// Not hovered and checked
+	else if (bSurfaceSnappingEnabled)
+	{
+		return ViewportToolbarCheckStyle.CheckedForeground;
+	}
+	// Hovered not checked
+	else if (bShouldAppearHovered)
+	{
+		return ViewportToolbarCheckStyle.HoveredForeground;
+	}
+	// Not hovered not checked
+	else
+	{
+		return ViewportToolbarCheckStyle.ForegroundColor;
+	}
 }
 
 TSharedRef< SWidget > STransformViewportToolBar::MakeTransformToolBar( const TSharedPtr< FExtender > InExtenders )
 {
-	FToolBarBuilder ToolbarBuilder( CommandList, FMultiBoxCustomization::None, InExtenders );
+	FSlimHorizontalToolBarBuilder ToolbarBuilder( CommandList, FMultiBoxCustomization::None, InExtenders );
 
 	// Use a custom style
-	FName ToolBarStyle = "LegacyViewportMenu";
+	FName ToolBarStyle = "EditorViewportToolBar";
 	ToolbarBuilder.SetStyle(&FEditorStyle::Get(), ToolBarStyle);
 	ToolbarBuilder.SetLabelVisibility(EVisibility::Collapsed);
 
@@ -207,7 +175,6 @@ TSharedRef< SWidget > STransformViewportToolBar::MakeTransformToolBar( const TSh
 	ToolbarBuilder.SetIsFocusable( false );
 
 	ToolbarBuilder.BeginSection("Transform");
-	ToolbarBuilder.BeginBlockGroup();
 	{
 		// Move Mode
 		static FName TranslateModeName = FName(TEXT("TranslateMode"));
@@ -229,16 +196,11 @@ TSharedRef< SWidget > STransformViewportToolBar::MakeTransformToolBar( const TSh
 		static FName ScaleModeName = FName(TEXT("ScaleMode"));
 		ToolbarBuilder.AddToolBarButton( FEditorViewportCommands::Get().ScaleMode, NAME_None, TAttribute<FText>(), TAttribute<FText>(), TAttribute<FSlateIcon>(), ScaleModeName );
 
-	}
-	ToolbarBuilder.EndBlockGroup();
-	ToolbarBuilder.EndSection();
-	
-	ToolbarBuilder.SetIsFocusable( true );
-	
-	ToolbarBuilder.BeginSection("LocalToWorld");
-	ToolbarBuilder.BeginBlockGroup();
-	{
-		// Move Mode
+
+		ToolbarBuilder.AddSeparator();
+
+		ToolbarBuilder.SetIsFocusable( true );
+
 		ToolbarBuilder.AddToolBarButton( FEditorViewportCommands::Get().CycleTransformGizmoCoordSystem,
 										NAME_None,
 										TAttribute<FText>(),
@@ -247,13 +209,18 @@ TSharedRef< SWidget > STransformViewportToolBar::MakeTransformToolBar( const TSh
 										FName(TEXT("CycleTransformGizmoCoordSystem"))
 										);
 	}
-	ToolbarBuilder.EndBlockGroup();
+
+	// Add a bit of spacing before the snapping section
+	ToolbarBuilder.AddWidget(SNew(SSpacer).Size(FVector2D(15.0f, 1.0f)));
+
 	ToolbarBuilder.EndSection();
-	
+
 	ToolbarBuilder.BeginSection("LocationGridSnap");
 	{
 		static FName SurfaceSnapName = FName(TEXT("SurfaceSnap"));
-		ToolbarBuilder.AddWidget( MakeSurfaceSnappingButton( ToolBarStyle ), SurfaceSnapName );
+		ToolbarBuilder.AddWidget( MakeSurfaceSnappingButton(), SurfaceSnapName );
+
+		ToolbarBuilder.AddSeparator();
 
 		// Grab the existing UICommand 
 		FUICommandInfo* Command = FEditorViewportCommands::Get().LocationGridSnap.Get();
@@ -263,8 +230,6 @@ TSharedRef< SWidget > STransformViewportToolBar::MakeTransformToolBar( const TSh
 		// Setup a GridSnapSetting with the UICommand
 		ToolbarBuilder.AddWidget(	SNew(SViewportToolBarComboMenu)
 									.Style(ToolBarStyle)
-									.BlockLocation(EMultiBlockLocation::Middle)
-									.Cursor( EMouseCursor::Default )
 									.IsChecked(this, &STransformViewportToolBar::IsLocationGridSnapChecked)
 									.OnCheckStateChanged(this, &STransformViewportToolBar::HandleToggleLocationGridSnap)
 									.Label(this, &STransformViewportToolBar::GetLocationGridLabel)
@@ -273,7 +238,7 @@ TSharedRef< SWidget > STransformViewportToolBar::MakeTransformToolBar( const TSh
 									.MenuButtonToolTip(LOCTEXT("LocationGridSnap_ToolTip", "Set the Position Grid Snap value"))
 									.Icon(Command->GetIcon())
 									.ParentToolBar(SharedThis(this))
-									, PositionSnapName );
+									,PositionSnapName );
 	}
 	ToolbarBuilder.EndSection();
 
@@ -286,7 +251,6 @@ TSharedRef< SWidget > STransformViewportToolBar::MakeTransformToolBar( const TSh
 
 		// Setup a GridSnapSetting with the UICommand
 		ToolbarBuilder.AddWidget(	SNew(SViewportToolBarComboMenu)
-									.Cursor( EMouseCursor::Default )
 									.Style(ToolBarStyle)
 									.IsChecked(this, &STransformViewportToolBar::IsRotationGridSnapChecked)
 									.OnCheckStateChanged(this, &STransformViewportToolBar::HandleToggleRotationGridSnap)
@@ -309,7 +273,6 @@ TSharedRef< SWidget > STransformViewportToolBar::MakeTransformToolBar( const TSh
 
 		TSharedRef<SWidget> SnapLayerPickerWidget =
 			SNew(SViewportToolBarComboMenu)
-			.Cursor(EMouseCursor::Default)
 			.Style(ToolBarStyle)
 			.Visibility(this, &STransformViewportToolBar::IsLayer2DSnapVisible)
 			.IsChecked(this, &STransformViewportToolBar::IsLayer2DSnapChecked)
@@ -352,21 +315,17 @@ TSharedRef< SWidget > STransformViewportToolBar::MakeTransformToolBar( const TSh
 	ToolbarBuilder.EndSection();
 
 	ToolbarBuilder.BeginSection("CameraSpeed");
-	ToolbarBuilder.BeginBlockGroup();
 	{
 		// Camera speed 
-		ToolbarBuilder.AddWidget(	SNew(SViewportToolBarIconMenu)
-									.Cursor(EMouseCursor::Default)
-									.Style(ToolBarStyle)
-									.Label(this, &STransformViewportToolBar::GetCameraSpeedLabel)
-									.OnGetMenuContent(this, &STransformViewportToolBar::FillCameraSpeedMenu)
-									.ToolTipText(LOCTEXT("CameraSpeed_ToolTip","Camera Speed"))
-									.Icon(FSlateIcon(FEditorStyle::GetStyleSetName(), "EditorViewport.CamSpeedSetting"))
-									.ParentToolBar(SharedThis(this))
-									.AddMetaData<FTagMetaData>(FTagMetaData(TEXT("CameraSpeedButton")))
-									);
+		ToolbarBuilder.AddWidget(	
+			SNew(SEditorViewportToolbarMenu)
+			.ParentToolBar(SharedThis(this))
+			.AddMetaData<FTagMetaData>(FTagMetaData(TEXT("CameraSpeedButton")))
+			.ToolTipText(LOCTEXT("CameraSpeed_ToolTip", "Camera Speed"))
+			.LabelIcon(FAppStyle::Get().GetBrush("EditorViewport.CamSpeedSetting"))
+			.Label(this, &STransformViewportToolBar::GetCameraSpeedLabel)
+			.OnGetMenuContent(this, &STransformViewportToolBar::FillCameraSpeedMenu));
 	}
-	ToolbarBuilder.EndBlockGroup();
 	ToolbarBuilder.EndSection();
 
 	return ToolbarBuilder.MakeWidget();
diff --git a/Engine/Source/Editor/UnrealEd/Private/SViewportToolBar.cpp b/Engine/Source/Editor/UnrealEd/Private/SViewportToolBar.cpp
--- a/Engine/Source/Editor/UnrealEd/Private/SViewportToolBar.cpp
+++ b/Engine/Source/Editor/UnrealEd/Private/SViewportToolBar.cpp
@@ -8,25 +8,9 @@
 
 #define LOCTEXT_NAMESPACE "ViewportToolBar"
 
-namespace ToolBarConstants
-{
-	/** The opacity when we are hovered */
-	const float HoveredOpacity = 1.0f;
-	/** The opacity when we are not hovered */
-	const float NonHoveredOpacity = .75f;
-	/** The amount of time to wait before fading out the toolbar after the mouse leaves it (to reduce poping when mouse moves in and out frequently */
-	const float TimeToFadeOut = 1.0f;
-	/** The amount of time spent actually fading in or out */
-	const float FadeTime = .15f;
-}
-
 void SViewportToolBar::Construct( const FArguments& InArgs )
 {
 	bIsHovered = false;
-
-	FadeInSequence = FCurveSequence( 0.0f, ToolBarConstants::FadeTime );
-	FadeOutSequence = FCurveSequence( ToolBarConstants::TimeToFadeOut, ToolBarConstants::FadeTime );
-	FadeOutSequence.JumpToEnd();
 }
 
 TWeakPtr<SMenuAnchor> SViewportToolBar::GetOpenMenu() const
@@ -44,60 +28,9 @@ void SViewportToolBar::SetOpenMenu( TSharedPtr< SMenuAnchor >& NewMenu )
 	OpenedMenu = NewMenu;
 }
 
-FLinearColor SViewportToolBar::OnGetColorAndOpacity() const
-{
-	FLinearColor Color = FLinearColor::White;
-	
-	if( OpenedMenu.IsValid() && OpenedMenu.Pin()->IsOpen() )
-	{
-		// Never fade out the toolbar if a menu is open
-		Color.A = ToolBarConstants::HoveredOpacity;
-	}
-	else if( FadeOutSequence.IsPlaying() || !bIsHovered )
-	{
-		Color.A = FMath::Lerp( ToolBarConstants::HoveredOpacity, ToolBarConstants::NonHoveredOpacity, FadeOutSequence.GetLerp() );
-	}
-	else
-	{
-		Color.A = FMath::Lerp( ToolBarConstants::NonHoveredOpacity, ToolBarConstants::HoveredOpacity, FadeInSequence.GetLerp() );
-	}
 
-	return Color;
-}
 
 
-void SViewportToolBar::OnMouseEnter( const FGeometry& MyGeometry, const FPointerEvent& MouseEvent )
-{
-	// The viewport could potentially be moved around inside the toolbar when the mouse is captured
-	// If that is the case we do not play the fade transition
-	if( !FSlateApplication::Get().IsUsingHighPrecisionMouseMovment() )
-	{
-		bIsHovered = true;
-		if( FadeOutSequence.IsPlaying() )
-		{
-			// Fade out is already playing so just force the fade in curve to the end so we don't have a "pop" 
-			// effect from quickly resetting the alpha
-			FadeInSequence.JumpToEnd();
-		}
-		else
-		{
-			FadeInSequence.Play( this->AsShared() );
-		}
-	}
-
-}
-
-void SViewportToolBar::OnMouseLeave( const FPointerEvent& MouseEvent )
-{
-	// The viewport could potentially be moved around inside the toolbar when the mouse is captured
-	// If that is the case we do not play the fade transition
-	if( !FSlateApplication::Get().IsUsingHighPrecisionMouseMovment() )
-	{
-		bIsHovered = false;
-		FadeOutSequence.Play( this->AsShared() );
-	}
-}
-
 FText SViewportToolBar::GetCameraMenuLabelFromViewportType(const ELevelViewportType ViewportType) const
 {
 	FText Label = LOCTEXT("CameraMenuTitle_Default", "Camera");
diff --git a/Engine/Source/Editor/UnrealEd/Private/SViewportToolBarComboMenu.cpp b/Engine/Source/Editor/UnrealEd/Private/SViewportToolBarComboMenu.cpp
--- a/Engine/Source/Editor/UnrealEd/Private/SViewportToolBarComboMenu.cpp
+++ b/Engine/Source/Editor/UnrealEd/Private/SViewportToolBarComboMenu.cpp
@@ -17,12 +17,8 @@ void SViewportToolBarComboMenu::Construct( const FArguments& InArgs )
 
 	EMultiBlockLocation::Type BlockLocation = InArgs._BlockLocation;
 
-	FName ButtonBlockStyle = EMultiBlockLocation::ToName(FEditorStyle::Join(InArgs._Style.Get(), ".Button"), BlockLocation);
-	const FButtonStyle& ButtonStyle = BlockLocation == EMultiBlockLocation::None ? ToolBarStyle.ButtonStyle : FAppStyle::Get().GetWidgetStyle<FButtonStyle>(ButtonBlockStyle);
-
-	FName CheckBoxBlockStyle = EMultiBlockLocation::ToName(FEditorStyle::Join(InArgs._Style.Get(), ".ToggleButton"), BlockLocation);
-	const FCheckBoxStyle& CheckBoxStyle = BlockLocation == EMultiBlockLocation::None ? ToolBarStyle.ToggleButton : FAppStyle::Get().GetWidgetStyle<FCheckBoxStyle>(CheckBoxBlockStyle);
-
+	const FButtonStyle& ButtonStyle = ToolBarStyle.ButtonStyle;
+	const FCheckBoxStyle& CheckBoxStyle = ToolBarStyle.ToggleButton;
 	const FTextBlockStyle& LabelStyle = ToolBarStyle.LabelStyle;
 
 	const FSlateIcon& Icon = InArgs._Icon.Get();
@@ -32,22 +28,19 @@ void SViewportToolBarComboMenu::Construct( const FArguments& InArgs )
 	TSharedPtr<SCheckBox> ToggleControl;
 	{
 		ToggleControl = SNew(SCheckBox)
-		.Cursor( EMouseCursor::Default )
-		.Padding(FMargin( 4.0f ))
+		.Padding(ToolBarStyle.CheckBoxPadding)
 		.Style(&CheckBoxStyle)
 		.OnCheckStateChanged(InArgs._OnCheckStateChanged)
 		.ToolTipText(InArgs._ToggleButtonToolTip)
 		.IsChecked(InArgs._IsChecked)
-		.Content()
 		[
 			SNew( SBox )
-			.WidthOverride( 16 )
-			.HeightOverride( 16 )
 			.HAlign(HAlign_Center)
 			.VAlign(VAlign_Center)
 			[
 				SNew(SImage)
 				.Image(Icon.GetIcon())
+				.ColorAndOpacity(FSlateColor::UseForeground())
 			]
 		];
 	}
@@ -56,44 +49,13 @@ void SViewportToolBarComboMenu::Construct( const FArguments& InArgs )
 		TSharedRef<SWidget> ButtonContents =
 			SNew(SButton)
 			.ButtonStyle(&ButtonStyle)
-			.ContentPadding( FMargin( 5.0f, 0.0f ) )
+			.ContentPadding( FMargin( 2.0f, 0.0f ) )
 			.ToolTipText(InArgs._MenuButtonToolTip)
 			.OnClicked(this, &SViewportToolBarComboMenu::OnMenuClicked)
 			[
-				SNew(SVerticalBox)
-				+SVerticalBox::Slot()
-				.AutoHeight()
-				.HAlign(HAlign_Center)
-				.VAlign(VAlign_Top)
-				[
-					SNew(STextBlock)
-					.TextStyle(&LabelStyle)
-					.Text(InArgs._Label)
-				]
-				+SVerticalBox::Slot()
-				.AutoHeight()
-				.VAlign(VAlign_Bottom)
-				[
-					SNew(SHorizontalBox)
-					+SHorizontalBox::Slot()
-					.FillWidth(1.0f)
-
-					+SHorizontalBox::Slot()
-					.AutoWidth()
-					[
-						SNew( SBox )
-						.WidthOverride( 4 )
-						.HeightOverride( 4 )
-						[
-							SNew(SImage)
-							.Image(FEditorStyle::GetBrush("ComboButton.Arrow"))
-							.ColorAndOpacity(FLinearColor::Black)
-						]
-					]
-
-					+SHorizontalBox::Slot()
-						.FillWidth(1.0f)
-				]
+				SNew(STextBlock)
+				.TextStyle(&LabelStyle)
+				.Text(InArgs._Label)
 			];
 		
 		if (InArgs._MinDesiredButtonWidth > 0.0f)
@@ -118,29 +80,17 @@ void SViewportToolBarComboMenu::Construct( const FArguments& InArgs )
 	ChildSlot
 	[
 		SNew(SHorizontalBox)
-
-		//Checkbox concept
 		+SHorizontalBox::Slot()
+		.VAlign(VAlign_Center)
 		.AutoWidth()
 		[
-			ToggleControl->AsShared()
+			ToggleControl.ToSharedRef()
 		]
-
-		// Black Separator line
-		+SHorizontalBox::Slot()
-		.AutoWidth()
-		[
-			SNew(SBorder)
-			.Padding(FMargin(1.0f, 0.0f, 0.0f, 0.0f))
-			.BorderImage(FEditorStyle::GetDefaultBrush())
-			.BorderBackgroundColor(FLinearColor::Black)
-		]
-
-		// Menu dropdown concept
 		+SHorizontalBox::Slot()
+		.VAlign(VAlign_Center)
 		.AutoWidth()
 		[
-			MenuAnchor->AsShared()
+			MenuAnchor.ToSharedRef()
 		]
 	];
 }
diff --git a/Engine/Source/Editor/UnrealEd/Private/SViewportToolBarIconMenu.cpp b/Engine/Source/Editor/UnrealEd/Private/SViewportToolBarIconMenu.cpp
--- a/Engine/Source/Editor/UnrealEd/Private/SViewportToolBarIconMenu.cpp
+++ b/Engine/Source/Editor/UnrealEd/Private/SViewportToolBarIconMenu.cpp
@@ -63,30 +63,6 @@ void SViewportToolBarIconMenu::Construct( const FArguments& InArgs )
 					.TextStyle(&ViewportToolbarStyle.LabelStyle)
 					.Text(InArgs._Label)
 				]
-				+SVerticalBox::Slot()
-				.AutoHeight()
-				.VAlign(VAlign_Bottom)
-				[
-					SNew(SHorizontalBox)
-					+SHorizontalBox::Slot()
-					.FillWidth(1.0f)
-
-					+SHorizontalBox::Slot()
-					.AutoWidth()
-					[
-						SNew( SBox )
-						.WidthOverride( 4 )
-						.HeightOverride( 4 )
-						[
-							SNew(SImage)
-							.Image(FEditorStyle::GetBrush("ComboButton.Arrow"))
-							.ColorAndOpacity(FLinearColor::Black)
-						]
-					]
-
-					+SHorizontalBox::Slot()
-					.FillWidth(1.0f)
-				]
 			]
 		]
 	];
diff --git a/Engine/Source/Editor/UnrealEd/Public/SCommonEditorViewportToolbarBase.h b/Engine/Source/Editor/UnrealEd/Public/SCommonEditorViewportToolbarBase.h
--- a/Engine/Source/Editor/UnrealEd/Public/SCommonEditorViewportToolbarBase.h
+++ b/Engine/Source/Editor/UnrealEd/Public/SCommonEditorViewportToolbarBase.h
@@ -82,12 +82,6 @@ private:
 	 */
 	FText GetCameraMenuLabel() const;
 
-	/**
-	 * Returns the label icon for the "Camera" tool bar menu, which changes depending on the viewport type
-	 *
-	 * @return	Label icon to use for this menu label
-	 */
-	const FSlateBrush* GetCameraMenuLabelIcon() const;
 
 	/**
 	 * Returns the label for the "View" tool bar menu, which changes depending on viewport show flags
@@ -96,13 +90,6 @@ private:
 	 */
 	FText GetViewMenuLabel() const;
 
-	/**
-	 * Returns the label icon for the "View" tool bar menu, which changes depending on viewport show flags
-	 *
-	 * @return	Label icon to use for this menu label
-	 */
-	const FSlateBrush* GetViewMenuLabelIcon() const;
-
 	/**
 	 * Generates the toolbar options menu content 
 	 *
diff --git a/Engine/Source/Editor/UnrealEd/Public/SEditorViewport.h b/Engine/Source/Editor/UnrealEd/Public/SEditorViewport.h
--- a/Engine/Source/Editor/UnrealEd/Public/SEditorViewport.h
+++ b/Engine/Source/Editor/UnrealEd/Public/SEditorViewport.h
@@ -251,6 +251,9 @@ private:
 	/** Ensures a Slate tick/paint pass when the viewport is realtime or was invalidated this frame */
 	EActiveTimerReturnType EnsureTick( double InCurrentTime, float InDeltaTime );
 
+	/** Gets the visibility of the active viewport border */
+	EVisibility GetActiveBorderVisibility() const;
+private:
 	/** The handle to the active timer */
 	TWeakPtr<FActiveTimerHandle> ActiveTimerHandle;
 
diff --git a/Engine/Source/Editor/UnrealEd/Public/SEditorViewportToolBarButton.h b/Engine/Source/Editor/UnrealEd/Public/SEditorViewportToolBarButton.h
--- a/Engine/Source/Editor/UnrealEd/Public/SEditorViewportToolBarButton.h
+++ b/Engine/Source/Editor/UnrealEd/Public/SEditorViewportToolBarButton.h
@@ -12,6 +12,7 @@
 #include "Styling/SlateTypes.h"
 #include "EditorStyleSet.h"
 #include "Styling/SlateWidgetStyleAsset.h"
+#include "Styling/ToolBarStyle.h"
 
 struct FSlateBrush;
 enum class ECheckBoxState : uint8;
@@ -22,10 +23,10 @@ enum class ECheckBoxState : uint8;
 class UNREALED_API SEditorViewportToolBarButton : public SCompoundWidget
 {
 public:
-	SLATE_BEGIN_ARGS( SEditorViewportToolBarButton )
+	SLATE_BEGIN_ARGS(SEditorViewportToolBarButton)
 		: _ButtonType(EUserInterfaceActionType::Button)
-		, _ButtonStyle(&FEditorStyle::Get().GetWidgetStyle<FButtonStyle>("EditorViewportToolBar.Button"))
-		, _CheckBoxStyle(&FEditorStyle::Get().GetWidgetStyle<FCheckBoxStyle>("LevelViewportToolBar.CheckBoxButton"))
+		, _ButtonStyle(&FEditorStyle::Get().GetWidgetStyle<FToolBarStyle>("EditorViewportToolBar").ButtonStyle)
+		, _CheckBoxStyle(&FEditorStyle::Get().GetWidgetStyle<FToolBarStyle>("EditorViewportToolBar").ToggleButton)
 		, _IsChecked(false)
 		{}
 
diff --git a/Engine/Source/Editor/UnrealEd/Public/SEditorViewportToolBarMenu.h b/Engine/Source/Editor/UnrealEd/Public/SEditorViewportToolBarMenu.h
--- a/Engine/Source/Editor/UnrealEd/Public/SEditorViewportToolBarMenu.h
+++ b/Engine/Source/Editor/UnrealEd/Public/SEditorViewportToolBarMenu.h
@@ -35,21 +35,24 @@ namespace EMenuItemType
 class UNREALED_API SEditorViewportToolbarMenu : public SCompoundWidget
 {
 public:
-	SLATE_BEGIN_ARGS( SEditorViewportToolbarMenu )
-		: _MenuStyle(&FEditorStyle::Get().GetWidgetStyle<FButtonStyle>("EditorViewportToolBar.MenuButton"))
+	SLATE_BEGIN_ARGS(SEditorViewportToolbarMenu)
+		: _MenuStyle(&FEditorStyle::Get().GetWidgetStyle<FButtonStyle>("EditorViewportToolBar.Button"))
+		, _ForegroundColor(FSlateColor::UseStyle())
 	{}
 		/** We need to know about the toolbar we are in */
-		SLATE_ARGUMENT( TSharedPtr<class SViewportToolBar>, ParentToolBar );
+		SLATE_ARGUMENT(TSharedPtr<class SViewportToolBar>, ParentToolBar);
 		/** Style to use */
 		SLATE_STYLE_ARGUMENT(FButtonStyle, MenuStyle)
 		/** The label to show in the menu */
-		SLATE_ATTRIBUTE( FText, Label )
+		SLATE_ATTRIBUTE(FText, Label)
 		/** Optional icon to display next to the label */
-		SLATE_ATTRIBUTE( const FSlateBrush*, LabelIcon )
+		SLATE_ATTRIBUTE(const FSlateBrush*, LabelIcon)
 		/** The image to show in the menu.  If both the label and image are valid, the button image is used.  Note that if this image is used, the label icon will not be displayed. */
-		SLATE_ARGUMENT( FName, Image )
+		SLATE_ARGUMENT(FName, Image)
 		/** Content to show in the menu */
-		SLATE_EVENT( FOnGetContent, OnGetMenuContent )
+		SLATE_EVENT(FOnGetContent, OnGetMenuContent)
+		/** The foreground color of the content */
+		SLATE_ATTRIBUTE(FSlateColor, ForegroundColor)
 	SLATE_END_ARGS()
 
 	/**
@@ -62,6 +65,11 @@ public:
 	 */
 	TWeakPtr<class SViewportToolBar> GetParentToolBar() const;
 
+	/**
+	 * @return true if this menu is open
+	 */
+	bool IsMenuOpen() const;
+
 private:
 	/**
 	 * Called when the menu button is clicked.  Will toggle the visibility of the menu content                   
diff --git a/Engine/Source/Editor/UnrealEd/Public/SEditorViewportViewMenu.h b/Engine/Source/Editor/UnrealEd/Public/SEditorViewportViewMenu.h
--- a/Engine/Source/Editor/UnrealEd/Public/SEditorViewportViewMenu.h
+++ b/Engine/Source/Editor/UnrealEd/Public/SEditorViewportViewMenu.h
@@ -25,7 +25,6 @@ public:
 
 private:
 	FText GetViewMenuLabel() const;
-	const FSlateBrush* GetViewMenuLabelIcon() const;
 	void FillViewMenu(UToolMenu* Menu) const;
 
 protected:
diff --git a/Engine/Source/Editor/UnrealEd/Public/STransformViewportToolbar.h b/Engine/Source/Editor/UnrealEd/Public/STransformViewportToolbar.h
--- a/Engine/Source/Editor/UnrealEd/Public/STransformViewportToolbar.h
+++ b/Engine/Source/Editor/UnrealEd/Public/STransformViewportToolbar.h
@@ -15,6 +15,7 @@ class FExtender;
 class FUICommandList;
 class SEditorViewport;
 class SSlider;
+class SEditorViewportToolbarMenu;
 enum class ECheckBoxState : uint8;
 
 /**
@@ -36,10 +37,11 @@ public:
 	 *
 	 * @return	New widget
 	 */
-	TSharedRef< SWidget > MakeTransformToolBar( const TSharedPtr< FExtender > InExtenders );
+	TSharedRef<SWidget> MakeTransformToolBar(const TSharedPtr< FExtender > InExtenders);
 
 private:
 	FSlateIcon GetLocalToWorldIcon() const;
+	FSlateIcon GetSurfaceSnappingIcon() const;
 
 	/** Camera speed menu construction callback */
 	TSharedRef<SWidget> FillCameraSpeedMenu();
@@ -103,7 +105,9 @@ private:
 	TSharedRef<SWidget> BuildRotationGridCheckBoxList(FName InExtentionHook, const FText& InHeading, const TArray<float>& InGridSizes, ERotationGridMode InGridMode) const;
 
 	/** Make the surface snapping toolbar checkbox button */
-	TSharedRef< SWidget > MakeSurfaceSnappingButton( FName ToolBarStyle );
+	TSharedRef< SWidget > MakeSurfaceSnappingButton();
+	TSharedRef<SWidget> GenerateSurfaceSnappingMenu();
+	FSlateColor GetSurfaceSnappingForegroundColor() const;
 
 	/** Grid Snap checked state callbacks */
 	ECheckBoxState IsLocationGridSnapChecked() const;
@@ -120,12 +124,13 @@ private:
 	void HandleToggleScaleGridSnap(ECheckBoxState InState);
 
 private:
+	TSharedPtr<SEditorViewportToolbarMenu> SurfaceSnappingMenu;
 
 	/** Reference to the camera slider used to display current camera speed */
-	mutable TSharedPtr< SSlider > CamSpeedSlider;
+	TSharedPtr<SSlider> CamSpeedSlider;
 
 	/** Reference to the camera spinbox used to display current camera speed scalar */
-	mutable TSharedPtr< SSpinBox<float> > CamSpeedScalarBox;
+	mutable TSharedPtr<SSpinBox<float>> CamSpeedScalarBox;
 
 	/** The editor viewport that we are in */
 	TWeakPtr<class SEditorViewport> Viewport;
diff --git a/Engine/Source/Editor/UnrealEd/Public/SViewportToolBar.h b/Engine/Source/Editor/UnrealEd/Public/SViewportToolBar.h
--- a/Engine/Source/Editor/UnrealEd/Public/SViewportToolBar.h
+++ b/Engine/Source/Editor/UnrealEd/Public/SViewportToolBar.h
@@ -35,31 +35,10 @@ public:
 	 */
 	void SetOpenMenu( TSharedPtr< SMenuAnchor >& NewMenu );
 
-	/**
-	 * @return The color and opacity of this viewport. 
-	 */ 
-	FLinearColor OnGetColorAndOpacity() const;
 
 	/** @return Whether the given viewmode is supported. */ 
 	virtual bool IsViewModeSupported(EViewModeIndex ViewModeIndex) const;
 
-private:
-	
-	/**
-	 * Called when the mouse enters the toolbar area.  We fade in the toolbar when this happens
-	 *
-	 * @param MyGeometry	Information about the size of the toolbar
-	 * @param MouseEvent	The mouse event that triggered this function
-	 */
-	virtual void OnMouseEnter( const FGeometry& MyGeometry, const FPointerEvent& MouseEvent ) override;
-
-	/**
-	 * Called when the mouse leaves the toolbar area.  We fade out the toolbar when this happens
-	 *
-	 * @param MouseEvent	The mouse event that triggered this function
-	 */
-	virtual void OnMouseLeave( const FPointerEvent& MouseEvent ) override;
-
 protected:
 	/**
 	* Returns the label for the "Camera" tool bar menu based on based viewport type
@@ -80,10 +59,6 @@ protected:
 	virtual const FSlateBrush* GetCameraMenuLabelIconFromViewportType(const ELevelViewportType ViewportType) const;
 
 private:
-	/** Curve sequence for fading out the toolbar */
-	FCurveSequence FadeOutSequence;
-	/** Curve sequence for fading in the toolbar */
-	FCurveSequence FadeInSequence;
 	/** The pulldown menu that is open if any */
 	TWeakPtr< SMenuAnchor > OpenedMenu;
 	/** True if the mouse is inside the toolbar */
diff --git a/Engine/Source/Editor/UnrealEd/Public/SViewportToolBarComboMenu.h b/Engine/Source/Editor/UnrealEd/Public/SViewportToolBarComboMenu.h
--- a/Engine/Source/Editor/UnrealEd/Public/SViewportToolBarComboMenu.h
+++ b/Engine/Source/Editor/UnrealEd/Public/SViewportToolBarComboMenu.h
@@ -23,7 +23,7 @@ class SMenuAnchor;
 class UNREALED_API SViewportToolBarComboMenu : public SCompoundWidget 
 {
 public:
-	SLATE_BEGIN_ARGS(SViewportToolBarComboMenu) : _BlockLocation(EMultiBlockLocation::Start), _MinDesiredButtonWidth(-1.0f) {}
+	SLATE_BEGIN_ARGS(SViewportToolBarComboMenu) : _BlockLocation(EMultiBlockLocation::None), _MinDesiredButtonWidth(-1.0f) {}
 	
 		/** We need to know about the toolbar we are in */
 		SLATE_ARGUMENT( TSharedPtr<class SViewportToolBar>, ParentToolBar );
diff --git a/Engine/Source/Runtime/Slate/Private/Framework/MultiBox/SToolBarComboButtonBlock.cpp b/Engine/Source/Runtime/Slate/Private/Framework/MultiBox/SToolBarComboButtonBlock.cpp
--- a/Engine/Source/Runtime/Slate/Private/Framework/MultiBox/SToolBarComboButtonBlock.cpp
+++ b/Engine/Source/Runtime/Slate/Private/Framework/MultiBox/SToolBarComboButtonBlock.cpp
@@ -179,7 +179,7 @@ void SToolBarComboButtonBlock::BuildMultiBlockWidget(const ISlateStyle* StyleSet
 
 	OpenForegroundColor = ButtonStyle.HoveredForeground;
 
-	const FComboButtonStyle* ComboStyle = &FAppStyle::Get().GetWidgetStyle<FComboButtonStyle>("ComboButton");
+	const FComboButtonStyle* ComboStyle = &ToolBarStyle.ComboButtonStyle;
 	if (ToolBarComboButtonBlock->bSimpleComboBox)
 	{
 		ComboStyle = &ToolBarStyle.SettingsComboButton;
diff --git a/Engine/Source/Runtime/Slate/Private/Framework/MultiBox/SWidgetBlock.cpp b/Engine/Source/Runtime/Slate/Private/Framework/MultiBox/SWidgetBlock.cpp
--- a/Engine/Source/Runtime/Slate/Private/Framework/MultiBox/SWidgetBlock.cpp
+++ b/Engine/Source/Runtime/Slate/Private/Framework/MultiBox/SWidgetBlock.cpp
@@ -117,7 +117,7 @@ void SWidgetBlock::BuildMultiBlockWidget(const ISlateStyle* StyleSet, const FNam
 			]
 		]
 		+SHorizontalBox::Slot()
-		.VAlign( VAlign_Bottom )
+		.VAlign(bHasLabel ? VAlign_Bottom : VAlign_Fill)
 		.FillWidth(1.f)
 		[
 			WidgetBlock->ContentWidget
diff --git a/Engine/Source/Runtime/SlateCore/Private/Styling/StarshipCoreStyle.cpp b/Engine/Source/Runtime/SlateCore/Private/Styling/StarshipCoreStyle.cpp
--- a/Engine/Source/Runtime/SlateCore/Private/Styling/StarshipCoreStyle.cpp
+++ b/Engine/Source/Runtime/SlateCore/Private/Styling/StarshipCoreStyle.cpp
@@ -68,7 +68,8 @@ const FLinearColor FStyleColors::AccentFolder       = COLOR("#B68F55FF");
 
 #define FONT(...) FSlateFontInfo(FLegacySlateFontInfoCache::Get().GetDefaultFont(), __VA_ARGS__)
 
-TSharedPtr< FStyleFonts > FStyleFonts::Instance = nullptr;
+TUniquePtr<FStyleFonts> FStyleFonts::Instance = nullptr;
+
 FStyleFonts::FStyleFonts()
 	: Normal(FONT(10, "Regular"))
 	, NormalBold(FONT(10, "Bold"))
@@ -1180,6 +1181,7 @@ TSharedRef<ISlateStyle> FStarshipCoreStyle::Create()
 			.SetSeparatorPadding(FMargin(1.f, 0.f, 1.f, 0.f))
 			.SetLabelStyle(FTextBlockStyle(NormalText).SetFont(StyleFonts.Normal))
 			.SetEditableTextStyle(FEditableTextBoxStyle(NormalEditableTextBoxStyle).SetFont(StyleFonts.Normal))
+			.SetComboButtonStyle(Style->GetWidgetStyle<FComboButtonStyle>("ComboButton"))
 			.SetBlockPadding(FMargin(2.0f, 2.0f, 4.0f, 4.0f))
 			.SetIndentedBlockPadding(FMargin(18.0f, 2.0f, 4.0f, 4.0f));
 
@@ -1235,6 +1237,7 @@ TSharedRef<ISlateStyle> FStarshipCoreStyle::Create()
 			.SetSeparatorBrush(FSlateColorBrush(FStyleColors::Input))
 			.SetSeparatorPadding(FMargin(8.f, 0))
 			.SetLabelStyle(FTextBlockStyle(NormalText))
+			.SetComboButtonStyle(Style->GetWidgetStyle<FComboButtonStyle>("ComboButton"))
 			.SetLabelPadding(FMargin(5, 9, 0, 9))
 			.SetEditableTextStyle(FEditableTextBoxStyle(NormalEditableTextBoxStyle));
 
@@ -1440,9 +1443,9 @@ TSharedRef<ISlateStyle> FStarshipCoreStyle::Create()
 		Style->Set("Docking.Border", new FSlateRoundedBoxBrush(FStyleColors::Background, 4) );
 
 		Style->Set("Docking.UnhideTabwellButton", FButtonStyle(Button)
-			.SetNormal(IMAGE_BRUSH("/Docking/ShowTabwellButton_Normal", FVector2D(10, 10)))
-			.SetPressed(IMAGE_BRUSH("/Docking/ShowTabwellButton_Pressed", FVector2D(10, 10)))
-			.SetHovered(IMAGE_BRUSH("/Docking/ShowTabwellButton_Hovered", FVector2D(10, 10)))
+			.SetNormal(IMAGE_BRUSH_SVG("Starship/Docking/show-tab-well", Icon8x8, FStyleColors::Primary))
+			.SetPressed(IMAGE_BRUSH_SVG("Starship/Docking/show-tab-well", Icon8x8, FStyleColors::PrimaryPress))
+			.SetHovered(IMAGE_BRUSH_SVG("Starship/Docking/show-tab-well", Icon8x8, FStyleColors::PrimaryHover))
 			.SetNormalPadding(0)
 			.SetPressedPadding(0)
 		);
diff --git a/Engine/Source/Runtime/SlateCore/Public/Styling/StarshipCoreStyle.h b/Engine/Source/Runtime/SlateCore/Public/Styling/StarshipCoreStyle.h
--- a/Engine/Source/Runtime/SlateCore/Public/Styling/StarshipCoreStyle.h
+++ b/Engine/Source/Runtime/SlateCore/Public/Styling/StarshipCoreStyle.h
@@ -129,7 +129,7 @@ struct SLATECORE_API FStyleFonts
 	{
 		if (Instance == nullptr)
 		{
-			Instance = MakeShared<FStyleFonts>(FStyleFonts());
+			Instance = MakeUnique<FStyleFonts>(FStyleFonts());
 		}
 		return *(Instance.Get());
 	}
@@ -142,6 +142,6 @@ struct SLATECORE_API FStyleFonts
 
   private: 
   	FStyleFonts();
-  	static TSharedPtr< struct FStyleFonts > Instance;
+  	static TUniquePtr<struct FStyleFonts> Instance;
 };
 
diff --git a/Engine/Source/Runtime/SlateCore/Public/Styling/ToolBarStyle.h b/Engine/Source/Runtime/SlateCore/Public/Styling/ToolBarStyle.h
--- a/Engine/Source/Runtime/SlateCore/Public/Styling/ToolBarStyle.h
+++ b/Engine/Source/Runtime/SlateCore/Public/Styling/ToolBarStyle.h
@@ -57,6 +57,10 @@ struct SLATECORE_API FToolBarStyle : public FSlateWidgetStyle
 	FCheckBoxStyle ToggleButton;
 	FToolBarStyle& SetToggleButtonStyle(const FCheckBoxStyle& InToggleButton) { ToggleButton = InToggleButton; return *this; }
 
+	UPROPERTY(EditAnywhere, BlueprintReadWrite, Category = Appearance)
+	FComboButtonStyle ComboButtonStyle;
+	FToolBarStyle& SetComboButtonStyle(const FComboButtonStyle& InComboButtonStyle) { ComboButtonStyle = InComboButtonStyle; return *this; }
+
 	UPROPERTY(EditAnywhere, BlueprintReadWrite, Category = Appearance)
 	FComboButtonStyle SettingsComboButton;
 	FToolBarStyle& SetSettingsComboButtonStyle(const FComboButtonStyle& InSettingsComboButton) { SettingsComboButton = InSettingsComboButton; return *this; }
