Fixed build.

diff --git a/Engine/Source/Editor/UnrealEd/Private/Fbx/FbxStaticMeshImport.cpp b/Engine/Source/Editor/UnrealEd/Private/Fbx/FbxStaticMeshImport.cpp
--- a/Engine/Source/Editor/UnrealEd/Private/Fbx/FbxStaticMeshImport.cpp
+++ b/Engine/Source/Editor/UnrealEd/Private/Fbx/FbxStaticMeshImport.cpp
@@ -1061,7 +1061,7 @@ bool UnFbx::FFbxImporter::BuildStaticMeshFromGeometry(FbxNode* Node, UStaticMesh
 						for (int32 VertexIndex = 0; VertexIndex < 3; ++VertexIndex)
 						{
 							int UVMapIndex = (FBXUVs.UVMappingMode[UVLayerIndex] == FbxLayerElement::eByControlPoint)
-								? MeshDescription->GetVertexInstanceVertex(CornerInstanceIDs[VertexIndex])
+								? Mesh->GetPolygonVertex(PolygonIndex, VertexIndex)
 								: SkippedVertexInstance + CurrentVertexInstanceIndex + VertexIndex;
 							int32 UVIndex = (FBXUVs.UVReferenceMode[UVLayerIndex] == FbxLayerElement::eDirect)
 								? UVMapIndex
