- Turnkey from Editor support (install Sdk, flash devices as needed, storing Sdk info in DDPI for display in Launch On menu, etc)
- There is a second Launch dropdown button for now, the original is still there for now (quick launch of previous device will not do Turnkey at this time)
- Various platform turnkey support updates

diff --git a/Engine/Build/BatchFiles/RunUAT.bat b/Engine/Build/BatchFiles/RunUAT.bat
--- a/Engine/Build/BatchFiles/RunUAT.bat
+++ b/Engine/Build/BatchFiles/RunUAT.bat
@@ -52,9 +52,24 @@ rem ## Run AutomationTool
 :DoRunUAT
 pushd %UATDirectory%
 %UATExecutable% %* %UATCompileArg%
-if not %ERRORLEVEL% == 0 goto Error_UATFailed
 popd
 
+for %%P in (%*) do if /I "%%P" == "-noturnkeyvariables" goto SkipTurnkey
+
+rem ## Turnkey needs to update env vars in the calling process so that if it is run multiple times the Sdk env var changes are in effect
+if EXIST ~dp0..\..\Intermediate\Turnkey\PostTurnkeyVariables.bat (
+	rem ## We need to endlocal so that the vars in the batch file work. NOTE: Working directory from pushd will be UNDONE here, but since we are about to quit, it's okay
+	endlocal 
+	echo Updating environment variables set by a Turnkey sub-process
+	call ~dp0..\..\Engine\Intermediate\Turnkey\PostTurnkeyVariables.bat
+	del ~dp0..\..\Engine\Intermediate\Turnkey\PostTurnkeyVariables.bat
+	rem ## setlocal again so that any popd's etc don't have an effect on calling process
+	setlocal
+)
+:SkipTurnkey
+
+if not %ERRORLEVEL% == 0 goto Error_UATFailed
+
 rem ## Success!
 goto Exit
 
diff --git a/Engine/Source/Developer/LauncherServices/Private/Launcher/LauncherUATTask.h b/Engine/Source/Developer/LauncherServices/Private/Launcher/LauncherUATTask.h
--- a/Engine/Source/Developer/LauncherServices/Private/Launcher/LauncherUATTask.h
+++ b/Engine/Source/Developer/LauncherServices/Private/Launcher/LauncherUATTask.h
@@ -17,12 +17,13 @@ class FLauncherUATTask
 {
 public:
 
-	FLauncherUATTask( const FString& InCommandLine, const FString& InName, const FString& InDesc, void* InReadPipe, void* InWritePipe, const FString& InEditorExe, FProcHandle& InProcessHandle, ILauncherWorker* InWorker, const FString& InCommandEnd)
+	FLauncherUATTask( const FString& InCommandLine, const FString& InName, const FString& InDesc, void* InReadPipe, void* InWritePipe, const FString& InEditorExe, FProcHandle& InProcessHandle, ILauncherWorker* InWorker, const FString& InCommandEnd, const FString& InBuildCookRunPreCommand)
 		: FLauncherTask(InName, InDesc, InReadPipe, InWritePipe)
 		, CommandLine(InCommandLine)
 		, EditorExe(InEditorExe)
 		, ProcessHandle(InProcessHandle)
 		, CommandText(InCommandEnd)
+		, PreCommand(InBuildCookRunPreCommand)
 	{
 		EndTextFound = false;
 		InWorker->OnOutputReceived().AddRaw(this, &FLauncherUATTask::HandleOutputReceived);
@@ -55,8 +56,9 @@ protected:
 		FString UATCommandLine;
 		FString ProjectPath = *ChainState.Profile->GetProjectPath();
 		ProjectPath = FPaths::ConvertRelativePathToFull(ProjectPath);
-		UATCommandLine = FString::Printf(TEXT("-ScriptsForProject=\"%s\" BuildCookRun -project=\"%s\" -noP4 -clientconfig=%s -serverconfig=%s"),
+		UATCommandLine = FString::Printf(TEXT("-ScriptsForProject=\"%s\" %s BuildCookRun -project=\"%s\" -noP4 -clientconfig=%s -serverconfig=%s"),
 			*ProjectPath,
+			*PreCommand,
 			*ProjectPath,
 			LexToString(ChainState.Profile->GetBuildConfiguration()),
 			LexToString(ChainState.Profile->GetBuildConfiguration()));
@@ -129,6 +131,9 @@ private:
 	// The editor executable that UAT should use
 	FString CommandText;
 
+	// a UAT command to run before BuildCookRun (in the same instance as the BuildCookRun)
+	FString PreCommand;
+
 	bool EndTextFound;
 };
 
diff --git a/Engine/Source/Developer/LauncherServices/Private/Launcher/LauncherWorker.cpp b/Engine/Source/Developer/LauncherServices/Private/Launcher/LauncherWorker.cpp
--- a/Engine/Source/Developer/LauncherServices/Private/Launcher/LauncherWorker.cpp
+++ b/Engine/Source/Developer/LauncherServices/Private/Launcher/LauncherWorker.cpp
@@ -939,7 +939,15 @@ void FLauncherWorker::CreateAndExecuteTasks( const ILauncherProfileRef& InProfil
 	TArray<FCommandDesc> Commands;
 	FString StartString;
 	FString UATCommand = CreateUATCommand(InProfile, Platforms, Commands, StartString);
-	TSharedPtr<FLauncherTask> BuildTask = MakeShareable(new FLauncherUATTask(UATCommand, TEXT("Build Task"), TEXT("Launching UAT..."), ReadPipe, WritePipe, InProfile->GetEditorExe(), ProcHandle, this, StartString));
+	
+	// have Turnkey 
+	FString TurnkeyCommand;
+	if (Profile->ShouldUpdateDeviceFlash() && DeviceGroup->GetDeviceIDs().Num() >= 0)
+	{
+		TurnkeyCommand = FString::Printf(TEXT("Turnkey -command=VerifySdk -type=Flash -device=%s -UpdateIfNeeded -EditorIO"), *FString::Join(DeviceGroup->GetDeviceIDs(), TEXT("+")));
+	}
+
+	TSharedPtr<FLauncherTask> BuildTask = MakeShareable(new FLauncherUATTask(UATCommand, TEXT("Build Task"), TEXT("Launching UAT..."), ReadPipe, WritePipe, InProfile->GetEditorExe(), ProcHandle, this, StartString, TurnkeyCommand));
 	BuildTask->OnStarted().AddRaw(this, &FLauncherWorker::OnTaskStarted);
 	BuildTask->OnCompleted().AddRaw(this, &FLauncherWorker::OnTaskCompleted);
 	NextTask->AddContinuation(BuildTask);
diff --git a/Engine/Source/Developer/LauncherServices/Private/Profiles/LauncherProfile.h b/Engine/Source/Developer/LauncherServices/Private/Profiles/LauncherProfile.h
--- a/Engine/Source/Developer/LauncherServices/Private/Profiles/LauncherProfile.h
+++ b/Engine/Source/Developer/LauncherServices/Private/Profiles/LauncherProfile.h
@@ -2446,6 +2446,19 @@ public:
 		return bUseIoStore;
 	}
 
+	virtual void SetShouldUpdateDeviceFlash(bool bInShouldUpdateFlash) override
+	{
+		bShouldUpdateFlash = bInShouldUpdateFlash;
+	}
+
+	/**
+	 * Whether or not the flash image/software on the device should attempt to be updated before running
+	 */
+	virtual bool ShouldUpdateDeviceFlash() const override
+	{
+		return bShouldUpdateFlash;
+	}
+
 	//~ End ILauncherProfile Interface
 
 protected:
@@ -2900,6 +2913,9 @@ private:
 	// Use I/O store.
 	bool bUseIoStore;
 
+	// Update flash on device before running
+	bool bShouldUpdateFlash;
+
 private:
 
 	// cook in the editor callbacks (not valid for any other cook mode)
diff --git a/Engine/Source/Developer/LauncherServices/Public/ILauncherProfile.h b/Engine/Source/Developer/LauncherServices/Public/ILauncherProfile.h
--- a/Engine/Source/Developer/LauncherServices/Public/ILauncherProfile.h
+++ b/Engine/Source/Developer/LauncherServices/Public/ILauncherProfile.h
@@ -1303,6 +1303,15 @@ public:
 	 */
 	virtual bool IsUsingIoStore() const = 0;
 
+	/**
+	 * Sets whether or not the flash image/software on the device should attempt to be updated before running
+	 */
+	virtual void SetShouldUpdateDeviceFlash(bool bInShouldUpdateFlash) = 0;
+
+	/**
+	 * Whether or not the flash image/software on the device should attempt to be updated before running
+	 */
+	virtual bool ShouldUpdateDeviceFlash() const = 0;
 public:
 
 	/** Virtual destructor. */
diff --git a/Engine/Source/Editor/MainFrame/Private/Frame/MainFrameActions.cpp b/Engine/Source/Editor/MainFrame/Private/Frame/MainFrameActions.cpp
--- a/Engine/Source/Editor/MainFrame/Private/Frame/MainFrameActions.cpp
+++ b/Engine/Source/Editor/MainFrame/Private/Frame/MainFrameActions.cpp
@@ -829,9 +829,12 @@ void FMainFrameActionCallbacks::PackageProject( const FName InPlatformInfoName )
 		OptionalParams += FString::Printf(TEXT(" -target=%s -clientconfig=%s"), *Target->Name, LexToString(ConfigurationInfo.Configuration));
 	}
 
+	FString TurnkeyCommandline = FString::Printf(TEXT("Turnkey -command=VerifySdk -platform=%s -UpdateIfNeeded -EditorIO"), *PlatformInfo->IniPlatformName);
+
 	FString ProjectPath = FPaths::IsProjectFilePathSet() ? FPaths::ConvertRelativePathToFull(FPaths::GetProjectFilePath()) : FPaths::RootDir() / FApp::GetProjectName() / FApp::GetProjectName() + TEXT(".uproject");
-	FString CommandLine = FString::Printf(TEXT("-ScriptsForProject=\"%s\" BuildCookRun %s%s -nop4 -project=\"%s\" -cook -stage -archive -archivedirectory=\"%s\" -package -ue4exe=\"%s\" %s -utf8output"),
+	FString CommandLine = FString::Printf(TEXT("-ScriptsForProject=\"%s\" %s BuildCookRun %s%s -nop4 -project=\"%s\" -cook -stage -archive -archivedirectory=\"%s\" -package -ue4exe=\"%s\" %s -utf8output"),
 		*ProjectPath,
+		*TurnkeyCommandline,
 		GetUATCompilationFlags(),
 		FApp::IsEngineInstalled() ? TEXT(" -installed") : TEXT(""),
 		*ProjectPath,
diff --git a/Engine/Source/Editor/UnrealEd/Private/Kismet2/DebuggerCommands.cpp b/Engine/Source/Editor/UnrealEd/Private/Kismet2/DebuggerCommands.cpp
--- a/Engine/Source/Editor/UnrealEd/Private/Kismet2/DebuggerCommands.cpp
+++ b/Engine/Source/Editor/UnrealEd/Private/Kismet2/DebuggerCommands.cpp
@@ -73,6 +73,13 @@
 #include "SBlueprintEditorToolbar.h"
 #include "SEnumCombobox.h"
 
+
+#include "Misc/DataDrivenPlatformInfoRegistry.h"
+#include "IUATHelperModule.h"
+#include "ISettingsEditorModule.h"
+#include "Async/Async.h"
+
+
 #define LOCTEXT_NAMESPACE "DebuggerCommands"
 
 void SGlobalPlayWorldActions::Construct(const FArguments& InArgs)
@@ -143,7 +150,7 @@ public:
 	static void PlayInSettings_Clicked();
 
 	// Launch On
-	static void HandleLaunchOnDeviceActionExecute(FString DevicedId, FString DeviceName);
+	static void HandleLaunchOnDeviceActionExecute(FString DevicedId, FString DeviceName, bool bUseTurnkey);
 	static bool HandleLaunchOnDeviceActionCanExecute(FString DeviceName);
 	static bool HandleLaunchOnDeviceActionIsChecked(FString DeviceName);
 
@@ -225,7 +232,7 @@ protected:
 	 *
 	 * @param PlatformName - The name of the platform to play the game on.
 	 */
-	static void LaunchOnDevice(const FString& DeviceId, const FString& DeviceName);
+	static void LaunchOnDevice(const FString& DeviceId, const FString& DeviceName, bool bUseTurnkey);
 
 	/** Get the player start location to use when starting PIE */
 	static EPlayModeLocations GetPlayModeLocation();
@@ -668,6 +675,19 @@ void FPlayWorldCommands::BuildToolbar(FToolMenuSection& InSection, bool bInclude
 					FSlateIcon(FEditorStyle::GetStyleSetName(), "LevelEditor.RepeatLastLaunch"),
 					true
 				));
+
+				FUIAction TurnkeyLaunchMenuAction;
+				TurnkeyLaunchMenuAction.IsActionVisibleDelegate = FIsActionButtonVisible::CreateStatic(&FInternalPlayWorldCommandCallbacks::CanShowNonPlayWorldOnlyActions);
+
+				InDynamicSection.AddEntry(FToolMenuEntry::InitComboButton(
+					"TurnkeyCombo",
+					TurnkeyLaunchMenuAction,
+					FOnGetContent::CreateStatic(&GenerateTurnkeyMenuContent, GlobalPlayWorldActions.ToSharedRef()),
+					LOCTEXT("LaunchCombo_Label", "Launch Options"),
+					LOCTEXT("PlayWithTurnkeyComboToolTip", "Options for launching on a device and manage platforms"),
+					FSlateIcon(FEditorStyle::GetStyleSetName(), "LevelEditor.RepeatLastLaunch"),
+					true
+				));
 			}
 		}));
 	}
@@ -1019,7 +1039,7 @@ static void MakeAllDevicesSubMenu(FMenuBuilder& InMenuBuilder, const PlatformInf
 
 		// create an action
 		FUIAction LaunchDeviceAction(
-			FExecuteAction::CreateStatic(&FInternalPlayWorldCommandCallbacks::HandleLaunchOnDeviceActionExecute, PlatformId, PlatformVariantStr),
+			FExecuteAction::CreateStatic(&FInternalPlayWorldCommandCallbacks::HandleLaunchOnDeviceActionExecute, PlatformId, PlatformVariantStr, false),
 			FCanExecuteAction::CreateStatic(&FInternalPlayWorldCommandCallbacks::HandleLaunchOnDeviceActionCanExecute, PlatformVariantStr),
 			FIsActionChecked::CreateStatic(&FInternalPlayWorldCommandCallbacks::HandleLaunchOnDeviceActionIsChecked, PlatformVariantStr)
 		);
@@ -1043,6 +1063,490 @@ static void MakeAllDevicesSubMenu(FMenuBuilder& InMenuBuilder, const PlatformInf
 	}
 }
 
+static void TurnkeyInstallSdk(FString PlatformName, bool bPreferFull, bool bForceInstall, const TCHAR* DeviceName=nullptr)
+{
+//	FString CommandLine = FString::Printf(TEXT("Turnkey -command=InstallSdk -platform=%s -BestAvailable -AllowAutoSdk -EditorIO -noturnkeyvariables"), *PlatformName);
+
+	FString OptionalOptions;
+	if (bPreferFull)
+	{
+		OptionalOptions += TEXT(" -PreferFull");
+	}
+	if (bForceInstall)
+	{
+		OptionalOptions += DeviceName != nullptr ? TEXT(" -ForceDeviceInstall") : TEXT(" -ForceSdkInstall");
+	}
+	if (DeviceName != nullptr)
+	{
+		OptionalOptions += FString::Printf(TEXT(" -Device=%s"), DeviceName);
+	}
+
+	FString CommandLine = FString::Printf(TEXT("Turnkey -command=VerifySdk -UpdateIfNeeded -platform=%s %s -EditorIO -noturnkeyvariables"), *PlatformName, *OptionalOptions);
+
+	FText TaskName = LOCTEXT("InstallingSdk", "Installing Sdk");
+	IUATHelperModule::Get().CreateUatTask(CommandLine, FText::FromString(PlatformName), TaskName, TaskName, FEditorStyle::GetBrush(TEXT("MainFrame.PackageProject")),
+		[](FString, double)
+	{
+		AsyncTask(ENamedThreads::GameThread, []()
+		{
+
+			// read in env var changes
+			// @todo turnkey move this and make it mac/linux aware
+			FString TurnkeyEnvVarsFilename = FPaths::Combine(FPaths::EngineIntermediateDir(), TEXT("Turnkey/PostTurnkeyVariables.bat"));
+
+			if (IFileManager::Get().FileExists(*TurnkeyEnvVarsFilename))
+			{
+				TArray<FString> Contents;
+				if (FFileHelper::LoadFileToStringArray(Contents, *TurnkeyEnvVarsFilename))
+				{
+					for (const FString& Line : Contents)
+					{
+						if (Line.StartsWith(TEXT("set ")))
+						{
+							// split the line
+							FString VariableLine = Line.Mid(4);
+							int Equals;
+							if (VariableLine.FindChar('=', Equals))
+							{
+								// set the key/value
+								FString Key = VariableLine.Mid(0, Equals);
+								FString Value = VariableLine.Mid(Equals + 1);
+
+								FPlatformMisc::SetEnvironmentVar(*Key, *Value);
+
+								UE_LOG(LogTemp, Display, TEXT("Turnkey setting env var: %s = %s"), *Key, *Value);
+							}
+						}
+					}
+				}
+			}
+
+			// update the Sdk status
+			FDataDrivenPlatformInfoRegistry::UpdateSdkStatus();
+
+			// show restart dialog
+			FModuleManager::GetModuleChecked<ISettingsEditorModule>("SettingsEditor").OnApplicationRestartRequired();
+		});
+	}
+	);
+}
+
+static TAttribute<FText> MakeSdkStatusAttribute(const PlatformInfo::FPlatformInfo* PlatformInfo, TSharedPtr< ITargetDeviceProxy> DeviceProxy)
+{
+// 	FFormatNamedArguments LabelArguments;
+// 	LabelArguments.Add(TEXT("DeviceName"), FText::FromString(DeviceProxy->GetName()));
+// 
+// 	if (!DeviceProxy->IsConnected())
+// 	{
+// 		LabelArguments.Add(TEXT("HostUser"), LOCTEXT("DisconnectedHint", " [Disconnected]"));
+// 	}
+// 	else if (DeviceProxy->GetHostUser() != FPlatformProcess::UserName(false))
+// 	{
+// 		LabelArguments.Add(TEXT("HostUser"), FText::FromString(DeviceProxy->GetHostUser()));
+// 	}
+// 	else
+// 	{
+// 		LabelArguments.Add(TEXT("HostUser"), FText::GetEmpty());
+// 	}
+// 
+// 	OutLabel = FText::Format(LOCTEXT("LaunchDeviceLabel", "{DeviceName}{HostUser}"), LabelArguments);
+// 	TAttribute<FText> = Make
+
+
+	FText DisplayString = DeviceProxy ? FText::FromString(DeviceProxy->GetName()) : PlatformInfo->DisplayName;
+
+	return TAttribute<FText>::Create(TAttribute<FText>::FGetter::CreateLambda([PlatformInfo, DisplayString]()
+	{
+		// get the status, or Unknown if it's not there
+		const FDataDrivenPlatformInfoRegistry::FPlatformInfo& Info = FDataDrivenPlatformInfoRegistry::GetPlatformInfo(PlatformInfo->IniPlatformName);
+		DDPIPlatformSdkStatus Status = Info.SdkStatus;
+
+		// @todo turnkey: Have premade FText's by SdkStatus for speed
+		FText Desc = FText::FromString(
+			Status == DDPIPlatformSdkStatus::Querying ? TEXT("Querying...") :
+			Status == DDPIPlatformSdkStatus::Valid ? TEXT("Valid Sdk") :
+			Status == DDPIPlatformSdkStatus::OutOfDate ? TEXT("Outdated Sdk") :
+			Status == DDPIPlatformSdkStatus::NoSdk ? TEXT("No Sdk") :
+			Status == DDPIPlatformSdkStatus::FlashValid ? TEXT("Valid Flash") :
+			Status == DDPIPlatformSdkStatus::FlashOutOfDate ? TEXT("Outdated Flash") :
+			TEXT("???"));
+
+		return FText::Format(FText::FromString(TEXT("{0} ({1})")), DisplayString, Desc);
+	}));
+}
+
+
+static void MakeTurnkeyPlatformMenu(FMenuBuilder& MenuBuilder, const PlatformInfo::FPlatformInfo* PlatformInfo, ITargetDeviceServicesModule* TargetDeviceServicesModule)
+{
+	MenuBuilder.BeginSection("AllDevices", LOCTEXT("TurnkeySection_AllDevices", "All Devices"));
+
+	TArray<TSharedPtr<ITargetDeviceProxy>> DeviceProxies;
+	TargetDeviceServicesModule->GetDeviceProxyManager()->GetAllProxies(*PlatformInfo->IniPlatformName, DeviceProxies);
+
+	FString PlatformName = PlatformInfo->UBTTargetId.ToString();
+	for (const TSharedPtr<ITargetDeviceProxy> Proxy : DeviceProxies)
+	{
+		FString DeviceName = Proxy->GetName();
+		MenuBuilder.AddSubMenu(
+			MakeSdkStatusAttribute(PlatformInfo, DeviceProxies[0]),
+			FText(),
+			FNewMenuDelegate::CreateLambda([PlatformName, DeviceName](FMenuBuilder& SubMenuBuilder) {
+				SubMenuBuilder.AddMenuEntry(
+					LOCTEXT("Turnkey_RepairDevice", "Repair Device as Needed"),
+					LOCTEXT("TurnkeyTooltip_RepairDevice", "Perform any fixup that may be needed on this device. If up to date already, nothing will be done."),
+					FSlateIcon(),
+					FExecuteAction::CreateStatic(TurnkeyInstallSdk, PlatformName, false, false, *DeviceName)
+				);
+
+				SubMenuBuilder.AddMenuEntry(
+					LOCTEXT("Turnkey_ForceRepairDevice", "Force Repair Device"),
+					LOCTEXT("TurnkeyTooltip_ForceRepairDevice", "Force repairing anything on the device needed (update firmware, etc). Will perform all steps possible, even if not needed."),
+					FSlateIcon(),
+					FExecuteAction::CreateStatic(TurnkeyInstallSdk, PlatformName, true, false, *DeviceName)
+				);
+			})
+		);
+	}
+
+	MenuBuilder.EndSection();
+
+	MenuBuilder.BeginSection("AllDevices", LOCTEXT("TurnkeySection_Sdks", "Sdk Managment"));
+
+	const FDataDrivenPlatformInfoRegistry::FPlatformInfo& Info = FDataDrivenPlatformInfoRegistry::GetPlatformInfo(PlatformInfo->IniPlatformName);
+	const TCHAR* NoDevice = nullptr;
+	if (Info.SdkStatus == DDPIPlatformSdkStatus::OutOfDate)
+	{
+		MenuBuilder.AddMenuEntry(
+			LOCTEXT("Turnkey_InstallSdkMinimal", "Update Sdk (Prefer Minimal)"),
+			LOCTEXT("TurnkeyTooltip_InstallSdkMinimal", "Attempt to update an Sdk, as hosted by your studio. Will attempt to install a minimal Sdk (useful for building/running only)"),
+			FSlateIcon(),
+			FExecuteAction::CreateStatic(TurnkeyInstallSdk, PlatformName, false, false, NoDevice)
+		);
+
+		MenuBuilder.AddMenuEntry(
+			LOCTEXT("Turnkey_InstallSdkFull", "Update Sdk (Prefer Full)"),
+			LOCTEXT("TurnkeyTooltip_InstallSdkMinimal", "Attempt to update an Sdk, as hosted by your studio. Will attempt to install a full Sdk (useful profiling or other use cases)"),
+			FSlateIcon(),
+			FExecuteAction::CreateStatic(TurnkeyInstallSdk, PlatformName, true, false, NoDevice)
+		);
+	}
+	else if (Info.SdkStatus == DDPIPlatformSdkStatus::Valid)
+	{
+		MenuBuilder.AddMenuEntry(
+			LOCTEXT("Turnkey_InstallSdkMinimal", "Force Reinstall Sdk (Prefer Minimal)"),
+			LOCTEXT("TurnkeyTooltip_InstallSdkMinimal", "Attempt to force re-install an Sdk, as hosted by your studio. Will attempt to install a minimal Sdk (useful for building/running only)"),
+			FSlateIcon(),
+			FExecuteAction::CreateStatic(TurnkeyInstallSdk, PlatformName, false, true, NoDevice)
+		);
+
+		MenuBuilder.AddMenuEntry(
+			LOCTEXT("Turnkey_InstallSdkFull", "Force Reinstall (Prefer Full)"),
+			LOCTEXT("TurnkeyTooltip_InstallSdkMinimal", "Attempt to force re-install an Sdk, as hosted by your studio. Will attempt to install a full Sdk (useful profiling or other use cases)"),
+			FSlateIcon(),
+			FExecuteAction::CreateStatic(TurnkeyInstallSdk, PlatformName, true, true, NoDevice)
+		);
+	}
+	else
+	{
+		MenuBuilder.AddMenuEntry(
+			LOCTEXT("Turnkey_InstallSdkMinimal", "Install Sdk (Prefer Minimal)"),
+			LOCTEXT("TurnkeyTooltip_InstallSdkMinimal", "Attempt to install an Sdk, as hosted by your studio. Will attempt to install a minimal Sdk (useful for building/running only)"),
+			FSlateIcon(),
+			FExecuteAction::CreateStatic(TurnkeyInstallSdk, PlatformName, false, false, NoDevice)
+		);
+
+		MenuBuilder.AddMenuEntry(
+			LOCTEXT("Turnkey_InstallSdkFull", "Install Sdk (Prefer Full)"),
+			LOCTEXT("TurnkeyTooltip_InstallSdkMinimal", "Attempt to install an Sdk, as hosted by your studio. Will attempt to install a full Sdk (useful profiling or other use cases)"),
+			FSlateIcon(),
+			FExecuteAction::CreateStatic(TurnkeyInstallSdk, PlatformName, true, false, NoDevice)
+		);
+	}
+}
+
+static void GenerateDeviceProxyMenuParams(TSharedPtr<ITargetDeviceProxy> DeviceProxy, const PlatformInfo::FPlatformInfo* EditorPlatformInfo, FName PlatformName, FUIAction& OutAction, FText& OutTooltip)
+{
+	// 	// create an All_<platform>_devices_on_<host> submenu
+	// 	if (DeviceProxy->IsAggregated())
+	// 	{
+	// 		FString AggregateDevicedName(FString::Printf(TEXT("  %s"), *DeviceProxy->GetName())); //align with the other menu entries
+	// 		FSlateIcon AggregateDeviceIcon(FEditorStyle::GetStyleSetName(), EditorPlatformInfo->GetIconStyleName(PlatformInfo::EPlatformIconSize::Normal));
+	// 
+	// 		MenuBuilder.AddSubMenu(
+	// 			FText::FromString(AggregateDevicedName),
+	// 			FText::FromString(AggregateDevicedName),
+	// 			FNewMenuDelegate::CreateStatic(&MakeAllDevicesSubMenu, EditorPlatformInfo, DeviceProxy),
+	// 			false, AggregateDeviceIcon, true
+	// 		);
+	// 		continue;
+	// 	}
+
+		// ... create an action...
+	OutAction = FUIAction(
+		FExecuteAction::CreateStatic(&FInternalPlayWorldCommandCallbacks::HandleLaunchOnDeviceActionExecute, DeviceProxy->GetTargetDeviceId(NAME_None), DeviceProxy->GetName(), true)
+//		, FCanExecuteAction::CreateStatic(&FInternalPlayWorldCommandCallbacks::HandleLaunchOnDeviceActionCanExecute, DeviceProxy->GetName())
+//		, FIsActionChecked::CreateStatic(&FInternalPlayWorldCommandCallbacks::HandleLaunchOnDeviceActionIsChecked, DeviceProxy->GetName())
+	);
+
+	// ... generate tooltip text
+	FFormatNamedArguments TooltipArguments;
+	TooltipArguments.Add(TEXT("DeviceID"), FText::FromString(DeviceProxy->GetName()));
+	TooltipArguments.Add(TEXT("DisplayName"), EditorPlatformInfo->DisplayName);
+	OutTooltip = FText::Format(LOCTEXT("LaunchDeviceToolTipText_ThisDevice", "Launch the game on this {DisplayName} device ({DeviceID})"), TooltipArguments);
+	if (!DeviceProxy->IsAuthorized())
+	{
+		OutTooltip = FText::Format(LOCTEXT("LaunchDeviceToolTipText_UnauthorizedOrLocked", "{DisplayName} device ({DeviceID}) is unauthorized or locked"), TooltipArguments);
+	}
+
+	FProjectStatus ProjectStatus;
+	if (IProjectManager::Get().QueryStatusForCurrentProject(ProjectStatus) && !ProjectStatus.IsTargetPlatformSupported(PlatformName))
+	{
+		FText TooltipLine2 = FText::Format(LOCTEXT("LaunchDevicePlatformWarning", "{DisplayName} is not listed as a target platform for this project, so may not run as expected."), TooltipArguments);
+		OutTooltip = FText::Format(FText::FromString(TEXT("{0}\n\n{1}")), OutTooltip, TooltipLine2);
+	}
+
+}
+
+TSharedRef< SWidget > FPlayWorldCommands::GenerateTurnkeyMenuContent(TSharedRef<FUICommandList> InCommandList)
+{
+	const bool bShouldCloseWindowAfterMenuSelection = true;
+	FMenuBuilder MenuBuilder(bShouldCloseWindowAfterMenuSelection, InCommandList);
+
+	TArray<PlatformInfo::FVanillaPlatformEntry> VanillaPlatforms = PlatformInfo::BuildPlatformHierarchy(PlatformInfo::EPlatformFilter::All);
+
+	VanillaPlatforms.Sort([](const PlatformInfo::FVanillaPlatformEntry& One, const PlatformInfo::FVanillaPlatformEntry& Two) -> bool
+	{
+		return One.PlatformInfo->DisplayName.CompareTo(Two.PlatformInfo->DisplayName) < 0;
+	});
+
+	// shared devices section
+	ITargetDeviceServicesModule* TargetDeviceServicesModule = static_cast<ITargetDeviceServicesModule*>(FModuleManager::Get().LoadModule(TEXT("TargetDeviceServices")));
+	IProjectTargetPlatformEditorModule& ProjectTargetPlatformEditorModule = FModuleManager::LoadModuleChecked<IProjectTargetPlatformEditorModule>("ProjectTargetPlatformEditor");
+
+	TArray<FString> DeviceIdsToQuery;
+
+	MenuBuilder.BeginSection("LevelEditorLaunchDevices", LOCTEXT("TurnkeySection_LaunchButtonDevices", "Quick Launch"));
+	{
+		for (const auto Pair : FDataDrivenPlatformInfoRegistry::GetAllPlatformInfos())
+		{
+			if (Pair.Value.bIsFakePlatform)
+			{
+				continue;
+			}
+
+			FName PlatformName(*Pair.Key);
+			const FDataDrivenPlatformInfoRegistry::FPlatformInfo& Info = Pair.Value;
+
+			// find a matching editor-facing FPlatformInfo object
+			const PlatformInfo::FPlatformInfo* EditorPlatformInfo = PlatformInfo::FindVanillaPlatformInfo(PlatformName);
+			check(EditorPlatformInfo != nullptr);
+
+			// look for devices for all platforms, even if the platform isn't installed - Turnkey can install Sdk after selecting LaunchOn
+			TArray<TSharedPtr<ITargetDeviceProxy>> DeviceProxies;
+			TargetDeviceServicesModule->GetDeviceProxyManager()->GetAllProxies(PlatformName, DeviceProxies);
+
+			if (DeviceProxies.Num() > 0)
+			{
+				// 					Algo::Sort(DeviceProxies, [](TSharedPtr<ITargetDeviceProxy> A, TSharedPtr<ITargetDeviceProxy> B)
+				// 					{ return A->Name == LastChosen || B->Name != LastChosen && A->IsDefault(); });
+
+
+				// always use the first one, after sorting
+				FUIAction Action;
+				FText Tooltip;
+				GenerateDeviceProxyMenuParams(DeviceProxies[0], EditorPlatformInfo, PlatformName, Action, Tooltip);
+
+				if (DeviceProxies.Num() == 1)
+				{
+					MenuBuilder.AddMenuEntry(
+						MakeSdkStatusAttribute(EditorPlatformInfo, DeviceProxies[0]),
+						Tooltip,
+						FSlateIcon(FEditorStyle::GetStyleSetName(), EditorPlatformInfo->GetIconStyleName(PlatformInfo::EPlatformIconSize::Normal)),
+						Action,
+						NAME_None,
+						EUserInterfaceActionType::Button
+					);
+				}
+				else
+				{
+					MenuBuilder.AddSubMenu(
+						MakeSdkStatusAttribute(EditorPlatformInfo, DeviceProxies[0]),
+						Tooltip,
+						FNewMenuDelegate::CreateLambda([TargetDeviceServicesModule, EditorPlatformInfo, PlatformName](FMenuBuilder& SubMenuBuilder) {
+							// re-get the proxies, just in case they changed
+							TArray<TSharedPtr<ITargetDeviceProxy>> DeviceProxies;
+							TargetDeviceServicesModule->GetDeviceProxyManager()->GetAllProxies(PlatformName, DeviceProxies);
+							// for each one, put an entry (even the one that was in the outer menu, for less confusion)
+							for (const TSharedPtr<ITargetDeviceProxy> Proxy : DeviceProxies)
+							{
+								FUIAction SubAction;
+								FText SubTooltip;
+								GenerateDeviceProxyMenuParams(Proxy, EditorPlatformInfo, PlatformName, SubAction, SubTooltip);
+								SubMenuBuilder.AddMenuEntry(
+									MakeSdkStatusAttribute(EditorPlatformInfo, Proxy),
+									SubTooltip,
+									FSlateIcon(FEditorStyle::GetStyleSetName(), EditorPlatformInfo->GetIconStyleName(PlatformInfo::EPlatformIconSize::Normal)),
+									SubAction,
+									NAME_None,
+									EUserInterfaceActionType::Button
+								);
+							}
+						}),
+						Action,
+						NAME_None,
+						EUserInterfaceActionType::Check,
+						false,
+						FSlateIcon(FEditorStyle::GetStyleSetName(), EditorPlatformInfo->GetIconStyleName(PlatformInfo::EPlatformIconSize::Normal)),
+						true
+					);
+				}
+
+				// gather any unknown status devices to query at the end
+				for (const TSharedPtr<ITargetDeviceProxy> Proxy : DeviceProxies)
+				{
+					FString DeviceId = Proxy->GetTargetDeviceId(NAME_None);
+					if (FDataDrivenPlatformInfoRegistry::GetPlatformInfo(EditorPlatformInfo->IniPlatformName).GetStatusForDeviceId(DeviceId) == DDPIPlatformSdkStatus::Unknown)
+					{
+						DeviceIdsToQuery.Add(DeviceId);
+					}
+				}
+			}
+		}
+	}
+	MenuBuilder.EndSection();
+
+	// tip section
+	MenuBuilder.BeginSection("LevelEditorLaunchHint");
+	{
+		MenuBuilder.AddWidget(
+			SNew(STextBlock)
+			.ColorAndOpacity(FSlateColor::UseSubduedForeground())
+			.Text(LOCTEXT("ZoomToFitHorizontal", "Launching a game on a different device will change your default 'Launch' device in the toolbar"))
+			.WrapTextAt(250),
+			FText::GetEmpty()
+		);
+	}
+	MenuBuilder.EndSection();
+
+	MenuBuilder.BeginSection("AllPlatforms", LOCTEXT("TurnkeyMenu_ManagePlatforms", "Sdk/Device Management"));
+	for (const auto Pair : FDataDrivenPlatformInfoRegistry::GetAllPlatformInfos())
+	{
+		if (Pair.Value.bIsFakePlatform)
+		{
+			continue;
+		}
+
+		FName PlatformName(*Pair.Key);
+		const FDataDrivenPlatformInfoRegistry::FPlatformInfo& Info = Pair.Value;
+		const PlatformInfo::FPlatformInfo* EditorPlatformInfo = PlatformInfo::FindVanillaPlatformInfo(PlatformName);
+
+		FString DeviceNameToCheckForStatus;
+		MenuBuilder.AddSubMenu(
+			MakeSdkStatusAttribute(EditorPlatformInfo, nullptr),
+			FText::FromString(PlatformName.ToString()),
+			FNewMenuDelegate::CreateStatic(&MakeTurnkeyPlatformMenu, EditorPlatformInfo, TargetDeviceServicesModule),
+			false, 
+			FSlateIcon(FEditorStyle::GetStyleSetName(), EditorPlatformInfo->GetIconStyleName(PlatformInfo::EPlatformIconSize::Normal)),
+			true
+		);
+	}
+
+	MenuBuilder.EndSection();
+
+
+	MenuBuilder.BeginSection("CookerSettings");
+
+	UCookerSettings* CookerSettings = GetMutableDefault<UCookerSettings>();
+
+	{
+		FUIAction UIAction;
+		UIAction.ExecuteAction = FExecuteAction::CreateLambda([CookerSettings]
+		{
+			CookerSettings->bCookOnTheFlyForLaunchOn = !CookerSettings->bCookOnTheFlyForLaunchOn;
+			CookerSettings->Modify(true);
+
+			// Update source control
+
+			FString ConfigPath = FPaths::ConvertRelativePathToFull(CookerSettings->GetDefaultConfigFilename());
+
+			if (FPlatformFileManager::Get().GetPlatformFile().FileExists(*ConfigPath))
+			{
+				if (ISourceControlModule::Get().IsEnabled())
+				{
+					FText ErrorMessage;
+
+					if (!SourceControlHelpers::CheckoutOrMarkForAdd(ConfigPath, FText::FromString(ConfigPath), NULL, ErrorMessage))
+					{
+						FNotificationInfo Info(ErrorMessage);
+						Info.ExpireDuration = 3.0f;
+						FSlateNotificationManager::Get().AddNotification(Info);
+					}
+				}
+				else
+				{
+					if (!FPlatformFileManager::Get().GetPlatformFile().SetReadOnly(*ConfigPath, false))
+					{
+						FNotificationInfo Info(FText::Format(LOCTEXT("FailedToMakeWritable", "Could not make {0} writable."), FText::FromString(ConfigPath)));
+						Info.ExpireDuration = 3.0f;
+						FSlateNotificationManager::Get().AddNotification(Info);
+					}
+				}
+			}
+
+			// Save settings
+			CookerSettings->UpdateSinglePropertyInConfigFile(CookerSettings->GetClass()->FindPropertyByName(GET_MEMBER_NAME_CHECKED(UCookerSettings, bCookOnTheFlyForLaunchOn)), CookerSettings->GetDefaultConfigFilename());
+		});
+
+		UIAction.GetActionCheckState = FGetActionCheckState::CreateLambda([CookerSettings]
+		{
+			return CookerSettings->bCookOnTheFlyForLaunchOn ? ECheckBoxState::Checked : ECheckBoxState::Unchecked;
+		});
+
+		MenuBuilder.AddMenuEntry(
+			LOCTEXT("CookOnTheFlyOnLaunch", "Enable cooking on the fly"),
+			LOCTEXT("CookOnTheFlyOnLaunchDescription", "Cook on the fly instead of cooking upfront when launching"),
+			FSlateIcon(),
+			UIAction,
+			NAME_None,
+			EUserInterfaceActionType::ToggleButton
+		);
+	}
+
+	MenuBuilder.EndSection();
+
+	// options section
+	MenuBuilder.BeginSection("LevelEditorLaunchOptions");
+	{
+		MenuBuilder.AddMenuEntry(FPlayWorldCommands::Get().OpenProjectLauncher,
+			NAME_None,
+			TAttribute<FText>(),
+			TAttribute<FText>(),
+			FSlateIcon(FEditorStyle::GetStyleSetName(), "Launcher.TabIcon")
+		);
+
+		MenuBuilder.AddMenuEntry(FPlayWorldCommands::Get().OpenDeviceManager,
+			NAME_None,
+			TAttribute<FText>(),
+			TAttribute<FText>(),
+			FSlateIcon(FEditorStyle::GetStyleSetName(), "DeviceDetails.TabIcon")
+		);
+
+		ProjectTargetPlatformEditorModule.AddOpenProjectTargetPlatformEditorMenuItem(MenuBuilder);
+	}
+	MenuBuilder.EndSection();
+
+	// now kick-off any devices that need to be updated
+	if (DeviceIdsToQuery.Num() > 0)
+	{
+		FDataDrivenPlatformInfoRegistry::UpdateDeviceSdkStatus(DeviceIdsToQuery);
+	}
+
+	return MenuBuilder.MakeWidget();
+}
+
+
+
+
 TSharedRef< SWidget > FPlayWorldCommands::GenerateLaunchMenuContent(TSharedRef<FUICommandList> InCommandList)
 {
 	const bool bShouldCloseWindowAfterMenuSelection = true;
@@ -1121,7 +1625,7 @@ TSharedRef< SWidget > FPlayWorldCommands::GenerateLaunchMenuContent(TSharedRef<F
 
 						// ... create an action...
 						FUIAction LaunchDeviceAction(
-							FExecuteAction::CreateStatic(&FInternalPlayWorldCommandCallbacks::HandleLaunchOnDeviceActionExecute, DeviceProxy->GetTargetDeviceId(NAME_None), DeviceProxy->GetName()),
+							FExecuteAction::CreateStatic(&FInternalPlayWorldCommandCallbacks::HandleLaunchOnDeviceActionExecute, DeviceProxy->GetTargetDeviceId(NAME_None), DeviceProxy->GetName(), false),
 							FCanExecuteAction::CreateStatic(&FInternalPlayWorldCommandCallbacks::HandleLaunchOnDeviceActionCanExecute, DeviceProxy->GetName()),
 							FIsActionChecked::CreateStatic(&FInternalPlayWorldCommandCallbacks::HandleLaunchOnDeviceActionIsChecked, DeviceProxy->GetName())
 						);
@@ -1310,7 +1814,6 @@ TSharedRef< SWidget > FPlayWorldCommands::GenerateLaunchMenuContent(TSharedRef<F
 				LabelArguments.Add(TEXT("PlatformName"), Platform.DisplayName);
 				FText Label = FText::Format(LOCTEXT("LaunchPlatformLabel", "{PlatformName} Support"), LabelArguments);
 
-
 				MenuBuilder.AddMenuEntry(
 					Label,
 					LOCTEXT("PlatformSDK", "Show information on setting up the platform tools"),
@@ -2059,7 +2562,7 @@ void FInternalPlayWorldCommandCallbacks::RepeatLastLaunch_Clicked()
 	case LaunchMode_OnDevice:
 		if (IsReadyToLaunchOnDevice(PlaySettings->LastExecutedLaunchDevice))
 		{
-			LaunchOnDevice(PlaySettings->LastExecutedLaunchDevice, PlaySettings->LastExecutedLaunchName);
+			LaunchOnDevice(PlaySettings->LastExecutedLaunchDevice, PlaySettings->LastExecutedLaunchName, false);
 		}
 		break;
 
@@ -2266,7 +2769,7 @@ bool FInternalPlayWorldCommandCallbacks::IsReadyToLaunchOnDevice(FString DeviceI
 	return true;
 }
 
-void FInternalPlayWorldCommandCallbacks::HandleLaunchOnDeviceActionExecute(FString DeviceId, FString DeviceName)
+void FInternalPlayWorldCommandCallbacks::HandleLaunchOnDeviceActionExecute(FString DeviceId, FString DeviceName, bool bUseTurnkey)
 {
 	if (IsReadyToLaunchOnDevice(DeviceId))
 	{
@@ -2280,7 +2783,8 @@ void FInternalPlayWorldCommandCallbacks::HandleLaunchOnDeviceActionExecute(FStri
 
 		PlaySettings->SaveConfig();
 
-		LaunchOnDevice(DeviceId, DeviceName);
+
+		LaunchOnDevice(DeviceId, DeviceName, bUseTurnkey);
 	}
 
 }
@@ -2596,7 +3100,7 @@ bool FInternalPlayWorldCommandCallbacks::CanLaunchOnDevice(const FString& Device
 }
 
 
-void FInternalPlayWorldCommandCallbacks::LaunchOnDevice(const FString& DeviceId, const FString& DeviceName)
+void FInternalPlayWorldCommandCallbacks::LaunchOnDevice(const FString& DeviceId, const FString& DeviceName, bool bUseTurnkey)
 {
 	FTargetDeviceId TargetDeviceId;
 	if (FTargetDeviceId::Parse(DeviceId, TargetDeviceId))
@@ -2620,6 +3124,7 @@ void FInternalPlayWorldCommandCallbacks::LaunchOnDevice(const FString& DeviceId,
 			FRequestPlaySessionParams::FLauncherDeviceInfo DeviceInfo;
 			DeviceInfo.DeviceId = DeviceId;
 			DeviceInfo.DeviceName = DeviceName;
+			DeviceInfo.bUpdateDeviceFlash = bUseTurnkey;
 
 			FRequestPlaySessionParams SessionParams;
 			SessionParams.SessionDestination = EPlaySessionDestinationType::Launcher;
diff --git a/Engine/Source/Editor/UnrealEd/Private/PlayLevelViaLauncher.cpp b/Engine/Source/Editor/UnrealEd/Private/PlayLevelViaLauncher.cpp
--- a/Engine/Source/Editor/UnrealEd/Private/PlayLevelViaLauncher.cpp
+++ b/Engine/Source/Editor/UnrealEd/Private/PlayLevelViaLauncher.cpp
@@ -217,6 +217,7 @@ void UEditorEngine::StartPlayUsingLauncherSession(FRequestPlaySessionParams& InR
 	LauncherProfile->SetDeployedDeviceGroup(DeviceGroup);
 	LauncherProfile->SetIncrementalDeploying(bIncrimentalCooking);
 	LauncherProfile->SetEditorExe(FUnrealEdMisc::Get().GetExecutableForCommandlets());
+	LauncherProfile->SetShouldUpdateDeviceFlash(InRequestParams.LauncherTargetDevice->bUpdateDeviceFlash);
 
 	const FString DummyIOSDeviceName(FString::Printf(TEXT("All_iOS_On_%s"), FPlatformProcess::ComputerName()));
 	const FString DummyTVOSDeviceName(FString::Printf(TEXT("All_tvOS_On_%s"), FPlatformProcess::ComputerName()));
diff --git a/Engine/Source/Editor/UnrealEd/Public/Kismet2/DebuggerCommands.h b/Engine/Source/Editor/UnrealEd/Public/Kismet2/DebuggerCommands.h
--- a/Engine/Source/Editor/UnrealEd/Public/Kismet2/DebuggerCommands.h
+++ b/Engine/Source/Editor/UnrealEd/Public/Kismet2/DebuggerCommands.h
@@ -132,6 +132,13 @@ protected:
 	 */
 	static TSharedRef< SWidget > GenerateLaunchMenuContent( TSharedRef<FUICommandList> InCommandList );
 
+	/**
+	 * Generates menu content for the Play On combo button drop down menu
+	 *
+	 * @return	Menu content widget
+	 */
+	static TSharedRef< SWidget > GenerateTurnkeyMenuContent(TSharedRef<FUICommandList> InCommandList);
+
 	// Add mobile PIE preview device commands
 	void AddPIEPreviewDeviceCommands();
 
diff --git a/Engine/Source/Editor/UnrealEd/Public/PlayInEditorDataTypes.h b/Engine/Source/Editor/UnrealEd/Public/PlayInEditorDataTypes.h
--- a/Engine/Source/Editor/UnrealEd/Public/PlayInEditorDataTypes.h
+++ b/Engine/Source/Editor/UnrealEd/Public/PlayInEditorDataTypes.h
@@ -164,6 +164,8 @@ public:
 		FString DeviceId;
 		/** The name of the Device selected in the Launch drop-down to launch on. */
 		FString DeviceName;
+		/** If True, a remote play session will attempt to update the flash/software on the target device if it's out of date */
+		bool bUpdateDeviceFlash = false;
 	};
 
 	/** Where should the session be launched? May be local or remote. */
diff --git a/Engine/Source/Editor/UnrealEd/UnrealEd.Build.cs b/Engine/Source/Editor/UnrealEd/UnrealEd.Build.cs
--- a/Engine/Source/Editor/UnrealEd/UnrealEd.Build.cs
+++ b/Engine/Source/Editor/UnrealEd/UnrealEd.Build.cs
@@ -184,6 +184,7 @@ public class UnrealEd : ModuleRules
 				"ScriptDisassembler",
 				"ToolMenus",
 				"FreeImage",
+				"UATHelper",
 				"IoStoreUtilities",
             }
 		);
diff --git a/Engine/Source/Programs/AutomationTool/IOS/IOSPlatform.Automation.cs b/Engine/Source/Programs/AutomationTool/IOS/IOSPlatform.Automation.cs
--- a/Engine/Source/Programs/AutomationTool/IOS/IOSPlatform.Automation.cs
+++ b/Engine/Source/Programs/AutomationTool/IOS/IOSPlatform.Automation.cs
@@ -175,6 +175,11 @@ public class IOSPlatform : Platform
 	}
 	public override string GetInstalledSdk()
 	{
+		if (HostPlatform.Current.HostEditorPlatform != UnrealTargetPlatform.Mac)
+		{
+			return base.GetInstalledSdk();
+		}
+
 		// @todo turnkey: get installed Xcode version
 		string Output = UnrealBuildTool.Utils.RunLocalProcessAndReturnStdOut("sh", "-c 'xcodebuild -version'");
 		Match Result = Regex.Match(Output, @"Xcode (\S*)");
@@ -191,6 +196,11 @@ public class IOSPlatform : Platform
 	}
 	public override bool CustomVersionUpdate(string CustomVersionId, string UpdateParams, FileRetriever Retriever)
 	{
+		if (HostPlatform.Current.HostEditorPlatform != UnrealTargetPlatform.Mac)
+		{
+			return base.CustomVersionUpdate(CustomVersionId, UpdateParams, Retriever);
+		}
+
 		if (CustomVersionId == "QueryIOSDeviceType")
 		{
 			Console.WriteLine("Getting DeviceType device with " + UpdateParams);
@@ -332,6 +342,11 @@ public class IOSPlatform : Platform
 
 	public override bool UpdateHostPrerequisites(BuildCommand Command, FileRetriever Retriever)
 	{
+		if (HostPlatform.Current.HostEditorPlatform != UnrealTargetPlatform.Mac)
+		{
+			return base.UpdateHostPrerequisites(Command, Retriever);
+		}
+
 		// make sure the Configurator is installed
 		string ConfiguratorLocation = GetConfiguratorLocation();
 
@@ -353,6 +368,11 @@ public class IOSPlatform : Platform
 
 	public override bool UpdateDevicePrerequisites(DeviceInfo Device, BuildCommand Command, FileRetriever Retriever)
 	{
+		if (HostPlatform.Current.HostEditorPlatform != UnrealTargetPlatform.Mac)
+		{
+			return base.UpdateDevicePrerequisites(Device, Command, Retriever);
+		}
+
 		VerifyIOSSettings Settings = new VerifyIOSSettings(Command);
 
 		// look if we have a cert that matches it
@@ -381,6 +401,11 @@ public class IOSPlatform : Platform
 
 	public override DeviceInfo[] GetDevices()
 	{
+		if (HostPlatform.Current.HostEditorPlatform != UnrealTargetPlatform.Mac)
+		{
+			return base.GetDevices();
+		}
+
 		string DeviceType = Platform == UnrealTargetPlatform.TVOS ? "tvOS" : "iOS";
 
 		// print out each connected device's udid, os version, and name (which might have spaces)
diff --git a/Engine/Source/Programs/AutomationTool/Turnkey/Commands/VerifySdk.cs b/Engine/Source/Programs/AutomationTool/Turnkey/Commands/VerifySdk.cs
--- a/Engine/Source/Programs/AutomationTool/Turnkey/Commands/VerifySdk.cs
+++ b/Engine/Source/Programs/AutomationTool/Turnkey/Commands/VerifySdk.cs
@@ -14,10 +14,11 @@ namespace Turnkey.Commands
 	{
 		protected override void Execute(string[] CommandOptions)
 		{
-			bool bUpdateIfNeeded = TurnkeyUtils.ParseParam("UpdateIfNeeded", CommandOptions);
 			bool bUnattended = TurnkeyUtils.ParseParam("Unattended", CommandOptions);
 			bool bPreferFullSdk = TurnkeyUtils.ParseParam("PreferFull", CommandOptions);
-			bool bForceInstall = TurnkeyUtils.ParseParam("ForceInstall", CommandOptions);
+			bool bForceSdkInstall = TurnkeyUtils.ParseParam("ForceSdkInstall", CommandOptions);
+			bool bForceDeviceInstall = TurnkeyUtils.ParseParam("ForceDeviceInstall", CommandOptions);
+			bool bUpdateIfNeeded = bForceSdkInstall || bForceDeviceInstall || TurnkeyUtils.ParseParam("UpdateIfNeeded", CommandOptions);
 
 			// track each platform to check, and 
 			Dictionary<UnrealTargetPlatform, List<string>> PlatformsAndDevices = null;
@@ -122,7 +123,7 @@ namespace Turnkey.Commands
 				bool bAutoSdkSetupSucceeded = false;
 
 				// install if out of date, or if forcing it
-				if (bForceInstall || (bUpdateIfNeeded && TurnkeyUtils.ExitCode != AutomationTool.ExitCode.Success))
+				if (bForceSdkInstall || (bUpdateIfNeeded && TurnkeyUtils.ExitCode != AutomationTool.ExitCode.Success))
 				{
 					TurnkeyUtils.ExitCode = AutomationTool.ExitCode.Success;
 
@@ -198,21 +199,24 @@ namespace Turnkey.Commands
 					{
 						AutomationPlatform.UpdateDevicePrerequisites(Device, TurnkeyUtils.CommandUtilHelper, Retriever);
 
-						bool bIsSoftwareValid = TurnkeyUtils.IsValueValid(Device.SoftwareVersion, AutomationPlatform.GetAllowedSoftwareVersions());
+						bool bIsSoftwareValid = TurnkeyUtils.IsValueValid(Device.SoftwareVersion, AutomationPlatform.GetAllowedSoftwareVersions(), AutomationPlatform);
 
 
-						if (bIsSoftwareValid)
+						if (!bForceDeviceInstall && bIsSoftwareValid)
 						{
 							TurnkeyUtils.Report("{0}@{1}: Valid: [{2}]", Platform, Device.Name, Device.SoftwareVersion);
 						}
 						else
 						{
-							TurnkeyUtils.Report("{0}@{1}: Invalid: [Has {2}, needs {3}]", Platform, Device.Name, Device.SoftwareVersion, AutomationPlatform.GetAllowedSoftwareVersions());
+							if (!bForceDeviceInstall)
+							{
+								TurnkeyUtils.Report("{0}@{1}: Invalid: [Has {2}, needs {3}]", Platform, Device.Name, Device.SoftwareVersion, AutomationPlatform.GetAllowedSoftwareVersions());
+							}
 							TurnkeyUtils.ExitCode = AutomationTool.ExitCode.Error_SDKNotFound;
 
 							if (bUpdateIfNeeded)
 							{
-								SdkInfo MatchingInstallableSdk = SdkInfo.FindMatchingSdk(AutomationPlatform, new SdkInfo.SdkType[] { SdkInfo.SdkType.Flash }, bSelectBest: bUnattended);
+								SdkInfo MatchingInstallableSdk = SdkInfo.FindMatchingSdk(AutomationPlatform, new SdkInfo.SdkType[] { SdkInfo.SdkType.Flash }, bSelectBest: bUnattended, Device.Type);
 
 								if (MatchingInstallableSdk == null)
 								{
diff --git a/Engine/Source/Programs/AutomationTool/Turnkey/Providers/GoogleDriveCopyProvider.cs b/Engine/Source/Programs/AutomationTool/Turnkey/Providers/GoogleDriveCopyProvider.cs
--- a/Engine/Source/Programs/AutomationTool/Turnkey/Providers/GoogleDriveCopyProvider.cs
+++ b/Engine/Source/Programs/AutomationTool/Turnkey/Providers/GoogleDriveCopyProvider.cs
@@ -471,8 +471,11 @@ namespace Turnkey
 					if (ExpansionSet != null)
 					{
 						NewExpansionSet = new List<string>(ExpansionSet);
-						// the match is the part of the filename that was not the *, so removing that will give us what we wanted to match
-						NewExpansionSet.Add(Result.Name.Replace(Match, ""));
+						if (Match != "")
+						{
+							// the match is the part of the filename that was not the *, so removing that will give us what we wanted to match
+							NewExpansionSet.Add(Result.Name.Replace(Match, ""));
+						}
 					}
 
 					// cache the file (it may have already been cached tho from a previous test)
diff --git a/Engine/Source/Programs/AutomationTool/Turnkey/Providers/HybridIOProvider.cs b/Engine/Source/Programs/AutomationTool/Turnkey/Providers/HybridIOProvider.cs
--- a/Engine/Source/Programs/AutomationTool/Turnkey/Providers/HybridIOProvider.cs
+++ b/Engine/Source/Programs/AutomationTool/Turnkey/Providers/HybridIOProvider.cs
@@ -33,6 +33,13 @@ namespace Turnkey
 
 		private  string ShowDialog(string Prompt, string Default, bool bIsList)
 		{
+			string Result = Microsoft.VisualBasic.Interaction.InputBox(Prompt, "Turnkey", Default);
+			if (Result == "" && bIsList)
+			{
+				return "0";
+			}
+			return Result;
+/*
 			// Create a new instance of the form.
 			Form Form1 = new Form();
 
@@ -136,6 +143,7 @@ namespace Turnkey
 			Form1.Dispose();
 
 			return Result;
+	*/
 		}
 
 		public override string ReadInput(string Prompt, string Default, bool bAppendNewLine)
@@ -176,6 +184,10 @@ namespace Turnkey
 				int Choice;
 				if (Int32.TryParse(ChoiceString, out Choice) == false || Choice < 0 || Choice >= Options.Count + (bIsCancellable ? 1 : 0))
 				{
+					if (Choice < 0 && bIsCancellable)
+					{
+						return 0;
+					}
 					TurnkeyUtils.Log("Invalid choice");
 				}
 				else
diff --git a/Engine/Source/Programs/AutomationTool/Turnkey/SdkInfo.cs b/Engine/Source/Programs/AutomationTool/Turnkey/SdkInfo.cs
--- a/Engine/Source/Programs/AutomationTool/Turnkey/SdkInfo.cs
+++ b/Engine/Source/Programs/AutomationTool/Turnkey/SdkInfo.cs
@@ -134,7 +134,7 @@ namespace Turnkey
 					foreach (DirectoryInfo Dir in PlatformDir.EnumerateDirectories())
 					{
 						// check for valid version number
-						if (TurnkeyUtils.IsValueValid(Dir.Name, AutomationPlatform.GetAllowedSdks()))
+						if (TurnkeyUtils.IsValueValid(Dir.Name, AutomationPlatform.GetAllowedSdks(), AutomationPlatform))
 						{
 							// make sure it actually has bits in it
 							if (File.Exists(Path.Combine(Dir.FullName, "setup.bat")))
@@ -154,7 +154,7 @@ namespace Turnkey
 			string InstalledVersion = AutomationPlatform.GetInstalledSdk();
 			if (!string.IsNullOrEmpty(InstalledVersion))
 			{
-				if (TurnkeyUtils.IsValueValid(InstalledVersion, AutomationPlatform.GetAllowedSdks()))
+				if (TurnkeyUtils.IsValueValid(InstalledVersion, AutomationPlatform.GetAllowedSdks(), AutomationPlatform))
 				{
 					Result |= LocalAvailability.InstalledSdk_ValidVersionExists;
 				}
@@ -168,14 +168,18 @@ namespace Turnkey
 			return Result;
 		}
 
-		static public SdkInfo FindMatchingSdk(AutomationTool.Platform Platform, SdkType[] TypePriority, bool bSelectBest)
+		static public SdkInfo FindMatchingSdk(AutomationTool.Platform Platform, SdkType[] TypePriority, bool bSelectBest, string DeviceType=null)
 		{
 			foreach (SdkType Type in TypePriority)
 			{
 				List<SdkInfo> Sdks = TurnkeyManifest.GetDiscoveredSdks();
 				Sdks = Sdks.FindAll(x => x.SupportsPlatform(Platform.IniPlatformType));
 				Sdks = Sdks.FindAll(x => x.Type == Type);
-				Sdks = Sdks.FindAll(x => (x.Version == null && x.CustomSdkId != null) || TurnkeyUtils.IsValueValid(x.Version, Type == SdkType.Flash ? Platform.GetAllowedSoftwareVersions() : Platform.GetAllowedSdks()));
+				Sdks = Sdks.FindAll(x => (x.Version == null && x.CustomSdkId != null) || TurnkeyUtils.IsValueValid(x.Version, Type == SdkType.Flash ? Platform.GetAllowedSoftwareVersions() : Platform.GetAllowedSdks(), Platform));
+				if (DeviceType != null)
+				{
+					Sdks = Sdks.FindAll(x => TurnkeyUtils.IsValueValid(DeviceType, x.AllowedFlashDeviceTypes, Platform));
+				}
 
 				// if none were found try next type
 				if (Sdks.Count == 0)
@@ -274,12 +278,12 @@ namespace Turnkey
 
 			if (Type == SdkType.Flash)
 			{
-				bool bIsValid = TurnkeyUtils.IsValueValid(Version, AutomationPlatforms[Platform].GetAllowedSoftwareVersions());
+				bool bIsValid = TurnkeyUtils.IsValueValid(Version, AutomationPlatforms[Platform].GetAllowedSoftwareVersions(), AutomationPlatforms[Platform]);
 				// if we were passed a device, also check if this Sdk is valid for that device
 				//				if (DeviceName != null)
 				{
 					DeviceInfo Device = GetDevice(Platform, DeviceName);
-					bIsValid = bIsValid && Device != null && TurnkeyUtils.IsValueValid(Device.Type, AllowedFlashDeviceTypes);
+					bIsValid = bIsValid && Device != null && TurnkeyUtils.IsValueValid(Device.Type, AllowedFlashDeviceTypes, AutomationPlatforms[Platform]);
 				}
 				return bIsValid;
 			}
@@ -290,7 +294,7 @@ namespace Turnkey
 			}
 			else
 			{
-				return TurnkeyUtils.IsValueValid(Version, AutomationPlatforms[Platform].GetAllowedSdks());
+				return TurnkeyUtils.IsValueValid(Version, AutomationPlatforms[Platform].GetAllowedSdks(), AutomationPlatforms[Platform]);
 			}
 		}
 
@@ -318,7 +322,7 @@ namespace Turnkey
 					return false;
 				}
 
-				return Device != null && !TurnkeyUtils.IsValueValid(Device.SoftwareVersion, AutomationPlatforms[Platform].GetAllowedSoftwareVersions());
+				return Device != null && !TurnkeyUtils.IsValueValid(Device.SoftwareVersion, AutomationPlatforms[Platform].GetAllowedSoftwareVersions(), AutomationPlatforms[Platform]);
 			}
 			// handle custom types, which can't use simple version checking
 			else if (CustomSdkId != null)
@@ -346,13 +350,13 @@ namespace Turnkey
 			}
 			else
 			{
-				return !TurnkeyUtils.IsValueValid(AutomationPlatforms[Platform].GetInstalledSdk(), AutomationPlatforms[Platform].GetAllowedSdks());
+				return !TurnkeyUtils.IsValueValid(AutomationPlatforms[Platform].GetInstalledSdk(), AutomationPlatforms[Platform].GetAllowedSdks(), AutomationPlatforms[Platform]);
 			}
 		}
 
-		public bool Install(UnrealTargetPlatform Platform, DeviceInfo Device = null, bool bUnattended = false)
+		public bool Install(UnrealTargetPlatform Platform, DeviceInfo Device=null, bool bUnattended=false, bool bSkipAutoSdkSetup=false)
 		{
-			if (Type == SdkType.AutoSdk)
+			if (Type == SdkType.AutoSdk && !bSkipAutoSdkSetup)
 			{
 				// AutoSdk has some extra setup needed
 				if (SdkInfo.ConditionalSetupAutoSdk(Platform, bUnattended))
@@ -481,7 +485,7 @@ namespace Turnkey
 
 					// now download it (AutoSdks don't "install") on download
 					// @todo turnkey: handle errors, handle p4 going to wrong location, handle one Sdk for multiple platforms
-					MatchingAutoSdk.Install(Platform, null, bUnattended);
+					MatchingAutoSdk.Install(Platform, null, bUnattended, bSkipAutoSdkSetup:true);
 
 					if (bSetupEnvVarAfterInstall)
 					{
diff --git a/Engine/Source/Programs/AutomationTool/Turnkey/Turnkey.Automation.csproj b/Engine/Source/Programs/AutomationTool/Turnkey/Turnkey.Automation.csproj
--- a/Engine/Source/Programs/AutomationTool/Turnkey/Turnkey.Automation.csproj
+++ b/Engine/Source/Programs/AutomationTool/Turnkey/Turnkey.Automation.csproj
@@ -98,6 +98,7 @@
       <HintPath>..\..\..\..\Binaries\DotNET\Ionic.Zip.Reduced.dll</HintPath>
       <Private>False</Private>
     </Reference>
+    <Reference Include="Microsoft.VisualBasic" />
     <Reference Include="System" />
     <Reference Include="System.Drawing" />
     <Reference Include="System.Windows.Forms" />
diff --git a/Engine/Source/Programs/AutomationTool/Turnkey/TurnkeyManifest.cs b/Engine/Source/Programs/AutomationTool/Turnkey/TurnkeyManifest.cs
--- a/Engine/Source/Programs/AutomationTool/Turnkey/TurnkeyManifest.cs
+++ b/Engine/Source/Programs/AutomationTool/Turnkey/TurnkeyManifest.cs
@@ -70,7 +70,8 @@ namespace Turnkey
 
 						if (Expansions.Count != ExpandedInstallerResults.Length)
 						{
-							throw new AutomationException("Bad expansions output from CopyProvider");
+							throw new AutomationException(string.Format("Bad expansions output from CopyProvider ({0} returned {1} count, expected {2}, from {3}",
+								Sdk.Expansion[0].Copy, Expansions.Count, ExpandedInstallerResults.Length, string.Join(", ", ExpandedInstallerResults)));
 						}
 
 						// @todo turnkey: this will be uysed in Builds also, make it a function with a lambda
diff --git a/Engine/Source/Runtime/Core/Private/Misc/DataDrivenPlatformInfoRegistry.cpp b/Engine/Source/Runtime/Core/Private/Misc/DataDrivenPlatformInfoRegistry.cpp
--- a/Engine/Source/Runtime/Core/Private/Misc/DataDrivenPlatformInfoRegistry.cpp
+++ b/Engine/Source/Runtime/Core/Private/Misc/DataDrivenPlatformInfoRegistry.cpp
@@ -7,6 +7,11 @@
 #include "Misc/ConfigCacheIni.h"
 #include "Modules/ModuleManager.h"
 
+#include "Async/Async.h"
+#include "Misc/MonitoredProcess.h"
+
+
+TMap<FString, FDataDrivenPlatformInfoRegistry::FPlatformInfo> FDataDrivenPlatformInfoRegistry::DataDrivenPlatforms;
 
 
 static const TArray<FString>& GetDataDrivenIniFilenames()
@@ -157,6 +162,7 @@ static void DDPIGetStringArray(const FConfigFile& IniFile, const TCHAR* Key, TAr
 static void LoadDDPIIniSettings(const FConfigFile& IniFile, FDataDrivenPlatformInfoRegistry::FPlatformInfo& Info)
 {
 	DDPIGetBool(IniFile, TEXT("bIsConfidential"), Info.bIsConfidential);
+	DDPIGetBool(IniFile, TEXT("bIsFakePlatform"), Info.bIsFakePlatform);
 	DDPIGetString(IniFile, TEXT("AudioCompressionSettingsIniSectionName"), Info.AudioCompressionSettingsIniSectionName);
 	DDPIGetStringArray(IniFile, TEXT("AdditionalRestrictedFolders"), Info.AdditionalRestrictedFolders);
 	
@@ -201,7 +207,6 @@ static void LoadDDPIIniSettings(const FConfigFile& IniFile, FDataDrivenPlatformI
 const TMap<FString, FDataDrivenPlatformInfoRegistry::FPlatformInfo>& FDataDrivenPlatformInfoRegistry::GetAllPlatformInfos()
 {
 	static bool bHasSearchedForPlatforms = false;
-	static TMap<FString, FDataDrivenPlatformInfoRegistry::FPlatformInfo> DataDrivenPlatforms;
 
 	// look on disk for special files
 	if (bHasSearchedForPlatforms == false)
@@ -242,11 +247,269 @@ const TMap<FString, FDataDrivenPlatformInfoRegistry::FPlatformInfo>& FDataDriven
 				It.Value.IniParentChain.Insert(CurrentPlatform, 0);
 			}
 		}
+
+
+#if DDPI_HAS_EXTENDED_PLATFORMINFO_DATA
+		UpdateSdkStatus();
+#endif
 	}
 
 	return DataDrivenPlatforms;
 }
 
+#if DDPI_HAS_EXTENDED_PLATFORMINFO_DATA
+
+// some shared functionality
+static void PrepForTurnkeyReport(FString& UAT, FString& ReportFilename)
+{
+	static int ReportIndex = 0;
+
+	// get path to AutomationTool - skipping RunUAT for speed
+//		FString AutomationTool = FPaths::Combine(FPaths::EngineDir(), TEXT("Binaries/DotNET/AutomationTool.exe"));
+	UAT = FPaths::ConvertRelativePathToFull(FPaths::Combine(FPaths::EngineDir(), TEXT("Build/BatchFiles/RunUAT.bat")));
+
+	ReportFilename = FPaths::ConvertRelativePathToFull(FPaths::Combine(FPaths::ProjectIntermediateDir(), *FString::Printf(TEXT("TurnkeyReport_%d.log"), ReportIndex++)));
+}
+
+static FString ConvertToDDPIPlatform(const FString& Platform)
+{
+	FString  New = Platform.Replace(TEXT("NoEditor"), TEXT("")).Replace(TEXT("Client"), TEXT("")).Replace(TEXT("Server"), TEXT(""));
+	if (New == TEXT("Win64"))
+	{
+		New = TEXT("Windows");
+	}
+	return New;
+}
+
+static FString ConvertToUATPlatform(const FString& Platform)
+{
+	FString New = ConvertToDDPIPlatform(Platform);
+	if (New == TEXT("Windows"))
+	{
+		New = TEXT("Win64");
+	}
+	return New;
+}
+
+static FString ConvertToUATDeviceId(const FString& DeviceId)
+{
+	TArray<FString> PlatformAndDevice;
+	DeviceId.ParseIntoArray(PlatformAndDevice, TEXT("@"), true);
+
+	return FString::Printf(TEXT("%s@%s"), *ConvertToUATPlatform(PlatformAndDevice[0]), *PlatformAndDevice[1]);
+}
+
+static FString ConvertToDDPIDeviceId(const FString& DeviceId)
+{
+	TArray<FString> PlatformAndDevice;
+	DeviceId.ParseIntoArray(PlatformAndDevice, TEXT("@"), true);
+
+	return FString::Printf(TEXT("%s@%s"), *ConvertToDDPIPlatform(PlatformAndDevice[0]), *PlatformAndDevice[1]);
+}
+
+void FDataDrivenPlatformInfoRegistry::UpdateSdkStatus()
+{
+	FString AutomationTool, ReportFilename;
+	PrepForTurnkeyReport(AutomationTool, ReportFilename);
+
+	// run Turnkey to get all of the platform statuses
+	FString Commandline = FString::Printf(TEXT("/c %s Turnkey -command=VerifySdk -platform=%s -ReportFilename=\"%s\""),
+		*AutomationTool,
+		*FString::JoinBy(DataDrivenPlatforms, TEXT("+"), [](TPair<FString, FDataDrivenPlatformInfoRegistry::FPlatformInfo> Pair) { return ConvertToUATPlatform(Pair.Key); }),
+		*ReportFilename);
+
+
+	// reset status to unknown
+	for (auto& It : DataDrivenPlatforms)
+	{
+		It.Value.SdkStatus = DDPIPlatformSdkStatus::Querying;
+		
+		// reset the per-device status when querying general Sdk status
+		It.Value.ClearDeviceStatus();
+	}
+
+	FMonitoredProcess* TurnkeyProcess = new FMonitoredProcess(TEXT("cmd.exe"), Commandline, true, false);
+	TurnkeyProcess->OnCompleted().BindLambda([ReportFilename, TurnkeyProcess](int32 ExitCode)
+	{
+		AsyncTask(ENamedThreads::GameThread, [ReportFilename, TurnkeyProcess, ExitCode]()
+		{
+			if (ExitCode == 0 || ExitCode == 10)
+			{
+				TArray<FString> Contents;
+				if (FFileHelper::LoadFileToStringArray(Contents, *ReportFilename))
+				{
+					for (FString& Line : Contents)
+					{
+						TArray<FString> Tokens;
+						Line.ParseIntoArray(Tokens, TEXT(": "), true);
+						// Tokens [0] is the platform, [1] is the status, [2] is information
+
+						DDPIPlatformSdkStatus Status = DDPIPlatformSdkStatus::Unknown;
+						if (Tokens[1] == TEXT("Valid"))
+						{
+							Status = DDPIPlatformSdkStatus::Valid;
+						}
+						else
+						{
+							if (Tokens[2].Contains(TEXT("AutoSdk_InvalidVersionExists")) || Tokens[2].Contains(TEXT("InstalledSdk_InvalidVersionExists")))
+							{
+								Status = DDPIPlatformSdkStatus::OutOfDate;
+							}
+							else
+							{
+								Status = DDPIPlatformSdkStatus::NoSdk;
+							}
+						}
+
+						// have to convert back to WIndows from Win64
+						FString PlatformName = ConvertToDDPIPlatform(Tokens[0]);
+						DataDrivenPlatforms[PlatformName].SdkStatus = Status;
+
+						UE_LOG(LogTemp, Log, TEXT("Turnkey Platform: %s - %d - %s"), *PlatformName, (int)Status, *Tokens[2]);
+					}
+				}
+			}
+			else
+			{
+				for (auto& It : DataDrivenPlatforms)
+				{
+					It.Value.SdkStatus = DDPIPlatformSdkStatus::Error;
+					// @todo turnkey error description!
+				}
+			}
+
+
+			for (auto& It : DataDrivenPlatforms)
+			{
+				if (It.Value.SdkStatus == DDPIPlatformSdkStatus::Querying)
+				{
+					if (It.Value.bIsFakePlatform)
+					{
+						It.Value.SdkStatus = DDPIPlatformSdkStatus::Unknown;
+					}
+					else
+					{
+						It.Value.SdkStatus = DDPIPlatformSdkStatus::Error;
+						It.Value.SdkErrorInformation = "The platform's Sdk status was not returned from Turnkey";
+						//					It.Value.SdkErrorInformation = NSLOCTEXT("Turnkey", "PlatformNotReturned", "The platform's Sdk status was not returned from Turnkey");
+					}
+				}
+			}
+
+			// cleanup
+			delete TurnkeyProcess;
+			IFileManager::Get().Delete(*ReportFilename);
+		});
+	});
+
+	// run it
+	TurnkeyProcess->Launch();
+}
+
+FDataDrivenPlatformInfoRegistry::FPlatformInfo& FDataDrivenPlatformInfoRegistry::DeviceIdToInfo(FString DeviceId, FString* OutDeviceName)
+{
+	TArray<FString> PlatformAndDevice;
+	DeviceId.ParseIntoArray(PlatformAndDevice, TEXT("@"), true);
+
+	if (OutDeviceName)
+	{
+		*OutDeviceName = PlatformAndDevice[1];
+	}
+
+	// have to convert back to Windows from Win64
+	return FDataDrivenPlatformInfoRegistry::DataDrivenPlatforms[ConvertToDDPIPlatform(PlatformAndDevice[0])];
+
+}
+
+void FDataDrivenPlatformInfoRegistry::UpdateDeviceSdkStatus(TArray<FString> PlatformDeviceIds)
+{
+	FString AutomationTool, ReportFilename;
+	PrepForTurnkeyReport(AutomationTool, ReportFilename);
+
+	// run Turnkey to get all of the platform statuses (turn the array into "Platform@Device+Platform@Device+...")
+	FString Commandline = FString::Printf(TEXT("/c %s Turnkey -command=VerifySdk -Device=\"%s\" -ReportFilename=\"%s\""),
+		*AutomationTool,
+		*FString::JoinBy(PlatformDeviceIds, TEXT("+"), [](const FString& DeviceId) { return ConvertToUATDeviceId(DeviceId); }),
+		*ReportFilename);
+
+	// set status to querying
+	for (const FString& Id : PlatformDeviceIds)
+	{
+		DeviceIdToInfo(Id).PerDeviceStatus.Add(ConvertToDDPIDeviceId(Id), DDPIPlatformSdkStatus::Querying);
+	}
+
+	FMonitoredProcess* TurnkeyProcess = new FMonitoredProcess(TEXT("cmd.exe"), Commandline, true, false);
+	TurnkeyProcess->OnCompleted().BindLambda([ReportFilename, TurnkeyProcess, PlatformDeviceIds](int32 ExitCode)
+	{
+		AsyncTask(ENamedThreads::GameThread, [ReportFilename, TurnkeyProcess, PlatformDeviceIds, ExitCode]()
+		{
+			if (ExitCode == 0 || ExitCode == 10)
+			{
+				TArray<FString> Contents;
+				if (FFileHelper::LoadFileToStringArray(Contents, *ReportFilename))
+				{
+					for (FString& Line : Contents)
+					{
+						TArray<FString> Tokens;
+						Line.ParseIntoArray(Tokens, TEXT(": "), true);
+						// Tokens [0] is platform:device, [1] is the status, [2] is information
+
+						// token[0] without an @ is just platform status, which we don't care about now
+						if (!Tokens[0].Contains(TEXT("@")))
+						{
+							continue;
+						}
+
+						TArray<FString> PlatformAndDevice;
+						Tokens[0].ParseIntoArray(PlatformAndDevice, TEXT("@"), true);
+						
+						DDPIPlatformSdkStatus Status = Tokens[1] == TEXT("Valid") ? DDPIPlatformSdkStatus::FlashValid : DDPIPlatformSdkStatus::FlashOutOfDate;
+
+						DeviceIdToInfo(Tokens[0]).PerDeviceStatus[ConvertToDDPIDeviceId(Tokens[0])] = Status;
+
+						UE_LOG(LogTemp, Log, TEXT("Turnkey Device: %s - %d - %s"), *Tokens[0], (int)Status, *Tokens[2]);
+					}
+				}
+			}
+
+			for (const FString& Id : PlatformDeviceIds)
+			{
+				FDataDrivenPlatformInfoRegistry::FPlatformInfo& Info = DeviceIdToInfo(Id);
+				
+				DDPIPlatformSdkStatus& Status = Info.PerDeviceStatus[ConvertToDDPIDeviceId(Id)];
+				if (Status == DDPIPlatformSdkStatus::Querying)
+				{
+					Status = DDPIPlatformSdkStatus::Error;
+					Info.SdkErrorInformation = "A device's Sdk status was not returned from Turnkey";
+				}
+			}
+
+			// cleanup
+			delete TurnkeyProcess;
+			IFileManager::Get().Delete(*ReportFilename);
+		});
+	});
+
+	// run it
+	TurnkeyProcess->Launch();
+}
+
+
+DDPIPlatformSdkStatus FDataDrivenPlatformInfoRegistry::FPlatformInfo::GetStatusForDeviceId(const FString& DeviceId) const
+{
+	// return the status, or Unknown if not known
+	return PerDeviceStatus.FindRef(ConvertToDDPIDeviceId(DeviceId));
+}
+
+void FDataDrivenPlatformInfoRegistry::FPlatformInfo::ClearDeviceStatus()
+{
+	PerDeviceStatus.Empty();
+}
+
+
+#endif
+
+
 const TArray<FString>& FDataDrivenPlatformInfoRegistry::GetValidPlatformDirectoryNames()
 {
 	static bool bHasSearchedForPlatforms = false;
diff --git a/Engine/Source/Runtime/Core/Public/Misc/DataDrivenPlatformInfoRegistry.h b/Engine/Source/Runtime/Core/Public/Misc/DataDrivenPlatformInfoRegistry.h
--- a/Engine/Source/Runtime/Core/Public/Misc/DataDrivenPlatformInfoRegistry.h
+++ b/Engine/Source/Runtime/Core/Public/Misc/DataDrivenPlatformInfoRegistry.h
@@ -9,6 +9,24 @@
 
 #define DDPI_HAS_EXTENDED_PLATFORMINFO_DATA WITH_EDITOR && !IS_MONOLITHIC
 
+#if DDPI_HAS_EXTENDED_PLATFORMINFO_DATA
+
+
+enum class DDPIPlatformSdkStatus : uint8
+{
+	Unknown,
+	Querying,
+	Valid,
+	OutOfDate,
+	NoSdk,
+	Error,
+	FlashValid,
+	FlashOutOfDate,
+};
+
+#endif
+
+
 struct CORE_API FDataDrivenPlatformInfoRegistry
 {
 	// Information about a platform loaded from disk
@@ -20,6 +38,9 @@ struct CORE_API FDataDrivenPlatformInfoRegistry
 		// is this platform confidential
 		bool bIsConfidential = false;
 
+		// some platforms are here just for IniParentChain needs and are not concrete platforms
+		bool bIsFakePlatform = false;
+
 		// the name of the ini section to use to load audio compression settings (used at runtime and cooktime)
 		FString AudioCompressionSettingsIniSectionName;
 
@@ -41,6 +62,20 @@ struct CORE_API FDataDrivenPlatformInfoRegistry
 		TArray<FString> AllTargetPlatformNames;
 		// list of UBT platform names (Win32, Win64, etc)
 		TArray<FString> AllUBTPlatformNames;
+
+		// Information about the validity of using a platform, discovered via Turnkey
+		DDPIPlatformSdkStatus SdkStatus;
+		FString SdkErrorInformation;
+
+		// Get the status of a device, or Unknown if not specified
+		CORE_API DDPIPlatformSdkStatus GetStatusForDeviceId(const FString& DeviceId) const;
+		CORE_API void ClearDeviceStatus();
+
+private:
+		friend FDataDrivenPlatformInfoRegistry;
+
+		// Information about the validity of each connected device (by string, discovered by Turnkey)
+		TMap<FString, DDPIPlatformSdkStatus> PerDeviceStatus;
 #endif
 	};
 
@@ -93,6 +128,24 @@ struct CORE_API FDataDrivenPlatformInfoRegistry
 		TargetPlatform,
 	};
 	static bool HasCompiledSupportForPlatform(const FString& PlatformName, EPlatformNameType PlatformNameType);
+
+	/**
+	 * Runs UAT to update Sdk status (can be called after user chooses an option that had a bad Sdk status)
+	 */
+	static void UpdateSdkStatus();
+
+	/**
+	 * Run UAT to update per device flash/software information (each Id should be in the standard DeviceId format of Platform@Device)
+	 */
+	static void UpdateDeviceSdkStatus(TArray<FString> PlatformDeviceIds);
+
+private:
+	static FDataDrivenPlatformInfoRegistry::FPlatformInfo& DeviceIdToInfo(FString DeviceId, FString* OutDeviceName = nullptr);
+
 #endif
+
+
+private:
+	static TMap<FString, FDataDrivenPlatformInfoRegistry::FPlatformInfo> DataDrivenPlatforms;
 };
 
