- Fixed non-unity error

diff --git a/Engine/Source/Editor/UnrealEd/Private/Kismet2/DebuggerCommands.cpp b/Engine/Source/Editor/UnrealEd/Private/Kismet2/DebuggerCommands.cpp
--- a/Engine/Source/Editor/UnrealEd/Private/Kismet2/DebuggerCommands.cpp
+++ b/Engine/Source/Editor/UnrealEd/Private/Kismet2/DebuggerCommands.cpp
@@ -78,6 +78,7 @@
 #include "IUATHelperModule.h"
 #include "ISettingsEditorModule.h"
 #include "Async/Async.h"
+#include "Misc/FileHelper.h"
 
 
 #define LOCTEXT_NAMESPACE "DebuggerCommands"
