Added editor scripting nodes for common actions like Selecting Actors, Deleting Actors and Inverting Selection. Added script support for Editor Delegates for cut/copy/paste/duplicate/delete as well.

diff --git a/Engine/Source/Editor/UnrealEd/Private/EditorActorUtilities.cpp b/Engine/Source/Editor/UnrealEd/Private/EditorActorUtilities.cpp
new file mode 100644
--- /dev/null
+++ b/Engine/Source/Editor/UnrealEd/Private/EditorActorUtilities.cpp
@@ -0,0 +1,105 @@
+// Copyright Epic Games, Inc. All Rights Reserved.
+
+#include "EditorActorUtilities.h"
+#include "CoreMinimal.h"
+#include "Engine/MeshMerging.h"
+#include "GameFramework/Actor.h"
+#include "Kismet/BlueprintFunctionLibrary.h"
+#include "UnrealEdGlobals.h"
+#include "ScopedTransaction.h"
+#include "Engine/Brush.h"
+#include "Editor.h"
+#include "Editor/UnrealEdEngine.h"
+#include "AssetSelection.h"
+
+#define LOCTEXT_NAMESPACE "EditorActorUtilities"
+
+void UEditorActorUtilities::DuplicateSelectedActors(UWorld* InWorld)
+{
+	if (!GEditor || !InWorld)
+	{
+		return;
+	}
+
+	bool bComponentsSelected = GEditor->GetSelectedComponentCount() > 0;
+	//@todo locked levels - if all actor levels are locked, cancel the transaction
+	const FScopedTransaction Transaction(bComponentsSelected ? NSLOCTEXT("UnrealEd", "DuplicateComponents", "Duplicate Components") : NSLOCTEXT("UnrealEd", "DuplicateActors", "Duplicate Actors"));
+
+	FEditorDelegates::OnDuplicateActorsBegin.Broadcast();
+
+	// duplicate selected
+	ABrush::SetSuppressBSPRegeneration(true);
+	GEditor->edactDuplicateSelected(InWorld->GetCurrentLevel(), GetDefault<ULevelEditorViewportSettings>()->GridEnabled);
+	ABrush::SetSuppressBSPRegeneration(false);
+
+	// Find out if any of the selected actors will change the BSP.
+	// and only then rebuild BSP as this is expensive.
+	const FSelectedActorInfo& SelectedActors = AssetSelectionUtils::GetSelectedActorInfo();
+	if (SelectedActors.bHaveBrush)
+	{
+		GEditor->RebuildAlteredBSP(); // Update the Bsp of any levels containing a modified brush
+	}
+
+	FEditorDelegates::OnDuplicateActorsEnd.Broadcast();
+
+	GEditor->RedrawLevelEditingViewports();
+}
+
+void UEditorActorUtilities::DeleteSelectedActors(UWorld* InWorld)
+{
+	if (!GEditor || !InWorld)
+	{
+		return;
+	}
+
+	bool bComponentsSelected = GEditor->GetSelectedComponentCount() > 0;
+
+	const FScopedTransaction Transaction(bComponentsSelected ? NSLOCTEXT("UnrealEd", "DeleteComponents", "Delete Components") : NSLOCTEXT("UnrealEd", "DeleteActors", "Delete Actors"));
+
+	FEditorDelegates::OnDeleteActorsBegin.Broadcast();
+	GEditor->edactDeleteSelected(InWorld);
+	FEditorDelegates::OnDeleteActorsEnd.Broadcast();
+}
+
+void UEditorActorUtilities::InvertSelection(UWorld* InWorld)
+{
+	if (!GUnrealEd || !InWorld)
+	{
+		return;
+	}
+
+	const FScopedTransaction Transaction(NSLOCTEXT("UnrealEd", "SelectInvert", "Select Invert"));
+	GUnrealEd->edactSelectInvert(InWorld);
+}
+
+void UEditorActorUtilities::SelectAll(UWorld* InWorld)
+{
+	if (!GUnrealEd || !InWorld)
+	{
+		return;
+	}
+
+	const FScopedTransaction Transaction(NSLOCTEXT("UnrealEd", "SelectAll", "Select All"));
+	GUnrealEd->edactSelectAll(InWorld);
+}
+
+void UEditorActorUtilities::SelectAllChildren(bool bRecurseChildren)
+{
+	if (!GUnrealEd)
+	{
+		return;
+	}
+
+	if (bRecurseChildren)
+	{
+		const FScopedTransaction Transaction(NSLOCTEXT("UnrealEd", "SelectAllDescendants", "Select All Descendants"));
+	}
+	else
+	{
+		const FScopedTransaction Transaction(NSLOCTEXT("UnrealEd", "SelectAllChildren", "Select All Children"));
+	}
+
+	GUnrealEd->edactSelectAllChildren(bRecurseChildren);
+}
+
+#undef LOCTEXT_NAMESPACE
diff --git a/Engine/Source/Editor/UnrealEd/Private/Subsystems/EditorDelegateSubsystem.cpp b/Engine/Source/Editor/UnrealEd/Private/Subsystems/EditorDelegateSubsystem.cpp
new file mode 100644
--- /dev/null
+++ b/Engine/Source/Editor/UnrealEd/Private/Subsystems/EditorDelegateSubsystem.cpp
@@ -0,0 +1,103 @@
+// Copyright Epic Games, Inc. All Rights Reserved.
+
+#include "Subsystems/EditorDelegateSubsystem.h"
+#include "Editor.h"
+
+UEditorDelegateSubsystem::UEditorDelegateSubsystem()
+	: UEditorSubsystem()
+{
+
+}
+
+void UEditorDelegateSubsystem::Initialize(FSubsystemCollectionBase& Collection)
+{
+	FEditorDelegates::OnEditCutActorsBegin.AddUObject(this, &UEditorDelegateSubsystem::BroadcastEditCutActorsBegin);
+	FEditorDelegates::OnEditCutActorsEnd.AddUObject(this, &UEditorDelegateSubsystem::BroadcastEditCutActorsEnd);
+	
+	FEditorDelegates::OnEditCopyActorsBegin.AddUObject(this, &UEditorDelegateSubsystem::BroadcastEditCopyActorsBegin);
+	FEditorDelegates::OnEditCopyActorsEnd.AddUObject(this, &UEditorDelegateSubsystem::BroadcastEditCopyActorsEnd);
+
+	FEditorDelegates::OnEditPasteActorsBegin.AddUObject(this, &UEditorDelegateSubsystem::BroadcastEditPasteActorsBegin);
+	FEditorDelegates::OnEditPasteActorsEnd.AddUObject(this, &UEditorDelegateSubsystem::BroadcastEditPasteActorsEnd);
+
+	FEditorDelegates::OnDuplicateActorsBegin.AddUObject(this, &UEditorDelegateSubsystem::BroadcastDuplicateActorsBegin);
+	FEditorDelegates::OnDuplicateActorsEnd.AddUObject(this, &UEditorDelegateSubsystem::BroadcastDuplicateActorsEnd);
+
+	FEditorDelegates::OnDeleteActorsBegin.AddUObject(this, &UEditorDelegateSubsystem::BroadcastDeleteActorsBegin);
+	FEditorDelegates::OnDeleteActorsEnd.AddUObject(this, &UEditorDelegateSubsystem::BroadcastDeleteActorsEnd);
+}
+
+void UEditorDelegateSubsystem::Deinitialize()
+{
+	FEditorDelegates::OnEditCutActorsBegin.RemoveAll(this);
+	FEditorDelegates::OnEditCutActorsEnd.RemoveAll(this);
+	FEditorDelegates::OnEditCopyActorsBegin.RemoveAll(this);
+	FEditorDelegates::OnEditCopyActorsEnd.RemoveAll(this);
+	FEditorDelegates::OnEditPasteActorsBegin.RemoveAll(this);
+	FEditorDelegates::OnEditPasteActorsEnd.RemoveAll(this);
+	FEditorDelegates::OnDuplicateActorsBegin.RemoveAll(this);
+	FEditorDelegates::OnDuplicateActorsEnd.RemoveAll(this);
+	FEditorDelegates::OnDeleteActorsBegin.RemoveAll(this);
+	FEditorDelegates::OnDeleteActorsEnd.RemoveAll(this);
+
+}
+
+/** To fire before an Actor is Cut */
+void UEditorDelegateSubsystem::BroadcastEditCutActorsBegin()
+{
+	OnEditCutActorsBegin.Broadcast();
+}
+
+/** To fire after an Actor is Cut */
+void UEditorDelegateSubsystem::BroadcastEditCutActorsEnd()
+{
+	OnEditCutActorsEnd.Broadcast();
+}
+
+/** To fire before an Actor is Copied */
+void UEditorDelegateSubsystem::BroadcastEditCopyActorsBegin()
+{
+	OnEditCopyActorsBegin.Broadcast();
+}
+
+/** To fire after an Actor is Copied */
+void UEditorDelegateSubsystem::BroadcastEditCopyActorsEnd()
+{
+	OnEditCopyActorsEnd.Broadcast();
+}
+
+/** To fire before an Actor is Pasted */
+void UEditorDelegateSubsystem::BroadcastEditPasteActorsBegin()
+{
+	OnEditPasteActorsBegin.Broadcast();
+}
+
+/** To fire after an Actor is Pasted */
+void UEditorDelegateSubsystem::BroadcastEditPasteActorsEnd()
+{
+	OnEditPasteActorsEnd.Broadcast();
+}
+
+/** To fire before an Actor is duplicated */
+void UEditorDelegateSubsystem::BroadcastDuplicateActorsBegin()
+{
+	OnDuplicateActorsBegin.Broadcast();
+}
+
+/** To fire after an Actor is duplicated */
+void UEditorDelegateSubsystem::BroadcastDuplicateActorsEnd()
+{
+	OnDuplicateActorsEnd.Broadcast();
+}
+
+/** To fire before an Actor is Deleted */
+void UEditorDelegateSubsystem::BroadcastDeleteActorsBegin()
+{
+	OnDeleteActorsBegin.Broadcast();
+}
+
+/** To fire after an Actor is Deleted */
+void UEditorDelegateSubsystem::BroadcastDeleteActorsEnd()
+{
+	OnDeleteActorsEnd.Broadcast();
+}
diff --git a/Engine/Source/Editor/UnrealEd/Private/UnrealEdSrv.cpp b/Engine/Source/Editor/UnrealEd/Private/UnrealEdSrv.cpp
--- a/Engine/Source/Editor/UnrealEd/Private/UnrealEdSrv.cpp
+++ b/Engine/Source/Editor/UnrealEd/Private/UnrealEdSrv.cpp
@@ -104,6 +104,7 @@
 #include "ActorGroupingUtils.h"
 #include "EdMode.h"
 #include "Subsystems/BrushEditingSubsystem.h"
+#include "EditorActorUtilities.h"
 
 DEFINE_LOG_CATEGORY_STATIC(LogUnrealEdSrv, Log, All);
 
@@ -2430,20 +2431,17 @@ bool UUnrealEdEngine::Exec_Actor( UWorld* InWorld, const TCHAR* Str, FOutputDevi
 			}
 			else if( FParse::Command(&Str, TEXT("CHILDREN")) ) // ACTOR SELECT ALL CHILDREN
 			{
-				const FScopedTransaction Transaction( NSLOCTEXT("UnrealEd", "SelectAllChildren", "Select All Children") );
-				edactSelectAllChildren( false );
+				UEditorActorUtilities::SelectAllChildren(false);
 				return true;
 			}
 			else if( FParse::Command(&Str, TEXT("DESCENDANTS")) ) // ACTOR SELECT ALL DESCENDANTS
 			{
-				const FScopedTransaction Transaction( NSLOCTEXT("UnrealEd", "SelectAllDescendants", "Select All Descendants") );
-				edactSelectAllChildren( true );
+				UEditorActorUtilities::SelectAllChildren(true);
 				return true;
 			}
 			else
 			{
-				const FScopedTransaction Transaction( NSLOCTEXT("UnrealEd", "SelectAll", "Select All") );
-				edactSelectAll( InWorld );
+				UEditorActorUtilities::SelectAll(InWorld);
 				return true;
 			}
 		}
@@ -2453,8 +2451,7 @@ bool UUnrealEdEngine::Exec_Actor( UWorld* InWorld, const TCHAR* Str, FOutputDevi
 		}
 		else if( FParse::Command(&Str,TEXT("INVERT") ) ) // ACTOR SELECT INVERT
 		{
-			const FScopedTransaction Transaction( NSLOCTEXT("UnrealEd", "SelectInvert", "Select Invert") );
-			edactSelectInvert( InWorld );
+			UEditorActorUtilities::InvertSelection(InWorld);
 			return true;
 		}
 		else if( FParse::Command(&Str,TEXT("OFCLASS")) ) // ACTOR SELECT OFCLASS CLASS=<class>
@@ -2554,10 +2551,7 @@ bool UUnrealEdEngine::Exec_Actor( UWorld* InWorld, const TCHAR* Str, FOutputDevi
 		// if not specially handled by the current editing mode,
 		if (!bHandled)
 		{
-			const FScopedTransaction Transaction( bComponentsSelected ? NSLOCTEXT("UnrealEd", "DeleteComponents", "Delete Components") : NSLOCTEXT("UnrealEd", "DeleteActors", "Delete Actors") );
-			FEditorDelegates::OnDeleteActorsBegin.Broadcast();
-			edactDeleteSelected( InWorld );
-			FEditorDelegates::OnDeleteActorsEnd.Broadcast();
+			UEditorActorUtilities::DeleteSelectedActors(InWorld);
 		}
 		return true;
 	}
@@ -2733,27 +2727,13 @@ bool UUnrealEdEngine::Exec_Actor( UWorld* InWorld, const TCHAR* Str, FOutputDevi
 		// if not specially handled by the current editing mode,
 		if (!bHandled)
 		{
-			//@todo locked levels - if all actor levels are locked, cancel the transaction
-			const FScopedTransaction Transaction( bComponentsSelected ? NSLOCTEXT("UnrealEd", "DuplicateComponents", "Duplicate Components") : NSLOCTEXT("UnrealEd", "DuplicateActors", "Duplicate Actors") );
-
-			FEditorDelegates::OnDuplicateActorsBegin.Broadcast();
-
-			// duplicate selected
-			ABrush::SetSuppressBSPRegeneration(true);
-			edactDuplicateSelected(InWorld->GetCurrentLevel(), GetDefault<ULevelEditorViewportSettings>()->GridEnabled);
-			ABrush::SetSuppressBSPRegeneration(false);
-
-			// Find out if any of the selected actors will change the BSP.
-			// and only then rebuild BSP as this is expensive.
-			const FSelectedActorInfo& SelectedActors = AssetSelectionUtils::GetSelectedActorInfo();
-			if( SelectedActors.bHaveBrush )
-			{
-				RebuildAlteredBSP(); // Update the Bsp of any levels containing a modified brush
-			}
-
-			FEditorDelegates::OnDuplicateActorsEnd.Broadcast();
+			UEditorActorUtilities::DuplicateSelectedActors( InWorld );
+		}
+		// DuplicateSelectedActors also calls RedrawLevelEditingViewports
+		else
+		{
+			RedrawLevelEditingViewports(); 
 		}
-		RedrawLevelEditingViewports();
 		return true;
 	}
 	else if( FParse::Command(&Str, TEXT("ALIGN")) )
diff --git a/Engine/Source/Editor/UnrealEd/Public/EditorActorUtilities.h b/Engine/Source/Editor/UnrealEd/Public/EditorActorUtilities.h
new file mode 100644
--- /dev/null
+++ b/Engine/Source/Editor/UnrealEd/Public/EditorActorUtilities.h
@@ -0,0 +1,53 @@
+// Copyright Epic Games, Inc. All Rights Reserved.
+
+#pragma once
+
+#include "CoreMinimal.h"
+#include "GameFramework/Actor.h"
+#include "Kismet/BlueprintFunctionLibrary.h"
+
+#include "EditorActorUtilities.generated.h"
+
+UCLASS()
+class UNREALED_API UEditorActorUtilities : public UBlueprintFunctionLibrary
+{
+	GENERATED_BODY()
+
+public:
+
+	/**
+	 * Duplicate all the selected actors in the given world
+	 * @param	InWorld 	The world the actors are selected in.
+	 */
+	UFUNCTION(BlueprintCallable, Category = "Editor Scripting | Level Utility")
+		static void DuplicateSelectedActors(UWorld* InWorld);
+
+	/**
+	 * Delete all the selected actors in the given world
+	 * @param	InWorld 	The world the actors are selected in.
+	 */
+	UFUNCTION(BlueprintCallable, Category = "Editor Scripting | Level Utility")
+		static void DeleteSelectedActors(UWorld* InWorld);
+
+	/**
+	 * Invert the selection in the given world
+	 * @param	InWorld 	The world the selection is in.
+	 */
+	UFUNCTION(BlueprintCallable, Category = "Editor Scripting | Level Utility")
+		static void InvertSelection(UWorld* InWorld);
+
+	/**
+	* Select all actors and BSP models in the given world, except those which are hidden
+	*  @param	InWorld 	The world the actors are to be selected in.
+	*/
+	UFUNCTION(BlueprintCallable, Category = "Editor Scripting | Level Utility")
+		static void SelectAll(UWorld* InWorld);
+
+	/**
+	 * Select all children actors of the current selection.
+	 *
+	 * @param   bRecurseChildren	True to recurse through all descendants of the children
+	 */
+	UFUNCTION(BlueprintCallable, Category = "Editor Scripting | Level Utility")
+		static void SelectAllChildren(bool bRecurseChildren);
+};
diff --git a/Engine/Source/Editor/UnrealEd/Public/Subsystems/EditorDelegateSubsystem.h b/Engine/Source/Editor/UnrealEd/Public/Subsystems/EditorDelegateSubsystem.h
new file mode 100644
--- /dev/null
+++ b/Engine/Source/Editor/UnrealEd/Public/Subsystems/EditorDelegateSubsystem.h
@@ -0,0 +1,107 @@
+// Copyright Epic Games, Inc. All Rights Reserved.
+
+#pragma once
+
+#include "CoreMinimal.h"
+#include "EditorSubsystem.h"
+
+#include "EditorDelegateSubsystem.generated.h"
+
+/** delegate type for before edit cut actors is handled */
+DECLARE_DYNAMIC_MULTICAST_DELEGATE(FOnEditCutActorsBegin);
+/** delegate type for after edit cut actors is handled */
+DECLARE_DYNAMIC_MULTICAST_DELEGATE(FOnEditCutActorsEnd);
+/** delegate type for before edit copy actors is handled */
+DECLARE_DYNAMIC_MULTICAST_DELEGATE(FOnEditCopyActorsBegin);
+/** delegate type for after edit copy actors is handled */
+DECLARE_DYNAMIC_MULTICAST_DELEGATE(FOnEditCopyActorsEnd);
+/** delegate type for before edit paste actors is handled */
+DECLARE_DYNAMIC_MULTICAST_DELEGATE(FOnEditPasteActorsBegin);
+/** delegate type for after edit paste actors is handled */
+DECLARE_DYNAMIC_MULTICAST_DELEGATE(FOnEditPasteActorsEnd);
+/** delegate type for before edit duplicate actors is handled */
+DECLARE_DYNAMIC_MULTICAST_DELEGATE(FOnDuplicateActorsBegin);
+/** delegate type for after edit duplicate actors is handled */
+DECLARE_DYNAMIC_MULTICAST_DELEGATE(FOnDuplicateActorsEnd);
+/** delegate type for before delete actors is handled */
+DECLARE_DYNAMIC_MULTICAST_DELEGATE(FOnDeleteActorsBegin);
+/** delegate type for after delete actors is handled */
+DECLARE_DYNAMIC_MULTICAST_DELEGATE(FOnDeleteActorsEnd);
+
+/**
+* UEditorDelegateSubsystem
+* Subsystem for exposing editor delegates to scripts,
+*/
+UCLASS()
+class UNREALED_API UEditorDelegateSubsystem : public UEditorSubsystem
+{
+	GENERATED_BODY()
+
+public:
+	UEditorDelegateSubsystem();
+
+	virtual void Initialize(FSubsystemCollectionBase& Collection) override;
+	virtual void Deinitialize() override;
+
+	UPROPERTY(BlueprintAssignable, Category = "Editor Scripting | Level Utility")
+	FOnEditCutActorsBegin OnEditCutActorsBegin;
+
+	UPROPERTY(BlueprintAssignable, Category = "Editor Scripting | Level Utility")
+	FOnEditCutActorsEnd OnEditCutActorsEnd;
+
+	UPROPERTY(BlueprintAssignable, Category = "Editor Scripting | Level Utility")
+	FOnEditCopyActorsBegin OnEditCopyActorsBegin;
+
+	UPROPERTY(BlueprintAssignable, Category = "Editor Scripting | Level Utility")
+	FOnEditCopyActorsEnd OnEditCopyActorsEnd;
+
+	UPROPERTY(BlueprintAssignable, Category = "Editor Scripting | Level Utility")
+	FOnEditPasteActorsBegin OnEditPasteActorsBegin;
+
+	UPROPERTY(BlueprintAssignable, Category = "Editor Scripting | Level Utility")
+	FOnEditPasteActorsEnd OnEditPasteActorsEnd;
+
+	UPROPERTY(BlueprintAssignable, Category = "Editor Scripting | Level Utility")
+	FOnEditCutActorsBegin OnDuplicateActorsBegin;
+
+	UPROPERTY(BlueprintAssignable, Category = "Editor Scripting | Level Utility")
+	FOnDuplicateActorsEnd OnDuplicateActorsEnd;
+
+	UPROPERTY(BlueprintAssignable, Category = "Editor Scripting | Level Utility")
+	FOnDeleteActorsBegin OnDeleteActorsBegin;
+
+	UPROPERTY(BlueprintAssignable, Category = "Editor Scripting | Level Utility")
+	FOnDeleteActorsEnd OnDeleteActorsEnd;
+
+private:
+	
+	/** To fire before an Actor is Cut */
+	void BroadcastEditCutActorsBegin();
+
+	/** To fire after an Actor is Cut */
+	void BroadcastEditCutActorsEnd();
+
+	/** To fire before an Actor is Copied */
+	void BroadcastEditCopyActorsBegin();
+
+	/** To fire after an Actor is Copied */
+	void BroadcastEditCopyActorsEnd();
+
+	/** To fire before an Actor is Pasted */
+	void BroadcastEditPasteActorsBegin();
+
+	/** To fire after an Actor is Pasted */
+	void BroadcastEditPasteActorsEnd();
+
+	/** To fire before an Actor is duplicated */
+	void BroadcastDuplicateActorsBegin();
+
+	/** To fire after an Actor is duplicated */
+	void BroadcastDuplicateActorsEnd();
+
+	/** To fire before an Actor is Deleted */
+	void BroadcastDeleteActorsBegin();
+
+	/** To fire after an Actor is Deleted */
+	void BroadcastDeleteActorsEnd();
+};
