CurveAtlasRowParameter nodes cannot be converted to constant nodes in the material editor anymore

diff --git a/Engine/Source/Editor/UnrealEd/Private/MaterialGraphNode.cpp b/Engine/Source/Editor/UnrealEd/Private/MaterialGraphNode.cpp
--- a/Engine/Source/Editor/UnrealEd/Private/MaterialGraphNode.cpp
+++ b/Engine/Source/Editor/UnrealEd/Private/MaterialGraphNode.cpp
@@ -32,7 +32,7 @@
 #include "Materials/MaterialExpressionViewProperty.h"
 #include "Materials/MaterialExpressionMaterialLayerOutput.h"
 #include "Materials/MaterialExpressionTextureObjectParameter.h"
-
+#include "Materials/MaterialExpressionCurveAtlasRowParameter.h"
 #include "MaterialEditorUtilities.h"
 #include "MaterialEditorActions.h"
 #include "GraphEditorActions.h"
@@ -371,8 +371,12 @@ void UMaterialGraphNode::GetNodeContextMenuActions(UToolMenu* Menu, UGraphNodeCo
 				}
 			}
 
+			// Constants cannot be converted back to a CurveAtlasRowParameter
+			const bool bIsAStandardScalarParam = MaterialExpression->IsA(UMaterialExpressionScalarParameter::StaticClass())
+											     && !MaterialExpression->IsA(UMaterialExpressionCurveAtlasRowParameter::StaticClass());
+
 			// Add a 'Convert To Constant' option for convertible types
-			if (MaterialExpression->IsA(UMaterialExpressionScalarParameter::StaticClass())
+			if (bIsAStandardScalarParam
 				|| MaterialExpression->IsA(UMaterialExpressionVectorParameter::StaticClass())
 				|| MaterialExpression->IsA(UMaterialExpressionTextureObjectParameter::StaticClass()))
 			{
