Removed SDockableTab which was the precursor to SDockTab and not used for anything since pre UE4 1.0

diff --git a/Engine/Plugins/Enterprise/DataprepEditor/Source/DataprepEditor/Private/DataprepEditor.h b/Engine/Plugins/Enterprise/DataprepEditor/Source/DataprepEditor/Private/DataprepEditor.h
--- a/Engine/Plugins/Enterprise/DataprepEditor/Source/DataprepEditor/Private/DataprepEditor.h
+++ b/Engine/Plugins/Enterprise/DataprepEditor/Source/DataprepEditor/Private/DataprepEditor.h
@@ -25,8 +25,7 @@ class FTabManager;
 class FUICommandList;
 class IMessageLogListing;
 class IMessageToken;
-class SDataprepGraphEditor;
-class SDockableTab;
+class SDataprepGraphEditor;;
 class SGraphEditor;
 class SGraphNodeDetailsWidget;
 class SInspectorView;
diff --git a/Engine/Plugins/Media/MediaFrameworkUtilities/Source/MediaFrameworkUtilitiesEditor/Private/AssetEditor/MediaBundleEditorToolkit.h b/Engine/Plugins/Media/MediaFrameworkUtilities/Source/MediaFrameworkUtilitiesEditor/Private/AssetEditor/MediaBundleEditorToolkit.h
--- a/Engine/Plugins/Media/MediaFrameworkUtilities/Source/MediaFrameworkUtilitiesEditor/Private/AssetEditor/MediaBundleEditorToolkit.h
+++ b/Engine/Plugins/Media/MediaFrameworkUtilities/Source/MediaFrameworkUtilitiesEditor/Private/AssetEditor/MediaBundleEditorToolkit.h
@@ -51,9 +51,6 @@ private:
 	void ExtendToolBar();
 
 private:
-	/** Dockable tab for properties */
-	TSharedPtr<class SDockableTab> PropertiesTab;
-
 	/** Details view */
 	TSharedPtr<class IDetailsView> DetailsView;
 };
diff --git a/Engine/Plugins/Media/MediaFrameworkUtilities/Source/MediaFrameworkUtilitiesEditor/Private/AssetEditor/MediaProfileEditorToolkit.h b/Engine/Plugins/Media/MediaFrameworkUtilities/Source/MediaFrameworkUtilitiesEditor/Private/AssetEditor/MediaProfileEditorToolkit.h
--- a/Engine/Plugins/Media/MediaFrameworkUtilities/Source/MediaFrameworkUtilitiesEditor/Private/AssetEditor/MediaProfileEditorToolkit.h
+++ b/Engine/Plugins/Media/MediaFrameworkUtilities/Source/MediaFrameworkUtilitiesEditor/Private/AssetEditor/MediaProfileEditorToolkit.h
@@ -60,9 +60,6 @@ private:
 	void HandleMediaProxiesChanged();
 
 private:
-	/** Dockable tab for properties */
-	TSharedPtr<class SDockableTab> PropertiesTab;
-
 	/** Details view */
 	TSharedPtr<class IDetailsView> DetailsView;
 	
diff --git a/Engine/Plugins/Media/TimecodeSynchronizer/Source/TimecodeSynchronizerEditor/Private/AssetEditor/TimecodeSynchronizerEditorToolkit.h b/Engine/Plugins/Media/TimecodeSynchronizer/Source/TimecodeSynchronizerEditor/Private/AssetEditor/TimecodeSynchronizerEditorToolkit.h
--- a/Engine/Plugins/Media/TimecodeSynchronizer/Source/TimecodeSynchronizerEditor/Private/AssetEditor/TimecodeSynchronizerEditorToolkit.h
+++ b/Engine/Plugins/Media/TimecodeSynchronizer/Source/TimecodeSynchronizerEditor/Private/AssetEditor/TimecodeSynchronizerEditorToolkit.h
@@ -52,9 +52,6 @@ private:
 	void HandleSynchronizationEvent(ETimecodeSynchronizationEvent Event);
 
 private:
-	/** Dockable tab for properties */
-	TSharedPtr<class SDockableTab> PropertiesTab;
-
 	/** Details view */
 	TSharedPtr<class IDetailsView> DetailsView;
 };
diff --git a/Engine/Plugins/Runtime/ApexDestruction/Source/ApexDestructionEditor/Private/DestructibleMeshEditor.h b/Engine/Plugins/Runtime/ApexDestruction/Source/ApexDestructionEditor/Private/DestructibleMeshEditor.h
--- a/Engine/Plugins/Runtime/ApexDestruction/Source/ApexDestructionEditor/Private/DestructibleMeshEditor.h
+++ b/Engine/Plugins/Runtime/ApexDestruction/Source/ApexDestructionEditor/Private/DestructibleMeshEditor.h
@@ -13,7 +13,6 @@
 #include "ApexDestructionStyle.h"
 
 class IDetailsView;
-class SDockableTab;
 class SSlider;
 class UDestructibleMesh;
 
@@ -141,9 +140,6 @@ private:
 	/** Preview Viewport widget */
 	TSharedPtr<class SDestructibleMeshEditorViewport>	Viewport;
 
-	/** List of open tool panels; used to ensure only one exists at any one time */
-	TMap< FName, TWeakPtr<class SDockableTab> >			SpawnedToolPanels;
-
 	/** Property View */
 	TSharedPtr<class IDetailsView>						DestructibleMeshDetailsView;
 	TSharedPtr<class SDockTab>							DestructibleMeshDetailsViewTab;
diff --git a/Engine/Plugins/Runtime/ApexDestruction/Source/ApexDestructionEditor/Private/SDestructibleMeshEditorViewport.cpp b/Engine/Plugins/Runtime/ApexDestruction/Source/ApexDestructionEditor/Private/SDestructibleMeshEditorViewport.cpp
--- a/Engine/Plugins/Runtime/ApexDestruction/Source/ApexDestructionEditor/Private/SDestructibleMeshEditorViewport.cpp
+++ b/Engine/Plugins/Runtime/ApexDestruction/Source/ApexDestructionEditor/Private/SDestructibleMeshEditorViewport.cpp
@@ -17,7 +17,6 @@
 #include "UObject/UObjectHash.h"
 #include "ComponentReregisterContext.h"
 #include "DestructibleMesh.h"
-#include "Widgets/Docking/SDockableTab.h"
 #include "Engine/StaticMesh.h"
 #include "PhysXPublic.h"
 #include "Settings/SkeletalMeshEditorSettings.h"
@@ -658,7 +657,7 @@ void SDestructibleMeshEditorViewport::UpdatePreviewMesh(UDestructibleMesh* InDes
 
 bool SDestructibleMeshEditorViewport::IsVisible() const
 {
-	return ViewportWidget.IsValid() && (!ParentTab.IsValid() || ParentTab.Pin()->IsForeground());
+	return ViewportWidget.IsValid();
 }
 
 void SDestructibleMeshEditorViewport::SetPreviewDepth(uint32 InPreviewDepth)
diff --git a/Engine/Plugins/Runtime/ApexDestruction/Source/ApexDestructionEditor/Private/SDestructibleMeshEditorViewport.h b/Engine/Plugins/Runtime/ApexDestruction/Source/ApexDestructionEditor/Private/SDestructibleMeshEditorViewport.h
--- a/Engine/Plugins/Runtime/ApexDestruction/Source/ApexDestructionEditor/Private/SDestructibleMeshEditorViewport.h
+++ b/Engine/Plugins/Runtime/ApexDestruction/Source/ApexDestructionEditor/Private/SDestructibleMeshEditorViewport.h
@@ -12,7 +12,6 @@
 
 class FDestructibleMeshEditorViewportClient;
 class IDestructibleMeshEditor;
-class SDockableTab;
 class UDestructibleComponent;
 class UDestructibleMesh;
 
@@ -43,9 +42,6 @@ public:
 	/** Component for the preview Destructible mesh. */
 	class UDestructibleComponent* PreviewComponent;
 
-	/** The parent tab where this viewport resides */
-	TWeakPtr<SDockableTab> ParentTab;
-
 	/** 
 	 *	Causes chunks at the given PreviewDepth to be displayed in the viewport.  Clamped to the range [0, depth count), where (depth count) = the number of chunk hierarchy depths in the destructible.
 	 *
diff --git a/Engine/Source/Editor/AudioEditor/Private/SoundClassEditor.h b/Engine/Source/Editor/AudioEditor/Private/SoundClassEditor.h
--- a/Engine/Source/Editor/AudioEditor/Private/SoundClassEditor.h
+++ b/Engine/Source/Editor/AudioEditor/Private/SoundClassEditor.h
@@ -11,7 +11,6 @@
 #include "ISoundClassEditor.h"
 
 class IDetailsView;
-class SDockableTab;
 class UEdGraph;
 class USoundClass;
 
@@ -111,9 +110,6 @@ private:
 	/** The SoundClass asset being inspected */
 	USoundClass* SoundClass;
 
-	/** List of open tool panels; used to ensure only one exists at any one time */
-	TMap< FName, TWeakPtr<class SDockableTab> > SpawnedToolPanels;
-
 	/** Graph Editor */
 	TSharedPtr<SGraphEditor> GraphEditor;
 
diff --git a/Engine/Source/Editor/AudioEditor/Private/SoundCueEditor.h b/Engine/Source/Editor/AudioEditor/Private/SoundCueEditor.h
--- a/Engine/Source/Editor/AudioEditor/Private/SoundCueEditor.h
+++ b/Engine/Source/Editor/AudioEditor/Private/SoundCueEditor.h
@@ -10,7 +10,6 @@
 #include "EditorUndoClient.h"
 
 class IDetailsView;
-class SDockableTab;
 class SGraphEditor;
 class SSoundCuePalette;
 class UEdGraphNode;
@@ -200,9 +199,6 @@ private:
 	/** The SoundCue asset being inspected */
 	USoundCue* SoundCue;
 
-	/** List of open tool panels; used to ensure only one exists at any one time */
-	TMap< FName, TWeakPtr<SDockableTab> > SpawnedToolPanels;
-
 	/** New Graph Editor */
 	TSharedPtr<SGraphEditor> SoundCueGraphEditor;
 
diff --git a/Engine/Source/Editor/AudioEditor/Private/SoundSubmixEditor.h b/Engine/Source/Editor/AudioEditor/Private/SoundSubmixEditor.h
--- a/Engine/Source/Editor/AudioEditor/Private/SoundSubmixEditor.h
+++ b/Engine/Source/Editor/AudioEditor/Private/SoundSubmixEditor.h
@@ -8,7 +8,6 @@
 #include "ISoundSubmixEditor.h"
 #include "Toolkits/IToolkitHost.h"
 #include "UObject/Object.h"
-#include "Widgets/Docking/SDockableTab.h"
 
 
 // Forward Declarations
@@ -100,9 +99,6 @@ private:
 	void RedoGraphAction();
 
 private:
-	/** List of open tool panels; used to ensure only one exists at any one time */
-	TMap<FName, TWeakPtr<SDockableTab>> SpawnedToolPanels;
-
 	/** Graph Editor */
 	TSharedPtr<SGraphEditor> GraphEditor;
 
diff --git a/Engine/Source/Editor/Cascade/Private/Cascade.h b/Engine/Source/Editor/Cascade/Private/Cascade.h
--- a/Engine/Source/Editor/Cascade/Private/Cascade.h
+++ b/Engine/Source/Editor/Cascade/Private/Cascade.h
@@ -22,7 +22,6 @@ class FFXSystemInterface;
 class IDetailsView;
 class SCascadeEmitterCanvas;
 class SCascadePreviewViewport;
-class SDockableTab;
 class UCascadeConfiguration;
 class UCascadeOptions;
 class UCascadeParticleSystemComponent;
@@ -412,9 +411,6 @@ private:
 	/** The ParticleSystem asset being inspected */
 	UParticleSystem* ParticleSystem;
 
-	/** List of open tool panels; used to ensure only one exists at any one time */
-	TMap< FName, TWeakPtr<SDockableTab> > SpawnedToolPanels;
-
 	/** Preview Viewport */
 	TSharedPtr<SCascadePreviewViewport> PreviewViewport;
 
diff --git a/Engine/Source/Editor/LevelEditor/Private/SLevelEditor.cpp b/Engine/Source/Editor/LevelEditor/Private/SLevelEditor.cpp
--- a/Engine/Source/Editor/LevelEditor/Private/SLevelEditor.cpp
+++ b/Engine/Source/Editor/LevelEditor/Private/SLevelEditor.cpp
@@ -450,12 +450,6 @@ void SLevelEditor::BringToFront()
 	}
 }
 
-TSharedRef< SDockTabStack > SLevelEditor::GetTabSpot( const EToolkitTabSpot::Type TabSpot )
-{
-	ensureMsgf(false, TEXT("Unimplemented"));
-	return TSharedPtr<SDockTabStack>().ToSharedRef();
-}
-
 void SLevelEditor::OnToolkitHostingStarted( const TSharedRef< class IToolkit >& Toolkit )
 {
 	// @todo toolkit minor: We should consider only allowing a single toolkit for a specific asset editor type hosted
diff --git a/Engine/Source/Editor/LevelEditor/Private/SLevelEditor.h b/Engine/Source/Editor/LevelEditor/Private/SLevelEditor.h
--- a/Engine/Source/Editor/LevelEditor/Private/SLevelEditor.h
+++ b/Engine/Source/Editor/LevelEditor/Private/SLevelEditor.h
@@ -115,7 +115,6 @@ public:
 	/** IToolKitHost interface */
 	virtual TSharedRef< class SWidget > GetParentWidget() override;
 	virtual void BringToFront() override;
-	virtual TSharedRef< class SDockTabStack > GetTabSpot( const EToolkitTabSpot::Type TabSpot ) override;
 	virtual TSharedPtr<FTabManager> GetTabManager() const override;
 	virtual void OnToolkitHostingStarted( const TSharedRef< class IToolkit >& Toolkit ) override;
 	virtual void OnToolkitHostingFinished( const TSharedRef< class IToolkit >& Toolkit ) override;
diff --git a/Engine/Source/Editor/LevelEditor/Private/SLevelEditorBuildAndSubmit.h b/Engine/Source/Editor/LevelEditor/Private/SLevelEditorBuildAndSubmit.h
--- a/Engine/Source/Editor/LevelEditor/Private/SLevelEditorBuildAndSubmit.h
+++ b/Engine/Source/Editor/LevelEditor/Private/SLevelEditorBuildAndSubmit.h
@@ -84,7 +84,7 @@ private:
 	/** Level editor that we're associated with */
 	TWeakPtr<class ILevelEditor> LevelEditor;
 
-	/** SDockableTab in the LevelEditor that we're associated with */
+	/** SDockTab in the LevelEditor that we're associated with */
 	TWeakPtr<SDockTab> ParentDockTab;
 
 	/** The package list that acts as the items source for the additional packages list widget */
diff --git a/Engine/Source/Editor/MaterialEditor/Private/MaterialEditor.h b/Engine/Source/Editor/MaterialEditor/Private/MaterialEditor.h
--- a/Engine/Source/Editor/MaterialEditor/Private/MaterialEditor.h
+++ b/Engine/Source/Editor/MaterialEditor/Private/MaterialEditor.h
@@ -27,7 +27,6 @@ class FCanvas;
 class FMaterialCompiler;
 class FScopedTransaction;
 class IMessageLogListing;
-class SDockableTab;
 class SFindInMaterial;
 class SGraphEditor;
 class SMaterialPalette;
@@ -788,9 +787,6 @@ private:
 	void GeneratorRowsRefreshed();
 	void UpdateGenerator();
 private:
-	/** List of open tool panels; used to ensure only one exists at any one time */
-	TMap< FName, TWeakPtr<class SDockableTab> > SpawnedToolPanels;
-
 	/** Property View */
 	TSharedPtr<class IDetailsView> MaterialDetailsView;
 
diff --git a/Engine/Source/Editor/Matinee/Private/Matinee.h b/Engine/Source/Editor/Matinee/Private/Matinee.h
--- a/Engine/Source/Editor/Matinee/Private/Matinee.h
+++ b/Engine/Source/Editor/Matinee/Private/Matinee.h
@@ -33,7 +33,6 @@ class FSceneView;
 class FViewport;
 class IDetailsView;
 class SBorder;
-class SDockableTab;
 class SMatineeRecorder;
 class SMatineeViewport;
 class UInterpFilter;
@@ -1536,9 +1535,6 @@ private:
 	UInterpData* GetInterpData() override { return IData; }
 	bool Hide3DTrackView() override { return bHide3DTrackView; }
 
-	/** List of open tool panels; used to ensure only one exists at any one time */
-	TMap< FName, TWeakPtr<SDockableTab> > SpawnedToolPanels;
-
 	/** Action Mapping */
 	void BindCommands();
 
diff --git a/Engine/Source/Editor/PhysicsAssetEditor/Private/PhysicsAssetEditor.h b/Engine/Source/Editor/PhysicsAssetEditor/Private/PhysicsAssetEditor.h
--- a/Engine/Source/Editor/PhysicsAssetEditor/Private/PhysicsAssetEditor.h
+++ b/Engine/Source/Editor/PhysicsAssetEditor/Private/PhysicsAssetEditor.h
@@ -25,7 +25,6 @@ struct FAssetData;
 class FPhysicsAssetEditorTreeInfo;
 class IDetailsView;
 class SComboButton;
-class SDockableTab;
 class SPhysicsAssetEditorPreviewViewport;
 class UAnimationAsset;
 class UAnimSequence;
diff --git a/Engine/Source/Editor/StaticMeshEditor/Private/StaticMeshEditor.h b/Engine/Source/Editor/StaticMeshEditor/Private/StaticMeshEditor.h
--- a/Engine/Source/Editor/StaticMeshEditor/Private/StaticMeshEditor.h
+++ b/Engine/Source/Editor/StaticMeshEditor/Private/StaticMeshEditor.h
@@ -26,7 +26,6 @@ class FStaticMeshDetails;
 class FEditorViewportClient;
 class IDetailsView;
 class SConvexDecomposition;
-class SDockableTab;
 class SStaticMeshEditorViewport;
 class UStaticMesh;
 class UStaticMeshComponent;
@@ -465,9 +464,6 @@ private:
 	// Tracking the active viewports in this editor.
 	TSharedPtr<class FEditorViewportTabContent> ViewportTabContent;
 
-	/** List of open tool panels; used to ensure only one exists at any one time */
-	TMap< FName, TWeakPtr<class SDockableTab> > SpawnedToolPanels;
-
 	/** Property View */
 	TSharedPtr<class IDetailsView> StaticMeshDetailsView;
 
diff --git a/Engine/Source/Editor/TextureEditor/Private/TextureEditorToolkit.h b/Engine/Source/Editor/TextureEditor/Private/TextureEditorToolkit.h
--- a/Engine/Source/Editor/TextureEditor/Private/TextureEditorToolkit.h
+++ b/Engine/Source/Editor/TextureEditor/Private/TextureEditorToolkit.h
@@ -12,7 +12,6 @@
 #include "IDetailsView.h"
 #include "TextureEditorSettings.h"
 
-class SDockableTab;
 class STextBlock;
 class STextureEditorViewport;
 class UFactory;
@@ -259,9 +258,6 @@ private:
 	/** The Texture asset being inspected */
 	UTexture* Texture;
 
-	/** List of open tool panels; used to ensure only one exists at any one time */
-	TMap<FName, TWeakPtr<SDockableTab>> SpawnedToolPanels;
-
 	/** Viewport */
 	TSharedPtr<STextureEditorViewport> TextureViewport;
 
diff --git a/Engine/Source/Editor/UnrealEd/Private/Toolkits/BaseToolkit.cpp b/Engine/Source/Editor/UnrealEd/Private/Toolkits/BaseToolkit.cpp
--- a/Engine/Source/Editor/UnrealEd/Private/Toolkits/BaseToolkit.cpp
+++ b/Engine/Source/Editor/UnrealEd/Private/Toolkits/BaseToolkit.cpp
@@ -2,8 +2,6 @@
 
 #include "Toolkits/BaseToolkit.h"
 #include "Toolkits/ToolkitManager.h"
-#include "Widgets/Docking/SDockableTab.h"
-#include "Widgets/Docking/SDockTabStack.h"
 #include "PropertyEditorModule.h"
 #include "Modules/ModuleManager.h"
 #include "IDetailsView.h"
@@ -26,23 +24,6 @@ FBaseToolkit::FBaseToolkit()
 
 FBaseToolkit::~FBaseToolkit()
 {
-	// Destroy any tabs that we still have hanging out.  This is just for convenience, so that the derived
-	// classes don't need to bother cleaning up their tabs when a world-centric editor is closed.  However,
-	// if the order that tabs are destroyed is important, it is suggested that you clean the up yourself!
-	{
-		for( auto CurSpotIt( ToolkitTabsInSpots.CreateConstIterator() ); CurSpotIt; ++CurSpotIt )
-		{
-			const auto& TabsForSpot = CurSpotIt.Value();
-			for( auto CurTabIt( TabsForSpot.CreateConstIterator() ); CurTabIt; ++CurTabIt )
-			{
-				const auto& PinnedTab = CurTabIt->Pin();
-				if( PinnedTab.IsValid() )
-				{
-					PinnedTab->RemoveTabFromParent();
-				}
-			}
-		}
-	}
 }
 
 
@@ -65,13 +46,6 @@ const TSharedRef< IToolkitHost > FBaseToolkit::GetToolkitHost() const
 	return ToolkitHost.Pin().ToSharedRef();
 }
 
-
-const TMap< EToolkitTabSpot::Type, TArray< TWeakPtr< SDockableTab > > >& FBaseToolkit::GetToolkitTabsInSpots() const
-{
-	return ToolkitTabsInSpots;
-}
-
-
 FName FBaseToolkit::GetToolkitContextFName() const
 {
 	return GetToolkitFName();
@@ -88,24 +62,6 @@ bool FBaseToolkit::ProcessCommandBindings( const FKeyEvent& InKeyEvent ) const
 	return false;
 }
 
-
-void FBaseToolkit::AddToolkitTab( const TSharedRef< SDockableTab >& TabToAdd, const EToolkitTabSpot::Type TabSpot )
-{
-	// Figure out where to put this tab by asking the toolkit host which spot to put it in.  It will give us
-	// back a dock tab stack that we can add the tab to!
-	TSharedRef< SDockTabStack > FoundTabStack = GetToolkitHost()->GetTabSpot( TabSpot );
-
-	// When a context menu is opening for this tab, ask the spawning app to fill the menu.
-	ensureMsgf(false, TEXT("TabToAdd->SetOnTabStackMenuOpening( FOnTabStackMenuOpening::CreateSP( GetToolkitHost(), &IToolkitHost::PopulateLayoutMenu ) );"));
-	//TabToAdd->SetOnTabStackMenuOpening( FOnTabStackMenuOpening::CreateSP( GetToolkitHost(), &IToolkitHost::PopulateLayoutMenu ) );
-
-	// Add the tab
-	FoundTabStack->AddTab( TabToAdd );
-
-	// Keep track of tabs
-	ToolkitTabsInSpots.FindOrAdd( TabSpot ).Add( TabToAdd );
-}
-
 FString FBaseToolkit::GetTabPrefix() const
 {
 	if( IsWorldCentricAssetEditor() )
@@ -131,66 +87,6 @@ void FBaseToolkit::BringToolkitToFront()
 	{
 		// Bring the host window to front
 		ToolkitHost.Pin()->BringToFront();
-
-		// First, figure out what the foreground tab is in each tab stack we have tabs docked inside of
-		TSet< SDockTabStack* > TabStacksWithOurTabsForegrounded;
-		{
-			for( auto CurSpotIt( ToolkitTabsInSpots.CreateConstIterator() ); CurSpotIt; ++CurSpotIt )
-			{
-				const auto& TabsForSpot = CurSpotIt.Value();
-				for( auto CurTabIt( TabsForSpot.CreateConstIterator() ); CurTabIt; ++CurTabIt )
-				{
-					const auto& PinnedTab = CurTabIt->Pin();
-					if( PinnedTab.IsValid() )
-					{
-						if( PinnedTab->IsForeground() )
-						{
-							const auto& TabStack = PinnedTab->GetParentDockTabStack();
-							if( TabStack.IsValid() )
-							{
-								TabStacksWithOurTabsForegrounded.Add( TabStack.Get() );
-							}
-						}
-					}
-				}
-			}
-		}
-
-		// @todo toolkit major: Also draw user's attention when clicked?
-
-		// @todo toolkit major: If any of the tabs are in their own floating windows, these should be brought to front
-
-		// Now, make sure that our tabs are foregrounded in their respective stacks!
-		// NOTE: We don't want to push tabs to the front that are in a stack where one of our other tabs is already front-most
-		for( auto CurSpotIt( ToolkitTabsInSpots.CreateConstIterator() ); CurSpotIt; ++CurSpotIt )
-		{
-			const auto& TabsForSpot = CurSpotIt.Value();
-			for( auto CurTabIt( TabsForSpot.CreateConstIterator() ); CurTabIt; ++CurTabIt )
-			{
-				const auto& PinnedTab = CurTabIt->Pin();
-
-				if( PinnedTab.IsValid() )
-				{
-					const auto& TabStack = PinnedTab->GetParentDockTabStack();
-					if( TabStack.IsValid() )
-					{
-						// Only foreground if we don't already have a tab foregrounded in this tab's stack
-						if( !TabStacksWithOurTabsForegrounded.Contains( TabStack.Get() ) )
-						{
-							PinnedTab->BringToFrontInParent();
-
-							// Take note that we've foregrounded a tab in this stack, no need to do that again
-							TabStacksWithOurTabsForegrounded.Add( TabStack.Get() );
-						}
-					}
-					else
-					{
-						// Just do what we can to foreground ourselves
-						PinnedTab->BringToFrontInParent();
-					}
-				}
-			}
-		}
 		// Tell the toolkit its been brought to the fore - give it a chance to update anything it needs to
 		ToolkitBroughtToFront();
 	}
diff --git a/Engine/Source/Editor/UnrealEd/Private/Toolkits/SStandaloneAssetEditorToolkitHost.cpp b/Engine/Source/Editor/UnrealEd/Private/Toolkits/SStandaloneAssetEditorToolkitHost.cpp
--- a/Engine/Source/Editor/UnrealEd/Private/Toolkits/SStandaloneAssetEditorToolkitHost.cpp
+++ b/Engine/Source/Editor/UnrealEd/Private/Toolkits/SStandaloneAssetEditorToolkitHost.cpp
@@ -236,13 +236,6 @@ void SStandaloneAssetEditorToolkitHost::BringToFront()
 	FGlobalTabmanager::Get()->DrawAttentionToTabManager( this->MyTabManager.ToSharedRef() );
 }
 
-
-TSharedRef< SDockTabStack > SStandaloneAssetEditorToolkitHost::GetTabSpot( const EToolkitTabSpot::Type TabSpot )
-{
-	return TSharedPtr<SDockTabStack>().ToSharedRef();
-}
-
-
 void SStandaloneAssetEditorToolkitHost::OnToolkitHostingStarted( const TSharedRef< class IToolkit >& Toolkit )
 {
 	// Keep track of the toolkit we're hosting
@@ -301,57 +294,9 @@ FReply SStandaloneAssetEditorToolkitHost::OnKeyDown( const FGeometry& MyGeometry
 	// If we are in debug mode do not process commands
 	if (FSlateApplication::Get().IsNormalExecution())
 	{
-		// Figure out if any of our toolkit's tabs is the active tab.  This is important because we want
-		// the toolkit to have it's own keybinds (which may overlap the level editor's keybinds or any
-		// other toolkit).  When a toolkit tab is active, we give that toolkit a chance to process
-		// commands instead of the level editor.
-		TSharedPtr< IToolkit > ActiveToolkit;
+		if (HostedAssetEditorToolkit->ProcessCommandBindings(InKeyEvent))
 		{
-			const TSharedPtr<SDockableTab> CurrentActiveTab;// = FSlateApplication::xxxGetGlobalTabManager()->GetActiveTab();
-
-			for (auto HostedToolkitIt = HostedToolkits.CreateConstIterator(); HostedToolkitIt && !ActiveToolkit.IsValid(); ++HostedToolkitIt)
-			{
-				const auto& CurToolkit = *HostedToolkitIt;
-				if (CurToolkit.IsValid())
-				{
-					// Iterate over this toolkits spawned tabs
-					const auto& ToolkitTabsInSpots = CurToolkit->GetToolkitTabsInSpots();
-
-					for (auto CurSpotIt(ToolkitTabsInSpots.CreateConstIterator()); CurSpotIt && !ActiveToolkit.IsValid(); ++CurSpotIt)
-					{
-						const auto& TabsForSpot = CurSpotIt.Value();
-						for (auto CurTabIt(TabsForSpot.CreateConstIterator()); CurTabIt; ++CurTabIt)
-						{
-							const auto& PinnedTab = CurTabIt->Pin();
-							if (PinnedTab.IsValid())
-							{
-								if (PinnedTab == CurrentActiveTab)
-								{
-									ActiveToolkit = CurToolkit;
-								}
-							}
-						}
-					}
-				}
-			}
-		}
-
-		//@TODO: This seems wrong (should prioritize it but not totally block it)
-		if (ActiveToolkit.IsValid())
-		{
-			// A toolkit tab is active, so direct all command processing to it
-			if (ActiveToolkit->ProcessCommandBindings(InKeyEvent))
-			{
-				return FReply::Handled();
-			}
-		}
-		else
-		{
-			// No toolkit tab is active, so let the underlying asset editor have a chance at the keystroke
-			if (HostedAssetEditorToolkit->ProcessCommandBindings(InKeyEvent))
-			{
-				return FReply::Handled();
-			}
+			return FReply::Handled();
 		}
 	}
 
diff --git a/Engine/Source/Editor/UnrealEd/Private/Toolkits/SStandaloneAssetEditorToolkitHost.h b/Engine/Source/Editor/UnrealEd/Private/Toolkits/SStandaloneAssetEditorToolkitHost.h
--- a/Engine/Source/Editor/UnrealEd/Private/Toolkits/SStandaloneAssetEditorToolkitHost.h
+++ b/Engine/Source/Editor/UnrealEd/Private/Toolkits/SStandaloneAssetEditorToolkitHost.h
@@ -52,7 +52,6 @@ public:
 	/** IToolkitHost interface */
 	virtual TSharedRef< class SWidget > GetParentWidget() override;
 	virtual void BringToFront() override;
-	virtual TSharedRef< class SDockTabStack > GetTabSpot( const EToolkitTabSpot::Type TabSpot ) override;
 	virtual TSharedPtr< class FTabManager > GetTabManager() const override { return MyTabManager; }
 	virtual void OnToolkitHostingStarted( const TSharedRef< class IToolkit >& Toolkit ) override;
 	virtual void OnToolkitHostingFinished( const TSharedRef< class IToolkit >& Toolkit ) override;
@@ -91,11 +90,8 @@ private:
 	/** Manages internal tab layout */
 	TSharedPtr<FTabManager> MyTabManager;
 
-	/** The widget that will house the default menu widget */
-	TSharedPtr<SBorder> MenuWidgetContent;
-
 	/** The widget that will house the overlay widgets (if any) */
-	TSharedPtr<SBorder> MenuOverlayWidgetContent;
+	TSharedPtr<SBox> MenuOverlayWidgetContent;
 
 	/** The default menu widget */
 	TSharedPtr< SWidget > DefaultMenuWidget;
diff --git a/Engine/Source/Editor/UnrealEd/Private/Toolkits/SimpleAssetEditor.cpp b/Engine/Source/Editor/UnrealEd/Private/Toolkits/SimpleAssetEditor.cpp
--- a/Engine/Source/Editor/UnrealEd/Private/Toolkits/SimpleAssetEditor.cpp
+++ b/Engine/Source/Editor/UnrealEd/Private/Toolkits/SimpleAssetEditor.cpp
@@ -43,7 +43,6 @@ FSimpleAssetEditor::~FSimpleAssetEditor()
 	GEditor->OnObjectsReplaced().RemoveAll(this);
 
 	DetailsView.Reset();
-	PropertiesTab.Reset();
 }
 
 
diff --git a/Engine/Source/Editor/UnrealEd/Public/Toolkits/BaseToolkit.h b/Engine/Source/Editor/UnrealEd/Public/Toolkits/BaseToolkit.h
--- a/Engine/Source/Editor/UnrealEd/Public/Toolkits/BaseToolkit.h
+++ b/Engine/Source/Editor/UnrealEd/Public/Toolkits/BaseToolkit.h
@@ -11,7 +11,6 @@
 
 class UInteractiveTool;
 class UInteractiveToolManager;
-class SDockableTab;
 class UEdMode;
 class IDetailsView;
 enum class EToolShutdownType;
@@ -39,7 +38,6 @@ public:
 	virtual bool ProcessCommandBindings(const FKeyEvent& InKeyEvent) const override;
 	virtual bool IsHosted() const override;
 	virtual const TSharedRef<IToolkitHost> GetToolkitHost() const override;
-	virtual const TMap<EToolkitTabSpot::Type, TArray<TWeakPtr<SDockableTab>>>& GetToolkitTabsInSpots() const override;
 	virtual void BringToolkitToFront() override;
 	virtual TSharedPtr<SWidget> GetInlineContent() const override;
 	virtual bool IsBlueprintEditor() const override;
@@ -56,15 +54,6 @@ public:
 		return ToolkitCommands;
 	}
 
-	/**
-	 * Adds an already-created toolkit tab to the toolkit.  Used for tabs that have no tab identifier, such as a "document" tab
-	 *
-	 * @param	TabToAdd		The dockable tab object to add
-	 * @param	TabSpot			Which spot to spawn this tab into
-	 */
-	void AddToolkitTab(const TSharedRef<SDockableTab>& TabToAdd, const EToolkitTabSpot::Type TabSpot);
-
-
 protected:
 
 	/** @return Returns the prefix string to use for tabs created for this toolkit.  In world-centric mode, tabs get a
@@ -87,9 +76,6 @@ protected:
 	    Use GetToolkitHost() method to access this member. */
 	TWeakPtr<IToolkitHost> ToolkitHost;
 
-	/** Map of toolkit tab spots to known tabs (these are weak pointers and may be invalid after tabs are closed.) */
-	TMap<EToolkitTabSpot::Type, TArray<TWeakPtr<SDockableTab>>> ToolkitTabsInSpots;
-
 	/** The workspace menu category of this toolkit */
 	TSharedPtr<FWorkspaceItem> WorkspaceMenuCategory;
 };
diff --git a/Engine/Source/Editor/UnrealEd/Public/Toolkits/IToolkit.h b/Engine/Source/Editor/UnrealEd/Public/Toolkits/IToolkit.h
--- a/Engine/Source/Editor/UnrealEd/Public/Toolkits/IToolkit.h
+++ b/Engine/Source/Editor/UnrealEd/Public/Toolkits/IToolkit.h
@@ -9,7 +9,6 @@
 class FEdMode;
 class UEdMode;
 class IToolkitHost;
-class SDockableTab;
 
 /** Some toolkits can be spawned as either standalone tools or within an existing level editing UI */
 namespace EToolkitMode
@@ -108,9 +107,6 @@ public:
 	/** @return Returns the toolkit host for this toolkit */
 	virtual const TSharedRef<IToolkitHost> GetToolkitHost() const = 0;
 
-	/** @return Returns a map of weak pointers to all of this toolkit's spawned tabs that may currently exist, indexed by the tab spots they are suitable for  */
-	virtual const TMap<EToolkitTabSpot::Type, TArray<TWeakPtr<SDockableTab>>>& GetToolkitTabsInSpots() const = 0;
-
 	/**
 	 * Processes any UI commands which are activated by the specified event
 	 *
diff --git a/Engine/Source/Editor/UnrealEd/Public/Toolkits/IToolkitHost.h b/Engine/Source/Editor/UnrealEd/Public/Toolkits/IToolkitHost.h
--- a/Engine/Source/Editor/UnrealEd/Public/Toolkits/IToolkitHost.h
+++ b/Engine/Source/Editor/UnrealEd/Public/Toolkits/IToolkitHost.h
@@ -20,9 +20,6 @@ public:
 	/** Brings this toolkit host's window (and tab, if it has one), to the front */
 	virtual void BringToFront() = 0;
 
-	/** Gets a tab stack to place a new tab for the specified toolkit area */
-	virtual TSharedRef< class SDockTabStack > GetTabSpot( const EToolkitTabSpot::Type TabSpot ) = 0;
-
 	/** Access the toolkit host's tab manager */
 	virtual TSharedPtr< class FTabManager > GetTabManager() const = 0;
 
diff --git a/Engine/Source/Editor/UnrealEd/Public/Toolkits/SimpleAssetEditor.h b/Engine/Source/Editor/UnrealEd/Public/Toolkits/SimpleAssetEditor.h
--- a/Engine/Source/Editor/UnrealEd/Public/Toolkits/SimpleAssetEditor.h
+++ b/Engine/Source/Editor/UnrealEd/Public/Toolkits/SimpleAssetEditor.h
@@ -8,7 +8,6 @@
 #include "Editor/PropertyEditor/Public/PropertyEditorDelegates.h"
 
 class IDetailsView;
-class SDockableTab;
 
 class UNREALED_API FSimpleAssetEditor : public FAssetEditorToolkit
 {
@@ -59,9 +58,6 @@ private:
 	/** Called when objects need to be swapped out for new versions, like after a blueprint recompile. */
 	void OnObjectsReplaced(const TMap<UObject*, UObject*>& ReplacementMap);
 
-	/** Dockable tab for properties */
-	TSharedPtr< SDockableTab > PropertiesTab;
-
 	/** Details view */
 	TSharedPtr< class IDetailsView > DetailsView;
 
diff --git a/Engine/Source/Runtime/Slate/Private/Framework/Docking/SDockTabStack.cpp b/Engine/Source/Runtime/Slate/Private/Framework/Docking/SDockTabStack.cpp
deleted file mode 100644
diff --git a/Engine/Source/Runtime/Slate/Private/Widgets/Docking/SDockableTab.cpp b/Engine/Source/Runtime/Slate/Private/Widgets/Docking/SDockableTab.cpp
deleted file mode 100644
diff --git a/Engine/Source/Runtime/Slate/Public/SlateExtras.h b/Engine/Source/Runtime/Slate/Public/SlateExtras.h
--- a/Engine/Source/Runtime/Slate/Public/SlateExtras.h
+++ b/Engine/Source/Runtime/Slate/Public/SlateExtras.h
@@ -154,6 +154,4 @@ MONOLITHIC_HEADER_BOILERPLATE()
 #include "Widgets/Input/SVirtualJoystick.h"
 
 // Docking Framework
-#include "Widgets/Docking/SDockTab.h"
-#include "Widgets/Docking/SDockableTab.h"
-#include "Widgets/Docking/SDockTabStack.h"
+#include "Widgets/Docking/SDockTab.h"
\ No newline at end of file
diff --git a/Engine/Source/Runtime/Slate/Public/Widgets/Docking/SDockTabStack.h b/Engine/Source/Runtime/Slate/Public/Widgets/Docking/SDockTabStack.h
deleted file mode 100644
diff --git a/Engine/Source/Runtime/Slate/Public/Widgets/Docking/SDockableTab.h b/Engine/Source/Runtime/Slate/Public/Widgets/Docking/SDockableTab.h
deleted file mode 100644
