- Renamed PlatformInfo::FPlatformInfo to PlatformInfo::FTargetPlatformInfo
- Renamed FDataDrivenPlatformRegistry::FPlatformInfo to FDataDrivenPlatformInfo
- Moved some fields from PlatformInfo::FTargetPlatformInfo to FDataDrivenPlatformInfo, and cleaned them up in the process
- Fixed the DataDrivePlatformInfo.ini files to match the previous items
- Removed FVanillaPlatformEntry, and now just using FTargetPlatformInfo to manage flavors under a vanilla PlatformInfo (see PlatformInfo::GetVanillaPlatformInfoArray())
- Cleaned up TPerPlatformValue, as it was often misused (took a group and platform name, but we can get the group from the platform name) [AnimationSharingManager.cpp, *Engine.cpp, Runtime\Engine\*, ]
- Fixed FBlueprintNativeCodeGenPaths::GetDefaultCodeGenPaths() to use IniPlatformName instead of some hacky code [BlueprintNativeCodeGenManifest.cpp]
- Fixed various Turnkey bugs that recent testing exposed
- Enabled AUTOSDKS_ENABLED in both Mac and Linux on the Editor side (not that there are much set up to use it)
- Using Turnkey to get the SDK status instead of ValidatePlatforms (gives more detailed information - may want to go back to ValidatePlatforms so only UBT is needed, but change it to give more info)
- Moved OnDeviceDiscovered and Lost delegates from each platform's TargetPlatform class to a static in ITargetPlatform - there was no need for per-platform implementations
- Started working on allowing for SDK to be installed with editor running and not need to restart editor - Work in Progress! It is not usable yet.

diff --git a/Engine/Config/Android/DataDrivenPlatformInfo.ini b/Engine/Config/Android/DataDrivenPlatformInfo.ini
--- a/Engine/Config/Android/DataDrivenPlatformInfo.ini
+++ b/Engine/Config/Android/DataDrivenPlatformInfo.ini
@@ -4,100 +4,50 @@ AudioCompressionSettingsIniSectionName=/Script/AndroidRuntimeSettings.AndroidRun
 Freezing_b32Bit=ini:Engine:Android:/Script/AndroidRuntimeSettings.AndroidRuntimeSettings:bBuildForArmV7
 Freezing_bForce64BitMemoryImagePointers=true
 
-[PlatformInfo Android]
-TargetPlatformName=Android
-DisplayName=Android
-PlatformType=Game
-PlatformFlags=None
 NormalIconPath=Launcher/Android/Platform_Android_24x
 LargeIconPath=Launcher/Android/Platform_Android_128x
 XLargeIconPath=
-UATCommandLine=
 AutoSDKPath=Android
 TutorialPath=/Engine/Tutorial/Mobile/SettingUpAndroidTutorial.SettingUpAndroidTutorial
 bIsEnabled=true
-BinariesDirectoryName=Android
-IniPlatformName=Android
 bUsesHostCompiler=false
 bUATClosesAfterLaunch=true
-UBTTargetID=Android
 PlatformGroupName=Mobile
 
+[PlatformInfo Android]
+TargetPlatformName=Android
+DisplayName=Android
+PlatformType=Game
+PlatformFlags=None
+UATCommandLine=
+
 [PlatformInfo Android_Multi]
 TargetPlatformName=Android_Multi
 DisplayName=Android (Multi)
 PlatformType=Game
 PlatformFlags=CookFlavor
-NormalIconPath=Launcher/Android/Platform_Android_24x
-LargeIconPath=Launcher/Android/Platform_Android_128x
-XLargeIconPath=
 UATCommandLine=-targetplatform=Android -cookflavor=Multi
-AutoSDKPath=Android
-TutorialPath=/Engine/Tutorial/Mobile/SettingUpAndroidTutorial.SettingUpAndroidTutorial
-bIsEnabled=true
-BinariesDirectoryName=Android
-IniPlatformName=Android
-bUsesHostCompiler=false
-bUATClosesAfterLaunch=true
-UBTTargetID=Android
-PlatformGroupName=Mobile
 
 [PlatformInfo Android_DXT]
 TargetPlatformName=Android_DXT
 DisplayName=Android (DXT)
 PlatformType=Game
 PlatformFlags=CookFlavor
-NormalIconPath=Launcher/Android/Platform_Android_DXT_24x
-LargeIconPath=Launcher/Android/Platform_Android_128x
-XLargeIconPath=
 UATCommandLine=-targetplatform=Android -cookflavor=DXT
-AutoSDKPath=Android
-TutorialPath=/Engine/Tutorial/Mobile/SettingUpAndroidTutorial.SettingUpAndroidTutorial
-bIsEnabled=true
-BinariesDirectoryName=Android
-IniPlatformName=Android
-bUsesHostCompiler=false
-bUATClosesAfterLaunch=true
-UBTTargetID=Android
-PlatformGroupName=Mobile
 
 [PlatformInfo Android_ETC2]
 TargetPlatformName=Android_ETC2
 DisplayName=Android (ETC2)
 PlatformType=Game
 PlatformFlags=CookFlavor
-NormalIconPath=Launcher/Android/Platform_Android_ETC2_24x
-LargeIconPath=Launcher/Android/Platform_Android_128x
-XLargeIconPath=
 UATCommandLine=-targetplatform=Android -cookflavor=ETC2
-AutoSDKPath=Android
-TutorialPath=/Engine/Tutorial/Mobile/SettingUpAndroidTutorial.SettingUpAndroidTutorial
-bIsEnabled=true
-BinariesDirectoryName=Android
-IniPlatformName=Android
-bUsesHostCompiler=false
-bUATClosesAfterLaunch=true
-UBTTargetID=Android
-PlatformGroupName=Mobile
 
 [PlatformInfo Android_ASTC]
 TargetPlatformName=Android_ASTC
 DisplayName=Android (ASTC)
 PlatformType=Game
 PlatformFlags=CookFlavor
-NormalIconPath=Launcher/Android/Platform_Android_ASTC_24x
-LargeIconPath=Launcher/Android/Platform_Android_128x
-XLargeIconPath=
 UATCommandLine=-targetplatform=Android -cookflavor=ASTC
-AutoSDKPath=Android
-TutorialPath=/Engine/Tutorial/Mobile/SettingUpAndroidTutorial.SettingUpAndroidTutorial
-bIsEnabled=true
-BinariesDirectoryName=Android
-IniPlatformName=Android
-bUsesHostCompiler=false
-bUATClosesAfterLaunch=true
-UBTTargetID=Android
-PlatformGroupName=Mobile
 
 
 [PlatformInfo AndroidClient]
@@ -105,92 +55,32 @@ TargetPlatformName=AndroidClient
 DisplayName=Android (Client-only)
 PlatformType=Client
 PlatformFlags=None
-NormalIconPath=Launcher/Android/Platform_Android_24x
-LargeIconPath=Launcher/Android/Platform_Android_128x
-XLargeIconPath=
 UATCommandLine=-client
-AutoSDKPath=Android
-TutorialPath=/Engine/Tutorial/Mobile/SettingUpAndroidTutorial.SettingUpAndroidTutorial
-bIsEnabled=true
-BinariesDirectoryName=Android
-IniPlatformName=Android
-bUsesHostCompiler=false
-bUATClosesAfterLaunch=true
-UBTTargetID=Android
-PlatformGroupName=Mobile
 
 [PlatformInfo Android_MultiClient]
 TargetPlatformName=Android_MultiClient
 DisplayName=Android (Multi) (Client-only)
 PlatformType=Client
 PlatformFlags=CookFlavor
-NormalIconPath=Launcher/Android/Platform_Android_24x
-LargeIconPath=Launcher/Android/Platform_Android_128x
-XLargeIconPath=
 UATCommandLine=-client -targetplatform=Android -cookflavor=Multi
-AutoSDKPath=Android
-TutorialPath=/Engine/Tutorial/Mobile/SettingUpAndroidTutorial.SettingUpAndroidTutorial
-bIsEnabled=true
-BinariesDirectoryName=Android
-IniPlatformName=Android
-bUsesHostCompiler=false
-bUATClosesAfterLaunch=true
-UBTTargetID=Android
-PlatformGroupName=Mobile
 
 [PlatformInfo Android_DXTClient]
 TargetPlatformName=Android_DXTClient
 DisplayName=Android (DXT) (Client-only)
 PlatformType=Client
 PlatformFlags=CookFlavor
-NormalIconPath=Launcher/Android/Platform_Android_DXT_24x
-LargeIconPath=Launcher/Android/Platform_Android_128x
-XLargeIconPath=
 UATCommandLine=-client -targetplatform=Android -cookflavor=DXT
-AutoSDKPath=Android
-TutorialPath=/Engine/Tutorial/Mobile/SettingUpAndroidTutorial.SettingUpAndroidTutorial
-bIsEnabled=true
-BinariesDirectoryName=Android
-IniPlatformName=Android
-bUsesHostCompiler=false
-bUATClosesAfterLaunch=true
-UBTTargetID=Android
-PlatformGroupName=Mobile
 
 [PlatformInfo Android_ETC2Client]
 TargetPlatformName=Android_ETC2Client
 DisplayName=Android (ETC2) (Client-only)
 PlatformType=Client
 PlatformFlags=CookFlavor
-NormalIconPath=Launcher/Android/Platform_Android_ETC2_24x
-LargeIconPath=Launcher/Android/Platform_Android_128x
-XLargeIconPath=
 UATCommandLine=-client -targetplatform=Android -cookflavor=ETC2
-AutoSDKPath=Android
-TutorialPath=/Engine/Tutorial/Mobile/SettingUpAndroidTutorial.SettingUpAndroidTutorial
-bIsEnabled=true
-BinariesDirectoryName=Android
-IniPlatformName=Android
-bUsesHostCompiler=false
-bUATClosesAfterLaunch=true
-UBTTargetID=Android
-PlatformGroupName=Mobile
 
 [PlatformInfo Android_ASTCClient]
 TargetPlatformName=Android_ASTCClient
 DisplayName=Android (ASTC) (Client-only)
 PlatformType=Client
 PlatformFlags=CookFlavor
-NormalIconPath=Launcher/Android/Platform_Android_ASTC_24x
-LargeIconPath=Launcher/Android/Platform_Android_128x
-XLargeIconPath=
 UATCommandLine=-client -targetplatform=Android -cookflavor=ASTC
-AutoSDKPath=Android
-TutorialPath=/Engine/Tutorial/Mobile/SettingUpAndroidTutorial.SettingUpAndroidTutorial
-bIsEnabled=true
-BinariesDirectoryName=Android
-IniPlatformName=Android
-bUsesHostCompiler=false
-bUATClosesAfterLaunch=true
-UBTTargetID=Android
-PlatformGroupName=Mobile
diff --git a/Engine/Config/HoloLens/DataDrivenPlatformInfo.ini b/Engine/Config/HoloLens/DataDrivenPlatformInfo.ini
--- a/Engine/Config/HoloLens/DataDrivenPlatformInfo.ini
+++ b/Engine/Config/HoloLens/DataDrivenPlatformInfo.ini
@@ -3,24 +3,22 @@ bIsConfidential=false
 Freezing_MaxFieldAlignment=8
 Freezing_bAlignBases=true
 
-[PlatformInfo HoloLens]
-TargetPlatformName=HoloLens
-DisplayName=HoloLens
-PlatformType=Game
-PlatformFlags=None
 NormalIconPath=Launcher/Windows/Platform_WindowsNoEditor_24x
 LargeIconPath=Launcher/Windows/Platform_WindowsNoEditor_128x
 XLargeIconPath=
-UATCommandLine=-targetplatform=HoloLens
 AutoSDKPath=
 TutorialPath=/Engine/Tutorial/Installation/InstallingVisualStudioTutorial.InstallingVisualStudioTutorial
 Windows:bIsEnabled=true
 Mac:bIsEnabled=false
 Linux:bIsEnabled=false
-BinariesDirectoryName=HoloLens
-IniPlatformName=HoloLens
 bUsesHostCompiler=true
-bTargetPlatformCanUseCrashReporter=false;
+bCanUseCrashReporter=false
 bUATClosesAfterLaunch=false
-UBTTargetID=HoloLens
-PlatformGroupName=Mobile
\ No newline at end of file
+PlatformGroupName=Mobile
+
+[PlatformInfo HoloLens]
+TargetPlatformName=HoloLens
+DisplayName=HoloLens
+PlatformType=Game
+PlatformFlags=None
+UATCommandLine=-targetplatform=HoloLens
diff --git a/Engine/Config/IOS/DataDrivenPlatformInfo.ini b/Engine/Config/IOS/DataDrivenPlatformInfo.ini
--- a/Engine/Config/IOS/DataDrivenPlatformInfo.ini
+++ b/Engine/Config/IOS/DataDrivenPlatformInfo.ini
@@ -2,15 +2,9 @@
 bIsConfidential=false
 AudioCompressionSettingsIniSectionName=/Script/IOSRuntimeSettings.IOSRuntimeSettings
 
-[PlatformInfo IOS]
-TargetPlatformName=IOS
-DisplayName=iOS
-PlatformType=Game
-PlatformFlags=None
 NormalIconPath=Launcher/IOS/Platform_IOS_24x
 LargeIconPath=Launcher/IOS/Platform_IOS_128x
 XLargeIconPath=
-UATCommandLine=
 AutoSDKPath=
 Windows:TutorialPath=/Engine/Tutorial/Mobile/InstallingiTunesTutorial.InstallingiTunesTutorial
 Mac:TutorialPath=/Engine/Tutorial/Installation/InstallingXCodeTutorial.InstallingXCodeTutorial
@@ -18,32 +12,20 @@ Linux:TutorialPath=/Engine/Tutorial/NotYetImplemented
 Windows:bIsEnabled=true
 Mac:bIsEnabled=true
 Linux:bIsEnabled=false
-BinariesDirectoryName=IOS
-IniPlatformName=IOS
 bUsesHostCompiler=false
 bUATClosesAfterLaunch=true
-UBTTargetID=IOS
 PlatformGroupName=Mobile
 
+[PlatformInfo IOS]
+TargetPlatformName=IOS
+DisplayName=iOS
+PlatformType=Game
+PlatformFlags=None
+UATCommandLine=
+
 [PlatformInfo IOSClient]
 TargetPlatformName=IOSClient
 DisplayName=iOS (Client-only)
 PlatformType=Client
 PlatformFlags=None
-NormalIconPath=Launcher/IOS/Platform_IOS_24x
-LargeIconPath=Launcher/IOS/Platform_IOS_24x
-XLargeIconPath=
 UATCommandLine=-client
-AutoSDKPath=
-Windows:TutorialPath=/Engine/Tutorial/Mobile/InstallingiTunesTutorial.InstallingiTunesTutorial
-Mac:TutorialPath=/Engine/Tutorial/Installation/InstallingXCodeTutorial.InstallingXCodeTutorial
-Linux:TutorialPath=/Engine/Tutorial/NotYetImplemented
-Windows:bIsEnabled=true
-Mac:bIsEnabled=true
-Linux:bIsEnabled=false
-BinariesDirectoryName=IOS
-IniPlatformName=IOS
-bUsesHostCompiler=false
-bUATClosesAfterLaunch=true
-UBTTargetID=IOS
-PlatformGroupName=Mobile
diff --git a/Engine/Config/Linux/DataDrivenPlatformInfo.ini b/Engine/Config/Linux/DataDrivenPlatformInfo.ini
--- a/Engine/Config/Linux/DataDrivenPlatformInfo.ini
+++ b/Engine/Config/Linux/DataDrivenPlatformInfo.ini
@@ -1,95 +1,44 @@
 [DataDrivenPlatformInfo]
 IniParent=Unix
 
-[PlatformInfo Linux]
-TargetPlatformName=Linux
-DisplayName=Linux (Editor)
-PlatformType=Editor
-PlatformFlags=None
 NormalIconPath=Launcher/Linux/Platform_Linux_24x
 LargeIconPath=Launcher/Linux/Platform_Linux_128x
 XLargeIconPath=
-UATCommandLine=
 AutoSDKPath=Linux_x64
 TutorialPath=/Platforms/Linux/GettingStarted
 Windows:bIsEnabled=true
 Mac:bIsEnabled=false
 Linux:bIsEnabled=true
-BinariesDirectoryName=Linux
-IniPlatformName=Linux
 Windows:bUsesHostCompiler=false
 Linux:bUsesHostCompiler=true
 bUATClosesAfterLaunch=true
-UBTTargetID=Linux
 PlatformGroupName=Desktop
 PlatformSubMenu=Linux
 
+[PlatformInfo Linux]
+TargetPlatformName=Linux
+DisplayName=Linux (Editor)
+PlatformType=Editor
+PlatformFlags=None
+UATCommandLine=
+
 [PlatformInfo LinuxNoEditor]
 TargetPlatformName=LinuxNoEditor
 DisplayName=Linux
 PlatformType=Game
 PlatformFlags=None
-NormalIconPath=Launcher/Linux/Platform_Linux_24x
-LargeIconPath=Launcher/Linux/Platform_Linux_128x
-XLargeIconPath=
 UATCommandLine=-targetplatform=Linux
-AutoSDKPath=Linux_x64
-TutorialPath=/Platforms/Linux/GettingStarted
-Windows:bIsEnabled=true
-Mac:bIsEnabled=false
-Linux:bIsEnabled=true
-BinariesDirectoryName=Linux
-IniPlatformName=Linux
-Windows:bUsesHostCompiler=false
-Linux:bUsesHostCompiler=true
-bUATClosesAfterLaunch=true
-UBTTargetID=Linux
-PlatformGroupName=Desktop
-PlatformSubMenu=Linux
 
 [PlatformInfo LinuxClient]
 TargetPlatformName=LinuxClient
 DisplayName=Linux (Client-only)
 PlatformType=Client
 PlatformFlags=None
-NormalIconPath=Launcher/Linux/Platform_Linux_24x
-LargeIconPath=Launcher/Linux/Platform_Linux_128x
-XLargeIconPath=
 UATCommandLine=-client
-AutoSDKPath=Linux_x64
-TutorialPath=/Platforms/Linux/GettingStarted
-Windows:bIsEnabled=true
-Mac:bIsEnabled=false
-Linux:bIsEnabled=true
-BinariesDirectoryName=Linux
-IniPlatformName=Linux
-Windows:bUsesHostCompiler=false
-Linux:bUsesHostCompiler=true
-bUATClosesAfterLaunch=true
-UBTTargetID=Linux
-PlatformGroupName=Desktop
-PlatformSubMenu=Linux
 
 [PlatformInfo LinuxServer]
 TargetPlatformName=LinuxServer
 DisplayName=Linux (Dedicated Server)
 PlatformType=Server
 PlatformFlags=None
-NormalIconPath=Launcher/Linux/Platform_Linux_24x
-LargeIconPath=Launcher/Linux/Platform_Linux_128x
-XLargeIconPath=
 UATCommandLine=
-AutoSDKPath=Linux_x64
-TutorialPath=/Platforms/Linux/GettingStarted
-Windows:bIsEnabled=true
-Mac:bIsEnabled=false
-Linux:bIsEnabled=true
-BinariesDirectoryName=Linux
-IniPlatformName=Linux
-Windows:bUsesHostCompiler=false
-Linux:bUsesHostCompiler=true
-bUATClosesAfterLaunch=true
-UBTTargetID=Linux
-PlatformGroupName=Desktop
-PlatformSubMenu=Linux
-
diff --git a/Engine/Config/LinuxAArch64/DataDrivenPlatformInfo.ini b/Engine/Config/LinuxAArch64/DataDrivenPlatformInfo.ini
--- a/Engine/Config/LinuxAArch64/DataDrivenPlatformInfo.ini
+++ b/Engine/Config/LinuxAArch64/DataDrivenPlatformInfo.ini
@@ -2,72 +2,37 @@
 IniParent=Unix
 bIsConfidential=false
 
-[PlatformInfo LinuxAArch64]
-TargetPlatformName=LinuxAArch64
-DisplayName=Linux (AArch64)
-PlatformType=Game
-PlatformFlags=BuildFlavor
 NormalIconPath=Launcher/Linux/Platform_Linux_24x
 LargeIconPath=Launcher/Linux/Platform_Linux_128x
 XLargeIconPath=
-UATCommandLine=-targetplatform=LinuxAArch64
 AutoSDKPath=Linux_x64
 TutorialPath=/Platforms/Linux/GettingStarted
 Windows:bIsEnabled=true
 Mac:bIsEnabled=false
 Linux:bIsEnabled=true
-BinariesDirectoryName=LinuxAArch64
-IniPlatformName=LinuxAArch64
 Windows:bUsesHostCompiler=false
 Linux:bUsesHostCompiler=true
 bUATClosesAfterLaunch=true
-UBTTargetID=LinuxAArch64
 PlatformGroupName=Mobile
 PlatformSubMenu=Linux
 
+[PlatformInfo LinuxAArch64]
+TargetPlatformName=LinuxAArch64
+DisplayName=Linux (AArch64)
+PlatformType=Game
+PlatformFlags=BuildFlavor
+UATCommandLine=-targetplatform=LinuxAArch64
+
 [PlatformInfo LinuxAArch64Client]
 TargetPlatformName=LinuxAArch64Client
 DisplayName=LinuxAArch64 (Client-only)
 PlatformType=Client
 PlatformFlags=None
-NormalIconPath=Launcher/Linux/Platform_Linux_24x
-LargeIconPath=Launcher/Linux/Platform_Linux_128x
-XLargeIconPath=
 UATCommandLine=-client
-AutoSDKPath=Linux_x64
-TutorialPath=/Platforms/Linux/GettingStarted
-Windows:bIsEnabled=true
-Mac:bIsEnabled=false
-Linux:bIsEnabled=true
-BinariesDirectoryName=LinuxAArch64
-IniPlatformName=LinuxAArch64
-Windows:bUsesHostCompiler=false
-Linux:bUsesHostCompiler=true
-bUATClosesAfterLaunch=true
-UBTTargetID=LinuxAArch64
-PlatformGroupName=Mobile
-PlatformSubMenu=Linux
 
 [PlatformInfo LinuxAArch64Server]
 TargetPlatformName=LinuxAArch64Server
 DisplayName=LinuxAArch64 (Dedicated Server)
 PlatformType=Server
 PlatformFlags=None
-NormalIconPath=Launcher/Linux/Platform_Linux_24x
-LargeIconPath=Launcher/Linux/Platform_Linux_128x
-XLargeIconPath=
 UATCommandLine=
-AutoSDKPath=Linux_x64
-TutorialPath=/Platforms/Linux/GettingStarted
-Windows:bIsEnabled=true
-Mac:bIsEnabled=false
-Linux:bIsEnabled=true
-BinariesDirectoryName=LinuxAArch64
-IniPlatformName=LinuxAArch64
-Windows:bUsesHostCompiler=false
-Linux:bUsesHostCompiler=true
-bUATClosesAfterLaunch=true
-UBTTargetID=LinuxAArch64
-PlatformGroupName=Mobile
-PlatformSubMenu=Linux
-
diff --git a/Engine/Config/Lumin/DataDrivenPlatformInfo.ini b/Engine/Config/Lumin/DataDrivenPlatformInfo.ini
--- a/Engine/Config/Lumin/DataDrivenPlatformInfo.ini
+++ b/Engine/Config/Lumin/DataDrivenPlatformInfo.ini
@@ -2,45 +2,28 @@
 bIsConfidential=false
 IniParent=Android
 
-
-[PlatformInfo Lumin]
-TargetPlatformName=Lumin
-DisplayName=Lumin
-PlatformType=Game
-PlatformFlags=None
 NormalIconPath=Launcher/Lumin/Platform_Lumin_24x
 LargeIconPath=Launcher/Lumin/Platform_Lumin_128x
 XLargeIconPath=
-UATCommandLine=
 AutoSDKPath=
 TutorialPath=/Engine/Tutorial/Mobile/SettingUpLuminTutorial
 Windows:bIsEnabled=true
 Mac:bIsEnabled=true
 Linux:bIsEnabled=false
-BinariesDirectoryName=Lumin
-IniPlatformName=Lumin
 bUsesHostCompiler=false
 bUATClosesAfterLaunch=true
-UBTTargetID=Lumin
 PlatformGroupName=Mobile
 
+[PlatformInfo Lumin]
+TargetPlatformName=Lumin
+DisplayName=Lumin
+PlatformType=Game
+PlatformFlags=None
+UATCommandLine=
+
 [PlatformInfo LuminClient]
 TargetPlatformName=LuminClient
 DisplayName=Lumin (Client-only)
 PlatformType=Client
 PlatformFlags=None
-NormalIconPath=Launcher/Lumin/Platform_Lumin_24x
-LargeIconPath=Launcher/Lumin/Platform_Lumin_128x
-XLargeIconPath=
 UATCommandLine=-client
-AutoSDKPath=
-TutorialPath=/Engine/Tutorial/Mobile/SettingUpLuminTutorial
-Windows:bIsEnabled=true
-Mac:bIsEnabled=true
-Linux:bIsEnabled=false
-BinariesDirectoryName=Lumin
-IniPlatformName=Lumin
-bUsesHostCompiler=false
-bUATClosesAfterLaunch=true
-UBTTargetID=Lumin
-PlatformGroupName=Mobile
diff --git a/Engine/Config/Mac/DataDrivenPlatformInfo.ini b/Engine/Config/Mac/DataDrivenPlatformInfo.ini
--- a/Engine/Config/Mac/DataDrivenPlatformInfo.ini
+++ b/Engine/Config/Mac/DataDrivenPlatformInfo.ini
@@ -2,86 +2,42 @@
 bIsConfidential=false
 AudioCompressionSettingsIniSectionName=/Script/MacTargetPlatform.MacTargetSettings
 
-[PlatformInfo MacNoEditor]
-TargetPlatformName=MacNoEditor
-DisplayName=Mac
-PlatformType=Game
-PlatformFlags=None
 NormalIconPath=Launcher/Mac/Platform_Mac_24x
 LargeIconPath=Launcher/Mac/Platform_Mac_128x
 XLargeIconPath=
-UATCommandLine=-targetplatform=Mac
 AutoSDKPath=
 TutorialPath=/Engine/Tutorial/Installation/InstallingXCodeTutorial.InstallingXCodeTutorial
 Windows:bIsEnabled=false
 Mac:bIsEnabled=true
 Linux:bIsEnabled=false
-BinariesDirectoryName=Mac
-IniPlatformName=Mac
 bUsesHostCompiler=true
 bUATClosesAfterLaunch=true
-UBTTargetID=Mac
 PlatformGroupName=Desktop
 
+[PlatformInfo MacNoEditor]
+TargetPlatformName=MacNoEditor
+DisplayName=Mac
+PlatformType=Game
+PlatformFlags=None
+UATCommandLine=-targetplatform=Mac
+
 [PlatformInfo Mac]
 TargetPlatformName=Mac
 DisplayName=Mac (Editor)
 PlatformType=Editor
 PlatformFlags=None
-NormalIconPath=Launcher/Mac/Platform_Mac_24x
-LargeIconPath=Launcher/Mac/Platform_Mac_128x
-XLargeIconPath=
 UATCommandLine=
-AutoSDKPath=
-TutorialPath=/Engine/Tutorial/Installation/InstallingXCodeTutorial.InstallingXCodeTutorial
-Windows:bIsEnabled=false
-Mac:bIsEnabled=true
-Linux:bIsEnabled=false
-BinariesDirectoryName=Mac
-IniPlatformName=Mac
-bUsesHostCompiler=true
-bUATClosesAfterLaunch=true
-UBTTargetID=Mac
-PlatformGroupName=Desktop
 
 [PlatformInfo MacClient]
 TargetPlatformName=MacClient
 DisplayName=Mac (Client-only)
 PlatformType=Client
 PlatformFlags=None
-NormalIconPath=Launcher/Mac/Platform_Mac_24x
-LargeIconPath=Launcher/Mac/Platform_Mac_128x
-XLargeIconPath=
 UATCommandLine=-client
-AutoSDKPath=
-TutorialPath=/Engine/Tutorial/Installation/InstallingXCodeTutorial.InstallingXCodeTutorial
-Windows:bIsEnabled=false
-Mac:bIsEnabled=true
-Linux:bIsEnabled=false
-BinariesDirectoryName=Mac
-IniPlatformName=Mac
-bUsesHostCompiler=true
-bUATClosesAfterLaunch=true
-UBTTargetID=Mac
-PlatformGroupName=Desktop
 
 [PlatformInfo MacServer]
 TargetPlatformName=MacServer
 DisplayName=Mac (Dedicated Server)
 PlatformType=Server
 PlatformFlags=None
-NormalIconPath=Launcher/Mac/Platform_Mac_24x
-LargeIconPath=Launcher/Mac/Platform_Mac_128x
-XLargeIconPath=
 UATCommandLine=
-AutoSDKPath=
-TutorialPath=/Engine/Tutorial/Installation/InstallingXCodeTutorial.InstallingXCodeTutorial
-Windows:bIsEnabled=false
-Mac:bIsEnabled=true
-Linux:bIsEnabled=false
-BinariesDirectoryName=Mac
-IniPlatformName=Mac
-bUsesHostCompiler=true
-bUATClosesAfterLaunch=true
-UBTTargetID=Mac
-PlatformGroupName=Desktop
diff --git a/Engine/Config/TVOS/DataDrivenPlatformInfo.ini b/Engine/Config/TVOS/DataDrivenPlatformInfo.ini
--- a/Engine/Config/TVOS/DataDrivenPlatformInfo.ini
+++ b/Engine/Config/TVOS/DataDrivenPlatformInfo.ini
@@ -3,44 +3,28 @@ bIsConfidential=false
 IniParent=IOS
 AudioCompressionSettingsIniSectionName=/Script/IOSRuntimeSettings.IOSRuntimeSettings
 
-[PlatformInfo TVOS]
-TargetPlatformName=TVOS
-DisplayName=tvOS
-PlatformType=Game
-PlatformFlags=None
 NormalIconPath=Launcher/TVOS/Platform_TVOS_24x
 LargeIconPath=Launcher/TVOS/Platform_TVOS_128x
 XLargeIconPath=
-UATCommandLine=
 AutoSDKPath=
 TutorialPath=
 Windows:bIsEnabled=true
 Mac:bIsEnabled=true
 Linux:bIsEnabled=false
-BinariesDirectoryName=TVOS
-IniPlatformName=TVOS
 bUsesHostCompiler=false
 bUATClosesAfterLaunch=true
-UBTTargetID=TVOS
 PlatformGroupName=Mobile
 
+[PlatformInfo TVOS]
+TargetPlatformName=TVOS
+DisplayName=tvOS
+PlatformType=Game
+PlatformFlags=None
+UATCommandLine=
+
 [PlatformInfo TVOSClient]
 TargetPlatformName=TVOSClient
 DisplayName=tvOS (Client-only)
 PlatformType=Client
 PlatformFlags=None
-NormalIconPath=Launcher/TVOS/Platform_TVOS_24x
-LargeIconPath=Launcher/TVOS/Platform_TVOS_24x
-XLargeIconPath=
 UATCommandLine=-client
-AutoSDKPath=
-TutorialPath=
-Windows:bIsEnabled=true
-Mac:bIsEnabled=true
-Linux:bIsEnabled=false
-BinariesDirectoryName=TVOS
-IniPlatformName=TVOS
-bUsesHostCompiler=false
-bUATClosesAfterLaunch=true
-UBTTargetID=TVOS
-PlatformGroupName=Mobile
diff --git a/Engine/Config/Windows/DataDrivenPlatformInfo.ini b/Engine/Config/Windows/DataDrivenPlatformInfo.ini
--- a/Engine/Config/Windows/DataDrivenPlatformInfo.ini
+++ b/Engine/Config/Windows/DataDrivenPlatformInfo.ini
@@ -7,89 +7,47 @@ Freezing_MaxFieldAlignment32=4
 Freezing_MaxFieldAlignment64=8
 Freezing_bAlignBases=true
 
-[PlatformInfo WindowsNoEditor]
-TargetPlatformName=WindowsNoEditor
-DisplayName=Windows
-PlatformType=Game
-PlatformFlags=None
 NormalIconPath=Launcher/Windows/Platform_WindowsNoEditor_24x
 LargeIconPath=Launcher/Windows/Platform_WindowsNoEditor_128x
 XLargeIconPath=
-UATCommandLine=-targetplatform=Win64
 AutoSDKPath=
 TutorialPath=/Engine/Tutorial/Installation/InstallingVisualStudioTutorial.InstallingVisualStudioTutorial
 Windows:bIsEnabled=true
 Mac:bIsEnabled=false
 Linux:bIsEnabled=false
 BinariesDirectoryName=Win64
-IniPlatformName=Windows
 bUsesHostCompiler=true
 bUATClosesAfterLaunch=false
-UBTTargetID=Win64
 PlatformGroupName=Desktop
 
+
+[PlatformInfo WindowsNoEditor]
+TargetPlatformName=WindowsNoEditor
+DisplayName=Windows
+PlatformType=Game
+PlatformFlags=None
+UATCommandLine=-targetplatform=Win64
+
 [PlatformInfo Windows]
 TargetPlatformName=Windows
 DisplayName=Windows (Editor)
 PlatformType=Editor
 PlatformFlags=None
-NormalIconPath=Launcher/Windows/Platform_Windows_24x
-LargeIconPath=Launcher/Windows/Platform_Windows_128x
-XLargeIconPath=
 UATCommandLine=
-AutoSDKPath=
-TutorialPath=/Engine/Tutorial/Installation/InstallingVisualStudioTutorial.InstallingVisualStudioTutorial
-Windows:bIsEnabled=true
-Mac:bIsEnabled=false
-Linux:bIsEnabled=false
-BinariesDirectoryName=Win64
-IniPlatformName=Windows
-bUsesHostCompiler=true
-bUATClosesAfterLaunch=false
-UBTTargetID=Win64
-PlatformGroupName=Desktop
 
 [PlatformInfo WindowsClient]
 TargetPlatformName=WindowsClient
 DisplayName=Windows (Client-only)
 PlatformType=Client
 PlatformFlags=None
-NormalIconPath=Launcher/Windows/Platform_Windows_24x
-LargeIconPath=Launcher/Windows/Platform_Windows_128x
-XLargeIconPath=
 UATCommandLine=-client
-AutoSDKPath=
-TutorialPath=/Engine/Tutorial/Installation/InstallingVisualStudioTutorial.InstallingVisualStudioTutorial
-Windows:bIsEnabled=true
-Mac:bIsEnabled=false
-Linux:bIsEnabled=false
-BinariesDirectoryName=Win64
-IniPlatformName=Windows
-bUsesHostCompiler=true
-bUATClosesAfterLaunch=false
-UBTTargetID=Win64
-PlatformGroupName=Desktop
 
 [PlatformInfo WindowsServer]
 TargetPlatformName=WindowsServer
 DisplayName=Windows (Dedicated Server)
 PlatformType=Server
 PlatformFlags=None
-NormalIconPath=Launcher/Windows/Platform_WindowsServer_24x
-LargeIconPath=Launcher/Windows/Platform_WindowsServer_128x
-XLargeIconPath=
 UATCommandLine=
-AutoSDKPath=
-TutorialPath=/Engine/Tutorial/Installation/InstallingVisualStudioTutorial.InstallingVisualStudioTutorial
-Windows:bIsEnabled=true
-Mac:bIsEnabled=false
-Linux:bIsEnabled=false
-BinariesDirectoryName=Win64
-IniPlatformName=Windows
-bUsesHostCompiler=true
-bUATClosesAfterLaunch=false
-UBTTargetID=Win64
-PlatformGroupName=Desktop
 
 [ShaderPlatform PCD3D_SM5]
 bSupportsRayTracing=true
diff --git a/Engine/Plugins/Developer/AnimationSharing/Source/AnimationSharing/Private/AnimationSharingManager.cpp b/Engine/Plugins/Developer/AnimationSharing/Source/AnimationSharing/Private/AnimationSharingManager.cpp
--- a/Engine/Plugins/Developer/AnimationSharing/Source/AnimationSharing/Private/AnimationSharingManager.cpp
+++ b/Engine/Plugins/Developer/AnimationSharing/Source/AnimationSharing/Private/AnimationSharingManager.cpp
@@ -186,10 +186,10 @@ void UAnimationSharingManager::Initialise(const UAnimationSharingSetup* InSetup)
 #if WITH_EDITOR
 		// Update local copy defaults with current platform value
 		const FName PlatformName = UAnimationSharingManager::GetPlatformName();
-		ScalabilitySettings.UseBlendTransitions = ScalabilitySettings.UseBlendTransitions.GetValueForPlatformIdentifiers(PlatformName, PlatformName);
-		ScalabilitySettings.BlendSignificanceValue = ScalabilitySettings.BlendSignificanceValue.GetValueForPlatformIdentifiers(PlatformName, PlatformName);
-		ScalabilitySettings.MaximumNumberConcurrentBlends = ScalabilitySettings.MaximumNumberConcurrentBlends.GetValueForPlatformIdentifiers(PlatformName, PlatformName);
-		ScalabilitySettings.TickSignificanceValue = ScalabilitySettings.TickSignificanceValue.GetValueForPlatformIdentifiers(PlatformName, PlatformName);
+		ScalabilitySettings.UseBlendTransitions = ScalabilitySettings.UseBlendTransitions.GetValueForPlatform(PlatformName);
+		ScalabilitySettings.BlendSignificanceValue = ScalabilitySettings.BlendSignificanceValue.GetValueForPlatform(PlatformName);
+		ScalabilitySettings.MaximumNumberConcurrentBlends = ScalabilitySettings.MaximumNumberConcurrentBlends.GetValueForPlatform(PlatformName);
+		ScalabilitySettings.TickSignificanceValue = ScalabilitySettings.TickSignificanceValue.GetValueForPlatform(PlatformName);
 #endif
 
 		// Debug materials
@@ -701,8 +701,7 @@ FName UAnimationSharingManager::GetPlatformName()
 	const FString PlatformString = CVarAnimSharing_PreviewScalabilityPlatform.GetValueOnAnyThread();
 	if (PlatformString.IsEmpty())
 	{
-		ITargetPlatform* CurrentPlatform = GetTargetPlatformManagerRef().GetRunningTargetPlatform();
-		return CurrentPlatform->GetPlatformInfo().PlatformGroupName;
+		return FName(FPlatformProperties::IniPlatformName());
 	}
 
 	FName PlatformNameFromString(*PlatformString);
@@ -862,7 +861,7 @@ void UAnimSharingInstance::SetupState(FPerStateData& StateData, const FAnimation
 	int32 MaximumNumberOfConcurrentInstances = StateEntry.MaximumNumberOfConcurrentInstances.Default;
 #if WITH_EDITOR
 	const FName PlatformName = UAnimationSharingManager::GetPlatformName();
-	MaximumNumberOfConcurrentInstances = StateEntry.MaximumNumberOfConcurrentInstances.GetValueForPlatformIdentifiers(PlatformName, PlatformName);
+	MaximumNumberOfConcurrentInstances = StateEntry.MaximumNumberOfConcurrentInstances.GetValueForPlatform(PlatformName);
 #endif
 
 	/** Ensure that we spread our number over the number of enabled setups */
@@ -874,7 +873,7 @@ void UAnimSharingInstance::SetupState(FPerStateData& StateData, const FAnimation
 			bool bEnabled = AnimationSetup.Enabled.Default;
 #if WITH_EDITOR
 			const FName PlatformName = UAnimationSharingManager::GetPlatformName();
-			bEnabled = AnimationSetup.Enabled.GetValueForPlatformIdentifiers(PlatformName, PlatformName);
+			bEnabled = AnimationSetup.Enabled.GetValueForPlatform(PlatformName);
 #endif
 			TotalEnabled += bEnabled ? 1 : 0;
 		}
@@ -900,7 +899,7 @@ void UAnimSharingInstance::SetupState(FPerStateData& StateData, const FAnimation
 
 		bool bEnabled = AnimationSetup.Enabled.Default;
 #if WITH_EDITOR			
-		bEnabled = AnimationSetup.Enabled.GetValueForPlatformIdentifiers(PlatformName, PlatformName);
+		bEnabled = AnimationSetup.Enabled.GetValueForPlatform(PlatformName);
 #endif
 
 		/** Only create component if the setup is enabled for this platform and we have a valid animation asset */
@@ -908,7 +907,7 @@ void UAnimSharingInstance::SetupState(FPerStateData& StateData, const FAnimation
 		{
 			int32 NumRandomizedInstances = AnimationSetup.NumRandomizedInstances.Default;
 #if WITH_EDITOR			
-			NumRandomizedInstances = AnimationSetup.NumRandomizedInstances.GetValueForPlatformIdentifiers(PlatformName, PlatformName);
+			NumRandomizedInstances = AnimationSetup.NumRandomizedInstances.GetValueForPlatform(PlatformName);
 #endif
 			const uint32 NumInstances = StateEntry.bOnDemand ? NumInstancesPerSetup	: FGenericPlatformMath::Max(NumRandomizedInstances, 1);
 			for (uint32 InstanceIndex = 0; InstanceIndex < NumInstances; ++InstanceIndex)
diff --git a/Engine/Plugins/Editor/MobileLauncherProfileWizard/Source/MobileLauncherProfileWizard/Private/SProfileWizardUI.cpp b/Engine/Plugins/Editor/MobileLauncherProfileWizard/Source/MobileLauncherProfileWizard/Private/SProfileWizardUI.cpp
--- a/Engine/Plugins/Editor/MobileLauncherProfileWizard/Source/MobileLauncherProfileWizard/Private/SProfileWizardUI.cpp
+++ b/Engine/Plugins/Editor/MobileLauncherProfileWizard/Source/MobileLauncherProfileWizard/Private/SProfileWizardUI.cpp
@@ -739,8 +739,8 @@ void SProfileWizardUI::CacheCookFlavorsList()
 		
 	DLCFlavorList.Empty();
 
-	FVanillaPlatformEntry PlatformEntry = BuildPlatformHierarchy(TargetPlatformName, EPlatformFilter::CookFlavor);
-	for (const FPlatformInfo* PlatformFlaworInfo : PlatformEntry.PlatformFlavors)
+	const PlatformInfo::FTargetPlatformInfo* Info = FindPlatformInfo(TargetPlatformName);
+	for (const FTargetPlatformInfo* PlatformFlaworInfo : Info->Flavors)
 	{
 		DLCFlavorList.Add(MakeShareable(new FString(PlatformFlaworInfo->PlatformInfoName.ToString())));
 	}
diff --git a/Engine/Plugins/FX/Niagara/Source/NiagaraEditor/Private/Customizations/NiagaraPlatformSetCustomization.cpp b/Engine/Plugins/FX/Niagara/Source/NiagaraEditor/Private/Customizations/NiagaraPlatformSetCustomization.cpp
--- a/Engine/Plugins/FX/Niagara/Source/NiagaraEditor/Private/Customizations/NiagaraPlatformSetCustomization.cpp
+++ b/Engine/Plugins/FX/Niagara/Source/NiagaraEditor/Private/Customizations/NiagaraPlatformSetCustomization.cpp
@@ -743,9 +743,9 @@ TSharedRef<ITableRow> FNiagaraPlatformSetCustomization::OnGenerateDeviceProfileT
 	//Top level profile. Look for a platform icon.
 	if (InItem->Profile->Parent == nullptr)
 	{
-		if (const PlatformInfo::FPlatformInfo* Info = PlatformInfo::FindPlatformInfo(*InItem->Profile->DeviceType))
+		if (const PlatformInfo::FTargetPlatformInfo* Info = PlatformInfo::FindPlatformInfo(*InItem->Profile->DeviceType))
 		{
-			const FSlateBrush* DeviceProfileTypeIcon = FEditorStyle::GetBrush(Info->GetIconStyleName(PlatformInfo::EPlatformIconSize::Normal));
+			const FSlateBrush* DeviceProfileTypeIcon = FEditorStyle::GetBrush(Info->GetIconStyleName(EPlatformIconSize::Normal));
 			if (DeviceProfileTypeIcon != FEditorStyle::Get().GetDefaultBrush())
 			{
 				RowContainer->AddSlot()
diff --git a/Engine/Plugins/Media/MediaPlayerEditor/Source/MediaPlayerEditor/Private/Customizations/BaseMediaSourceCustomization.cpp b/Engine/Plugins/Media/MediaPlayerEditor/Source/MediaPlayerEditor/Private/Customizations/BaseMediaSourceCustomization.cpp
--- a/Engine/Plugins/Media/MediaPlayerEditor/Source/MediaPlayerEditor/Private/Customizations/BaseMediaSourceCustomization.cpp
+++ b/Engine/Plugins/Media/MediaPlayerEditor/Source/MediaPlayerEditor/Private/Customizations/BaseMediaSourceCustomization.cpp
@@ -133,19 +133,19 @@ TSharedRef<SWidget> FBaseMediaSourceCustomization::MakePlatformPlayerNamesValueW
 	}
 
 	// get available platforms
-	TArray<const PlatformInfo::FPlatformInfo*> AvailablePlatforms;
+	TArray<const PlatformInfo::FTargetPlatformInfo*> AvailablePlatforms;
 
-	for (const PlatformInfo::FPlatformInfo& PlatformInfo : PlatformInfo::GetPlatformInfoArray())
+	for (const PlatformInfo::FTargetPlatformInfo* PlatformInfo : PlatformInfo::GetVanillaPlatformInfoArray())
 	{
-		if (PlatformInfo.IsVanilla() && (PlatformInfo.PlatformType == EBuildTargetType::Game))
+		if (PlatformInfo->PlatformType == EBuildTargetType::Game)
 		{
-			AvailablePlatforms.Add(&PlatformInfo);
+			AvailablePlatforms.Add(PlatformInfo);
 		}
 	}
 
-	AvailablePlatforms.Sort([](const PlatformInfo::FPlatformInfo& One, const PlatformInfo::FPlatformInfo& Two) -> bool
+	Algo::Sort(AvailablePlatforms, [](const PlatformInfo::FTargetPlatformInfo* One, const PlatformInfo::FTargetPlatformInfo* Two) -> bool
 	{
-		return One.DisplayName.CompareTo(Two.DisplayName) < 0;
+		return One->DisplayName.CompareTo(Two->DisplayName) < 0;
 	});
 
 	// build value widget
@@ -153,14 +153,14 @@ TSharedRef<SWidget> FBaseMediaSourceCustomization::MakePlatformPlayerNamesValueW
 
 	for (int32 PlatformIndex = 0; PlatformIndex < AvailablePlatforms.Num(); ++PlatformIndex)
 	{
-		const PlatformInfo::FPlatformInfo* Platform = AvailablePlatforms[PlatformIndex];
+		const PlatformInfo::FTargetPlatformInfo* Platform = AvailablePlatforms[PlatformIndex];
 
 		// platform icon
 		PlatformPanel->AddSlot(0, PlatformIndex)
 			.VAlign(VAlign_Center)
 			[
 				SNew(SImage)
-					.Image(FEditorStyle::GetBrush(Platform->GetIconStyleName(PlatformInfo::EPlatformIconSize::Normal)))
+					.Image(FEditorStyle::GetBrush(Platform->GetIconStyleName(EPlatformIconSize::Normal)))
 			];
 
 		// platform name
@@ -180,13 +180,13 @@ TSharedRef<SWidget> FBaseMediaSourceCustomization::MakePlatformPlayerNamesValueW
 					.ButtonContent()
 					[
 						SNew(STextBlock)
-							.Text(this, &FBaseMediaSourceCustomization::HandlePlatformPlayersComboButtonText, Platform->IniPlatformName)
+							.Text(this, &FBaseMediaSourceCustomization::HandlePlatformPlayersComboButtonText, Platform->IniPlatformName.ToString())
 							.ToolTipText(LOCTEXT("PlatformPlayerComboButtonToolTipText", "Choose desired player for this platform"))
 					]
 					.ContentPadding(FMargin(6.0f, 2.0f))
 					.MenuContent()
 					[
-						MakePlatformPlayersMenu(Platform->IniPlatformName, PlayerFactories)
+						MakePlatformPlayersMenu(Platform->IniPlatformName.ToString(), PlayerFactories)
 					]
 			];
 	}
diff --git a/Engine/Plugins/Media/MediaPlayerEditor/Source/MediaPlayerEditor/Private/Customizations/PlatformMediaSourceCustomization.cpp b/Engine/Plugins/Media/MediaPlayerEditor/Source/MediaPlayerEditor/Private/Customizations/PlatformMediaSourceCustomization.cpp
--- a/Engine/Plugins/Media/MediaPlayerEditor/Source/MediaPlayerEditor/Private/Customizations/PlatformMediaSourceCustomization.cpp
+++ b/Engine/Plugins/Media/MediaPlayerEditor/Source/MediaPlayerEditor/Private/Customizations/PlatformMediaSourceCustomization.cpp
@@ -71,25 +71,25 @@ TSharedRef<SWidget> FPlatformMediaSourceCustomization::MakePlatformMediaSourcesV
 	const TArray<IMediaPlayerFactory*>& PlayerFactories = MediaModule->GetPlayerFactories();
 
 	// get available platforms
-	TArray<const PlatformInfo::FPlatformInfo*> AvailablePlatforms;
+	TArray<PlatformInfo::FTargetPlatformInfo*> AvailablePlatforms;
 
-	for (const PlatformInfo::FPlatformInfo& PlatformInfo : PlatformInfo::GetPlatformInfoArray())
+	for (PlatformInfo::FTargetPlatformInfo* PlatformInfo : PlatformInfo::GetVanillaPlatformInfoArray())
 	{
-		if (PlatformInfo.IsVanilla() && (PlatformInfo.PlatformType == EBuildTargetType::Game))
+		if (PlatformInfo->PlatformType == EBuildTargetType::Game)
 		{
-			if (PlatformInfo.PlatformInfoName == TEXT("TVOS"))
+			if (PlatformInfo->PlatformInfoName == TEXT("TVOS"))
 			{
 				continue; // tvOS is just iOS for now
 			}
 
-			AvailablePlatforms.Add(&PlatformInfo);
+			AvailablePlatforms.Add(PlatformInfo);
 		}
 	}
 
 	// sort available platforms alphabetically
-	AvailablePlatforms.Sort([](const PlatformInfo::FPlatformInfo& One, const PlatformInfo::FPlatformInfo& Two) -> bool
+	Algo::Sort(AvailablePlatforms, [](const PlatformInfo::FTargetPlatformInfo* One, const PlatformInfo::FTargetPlatformInfo* Two) -> bool
 	{
-		return One.DisplayName.CompareTo(Two.DisplayName) < 0;
+		return One->DisplayName.CompareTo(Two->DisplayName) < 0;
 	});
 
 	// build value widget
@@ -97,14 +97,14 @@ TSharedRef<SWidget> FPlatformMediaSourceCustomization::MakePlatformMediaSourcesV
 
 	for (int32 PlatformIndex = 0; PlatformIndex < AvailablePlatforms.Num(); ++PlatformIndex)
 	{
-		const PlatformInfo::FPlatformInfo* Platform = AvailablePlatforms[PlatformIndex];
+		const PlatformInfo::FTargetPlatformInfo* Platform = AvailablePlatforms[PlatformIndex];
 
 		// platform icon
 		PlatformPanel->AddSlot(0, PlatformIndex)
 			.VAlign(VAlign_Center)
 			[
 				SNew(SImage)
-					.Image(FEditorStyle::GetBrush(Platform->GetIconStyleName(PlatformInfo::EPlatformIconSize::Normal)))
+					.Image(FEditorStyle::GetBrush(Platform->GetIconStyleName(EPlatformIconSize::Normal)))
 			];
 
 		// platform name
@@ -123,8 +123,8 @@ TSharedRef<SWidget> FPlatformMediaSourceCustomization::MakePlatformMediaSourcesV
 				SNew(SObjectPropertyEntryBox)
 					.AllowedClass(UMediaSource::StaticClass())
 					.AllowClear(true)
-					.ObjectPath(this, &FPlatformMediaSourceCustomization::HandleMediaSourceEntryBoxObjectPath, Platform->IniPlatformName)
-					.OnObjectChanged(this, &FPlatformMediaSourceCustomization::HandleMediaSourceEntryBoxChanged, Platform->IniPlatformName)
+					.ObjectPath(this, &FPlatformMediaSourceCustomization::HandleMediaSourceEntryBoxObjectPath, Platform->IniPlatformName.ToString())
+					.OnObjectChanged(this, &FPlatformMediaSourceCustomization::HandleMediaSourceEntryBoxChanged, Platform->IniPlatformName.ToString())
 					.OnShouldFilterAsset(this, &FPlatformMediaSourceCustomization::HandleMediaSourceEntryBoxShouldFilterAsset)
 			];
 	}
diff --git a/Engine/Plugins/Runtime/AR/Microsoft/HoloLensAR/Source/HoloLensTargetPlatform/Private/HoloLensTargetPlatform.cpp b/Engine/Plugins/Runtime/AR/Microsoft/HoloLensAR/Source/HoloLensTargetPlatform/Private/HoloLensTargetPlatform.cpp
--- a/Engine/Plugins/Runtime/AR/Microsoft/HoloLensAR/Source/HoloLensTargetPlatform/Private/HoloLensTargetPlatform.cpp
+++ b/Engine/Plugins/Runtime/AR/Microsoft/HoloLensAR/Source/HoloLensTargetPlatform/Private/HoloLensTargetPlatform.cpp
@@ -140,7 +140,7 @@ void FHoloLensTargetPlatform::OnDeviceDetected(const FHoloLensDeviceInfo& Info)
 		FScopeLock Lock(&DevicesLock);
 		Devices.Add(NewDevice);
 	}
-	DeviceDiscoveredEvent.Broadcast(NewDevice.ToSharedRef());
+	OnDeviceDiscovered().Broadcast(NewDevice.ToSharedRef());
 }
 
 bool FHoloLensTargetPlatform::SupportsBuildTarget(EBuildTargetType BuildTarget) const
diff --git a/Engine/Plugins/Runtime/AR/Microsoft/HoloLensAR/Source/HoloLensTargetPlatform/Private/HoloLensTargetPlatform.h b/Engine/Plugins/Runtime/AR/Microsoft/HoloLensAR/Source/HoloLensTargetPlatform/Private/HoloLensTargetPlatform.h
--- a/Engine/Plugins/Runtime/AR/Microsoft/HoloLensAR/Source/HoloLensTargetPlatform/Private/HoloLensTargetPlatform.h
+++ b/Engine/Plugins/Runtime/AR/Microsoft/HoloLensAR/Source/HoloLensTargetPlatform/Private/HoloLensTargetPlatform.h
@@ -104,18 +104,6 @@ public:
 
 #endif //WITH_ENGINE
 
-	DECLARE_DERIVED_EVENT(FHoloLensTargetPlatform, ITargetPlatform::FOnTargetDeviceDiscovered, FOnTargetDeviceDiscovered);
-	virtual FOnTargetDeviceDiscovered& OnDeviceDiscovered( ) override
-	{
-		return DeviceDiscoveredEvent;
-	}
-
-	DECLARE_DERIVED_EVENT(FHoloLensTargetPlatform, ITargetPlatform::FOnTargetDeviceLost, FOnTargetDeviceLost);
-	virtual FOnTargetDeviceLost& OnDeviceLost( ) override
-	{
-		return DeviceLostEvent;
-	}
-
 	virtual EPlatformAuthentication RequiresUserCredentials() const override
 	{
 		return EPlatformAuthentication::Possible;
@@ -172,13 +160,6 @@ private:
 	FStaticMeshLODSettings StaticMeshLODSettings;
 #endif //WITH_ENGINE
 
-private:
-
-	// Holds an event delegate that is executed when a new target device has been discovered.
-	FOnTargetDeviceDiscovered DeviceDiscoveredEvent;
-
-	// Holds an event delegate that is executed when a target device has been lost, i.e. disconnected or timed out.
-	FOnTargetDeviceLost DeviceLostEvent;
 };
 
 
diff --git a/Engine/Source/Developer/Android/AndroidTargetPlatform/Private/AndroidTargetPlatform.cpp b/Engine/Source/Developer/Android/AndroidTargetPlatform/Private/AndroidTargetPlatform.cpp
--- a/Engine/Source/Developer/Android/AndroidTargetPlatform/Private/AndroidTargetPlatform.cpp
+++ b/Engine/Source/Developer/Android/AndroidTargetPlatform/Private/AndroidTargetPlatform.cpp
@@ -755,7 +755,7 @@ bool FAndroidTargetPlatform::HandleTicker( float DeltaTime )
 				TestDevice->SetConnected(false);
 				Devices.Remove(DeviceIt.Key());
 
-				DeviceLostEvent.Broadcast(TestDevice.ToSharedRef());
+				OnDeviceLost().Broadcast(TestDevice.ToSharedRef());
 			}
 
 			// check if this platform is supported by the extensions and version
@@ -775,7 +775,7 @@ bool FAndroidTargetPlatform::HandleTicker( float DeltaTime )
 			Device->SetAuthorized(DeviceInfo.bAuthorizedDevice);
 			Device->SetVersions(DeviceInfo.SDKVersion, DeviceInfo.HumanAndroidVersion);
 
-			DeviceDiscoveredEvent.Broadcast(Device.ToSharedRef());
+			OnDeviceDiscovered().Broadcast(Device.ToSharedRef());
 		}
 	}
 
@@ -789,7 +789,7 @@ bool FAndroidTargetPlatform::HandleTicker( float DeltaTime )
 
 			Iter.RemoveCurrent();
 
-			DeviceLostEvent.Broadcast(RemovedDevice.ToSharedRef());
+			OnDeviceLost().Broadcast(RemovedDevice.ToSharedRef());
 		}
 	}
 
diff --git a/Engine/Source/Developer/Android/AndroidTargetPlatform/Private/AndroidTargetPlatform.h b/Engine/Source/Developer/Android/AndroidTargetPlatform/Private/AndroidTargetPlatform.h
--- a/Engine/Source/Developer/Android/AndroidTargetPlatform/Private/AndroidTargetPlatform.h
+++ b/Engine/Source/Developer/Android/AndroidTargetPlatform/Private/AndroidTargetPlatform.h
@@ -226,18 +226,6 @@ public:
 		InStringKeys.Add(TEXT("NDKAPILevel"));
 	}
 
-	DECLARE_DERIVED_EVENT(FAndroidTargetPlatform, ITargetPlatform::FOnTargetDeviceDiscovered, FOnTargetDeviceDiscovered);
-	virtual FOnTargetDeviceDiscovered& OnDeviceDiscovered( ) override
-	{
-		return DeviceDiscoveredEvent;
-	}
-
-	DECLARE_DERIVED_EVENT(FAndroidTargetPlatform, ITargetPlatform::FOnTargetDeviceLost, FOnTargetDeviceLost);
-	virtual FOnTargetDeviceLost& OnDeviceLost( ) override
-	{
-		return DeviceLostEvent;
-	}
-
 	virtual bool ShouldExpandTo32Bit(const uint16* Indices, const int32 NumIndices) const override;
 	//~ End ITargetPlatform Interface
 
@@ -311,12 +299,6 @@ protected:
 
 	ITargetDevicePtr DefaultDevice;
 #endif //WITH_ENGINE
-
-	// Holds an event delegate that is executed when a new target device has been discovered.
-	FOnTargetDeviceDiscovered DeviceDiscoveredEvent;
-
-	// Holds an event delegate that is executed when a target device has been lost, i.e. disconnected or timed out.
-	FOnTargetDeviceLost DeviceLostEvent;
 };
 
 
diff --git a/Engine/Source/Developer/BlueprintNativeCodeGen/Private/BlueprintNativeCodeGenManifest.cpp b/Engine/Source/Developer/BlueprintNativeCodeGen/Private/BlueprintNativeCodeGenManifest.cpp
--- a/Engine/Source/Developer/BlueprintNativeCodeGen/Private/BlueprintNativeCodeGenManifest.cpp
+++ b/Engine/Source/Developer/BlueprintNativeCodeGen/Private/BlueprintNativeCodeGenManifest.cpp
@@ -319,19 +319,8 @@ FBlueprintNativeCodeGenPaths FBlueprintNativeCodeGenPaths::GetDefaultCodeGenPath
 	FString DefaultPluginPath = FPaths::Combine(*FPaths::ProjectIntermediateDir(), TEXT("Plugins"), *DefaultPluginName);
 	if (!PlatformName.IsNone())
 	{
-		for (const PlatformInfo::FPlatformInfo& PlatformInfo : PlatformInfo::GetPlatformInfoArray())
-		{
-			if (PlatformInfo.TargetPlatformName == PlatformName)
-			{
-				static const FName UBTTargetId_Win32 = FName(TEXT("Win32"));
-				static const FName UBTTargetId_Win64 = FName(TEXT("Win64"));
-				static const FName UBTTargetId_Windows = FName(TEXT("Windows"));
-
-				const FName UBTTargetId = (PlatformInfo.UBTTargetId == UBTTargetId_Win32 || PlatformInfo.UBTTargetId == UBTTargetId_Win64) ? UBTTargetId_Windows : PlatformInfo.UBTTargetId;
-				DefaultPluginPath = FPaths::Combine(*DefaultPluginPath, *LexToString(UBTTargetId), LexToString(PlatformInfo.PlatformType));
-				break;
-			}
-		}
+		const PlatformInfo::FTargetPlatformInfo* PlatformInfo = PlatformInfo::FindPlatformInfo(PlatformName);
+		DefaultPluginPath = FPaths::Combine(*DefaultPluginPath, *LexToString(PlatformInfo->IniPlatformName), LexToString(PlatformInfo->PlatformType));
 	}
 
 	return FBlueprintNativeCodeGenPaths(DefaultPluginName, DefaultPluginPath, PlatformName);
diff --git a/Engine/Source/Developer/BlueprintNativeCodeGen/Private/BlueprintNativeCodeGenUtils.cpp b/Engine/Source/Developer/BlueprintNativeCodeGen/Private/BlueprintNativeCodeGenUtils.cpp
--- a/Engine/Source/Developer/BlueprintNativeCodeGen/Private/BlueprintNativeCodeGenUtils.cpp
+++ b/Engine/Source/Developer/BlueprintNativeCodeGen/Private/BlueprintNativeCodeGenUtils.cpp
@@ -123,23 +123,24 @@ static bool BlueprintNativeCodeGenUtilsImpl::GeneratePluginDescFile(const FBluep
 		ModuleDesc->LoadingPhase = ELoadingPhase::Default;
 
 		const FName PlatformName = TargetPaths.GetTargetPlatformName();
-		for (const PlatformInfo::FPlatformInfo& PlatformInfo : PlatformInfo::GetPlatformInfoArray())
+		for (const PlatformInfo::FTargetPlatformInfo* PlatformInfo : PlatformInfo::GetPlatformInfoArray())
 		{
-			if (PlatformInfo.TargetPlatformName == PlatformName)
+			if (PlatformInfo->TargetPlatformName == PlatformName)
 			{
-				// We use the 'UBTTargetId' because this white-list expects the 
+				// We use the 'UBTPlatformName' because this white-list expects the 
 				// string to correspond to UBT's UnrealTargetPlatform enum (and by proxy, FPlatformMisc::GetUBTPlatform)
-				ModuleDesc->WhitelistPlatforms.AddUnique(PlatformInfo.UBTTargetId.ToString());
+				ModuleDesc->WhitelistPlatforms.AddUnique(PlatformInfo->DataDrivenPlatformInfo->UBTPlatformString);
 
+				FName UBTPlatformName = PlatformInfo->DataDrivenPlatformInfo->UBTPlatformName;
 				// Hack to allow clients for PS4/XboxOne (etc.) to build the nativized assets plugin
-				const bool bIsClientValidForPlatform = PlatformInfo.UBTTargetId == TEXT("Win32") ||
-					PlatformInfo.UBTTargetId == TEXT("Win64") ||
-					PlatformInfo.UBTTargetId == TEXT("Linux") ||
-					PlatformInfo.UBTTargetId == TEXT("LinuxAArch64") ||
-					PlatformInfo.UBTTargetId == TEXT("Mac");
+				const bool bIsClientValidForPlatform = UBTPlatformName == TEXT("Win32") ||
+					UBTPlatformName == TEXT("Win64") ||
+					UBTPlatformName == TEXT("Linux") ||
+					UBTPlatformName == TEXT("LinuxAArch64") ||
+					UBTPlatformName == TEXT("Mac");
 
 				// should correspond to UnrealBuildTool::TargetType in TargetRules.cs
-				switch (PlatformInfo.PlatformType)
+				switch (PlatformInfo->PlatformType)
 				{
 				case EBuildTargetType::Game:
 					ModuleDesc->WhitelistTargets.AddUnique(EBuildTargetType::Game);
@@ -161,7 +162,7 @@ static bool BlueprintNativeCodeGenUtilsImpl::GeneratePluginDescFile(const FBluep
 					break;
 
 				case EBuildTargetType::Editor:
-					ensureMsgf(PlatformInfo.PlatformType != EBuildTargetType::Editor, TEXT("Nativized Blueprint plugin is for cooked projects only - it isn't supported in editor builds."));
+					ensureMsgf(PlatformInfo->PlatformType != EBuildTargetType::Editor, TEXT("Nativized Blueprint plugin is for cooked projects only - it isn't supported in editor builds."));
 					break;
 				};				
 			}
diff --git a/Engine/Source/Developer/DesktopPlatform/Private/PlatformInfo.cpp b/Engine/Source/Developer/DesktopPlatform/Private/PlatformInfo.cpp
--- a/Engine/Source/Developer/DesktopPlatform/Private/PlatformInfo.cpp
+++ b/Engine/Source/Developer/DesktopPlatform/Private/PlatformInfo.cpp
@@ -10,6 +10,8 @@
 
 #define LOCTEXT_NAMESPACE "PlatformInfo"
 
+#if DDPI_HAS_EXTENDED_PLATFORMINFO_DATA
+
 namespace PlatformInfo
 {
 	TArray<FName> AllPlatformGroupNames;
@@ -18,71 +20,81 @@ namespace PlatformInfo
 namespace
 {
 
-TArray<FPlatformInfo> AllPlatformInfoArray;
+TArray<FTargetPlatformInfo*> AllPlatformInfoArray;
+TArray<FTargetPlatformInfo*> VanillaPlatformInfoArray;
+
 
 // we don't need any of this without the editor, although we would ideally not even compile this outside of the editor
 // @todo platplug: Figure out why this is compiled on target devices
-#if WITH_EDITOR || IS_PROGRAM
 
-void BuildPlatformInfo(const FName& InPlatformInfoName, const FName& InTargetPlatformName, const FText& InDisplayName, const EBuildTargetType InPlatformType, const EPlatformFlags::Flags InPlatformFlags, const FPlatformIconPaths& InIconPaths, const FString& InUATCommandLine, const FString& InAutoSDKPath, EPlatformSDKStatus InStatus, const FString& InTutorial, bool InEnabled, FString InBinaryFolderName, FString InIniPlatformName, bool InUsesHostCompiler, bool InUATClosesAfterLaunch, bool InIsConfidential, const FName& InUBTTargetId, const FName& InPlatformGroupName, const FName& InPlatformSubMenu, bool InTargetPlatformCanUseCrashReporter)
+void BuildPlatformInfo(FName InPlatformInfoName, FName InTargetPlatformName, const FText& InDisplayName, const EBuildTargetType InPlatformType, 
+	const EPlatformFlags::Flags InPlatformFlags, FName InIniPlatformName, FName InUBTPlatformName)
 {
-	FPlatformInfo& PlatformInfo = AllPlatformInfoArray.Emplace_GetRef();
-
-	PlatformInfo.PlatformInfoName = InPlatformInfoName;
-	PlatformInfo.TargetPlatformName = InTargetPlatformName;
-
-	// See if this name also contains a flavor
-	const FString InPlatformInfoNameString = InPlatformInfoName.ToString();
+	// some verification before we create the object
+	FName VanillaName = NAME_None;
+	FName FlavorName = NAME_None;
+	FTargetPlatformInfo* VanillaInfo = nullptr;
+	// See if this contains a flavor
+	if (InPlatformFlags & EPlatformFlags::CookFlavor)
 	{
+		FString InPlatformInfoNameString = InPlatformInfoName.ToString();
 		int32 UnderscoreLoc;
-		if(InPlatformInfoNameString.FindChar(TEXT('_'), UnderscoreLoc))
+		if (InPlatformInfoNameString.FindChar(TEXT('_'), UnderscoreLoc))
 		{
-			PlatformInfo.VanillaPlatformName = *InPlatformInfoNameString.Mid(0, UnderscoreLoc);
-			PlatformInfo.PlatformFlavor = *InPlatformInfoNameString.Mid(UnderscoreLoc + 1);
+			// removing VanillaPlatformName from the TP info means we need to verify, as it assumed the vanilla platform name is 
+			// if this assumption breaks, then we can store a private copy of this name to find the vanilla platform later
+			VanillaName = *InPlatformInfoNameString.Mid(0, UnderscoreLoc);
+			FlavorName = *InPlatformInfoNameString.Mid(UnderscoreLoc + 1);
+
+			FTargetPlatformInfo** FoundInfo = AllPlatformInfoArray.FindByPredicate([VanillaName](const FTargetPlatformInfo* Item) -> bool
+			{
+				return Item->PlatformInfoName == VanillaName;
+			});
+
+			VanillaInfo = FoundInfo ? *FoundInfo : nullptr;
 		}
-		else
+
+		// make sure it was good
+		if (VanillaInfo == nullptr)
 		{
-			PlatformInfo.VanillaPlatformName = InPlatformInfoName;
+			UE_LOG(LogDesktopPlatform, Error, TEXT("TargetPlatformInfo %s is a flavor, but wasn't in the form Parent_Flavor (Parent needs to be specified before flavor in the .ini file)"), *InPlatformInfoName.ToString());
+			return;
 		}
 	}
 
-	if (PlatformInfo.VanillaPlatformName != NAME_None)
+	// create the platform info
+	FTargetPlatformInfo* PlatformInfo = new FTargetPlatformInfo();
+	AllPlatformInfoArray.Add(PlatformInfo);
+
+	PlatformInfo->PlatformInfoName = InPlatformInfoName;
+	PlatformInfo->TargetPlatformName = InTargetPlatformName;
+	PlatformInfo->DisplayName = InDisplayName;
+	PlatformInfo->PlatformType = InPlatformType;
+	PlatformInfo->PlatformFlags = InPlatformFlags;
+	PlatformInfo->IniPlatformName = InIniPlatformName;
+	PlatformInfo->PlatformFlavor = FlavorName;
+
+	// add this object to either a parent, or the list of vanilla platforms
+	if (VanillaInfo != nullptr)
 	{
-		PlatformInfo::AllVanillaPlatformNames.AddUnique(PlatformInfo.VanillaPlatformName);
+		PlatformInfo->VanillaInfo = VanillaInfo;
+		VanillaInfo->Flavors.Add(PlatformInfo);
 	}
-
-	PlatformInfo.DisplayName = InDisplayName;
-	PlatformInfo.PlatformType = InPlatformType;
-	PlatformInfo.PlatformFlags = InPlatformFlags;
-	PlatformInfo.IconPaths = InIconPaths;
-	PlatformInfo.UATCommandLine = InUATCommandLine;
-	PlatformInfo.AutoSDKPath = InAutoSDKPath;
-	PlatformInfo.BinaryFolderName = InBinaryFolderName;
-	PlatformInfo.IniPlatformName = InIniPlatformName;
-	PlatformInfo.UBTTargetId = InUBTTargetId;
-	PlatformInfo.PlatformGroupName = InPlatformGroupName;
-	PlatformInfo.PlatformSubMenu = InPlatformSubMenu;
-
-	if (InPlatformGroupName != NAME_None)
+	else
 	{
-		PlatformInfo::AllPlatformGroupNames.AddUnique(InPlatformGroupName);
+		// if we are vanilla, then point to ourself (this way we can always just use VanillaInfo without checking for null)
+		PlatformInfo->VanillaInfo = PlatformInfo;
+
+		VanillaPlatformInfoArray.Add(PlatformInfo);
+		PlatformInfo::AllVanillaPlatformNames.AddUnique(InPlatformInfoName);
 	}
 
-	// Generate the icon style names for FEditorStyle
-	PlatformInfo.IconPaths.NormalStyleName = *FString::Printf(TEXT("Launcher.Platform_%s"), *InPlatformInfoNameString);
-	PlatformInfo.IconPaths.LargeStyleName  = *FString::Printf(TEXT("Launcher.Platform_%s.Large"), *InPlatformInfoNameString);
-	PlatformInfo.IconPaths.XLargeStyleName = *FString::Printf(TEXT("Launcher.Platform_%s.XLarge"), *InPlatformInfoNameString);
-
-	// SDK data
-	PlatformInfo.SDKStatus = InStatus;
-	PlatformInfo.SDKTutorial = InTutorial;
-
-	// Distribution data
-	PlatformInfo.bEnabledForUse = InEnabled;
-	PlatformInfo.bUsesHostCompiler = InUsesHostCompiler;
-	PlatformInfo.bUATClosesAfterLaunch = InUATClosesAfterLaunch;
-	PlatformInfo.bIsConfidential = InIsConfidential;
-	PlatformInfo.bTargetPlatformCanUseCrashReporter = InTargetPlatformCanUseCrashReporter;
+	// build up used platform group names
+	PlatformInfo->DataDrivenPlatformInfo = &FDataDrivenPlatformInfoRegistry::GetPlatformInfo(InIniPlatformName.ToString());
+	if (PlatformInfo->DataDrivenPlatformInfo->PlatformGroupName != NAME_None)
+	{
+		PlatformInfo::AllPlatformGroupNames.AddUnique(PlatformInfo->DataDrivenPlatformInfo->PlatformGroupName);
+	}
 }
 
 // Gets a string from a section, or empty string if it didn't exist
@@ -116,51 +128,20 @@ namespace EPlatformSection
 	const FName DisplayName = FName(TEXT("DisplayName"));
 	const FName PlatformType = FName(TEXT("PlatformType"));
 	const FName PlatformFlags = FName(TEXT("PlatformFlags"));
-	const FName NormalIconPath = FName(TEXT("NormalIconPath"));
-	const FName LargeIconPath = FName(TEXT("LargeIconPath"));
-	const FName XLargeIconPath = FName(TEXT("XLargeIconPath"));
 	const FName UATCommandLine = FName(TEXT("UATCommandLine"));
-	const FName AutoSDKPath = FName(TEXT("AutoSDKPath"));
-	const FName TutorialPath = FName(TEXT("TutorialPath"));
-	const FName bIsEnabled = FName(TEXT("bIsEnabled"));
-	const FName BinariesDirectoryName = FName(TEXT("BinariesDirectoryName"));
-	const FName IniPlatformName = FName(TEXT("IniPlatformName"));
-	const FName bUsesHostCompiler = FName(TEXT("bUsesHostCompiler"));
-	const FName bUATClosesAfterLaunch = FName(TEXT("bUATClosesAfterLaunch"));
-	const FName bIsConfidential = FName(TEXT("bIsConfidential"));
-	const FName UBTTargetID = FName(TEXT("UBTTargetID"));
-	const FName PlatformGroupName = FName(TEXT("PlatformGroupName"));
-	const FName PlatformSubMenu = FName(TEXT("PlatformSubMenu"));
+	const FName UBTPlatformName = FName(TEXT("UBTPlatformName"));
 }
 
-void ParseDataDrivenPlatformInfo(const TCHAR* Name, const FConfigSection& Section)
+void ParseDataDrivenPlatformInfo(const TCHAR* Name, const FConfigSection& Section, FName IniPlatformName)
 {
 	FName TargetPlatformName = *GetSectionString(Section, EPlatformSection::TargetPlatformName);
 	FString DisplayName = GetSectionString(Section, EPlatformSection::DisplayName);
 	FString PlatformType = GetSectionString(Section, EPlatformSection::PlatformType);
 	FString PlatformFlags = GetSectionString(Section, EPlatformSection::PlatformFlags);
-	FString NormalIconPath = GetSectionString(Section, EPlatformSection::NormalIconPath);
-	FString LargeIconPath = GetSectionString(Section, EPlatformSection::LargeIconPath);
-	FString XLargeIconPath = GetSectionString(Section, EPlatformSection::XLargeIconPath);
-	// no one has an XLarge path yet, but in case they add one, this will use it
-	if (XLargeIconPath == TEXT(""))
-	{
-		XLargeIconPath = LargeIconPath;
-	}
 	FString UATCommandLine = GetSectionString(Section, EPlatformSection::UATCommandLine);
-	FString AutoSDKPath = GetSectionString(Section, EPlatformSection::AutoSDKPath);
-	FString TutorialPath = GetSectionString(Section, EPlatformSection::TutorialPath);
-	bool bIsEnabled = GetSectionBool(Section, EPlatformSection::bIsEnabled);
-	FString BinariesDirectoryName = GetSectionString(Section, EPlatformSection::BinariesDirectoryName);
-	FString IniPlatformName = GetSectionString(Section, EPlatformSection::IniPlatformName);
-	bool bUsesHostCompiler = GetSectionBool(Section, EPlatformSection::bUsesHostCompiler);
-	bool bUATClosesAfterLaunch = GetSectionBool(Section, EPlatformSection::bUATClosesAfterLaunch);
-	bool bIsConfidential = GetSectionBool(Section, EPlatformSection::bIsConfidential);
-	FName UBTTargetID = *GetSectionString(Section, EPlatformSection::UBTTargetID);
-	FName PlatformGroupName = *GetSectionString(Section, EPlatformSection::PlatformGroupName);
-	FName PlatformSubMenu = *GetSectionString(Section, EPlatformSection::PlatformSubMenu);
-	FString TargetPlatformCanUseCrashReporterString = GetSectionString(Section, TEXT("bTargetPlatformCanUseCrashReporter"));
-	bool bTargetPlatformCanUseCrashReporter = TargetPlatformCanUseCrashReporterString.IsEmpty() ? true : GetSectionBool(Section, TEXT("bTargetPlatformCanUseCrashReporter"));
+	FString UBTTargetValue = GetSectionString(Section, EPlatformSection::UBTPlatformName);
+	// in almost all cases the UBT platform name matches the IniPlatformName, so it's optional
+	FName UBTPlatformName = UBTTargetValue.Len() == 0 ? IniPlatformName : FName(*UBTTargetValue);
 
 	EBuildTargetType TargetType;
 	if (!LexTryParseString(TargetType, *PlatformType))
@@ -168,8 +149,7 @@ void ParseDataDrivenPlatformInfo(const TCHAR* Name, const FConfigSection& Sectio
 		TargetType = EBuildTargetType::Unknown;
 	}
 
-	BuildPlatformInfo(Name, TargetPlatformName, FText::FromString(DisplayName), TargetType, ConvertPlatformFlags(PlatformFlags), FPlatformIconPaths(NormalIconPath, LargeIconPath, XLargeIconPath), UATCommandLine,
-		AutoSDKPath, PlatformInfo::EPlatformSDKStatus::Unknown, TutorialPath, bIsEnabled, BinariesDirectoryName, IniPlatformName, bUsesHostCompiler, bUATClosesAfterLaunch, bIsConfidential, UBTTargetID, PlatformGroupName, PlatformSubMenu, bTargetPlatformCanUseCrashReporter);
+	BuildPlatformInfo(Name, TargetPlatformName, FText::FromString(DisplayName), TargetType, ConvertPlatformFlags(PlatformFlags), IniPlatformName, UBTPlatformName);
 }
 
 void LoadDataDrivenPlatforms()
@@ -183,13 +163,15 @@ void LoadDataDrivenPlatforms()
 
 		FDataDrivenPlatformInfoRegistry::LoadDataDrivenIniFile(Index, IniFile, PlatformName);
 
+		FName PlatformFName = *PlatformName;
+
 		// now walk over the file, looking for ShaderPlatformInfo sections
 		for (auto Section : IniFile)
 		{
 			if (Section.Key.StartsWith(TEXT("PlatformInfo ")))
 			{
 				const FString& SectionName = Section.Key;
-				ParseDataDrivenPlatformInfo(*SectionName.Mid(13), Section.Value);
+				ParseDataDrivenPlatformInfo(*SectionName.Mid(13), Section.Value, PlatformFName);
 			}
 		}
 	}
@@ -200,17 +182,16 @@ FDelayedAutoRegisterHelper GPlatformInfoInit(EDelayedRegisterRunPhase::FileSyste
 	LoadDataDrivenPlatforms();
 });
 
-#endif
 
 } // anonymous namespace
 
-const FPlatformInfo* FindPlatformInfo(const FName& InPlatformName)
+const FTargetPlatformInfo* FindPlatformInfo(const FName& InPlatformName)
 {
-	for(const FPlatformInfo& PlatformInfo : AllPlatformInfoArray)
+	for(const FTargetPlatformInfo* PlatformInfo : AllPlatformInfoArray)
 	{
-		if(PlatformInfo.PlatformInfoName == InPlatformName)
+		if(PlatformInfo->PlatformInfoName == InPlatformName)
 		{
-			return &PlatformInfo;
+			return PlatformInfo;
 		}
 	}
 
@@ -219,107 +200,31 @@ const FPlatformInfo* FindPlatformInfo(const FName& InPlatformName)
 	return nullptr;
 }
 
-const FPlatformInfo* FindVanillaPlatformInfo(const FName& InPlatformName)
+const FTargetPlatformInfo* FindVanillaPlatformInfo(const FName& InPlatformName)
 {
-	const FPlatformInfo* const FoundInfo = FindPlatformInfo(InPlatformName);
-	return (FoundInfo) ? (FoundInfo->IsVanilla()) ? FoundInfo : FindPlatformInfo(FoundInfo->VanillaPlatformName) : nullptr;
-}
-
-void UpdatePlatformSDKStatus(FString InPlatformName, EPlatformSDKStatus InStatus)
-{
-	for(const FPlatformInfo& PlatformInfo : AllPlatformInfoArray)
-	{
-		if(PlatformInfo.VanillaPlatformName == FName(*InPlatformName))
-		{
-			const_cast<FPlatformInfo&>(PlatformInfo).SDKStatus = InStatus;
-		}
-	}
+	const FTargetPlatformInfo* const FoundInfo = FindPlatformInfo(InPlatformName);
+	return FoundInfo ? FoundInfo->VanillaInfo : nullptr;
 }
 
 void UpdatePlatformDisplayName(FString InPlatformName, FText InDisplayName)
 {
-	for (const FPlatformInfo& PlatformInfo : AllPlatformInfoArray)
+	for (FTargetPlatformInfo* PlatformInfo : AllPlatformInfoArray)
 	{
-		if (PlatformInfo.TargetPlatformName == FName(*InPlatformName))
+		if (PlatformInfo->TargetPlatformName == FName(*InPlatformName))
 		{
-			const_cast<FPlatformInfo&>(PlatformInfo).DisplayName = InDisplayName;
+			PlatformInfo->DisplayName = InDisplayName;
 		}
 	}
 }
 
-const TArray<FPlatformInfo>& GetPlatformInfoArray()
+const TArray<FTargetPlatformInfo*>& GetPlatformInfoArray()
 {
 	return AllPlatformInfoArray;
 }
 
-TArray<FVanillaPlatformEntry> BuildPlatformHierarchy(const EPlatformFilter InFilter)
+const TArray<FTargetPlatformInfo*>& GetVanillaPlatformInfoArray()
 {
-	TArray<FVanillaPlatformEntry> VanillaPlatforms;
-
-	// Build up a tree from the platforms we support (vanilla outers, with a list of flavors)
-	// PlatformInfoArray should be ordered in such a way that the vanilla platforms always appear before their flavors
-	for (const PlatformInfo::FPlatformInfo& PlatformInfo : AllPlatformInfoArray)
-	{
-		if(PlatformInfo.IsVanilla())
-		{
-			VanillaPlatforms.Add(FVanillaPlatformEntry(&PlatformInfo));
-		}
-		else
-		{
-			const bool bHasBuildFlavor = !!(PlatformInfo.PlatformFlags & EPlatformFlags::BuildFlavor);
-			const bool bHasCookFlavor  = !!(PlatformInfo.PlatformFlags & EPlatformFlags::CookFlavor);
-			
-			const bool bValidFlavor = 
-				InFilter == EPlatformFilter::All || 
-				(InFilter == EPlatformFilter::BuildFlavor && bHasBuildFlavor) || 
-				(InFilter == EPlatformFilter::CookFlavor && bHasCookFlavor);
-
-			if(bValidFlavor)
-			{
-				const FName VanillaPlatformName = PlatformInfo.VanillaPlatformName;
-				FVanillaPlatformEntry* const VanillaEntry = VanillaPlatforms.FindByPredicate([VanillaPlatformName](const FVanillaPlatformEntry& Item) -> bool
-				{
-					return Item.PlatformInfo->PlatformInfoName == VanillaPlatformName;
-				});
-				check(VanillaEntry);
-				VanillaEntry->PlatformFlavors.Add(&PlatformInfo);
-			}
-		}
-	}
-
-	return VanillaPlatforms;
-}
-
-FVanillaPlatformEntry BuildPlatformHierarchy(const FName& InPlatformName, const EPlatformFilter InFilter)
-{
-	FVanillaPlatformEntry VanillaPlatformEntry;
-	const FPlatformInfo* VanillaPlatformInfo = FindVanillaPlatformInfo(InPlatformName);
-	
-	if (VanillaPlatformInfo)
-	{
-		VanillaPlatformEntry.PlatformInfo = VanillaPlatformInfo;
-		
-		for (const PlatformInfo::FPlatformInfo& PlatformInfo : AllPlatformInfoArray)
-		{
-			if (!PlatformInfo.IsVanilla() && PlatformInfo.VanillaPlatformName == VanillaPlatformInfo->PlatformInfoName)
-			{
-				const bool bHasBuildFlavor = !!(PlatformInfo.PlatformFlags & EPlatformFlags::BuildFlavor);
-				const bool bHasCookFlavor = !!(PlatformInfo.PlatformFlags & EPlatformFlags::CookFlavor);
-
-				const bool bValidFlavor =
-					InFilter == EPlatformFilter::All ||
-					(InFilter == EPlatformFilter::BuildFlavor && bHasBuildFlavor) ||
-					(InFilter == EPlatformFilter::CookFlavor && bHasCookFlavor);
-
-				if (bValidFlavor)
-				{
-					VanillaPlatformEntry.PlatformFlavors.Add(&PlatformInfo);
-				}
-			}
-		}
-	}
-	
-	return VanillaPlatformEntry;
+	return VanillaPlatformInfoArray;
 }
 
 const TArray<FName>& GetAllPlatformGroupNames()
@@ -334,4 +239,7 @@ const TArray<FName>& GetAllVanillaPlatformNames()
 
 } // namespace PlatformInfo
 
+#endif
+
+
 #undef LOCTEXT_NAMESPACE
diff --git a/Engine/Source/Developer/DesktopPlatform/Public/PlatformInfo.h b/Engine/Source/Developer/DesktopPlatform/Public/PlatformInfo.h
--- a/Engine/Source/Developer/DesktopPlatform/Public/PlatformInfo.h
+++ b/Engine/Source/Developer/DesktopPlatform/Public/PlatformInfo.h
@@ -3,35 +3,12 @@
 #pragma once
 
 #include "CoreMinimal.h"
+#include "Misc/DataDrivenPlatformInfoRegistry.h"
+
+#if DDPI_HAS_EXTENDED_PLATFORMINFO_DATA
 
 namespace PlatformInfo
 {
-	/** Available icon sizes (see FPlatformIconPaths) */
-	enum class EPlatformIconSize : uint8
-	{
-		/** Normal sized icon (24x24) */
-		Normal,
-
-		/** Large sized icon (64x64) */
-		Large,
-
-		/** Extra large sized icon (128x128) */
-		XLarge,
-	};
-
-	/** Flavor types used when filtering the platforms based upon their flags */
-	enum class EPlatformFilter : uint8
-	{
-		/** Include all platform types */
-		All,
-
-		/** Include only build flavors */
-		BuildFlavor,
-
-		/** Include only cook flavors */
-		CookFlavor,
-	};
-
 	/** Flags describing platform variants */
 	namespace EPlatformFlags
 	{
@@ -42,65 +19,29 @@ namespace PlatformInfo
 			None = 0,
 
 			/** The flavor generates different output when building (eg, 32 or 64-bit) */
-			BuildFlavor = 1<<0,
+			BuildFlavor = 1 << 0,
 
 			/** The flavor generates different output when cooking (eg, ETC2 or ASTC texture format) */
-			CookFlavor = 1<<1,
+			CookFlavor = 1 << 1,
 		};
 	}
 
-	enum class EPlatformSDKStatus : uint8
+	/** Flavor types used when filtering the platforms based upon their flags */
+	enum class EPlatformFilter : uint8
 	{
-		/** SDK status is unknown */
-		Unknown,
+		/** Include all platform types */
+		All,
 
-		/** SDK is installed */
-		Installed,
+		/** Include only build flavors */
+		BuildFlavor,
 
-		/** SDK is not installed */
-		NotInstalled,
+		/** Include only cook flavors */
+		CookFlavor,
 	};
 
-	/** Information about where to find the platform icons (for use by FEditorStyle) */
-	struct FPlatformIconPaths
-	{
-		FPlatformIconPaths()
-		{
-		}
-
-		FPlatformIconPaths(const FString& InIconPath)
-			: NormalPath(InIconPath)
-			, LargePath(InIconPath)
-			, XLargePath(InIconPath)
-		{
-		}
-
-		FPlatformIconPaths(const FString& InNormalPath, const FString& InLargePath)
-			: NormalPath(InNormalPath)
-			, LargePath(InLargePath)
-			, XLargePath(InLargePath)
-		{
-		}
-
-		FPlatformIconPaths(const FString& InNormalPath, const FString& InLargePath, const FString& InXLargePath)
-			: NormalPath(InNormalPath)
-			, LargePath(InLargePath)
-			, XLargePath(InXLargePath)
-		{
-		}
-
-		FName NormalStyleName;
-		FString NormalPath;
-
-		FName LargeStyleName;
-		FString LargePath;
-
-		FName XLargeStyleName;
-		FString XLargePath;
-	};
 
 	/** Information about a given platform */
-	struct FPlatformInfo
+	struct FTargetPlatformInfo
 	{
 		/** Name used to identify this platform, eg "Android_ETC2" */
 		FName PlatformInfoName;
@@ -108,9 +49,6 @@ namespace PlatformInfo
 		/** Name used to find the corresponding ITargetPlatform for this platform (also used by UAT) */
 		FName TargetPlatformName;
 
-		/** Vanilla name for this platform, eg "Android" for "Android_ETC2" */
-		FName VanillaPlatformName;
-
 		/** Platform flavor, eg "ETC2" for "Android_ETC2" */
 		FName PlatformFlavor;
 
@@ -123,50 +61,20 @@ namespace PlatformInfo
 		/** Flags for this platform */
 		EPlatformFlags::Flags PlatformFlags;
 
-		/** Information about where to find the platform icons (for use by FEditorStyle) */
-		FPlatformIconPaths IconPaths;
-
 		/** Additional argument string data to append to UAT commands relating to this platform */
 		FString UATCommandLine;
 
-		/** Path under CarefullyRedist for the SDK.  FString so case sensitive platforms don't get messed up by a pre-existing FName of a different casing. */
-		FString AutoSDKPath;
-
-		/** Whether or not this platform SDK has been properly installed */
-		EPlatformSDKStatus SDKStatus;
-
-		/** Tutorial path for tutorial to install SDK */
-		FString SDKTutorial;
+		/** Name of this platform when loading INI files (and finding DataDrivenPlatformInfo) */
+		FName IniPlatformName;
 
-		/** Name of sub-folder where binaries will be placed */
-		FString BinaryFolderName;
+		/** For flavors, this points to the vanilla (parent) object - for vanilla objects, this points to itself so ->VanillaInfo can be used without checking for null */
+		FTargetPlatformInfo* VanillaInfo;
 
-		/** Name of this platform when loading INI files */
-		FString IniPlatformName;
+		/** For vanilla objects, this contains the flavors (children) */
+		TArray<const FTargetPlatformInfo*> Flavors;
 
-		/** Enabled for use */
-		bool bEnabledForUse;
-
-		/** Whether code projects for this platform require the host platform compiler to be installed. Host platforms typically have a SDK status of valid, but they can't necessarily build. */
-		bool bUsesHostCompiler;
-
-		/** Whether UAT closes immediately after launching on this platform, or if it sticks around to read output from the running process */
-		bool bUATClosesAfterLaunch;
-
-		/** Whether or not the platform is confidential in nature */
-		bool bIsConfidential;
-
-		/** Whether or not the platform can use Crash Reporter */
-		bool bTargetPlatformCanUseCrashReporter;
-
-		/** An identifier that corresponds to UBT's UnrealTargetPlatform enum (and by proxy, FGenericPlatformMisc::GetUBTPlatform()) */
-		FName UBTTargetId;
-
-		/** An identifier to group similar platforms together, such as "Mobile" and "Console". Used for Per-Platform Override Properties. */
-		FName PlatformGroupName;
-
-		/** Submenu name to group similar platforms together in menus, such as "Linux" and "LinuxAArch64".  */
-		FName PlatformSubMenu;
+		/** Cached pointer to the DDPI */
+		const FDataDrivenPlatformInfo* DataDrivenPlatformInfo;
 
 		/** Returns true if this platform is vanilla */
 		FORCEINLINE bool IsVanilla() const
@@ -180,60 +88,12 @@ namespace PlatformInfo
 			return !PlatformFlavor.IsNone();
 		}
 
-		/** Get the icon name (for FEditorStyle) used by the given icon type for this platform */
+		// convenience function
 		FName GetIconStyleName(const EPlatformIconSize InIconSize) const
 		{
-			switch(InIconSize)
-			{
-			case EPlatformIconSize::Normal:
-				return IconPaths.NormalStyleName;
-			case EPlatformIconSize::Large:
-				return IconPaths.LargeStyleName;
-			case EPlatformIconSize::XLarge:
-				return IconPaths.XLargeStyleName;
-			default:
-				break;
-			}
-			return NAME_None;
+			return DataDrivenPlatformInfo->GetIconStyleName(InIconSize);
 		}
 
-		/** Get the path to the icon on disk (for FEditorStyle) for the given icon type for this platform */
-		const FString& GetIconPath(const EPlatformIconSize InIconSize) const
-		{
-			switch(InIconSize)
-			{
-			case EPlatformIconSize::Normal:
-				return IconPaths.NormalPath;
-			case EPlatformIconSize::Large:
-				return IconPaths.LargePath;
-			case EPlatformIconSize::XLarge:
-				return IconPaths.XLargePath;
-			default:
-				break;
-			}
-			static const FString EmptyString = TEXT("");
-			return EmptyString;
-		}
-	};
-
-	/** Vanilla platform that may contain a set of flavors */
-	struct FVanillaPlatformEntry
-	{
-		FVanillaPlatformEntry()
-			: PlatformInfo(nullptr)
-		{
-		}
-
-		FVanillaPlatformEntry(const FPlatformInfo* const InPlatformInfo)
-			: PlatformInfo(InPlatformInfo)
-		{
-		}
-
-		/** Information for this platform */
-		const FPlatformInfo* PlatformInfo;
-
-		/** Any flavors for this platform */
-		TArray<const FPlatformInfo*> PlatformFlavors;
 	};
 
 	/**
@@ -241,42 +101,26 @@ namespace PlatformInfo
 	 * @param InPlatformName - The name of the platform to find
 	 * @return The platform info if the platform was found, null otherwise
 	 */
-	DESKTOPPLATFORM_API const FPlatformInfo* FindPlatformInfo(const FName& InPlatformName);
+	DESKTOPPLATFORM_API const FTargetPlatformInfo* FindPlatformInfo(const FName& InPlatformName);
 
 	/**
 	 * Try and find the vanilla information for the given platform
 	 * @param InPlatformName - The name of the platform to find (can be a flavor, but you'll still get back the vanilla platform)
 	 * @return The platform info if the platform was found, null otherwise
 	 */
-	DESKTOPPLATFORM_API const FPlatformInfo* FindVanillaPlatformInfo(const FName& InPlatformName);
+	DESKTOPPLATFORM_API const FTargetPlatformInfo* FindVanillaPlatformInfo(const FName& InPlatformName);
 
 	/**
 	 * Get an array of all the platforms we know about
 	 * @return The pointer to the start of the platform array
 	 */
-	DESKTOPPLATFORM_API const TArray<FPlatformInfo>& GetPlatformInfoArray();
+	DESKTOPPLATFORM_API const TArray<FTargetPlatformInfo*>& GetPlatformInfoArray();
 
 	/**
-	 * Build a hierarchy mapping vanilla platforms to their flavors
-	 * @param InFilter - Flags to control which kinds of flavors you want to include
-	 * @return An array of vanilla platforms, potentially containing flavors
-	 */
-	DESKTOPPLATFORM_API TArray<FVanillaPlatformEntry> BuildPlatformHierarchy(const EPlatformFilter InFilter);
-
-	/**
-	* Build a hierarchy mapping for specified vanilla platform to it flavors
-	* @param InPlatformName - Platform name to build hierarchy for, could be vanilla or flavor name
-	* @param InFilter - Flags to control which kinds of flavors you want to include
-	* @return Vanilla platform potentially containing flavors
-	*/
-	DESKTOPPLATFORM_API FVanillaPlatformEntry BuildPlatformHierarchy(const FName& InPlatformName, const EPlatformFilter InFilter);
-
-	/**
-	 * Get an array of all the platforms we know about
-	 * @param OutNumPlatforms - The number of platforms in the array
+	 * Get an array of only the vanilla platforms
 	 * @return The pointer to the start of the platform array
 	 */
-	DESKTOPPLATFORM_API void UpdatePlatformSDKStatus(FString InPlatformName, EPlatformSDKStatus InStatus);
+	DESKTOPPLATFORM_API const TArray<FTargetPlatformInfo*>& GetVanillaPlatformInfoArray();
 
 	/**
 	* Update the display name for a platform
@@ -299,3 +143,5 @@ namespace PlatformInfo
 	*/
 	DESKTOPPLATFORM_API const TArray<FName>& GetAllVanillaPlatformNames();
 }
+
+#endif
diff --git a/Engine/Source/Developer/DeviceManager/Private/Widgets/Browser/SDeviceBrowserDeviceAdder.cpp b/Engine/Source/Developer/DeviceManager/Private/Widgets/Browser/SDeviceBrowserDeviceAdder.cpp
--- a/Engine/Source/Developer/DeviceManager/Private/Widgets/Browser/SDeviceBrowserDeviceAdder.cpp
+++ b/Engine/Source/Developer/DeviceManager/Private/Widgets/Browser/SDeviceBrowserDeviceAdder.cpp
@@ -61,7 +61,7 @@ void SDeviceBrowserDeviceAdder::Construct(const FArguments& InArgs, const TShare
 
 	// callback for generating widgets for the platforms combo box
 	auto PlatformComboBoxGenerateWidget = [this](TSharedPtr<FString> Item) -> TSharedRef<SWidget> {
-		const PlatformInfo::FPlatformInfo* const PlatformInfo = PlatformInfo::FindPlatformInfo(**Item);
+		const PlatformInfo::FTargetPlatformInfo* const PlatformInfo = PlatformInfo::FindPlatformInfo(**Item);
 
 		return
 			SNew(SHorizontalBox)
@@ -75,7 +75,7 @@ void SDeviceBrowserDeviceAdder::Construct(const FArguments& InArgs, const TShare
 						.HeightOverride(24)
 						[
 							SNew(SImage)
-								.Image((PlatformInfo) ? FEditorStyle::GetBrush(PlatformInfo->GetIconStyleName(PlatformInfo::EPlatformIconSize::Normal)) : FStyleDefaults::GetNoBrush())
+								.Image((PlatformInfo) ? FEditorStyle::GetBrush(PlatformInfo->GetIconStyleName(EPlatformIconSize::Normal)) : FStyleDefaults::GetNoBrush())
 						]
 				]
 
diff --git a/Engine/Source/Developer/DeviceManager/Private/Widgets/Browser/SDeviceBrowserDeviceListRow.h b/Engine/Source/Developer/DeviceManager/Private/Widgets/Browser/SDeviceBrowserDeviceListRow.h
--- a/Engine/Source/Developer/DeviceManager/Private/Widgets/Browser/SDeviceBrowserDeviceListRow.h
+++ b/Engine/Source/Developer/DeviceManager/Private/Widgets/Browser/SDeviceBrowserDeviceListRow.h
@@ -82,7 +82,7 @@ public:
 		}
 		else if (ColumnName == TEXT("Icon"))
 		{
-			const PlatformInfo::FPlatformInfo* const PlatformInfo = PlatformInfo::FindPlatformInfo(DeviceService->GetDevicePlatformName());
+			const PlatformInfo::FTargetPlatformInfo* const PlatformInfo = PlatformInfo::FindPlatformInfo(DeviceService->GetDevicePlatformName());
 
 			return SNew(SBox)
 				.Padding(FMargin(4.0f, 0.0f))
@@ -90,7 +90,7 @@ public:
 				.HeightOverride(24)
 				[
 					SNew(SImage)
-						.Image((PlatformInfo) ? FEditorStyle::GetBrush(PlatformInfo->GetIconStyleName(PlatformInfo::EPlatformIconSize::Normal)) : FStyleDefaults::GetNoBrush())
+						.Image((PlatformInfo) ? FEditorStyle::GetBrush(PlatformInfo->GetIconStyleName(EPlatformIconSize::Normal)) : FStyleDefaults::GetNoBrush())
 				];
 		}
 		else if (ColumnName == TEXT("Name"))
diff --git a/Engine/Source/Developer/DeviceManager/Private/Widgets/Browser/SDeviceBrowserFilterBar.cpp b/Engine/Source/Developer/DeviceManager/Private/Widgets/Browser/SDeviceBrowserFilterBar.cpp
--- a/Engine/Source/Developer/DeviceManager/Private/Widgets/Browser/SDeviceBrowserFilterBar.cpp
+++ b/Engine/Source/Developer/DeviceManager/Private/Widgets/Browser/SDeviceBrowserFilterBar.cpp
@@ -50,7 +50,7 @@ void SDeviceBrowserFilterBar::Construct(const FArguments& InArgs, TSharedRef<FDe
 
 	// callback for generating a row widget for the platform filter list
 	auto PlatformListViewGenerateRow = [this](TSharedPtr<FDeviceBrowserFilterEntry> PlatformEntry, const TSharedRef<STableViewBase>& OwnerTable) -> TSharedRef<ITableRow> {
-		const PlatformInfo::FPlatformInfo* const PlatformInfo = PlatformInfo::FindPlatformInfo(PlatformEntry->PlatformLookup);
+		const PlatformInfo::FTargetPlatformInfo* const PlatformInfo = PlatformInfo::FindPlatformInfo(PlatformEntry->PlatformLookup);
 
 		return SNew(STableRow<TSharedPtr<FString> >, OwnerTable)
 			.Content()
@@ -81,7 +81,7 @@ void SDeviceBrowserFilterBar::Construct(const FArguments& InArgs, TSharedRef<FDe
 								.HeightOverride(24)
 								[
 									SNew(SImage)
-									.Image((PlatformInfo) ? FEditorStyle::GetBrush(PlatformInfo->GetIconStyleName(PlatformInfo::EPlatformIconSize::Normal)) : FStyleDefaults::GetNoBrush())
+									.Image((PlatformInfo) ? FEditorStyle::GetBrush(PlatformInfo->GetIconStyleName(EPlatformIconSize::Normal)) : FStyleDefaults::GetNoBrush())
 								]
 						]
 
diff --git a/Engine/Source/Developer/DeviceManager/Private/Widgets/Shared/SDeviceQuickInfo.h b/Engine/Source/Developer/DeviceManager/Private/Widgets/Shared/SDeviceQuickInfo.h
--- a/Engine/Source/Developer/DeviceManager/Private/Widgets/Shared/SDeviceQuickInfo.h
+++ b/Engine/Source/Developer/DeviceManager/Private/Widgets/Shared/SDeviceQuickInfo.h
@@ -261,10 +261,10 @@ private:
 	{
 		if (DeviceService.IsValid())
 		{
-			const PlatformInfo::FPlatformInfo* const PlatformInfo = PlatformInfo::FindPlatformInfo(DeviceService->GetDevicePlatformName());
+			const PlatformInfo::FTargetPlatformInfo* const PlatformInfo = PlatformInfo::FindPlatformInfo(DeviceService->GetDevicePlatformName());
 			if(PlatformInfo)
 			{
-				return FEditorStyle::GetBrush(PlatformInfo->GetIconStyleName(PlatformInfo::EPlatformIconSize::XLarge));
+				return FEditorStyle::GetBrush(PlatformInfo->GetIconStyleName(EPlatformIconSize::XLarge));
 			}
 		}
 
diff --git a/Engine/Source/Developer/IOS/IOSPlatformEditor/Private/IOSCustomIconProjectBuildMutatorFeature.cpp b/Engine/Source/Developer/IOS/IOSPlatformEditor/Private/IOSCustomIconProjectBuildMutatorFeature.cpp
--- a/Engine/Source/Developer/IOS/IOSPlatformEditor/Private/IOSCustomIconProjectBuildMutatorFeature.cpp
+++ b/Engine/Source/Developer/IOS/IOSPlatformEditor/Private/IOSCustomIconProjectBuildMutatorFeature.cpp
@@ -47,10 +47,10 @@ static bool RequiresBuild()
 
 bool FIOSCustomIconProjectBuildMutatorFeature ::RequiresProjectBuild(const FName& InPlatformInfoName, FText& OutReason) const
 {
-	const PlatformInfo::FPlatformInfo* const PlatInfo = PlatformInfo::FindPlatformInfo(InPlatformInfoName);
+	const PlatformInfo::FTargetPlatformInfo* const PlatInfo = PlatformInfo::FindPlatformInfo(InPlatformInfoName);
 	check(PlatInfo);
 
-	if (PlatInfo->SDKStatus == PlatformInfo::EPlatformSDKStatus::Installed)
+	if (PlatInfo->DataDrivenPlatformInfo->GetSdkStatus() == DDPIPlatformSdkStatus::Valid)
 	{
 		const ITargetPlatform* const Platform = GetTargetPlatformManager()->FindTargetPlatform(PlatInfo->TargetPlatformName.ToString());
 		if (Platform)
diff --git a/Engine/Source/Developer/IOS/IOSTargetPlatform/Private/IOSTargetPlatform.cpp b/Engine/Source/Developer/IOS/IOSTargetPlatform/Private/IOSTargetPlatform.cpp
--- a/Engine/Source/Developer/IOS/IOSTargetPlatform/Private/IOSTargetPlatform.cpp
+++ b/Engine/Source/Developer/IOS/IOSTargetPlatform/Private/IOSTargetPlatform.cpp
@@ -222,7 +222,7 @@ int32 FIOSTargetPlatform::CheckRequirements(bool bProjectHasCode, EBuildConfigur
 	OutTutorialPath = FString("/Engine/Tutorial/Installation/InstallingXCodeTutorial.InstallingXCodeTutorial");
     // shell to certtool
 #else
-	if (!FInstalledPlatformInfo::Get().IsValidPlatform(GetPlatformInfo().BinaryFolderName, EProjectType::Code))
+	if (!FInstalledPlatformInfo::Get().IsValidPlatform(GetPlatformInfo().UBTPlatformString, EProjectType::Code))
 	{
 		if (bProjectHasCode)
 		{
@@ -383,7 +383,7 @@ void FIOSTargetPlatform::HandlePongMessage( const FIOSLaunchDaemonPong& Message,
 		Device->SetDeviceEndpoint(Context->GetSender());
 		Device->SetIsSimulated(Message.DeviceID.Contains(TEXT("Simulator")));
 
-		DeviceDiscoveredEvent.Broadcast(Device.ToSharedRef());
+		OnDeviceDiscovered().Broadcast(Device.ToSharedRef());
 	}
 
 	Device->LastPinged = FDateTime::UtcNow();
@@ -410,7 +410,7 @@ void FIOSTargetPlatform::HandleDeviceConnected(const FIOSLaunchDaemonPong& Messa
 			Device->SetDeviceType(Message.DeviceType);
 			Device->SetIsSimulated(Message.DeviceID.Contains(TEXT("Simulator")));
 
-			DeviceDiscoveredEvent.Broadcast(Device.ToSharedRef());
+			OnDeviceDiscovered().Broadcast(Device.ToSharedRef());
 		}
 		else
 		{
@@ -432,7 +432,7 @@ void FIOSTargetPlatform::HandleDeviceDisconnected(const FIOSLaunchDaemonPong& Me
 	
 	if (Device.IsValid())
 	{
-		DeviceLostEvent.Broadcast(Device.ToSharedRef());
+		OnDeviceLost().Broadcast(Device.ToSharedRef());
 		Devices.Remove(DeviceId);
 	}
 }
diff --git a/Engine/Source/Developer/IOS/IOSTargetPlatform/Private/IOSTargetPlatform.h b/Engine/Source/Developer/IOS/IOSTargetPlatform/Private/IOSTargetPlatform.h
--- a/Engine/Source/Developer/IOS/IOSTargetPlatform/Private/IOSTargetPlatform.h
+++ b/Engine/Source/Developer/IOS/IOSTargetPlatform/Private/IOSTargetPlatform.h
@@ -151,18 +151,6 @@ public:
 		}
 	}
 
-	DECLARE_DERIVED_EVENT(FIOSTargetPlatform, ITargetPlatform::FOnTargetDeviceDiscovered, FOnTargetDeviceDiscovered);
-	virtual FOnTargetDeviceDiscovered& OnDeviceDiscovered( ) override
-	{
-		return DeviceDiscoveredEvent;
-	}
-
-	DECLARE_DERIVED_EVENT(FIOSTargetPlatform, ITargetPlatform::FOnTargetDeviceLost, FOnTargetDeviceLost);
-	virtual FOnTargetDeviceLost& OnDeviceLost( ) override
-	{
-		return DeviceLostEvent;
-	}
-
 	//~ Begin ITargetPlatform Interface
 
 protected:
@@ -217,11 +205,4 @@ private:
     // holds usb device helper
 	FIOSDeviceHelper DeviceHelper;
 
-private:
-
-	// Holds an event delegate that is executed when a new target device has been discovered.
-	FOnTargetDeviceDiscovered DeviceDiscoveredEvent;
-
-	// Holds an event delegate that is executed when a target device has been lost, i.e. disconnected or timed out.
-	FOnTargetDeviceLost DeviceLostEvent;
 };
diff --git a/Engine/Source/Developer/LauncherServices/Private/Launcher/LauncherWorker.cpp b/Engine/Source/Developer/LauncherServices/Private/Launcher/LauncherWorker.cpp
--- a/Engine/Source/Developer/LauncherServices/Private/Launcher/LauncherWorker.cpp
+++ b/Engine/Source/Developer/LauncherServices/Private/Launcher/LauncherWorker.cpp
@@ -308,60 +308,15 @@ FString FLauncherWorker::CreateUATCommand( const ILauncherProfileRef& InProfile,
 	for (int32 PlatformIndex = 0; PlatformIndex < InPlatforms.Num(); ++PlatformIndex)
 	{
 		// Platform info for the given platform
-		const PlatformInfo::FPlatformInfo* PlatformInfo = PlatformInfo::FindPlatformInfo(FName(*InPlatforms[PlatformIndex]));
+		const PlatformInfo::FTargetPlatformInfo* PlatformInfo = PlatformInfo::FindPlatformInfo(FName(*InPlatforms[PlatformIndex]));
 
 		if (ensure(PlatformInfo))
 		{
-			// switch server and no editor platforms to the proper type
-			if (PlatformInfo->TargetPlatformName == FName("LinuxServer"))
-			{
-				ServerPlatforms += TEXT("+Linux");
-			}
-			else if (PlatformInfo->TargetPlatformName == FName("LinuxAArch64Server"))
-			{
-				ServerPlatforms += TEXT("+LinuxAArch64");
-			}
-			else if (PlatformInfo->TargetPlatformName == FName("WindowsServer"))
-			{
-				ServerPlatforms += TEXT("+Win64");
-			}
-			else if (PlatformInfo->TargetPlatformName == FName("MacServer"))
-			{
-				ServerPlatforms += TEXT("+Mac");
-			}
-			else if (PlatformInfo->TargetPlatformName == FName("LinuxNoEditor") || PlatformInfo->TargetPlatformName == FName("LinuxClient"))
-			{
-				Platforms += TEXT("+Linux");
-			}
-			else if (PlatformInfo->TargetPlatformName == FName("LinuxAArch64NoEditor") || PlatformInfo->TargetPlatformName == FName("LinuxAArch64Client"))
-			{
-				Platforms += TEXT("+LinuxAArch64");
-			}
-			else if (PlatformInfo->TargetPlatformName == FName("WindowsNoEditor") || PlatformInfo->TargetPlatformName == FName("Windows") || PlatformInfo->TargetPlatformName == FName("WindowsClient"))
-			{
-				Platforms += TEXT("+Win64");
-			}
-			else if (PlatformInfo->TargetPlatformName == FName("MacNoEditor") || PlatformInfo->TargetPlatformName == FName("MacClient"))
-			{
-				Platforms += TEXT("+Mac");
-			}
-			else if (PlatformInfo->TargetPlatformName == FName("IOSClient"))
-			{
-				Platforms += TEXT("+IOS");
-			}
-			else if (PlatformInfo->TargetPlatformName == FName("TVOSClient"))
-			{
-				Platforms += TEXT("+TVOS");
-			}
-			else if (PlatformInfo->TargetPlatformName == FName("HoloLens"))
-			{
-				Platforms += TEXT("+HoloLens");
-			}
-			else
-			{
-				Platforms += TEXT("+");
-				Platforms += PlatformInfo->UBTTargetId.ToString();
-			}
+			// separate out Server platforms
+			FString& PlatformString = (PlatformInfo->PlatformType == EBuildTargetType::Server) ? ServerPlatforms : Platforms;
+
+			PlatformString += TEXT("+");
+			PlatformString += PlatformInfo->DataDrivenPlatformInfo->UBTPlatformString;
 
 			// Append any extra UAT flags specified for this platform flavor
 			if (!PlatformInfo->UATCommandLine.IsEmpty())
@@ -385,7 +340,7 @@ FString FLauncherWorker::CreateUATCommand( const ILauncherProfileRef& InProfile,
 				OptionalParams += TEXT(" ");
 				OptionalParams += OptionalUATCommandLine;
 			}
-			bUATClosesAfterLaunch |= PlatformInfo->bUATClosesAfterLaunch;
+			bUATClosesAfterLaunch |= PlatformInfo->DataDrivenPlatformInfo->bUATClosesAfterLaunch;
 		}
 	}
 	if (ServerPlatforms.Len() > 0)
diff --git a/Engine/Source/Developer/LauncherServices/Private/Profiles/LauncherProfile.h b/Engine/Source/Developer/LauncherServices/Private/Profiles/LauncherProfile.h
--- a/Engine/Source/Developer/LauncherServices/Private/Profiles/LauncherProfile.h
+++ b/Engine/Source/Developer/LauncherServices/Private/Profiles/LauncherProfile.h
@@ -695,7 +695,7 @@ public:
 				{
 					// Get the target we're building for
 					const ITargetPlatform* TargetPlatform = GetTargetPlatformManager()->FindTargetPlatform(TargetPlatformName);
-					const PlatformInfo::FPlatformInfo& PlatformInfo = TargetPlatform->GetPlatformInfo();
+					const PlatformInfo::FTargetPlatformInfo& PlatformInfo = TargetPlatform->GetTargetPlatformInfo();
 
 					// Figure out which target we're building
 					FString ReceiptDir;
@@ -725,7 +725,7 @@ public:
 					}
 
 					// Check if the existing target is valid
-					FString BuildPlatform = PlatformInfo.UBTTargetId.ToString();
+					FString BuildPlatform = PlatformInfo.DataDrivenPlatformInfo->UBTPlatformString;
 					if (!HasPromotedTarget(*ReceiptDir, *TargetName, *BuildPlatform, BuildConfiguration, nullptr))
 					{
 						break;
@@ -1642,7 +1642,7 @@ public:
 		for (int32 PlatformIndex = 0; PlatformIndex < InPlatforms.Num(); ++PlatformIndex)
 		{
 			// Platform info for the given platform
-			const PlatformInfo::FPlatformInfo* PlatformInfo = PlatformInfo::FindPlatformInfo(FName(*InPlatforms[PlatformIndex]));
+			const PlatformInfo::FTargetPlatformInfo* PlatformInfo = PlatformInfo::FindPlatformInfo(FName(*InPlatforms[PlatformIndex]));
 			if (PlatformInfo == nullptr)
 			{
 				return false;
@@ -1693,7 +1693,7 @@ public:
 				OptionalParams += PlatformInfo->UATCommandLine;
 			}
 
-			bUATClosesAfterLaunch |= PlatformInfo->bUATClosesAfterLaunch;
+			bUATClosesAfterLaunch |= PlatformInfo->DataDrivenPlatformInfo->bUATClosesAfterLaunch;
 		}
 		return bUATClosesAfterLaunch;
 	}
diff --git a/Engine/Source/Developer/Linux/LinuxTargetPlatform/Private/LinuxTargetPlatform.h b/Engine/Source/Developer/Linux/LinuxTargetPlatform/Private/LinuxTargetPlatform.h
--- a/Engine/Source/Developer/Linux/LinuxTargetPlatform/Private/LinuxTargetPlatform.h
+++ b/Engine/Source/Developer/Linux/LinuxTargetPlatform/Private/LinuxTargetPlatform.h
@@ -109,7 +109,7 @@ public:
 			Device->SetUserCredentials(Username, Password);
 		}
 
-		DeviceDiscoveredEvent.Broadcast(Device.ToSharedRef());
+		this->OnDeviceDiscovered().Broadcast(Device.ToSharedRef());
 		return true;
 	}
 
@@ -218,7 +218,7 @@ public:
 		int32 ReadyToBuild = TSuper::CheckRequirements(bProjectHasCode, Configuration, bRequiresAssetNativization, OutTutorialPath, OutDocumentationPath, CustomizedLogMessage);
 
 		// do not support code/plugins in Installed builds if the required libs aren't bundled (on Windows/Mac)
-		if (!PLATFORM_LINUX && !FInstalledPlatformInfo::Get().IsValidPlatform(TSuper::GetPlatformInfo().BinaryFolderName, EProjectType::Code))
+		if (!PLATFORM_LINUX && !FInstalledPlatformInfo::Get().IsValidPlatform(TSuper::GetPlatformInfo().UBTPlatformString, EProjectType::Code))
 		{
 			if (bProjectHasCode)
 			{
@@ -388,18 +388,6 @@ public:
 		return TProperties::GetVariantPriority();
 	}
 
-	DECLARE_DERIVED_EVENT(TLinuxTargetPlatform, ITargetPlatform::FOnTargetDeviceDiscovered, FOnTargetDeviceDiscovered);
-	virtual FOnTargetDeviceDiscovered& OnDeviceDiscovered( ) override
-	{
-		return DeviceDiscoveredEvent;
-	}
-
-	DECLARE_DERIVED_EVENT(TLinuxTargetPlatform, ITargetPlatform::FOnTargetDeviceLost, FOnTargetDeviceLost);
-	virtual FOnTargetDeviceLost& OnDeviceLost( ) override
-	{
-		return DeviceLostEvent;
-	}
-
 	//~ End ITargetPlatform Interface
 
 protected:
@@ -520,14 +508,6 @@ protected:
 	// True if the project requires encoded HDR reflection captures
 	bool bRequiresEncodedHDRReflectionCaptures;
 #endif // WITH_ENGINE
-
-private:
-
-	// Holds an event delegate that is executed when a new target device has been discovered.
-	FOnTargetDeviceDiscovered DeviceDiscoveredEvent;
-
-	// Holds an event delegate that is executed when a target device has been lost, i.e. disconnected or timed out.
-	FOnTargetDeviceLost DeviceLostEvent;
 };
 
 #undef LOCTEXT_NAMESPACE
diff --git a/Engine/Source/Developer/Localization/Private/LocTextHelper.cpp b/Engine/Source/Developer/Localization/Private/LocTextHelper.cpp
--- a/Engine/Source/Developer/Localization/Private/LocTextHelper.cpp
+++ b/Engine/Source/Developer/Localization/Private/LocTextHelper.cpp
@@ -43,14 +43,9 @@ const TArray<FString>& FLocTextPlatformSplitUtils::GetPlatformsToSplit(const ELo
 		{
 			static const TArray<FString> AllPlatformNames = []()
 			{
+				// the Keys of the FDataDrivenPlatformInfoRegistry platforms are the known ini platform names
 				TArray<FString> TmpArray;
-				for (const PlatformInfo::FPlatformInfo& Info : PlatformInfo::GetPlatformInfoArray())
-				{
-					if (!Info.IniPlatformName.IsEmpty())
-					{
-						TmpArray.AddUnique(*Info.IniPlatformName);
-					}
-				}
+				FDataDrivenPlatformInfoRegistry::GetAllPlatformInfos().GetKeys(TmpArray);
 				TmpArray.Sort();
 				return TmpArray;
 			}();
diff --git a/Engine/Source/Developer/Mac/MacTargetPlatform/Private/GenericMacTargetPlatform.h b/Engine/Source/Developer/Mac/MacTargetPlatform/Private/GenericMacTargetPlatform.h
--- a/Engine/Source/Developer/Mac/MacTargetPlatform/Private/GenericMacTargetPlatform.h
+++ b/Engine/Source/Developer/Mac/MacTargetPlatform/Private/GenericMacTargetPlatform.h
@@ -385,19 +385,6 @@ return TSuper::SupportsFeature(Feature);
 		return TProperties::GetVariantPriority();
 	}
 
-
-	DECLARE_DERIVED_EVENT(TGenericMacTargetPlatform, ITargetPlatform::FOnTargetDeviceDiscovered, FOnTargetDeviceDiscovered);
-	virtual FOnTargetDeviceDiscovered& OnDeviceDiscovered( ) override
-	{
-		return DeviceDiscoveredEvent;
-	}
-
-	DECLARE_DERIVED_EVENT(TGenericMacTargetPlatform, ITargetPlatform::FOnTargetDeviceLost, FOnTargetDeviceLost);
-	virtual FOnTargetDeviceLost& OnDeviceLost( ) override
-	{
-		return DeviceLostEvent;
-	}
-
 	//~ End ITargetPlatform Interface
 
 private:
@@ -417,13 +404,6 @@ private:
 
 #endif // WITH_ENGINE
 
-private:
-
-	// Holds an event delegate that is executed when a new target device has been discovered.
-	FOnTargetDeviceDiscovered DeviceDiscoveredEvent;
-
-	// Holds an event delegate that is executed when a target device has been lost, i.e. disconnected or timed out.
-	FOnTargetDeviceLost DeviceLostEvent;
 };
 
 #undef LOCTEXT_NAMESPACE
diff --git a/Engine/Source/Developer/OutputLog/Private/SDeviceOutputLog.cpp b/Engine/Source/Developer/OutputLog/Private/SDeviceOutputLog.cpp
--- a/Engine/Source/Developer/OutputLog/Private/SDeviceOutputLog.cpp
+++ b/Engine/Source/Developer/OutputLog/Private/SDeviceOutputLog.cpp
@@ -104,19 +104,11 @@ void SDeviceOutputLog::Construct( const FArguments& InArgs )
 	bIsUserScrolled = false;
 	RequestForceScroll();
 	
-	//
-	TArray<ITargetPlatform*> Platforms = GetTargetPlatformManager()->GetTargetPlatforms();
-	for (ITargetPlatform* Platform : Platforms)
-	{
-		if (IsSupportedPlatform(Platform))
-		{
-			Platform->OnDeviceDiscovered().AddRaw(this, &SDeviceOutputLog::HandleTargetPlatformDeviceDiscovered);
-			Platform->OnDeviceLost().AddRaw(this, &SDeviceOutputLog::HandleTargetPlatformDeviceLost);
-		}
-	}
+	ITargetPlatform::OnDeviceDiscovered().AddRaw(this, &SDeviceOutputLog::HandleTargetPlatformDeviceDiscovered);
+	ITargetPlatform::OnDeviceLost().AddRaw(this, &SDeviceOutputLog::HandleTargetPlatformDeviceLost);
 		
 	// Get list of available devices
-	for (ITargetPlatform* Platform : Platforms)
+	for (ITargetPlatform* Platform : GetTargetPlatformManager()->GetTargetPlatforms())
 	{
 		if (IsSupportedPlatform(Platform))
 		{
@@ -136,17 +128,8 @@ void SDeviceOutputLog::Construct( const FArguments& InArgs )
 
 SDeviceOutputLog::~SDeviceOutputLog()
 {
-	ITargetPlatformManagerModule* Module = FModuleManager::GetModulePtr<ITargetPlatformManagerModule>("TargetPlatform");
-	if (Module)
-	{
-		TArray<ITargetPlatform*> Platforms = Module->GetTargetPlatforms();
-		for (ITargetPlatform* Platform : Platforms)
-		{
-			Platform->OnDeviceDiscovered().RemoveAll(this);
-			Platform->OnDeviceLost().RemoveAll(this);
-		}
-	}
-	CurrentDeviceOutputPtr.Reset();
+	ITargetPlatform::OnDeviceDiscovered().RemoveAll(this);
+	ITargetPlatform::OnDeviceLost().RemoveAll(this);
 }
 
 void SDeviceOutputLog::Tick(const FGeometry& AllottedGeometry, const double InCurrentTime, const float InDeltaTime)
diff --git a/Engine/Source/Developer/ProjectLauncher/Private/Widgets/Deploy/SProjectLauncherDeployTargetListRow.h b/Engine/Source/Developer/ProjectLauncher/Private/Widgets/Deploy/SProjectLauncherDeployTargetListRow.h
--- a/Engine/Source/Developer/ProjectLauncher/Private/Widgets/Deploy/SProjectLauncherDeployTargetListRow.h
+++ b/Engine/Source/Developer/ProjectLauncher/Private/Widgets/Deploy/SProjectLauncherDeployTargetListRow.h
@@ -271,10 +271,10 @@ private:
 	{
 		if (DeviceProxy->HasVariant(NAME_None))
 		{
-			const PlatformInfo::FPlatformInfo* const PlatformInfo = PlatformInfo::FindPlatformInfo(*DeviceProxy->GetTargetPlatformName(NAME_None));
+			const PlatformInfo::FTargetPlatformInfo* const PlatformInfo = PlatformInfo::FindPlatformInfo(*DeviceProxy->GetTargetPlatformName(NAME_None));
 			if (PlatformInfo)
 			{
-				FEditorStyle::GetBrush(PlatformInfo->GetIconStyleName(PlatformInfo::EPlatformIconSize::Normal));
+				FEditorStyle::GetBrush(PlatformInfo->GetIconStyleName(EPlatformIconSize::Normal));
 			}
 		}
 		return FStyleDefaults::GetNoBrush();
diff --git a/Engine/Source/Developer/ProjectLauncher/Private/Widgets/Deploy/SProjectLauncherDeployTargets.cpp b/Engine/Source/Developer/ProjectLauncher/Private/Widgets/Deploy/SProjectLauncherDeployTargets.cpp
--- a/Engine/Source/Developer/ProjectLauncher/Private/Widgets/Deploy/SProjectLauncherDeployTargets.cpp
+++ b/Engine/Source/Developer/ProjectLauncher/Private/Widgets/Deploy/SProjectLauncherDeployTargets.cpp
@@ -134,17 +134,12 @@ void SProjectLauncherDeployTargets::Construct(const FArguments& InArgs, const TS
 	VanillaNames.Add(NAME_None);
 	VanillaPlatformList.Add(MakeShareable(new FName(NAME_None)));
 
-	TArray<ITargetPlatform*> Platforms = GetTargetPlatformManager()->GetTargetPlatforms();
-	if (Platforms.Num() > 0)
+	for (FName VanillaName : PlatformInfo::GetAllVanillaPlatformNames())
 	{
-		for (int32 PlatformIndex = 0; PlatformIndex < Platforms.Num(); ++PlatformIndex)
+		if (!VanillaNames.Contains(VanillaName))
 		{
-			FName VanillaName = Platforms[PlatformIndex]->GetPlatformInfo().VanillaPlatformName;
-			if (!VanillaNames.Contains(VanillaName))
-			{
-				VanillaNames.Add(VanillaName);
-				VanillaPlatformList.AddUnique(MakeShareable(new FName(VanillaName)));
-			}
+			VanillaNames.Add(VanillaName);
+			VanillaPlatformList.AddUnique(MakeShareable(new FName(VanillaName)));
 		}
 	}
 
diff --git a/Engine/Source/Developer/ProjectLauncher/Private/Widgets/Deploy/SProjectLauncherSimpleDeviceListRow.h b/Engine/Source/Developer/ProjectLauncher/Private/Widgets/Deploy/SProjectLauncherSimpleDeviceListRow.h
--- a/Engine/Source/Developer/ProjectLauncher/Private/Widgets/Deploy/SProjectLauncherSimpleDeviceListRow.h
+++ b/Engine/Source/Developer/ProjectLauncher/Private/Widgets/Deploy/SProjectLauncherSimpleDeviceListRow.h
@@ -427,10 +427,10 @@ private:
 	{
 		if (LaunchProfile.IsValid())
 		{
-			const PlatformInfo::FPlatformInfo* const PlatformInfo = PlatformInfo::FindPlatformInfo(*DeviceProxy->GetTargetPlatformName(SimpleProfile->GetDeviceVariant()));
+			const PlatformInfo::FTargetPlatformInfo* const PlatformInfo = PlatformInfo::FindPlatformInfo(*DeviceProxy->GetTargetPlatformName(SimpleProfile->GetDeviceVariant()));
 			if (PlatformInfo)
 			{
-				return FEditorStyle::GetBrush(PlatformInfo->GetIconStyleName(PlatformInfo::EPlatformIconSize::Large));
+				return FEditorStyle::GetBrush(PlatformInfo->GetIconStyleName(EPlatformIconSize::Large));
 			}
 		}
 		return FStyleDefaults::GetNoBrush();
diff --git a/Engine/Source/Developer/ProjectLauncher/Private/Widgets/Shared/SProjectLauncherProfileNameDescEditor.h b/Engine/Source/Developer/ProjectLauncher/Private/Widgets/Shared/SProjectLauncherProfileNameDescEditor.h
--- a/Engine/Source/Developer/ProjectLauncher/Private/Widgets/Shared/SProjectLauncherProfileNameDescEditor.h
+++ b/Engine/Source/Developer/ProjectLauncher/Private/Widgets/Shared/SProjectLauncherProfileNameDescEditor.h
@@ -43,7 +43,7 @@ private:
 	// Callback for getting the icon image of the profile.
 	const FSlateBrush* HandleProfileImage() const
 	{
-		//const PlatformInfo::FPlatformInfo* const PlatformInfo = PlatformInfo::FindPlatformInfo(*DeviceProxy->GetTargetPlatformName(SimpleProfile->GetDeviceVariant()));
+		//const PlatformInfo::FTargetPlatformInfo* const PlatformInfo = PlatformInfo::FindPlatformInfo(*DeviceProxy->GetTargetPlatformName(SimpleProfile->GetDeviceVariant()));
 		//return (PlatformInfo) ? FEditorStyle::GetBrush(PlatformInfo->GetIconStyleName(PlatformInfo::EPlatformIconSize::Large)) : FStyleDefaults::GetNoBrush();
 		return FEditorStyle::GetBrush("LauncherCommand.QuickLaunch");
 	}
diff --git a/Engine/Source/Developer/ScreenShotComparisonTools/Private/ScreenShotManager.cpp b/Engine/Source/Developer/ScreenShotComparisonTools/Private/ScreenShotManager.cpp
--- a/Engine/Source/Developer/ScreenShotComparisonTools/Private/ScreenShotManager.cpp
+++ b/Engine/Source/Developer/ScreenShotComparisonTools/Private/ScreenShotManager.cpp
@@ -91,7 +91,7 @@ FString FScreenShotManager::GetApprovedFolderForImageWithOptions(const FAutomati
 {
 	// plaform will be something like PS4
 	// RHI = SM5
-	const FDataDrivenPlatformInfoRegistry::FPlatformInfo& PlatInfo = FDataDrivenPlatformInfoRegistry::GetPlatformInfo(MetaData.Platform);
+	const FDataDrivenPlatformInfo& PlatInfo = FDataDrivenPlatformInfoRegistry::GetPlatformInfo(MetaData.Platform);
 
 	bool bUsePlatformPath = PlatInfo.bIsConfidential && (InOptions & EApprovedFolderOptions::UsePlatformFolders) == 0;
 
diff --git a/Engine/Source/Developer/SessionFrontend/Private/Widgets/Browser/SSessionBrowserTreeInstanceRow.h b/Engine/Source/Developer/SessionFrontend/Private/Widgets/Browser/SSessionBrowserTreeInstanceRow.h
--- a/Engine/Source/Developer/SessionFrontend/Private/Widgets/Browser/SSessionBrowserTreeInstanceRow.h
+++ b/Engine/Source/Developer/SessionFrontend/Private/Widgets/Browser/SSessionBrowserTreeInstanceRow.h
@@ -72,7 +72,7 @@ public:
 
 			if (InstanceInfo.IsValid())
 			{
-				const PlatformInfo::FPlatformInfo* const PlatformInfo = PlatformInfo::FindPlatformInfo(*InstanceInfo->GetPlatformName());
+				const PlatformInfo::FTargetPlatformInfo* const PlatformInfo = PlatformInfo::FindPlatformInfo(*InstanceInfo->GetPlatformName());
 
 				return SNew(SHorizontalBox)
 
@@ -85,7 +85,7 @@ public:
 					.HeightOverride(24)
 					[
 						SNew(SImage)
-						.Image((PlatformInfo) ? FEditorStyle::GetBrush(PlatformInfo->GetIconStyleName(PlatformInfo::EPlatformIconSize::Normal)) : FStyleDefaults::GetNoBrush())
+						.Image((PlatformInfo) ? FEditorStyle::GetBrush(PlatformInfo->GetIconStyleName(EPlatformIconSize::Normal)) : FStyleDefaults::GetNoBrush())
 					]
 				]
 
diff --git a/Engine/Source/Developer/TargetDeviceServices/Private/Services/TargetDeviceService.cpp b/Engine/Source/Developer/TargetDeviceServices/Private/Services/TargetDeviceService.cpp
--- a/Engine/Source/Developer/TargetDeviceServices/Private/Services/TargetDeviceService.cpp
+++ b/Engine/Source/Developer/TargetDeviceServices/Private/Services/TargetDeviceService.cpp
@@ -87,9 +87,10 @@ void FTargetDeviceService::AddTargetDevice(TSharedPtr<ITargetDevice, ESPMode::Th
 	{
 		// If this seems nasty your right!
 		// This is just one more nastiness in this class due to the fact that we intend to refactor the target platform stuff as a separate task.
-		const PlatformInfo::FPlatformInfo& Info = InDevice->GetTargetPlatform().GetPlatformInfo();
+		const PlatformInfo::FTargetPlatformInfo& Info = InDevice->GetTargetPlatform().GetTargetPlatformInfo();
+		const PlatformInfo::FTargetPlatformInfo* VanillaInfo = Info.VanillaInfo;
+
 		DevicePlatformName = Info.PlatformInfoName;
-		const PlatformInfo::FPlatformInfo* VanillaInfo = PlatformInfo::FindVanillaPlatformInfo(Info.VanillaPlatformName);
 		DevicePlatformDisplayName = VanillaInfo->DisplayName.ToString();
 		
 		// Sigh the hacks... Should be able to remove if platform info gets cleaned up.... Windows doesn't have a reasonable vanilla platform.
@@ -492,7 +493,7 @@ void FTargetDeviceService::HandlePingMessage(const FTargetDeviceServicePing& InM
 	if (DefaultDevice.IsValid())
 	{
 		const FString& PlatformName = DefaultDevice->GetTargetPlatform().PlatformName();
-		const PlatformInfo::FPlatformInfo* VanillaInfo = PlatformInfo::FindVanillaPlatformInfo(FName(*PlatformName));
+		const PlatformInfo::FTargetPlatformInfo* VanillaInfo = DefaultDevice->GetTargetPlatform().GetTargetPlatformInfo().VanillaInfo;
 
 		// message is going to be deleted by FMemory::Free() (see FMessageContext destructor), so allocate it with Malloc
 		void* Memory = FMemory::Malloc(sizeof(FTargetDeviceServicePong), alignof(FTargetDeviceServicePong));
@@ -517,7 +518,7 @@ void FTargetDeviceService::HandlePingMessage(const FTargetDeviceServicePing& InM
 
 		// Check if we should also create an aggregate (All_<platform>_devices_on_<host>) proxy
 		Message->Aggregated = DefaultDevice->IsPlatformAggregated();
-		Message->AllDevicesName = DefaultDevice->GetAllDevicesName().IsEmpty() ? VanillaInfo->VanillaPlatformName.ToString() : DefaultDevice->GetAllDevicesName();
+		Message->AllDevicesName = DefaultDevice->GetAllDevicesName().IsEmpty() ? VanillaInfo->PlatformInfoName.ToString() : DefaultDevice->GetAllDevicesName();
 		Message->AllDevicesDefaultVariant = DefaultDevice->GetAllDevicesDefaultVariant().IsNone() ? Message->DefaultVariant : DefaultDevice->GetAllDevicesDefaultVariant();
 
 		// Add the data for all the flavors
@@ -527,7 +528,7 @@ void FTargetDeviceService::HandlePingMessage(const FTargetDeviceServicePing& InM
 		for (auto TargetDeviceIt = TargetDevicePtrs.CreateIterator(); TargetDeviceIt; ++TargetDeviceIt, ++Index)
 		{
 			const ITargetDevicePtr& TargetDevice = TargetDeviceIt.Value().Pin();
-			const PlatformInfo::FPlatformInfo& Info = TargetDevice->GetTargetPlatform().GetPlatformInfo();
+			const PlatformInfo::FTargetPlatformInfo& Info = TargetDevice->GetTargetPlatform().GetTargetPlatformInfo();
 
 			FTargetDeviceVariant& Variant = Message->Variants[Index];
 
@@ -535,7 +536,7 @@ void FTargetDeviceService::HandlePingMessage(const FTargetDeviceServicePing& InM
 			Variant.VariantName = TargetDeviceIt.Key();
 			Variant.TargetPlatformName = TargetDevice->GetTargetPlatform().PlatformName();
 			Variant.TargetPlatformId = Info.TargetPlatformName;
-			Variant.VanillaPlatformId = Info.VanillaPlatformName;
+			Variant.VanillaPlatformId = Info.VanillaInfo->PlatformInfoName;
 			Variant.PlatformDisplayName = Info.DisplayName.ToString();
 		}
 
diff --git a/Engine/Source/Developer/TargetDeviceServices/Private/Services/TargetDeviceServiceManager.cpp b/Engine/Source/Developer/TargetDeviceServices/Private/Services/TargetDeviceServiceManager.cpp
--- a/Engine/Source/Developer/TargetDeviceServices/Private/Services/TargetDeviceServiceManager.cpp
+++ b/Engine/Source/Developer/TargetDeviceServices/Private/Services/TargetDeviceServiceManager.cpp
@@ -153,15 +153,11 @@ bool FTargetDeviceServiceManager::AddTargetDevice(ITargetDevicePtr InDevice)
 
 void FTargetDeviceServiceManager::InitializeTargetPlatforms()
 {
-	TArray<ITargetPlatform*> Platforms = GetTargetPlatformManager()->GetTargetPlatforms();
-
-	for (int32 PlatformIndex = 0; PlatformIndex < Platforms.Num(); ++PlatformIndex)
+	ITargetPlatform::OnDeviceDiscovered().AddRaw(this, &FTargetDeviceServiceManager::HandleTargetPlatformDeviceDiscovered);
+	ITargetPlatform::OnDeviceLost().AddRaw(this, &FTargetDeviceServiceManager::HandleTargetPlatformDeviceLost);
+	
+	for (ITargetPlatform* Platform : GetTargetPlatformManager()->GetTargetPlatforms())
 	{
-		// set up target platform callbacks
-		ITargetPlatform* Platform = Platforms[PlatformIndex];
-		Platform->OnDeviceDiscovered().AddRaw(this, &FTargetDeviceServiceManager::HandleTargetPlatformDeviceDiscovered);
-		Platform->OnDeviceLost().AddRaw(this, &FTargetDeviceServiceManager::HandleTargetPlatformDeviceLost);
-
 		// add services for existing devices
 		TArray<ITargetDevicePtr> Devices;
 		Platform->GetAllDevices(Devices);
@@ -329,19 +325,8 @@ void FTargetDeviceServiceManager::SaveSettings()
 
 void FTargetDeviceServiceManager::ShutdownTargetPlatforms()
 {
-	ITargetPlatformManagerModule* Module = FModuleManager::GetModulePtr<ITargetPlatformManagerModule>("TargetPlatform");
-	if (Module)
-	{
-		TArray<ITargetPlatform*> Platforms = Module->GetTargetPlatforms();
-
-		for (int32 PlatformIndex = 0; PlatformIndex < Platforms.Num(); ++PlatformIndex)
-		{
-			// set up target platform callbacks
-			ITargetPlatform* Platform = Platforms[PlatformIndex];
-			Platform->OnDeviceDiscovered().RemoveAll(this);
-			Platform->OnDeviceLost().RemoveAll(this);
-		}
-	}
+	ITargetPlatform::OnDeviceDiscovered().RemoveAll(this);
+	ITargetPlatform::OnDeviceLost().RemoveAll(this);
 }
 
 
diff --git a/Engine/Source/Developer/TargetPlatform/Private/TargetPlatformBase.cpp b/Engine/Source/Developer/TargetPlatform/Private/TargetPlatformBase.cpp
--- a/Engine/Source/Developer/TargetPlatform/Private/TargetPlatformBase.cpp
+++ b/Engine/Source/Developer/TargetPlatform/Private/TargetPlatformBase.cpp
@@ -200,7 +200,7 @@ bool FTargetPlatformBase::RequiresTempTarget(bool bProjectHasCode, EBuildConfigu
 
 	// check if there's a non-default plugin change
 	FText Reason;
-	if (IPluginManager::Get().RequiresTempTargetForCodePlugin(Project, PlatformInfo->UBTTargetId.ToString(), Configuration, PlatformInfo->PlatformType, Reason))
+	if (IPluginManager::Get().RequiresTempTargetForCodePlugin(Project, GetPlatformInfo().UBTPlatformString, Configuration, PlatformInfo->PlatformType, Reason))
 	{
 		OutReason = Reason;
 		return true;
@@ -232,8 +232,7 @@ bool FTargetPlatformBase::HasDefaultBuildSettings() const
 	BuildKeys.Add(TEXT("bCompileCEF3")); 
 	BuildKeys.Add(TEXT("bCompileCustomSQLitePlatform"));
 
-	const PlatformInfo::FPlatformInfo& PlatInfo = GetPlatformInfo();
-	if (!DoProjectSettingsMatchDefault(PlatInfo.TargetPlatformName.ToString(), TEXT("/Script/BuildSettings.BuildSettings"), &BuildKeys, nullptr, nullptr))
+	if (!DoProjectSettingsMatchDefault(IniPlatformName(), TEXT("/Script/BuildSettings.BuildSettings"), &BuildKeys, nullptr, nullptr))
 	{
 		return false;
 	}
@@ -241,7 +240,7 @@ bool FTargetPlatformBase::HasDefaultBuildSettings() const
 	FString PlatformSection;
 	GetBuildProjectSettingKeys(PlatformSection, BoolKeys, IntKeys, StringKeys);
 
-	if(!DoProjectSettingsMatchDefault(PlatInfo.TargetPlatformName.ToString(), PlatformSection, &BoolKeys, &IntKeys, &StringKeys))
+	if(!DoProjectSettingsMatchDefault(IniPlatformName(), PlatformSection, &BoolKeys, &IntKeys, &StringKeys))
 	{
 		return false;
 	}
diff --git a/Engine/Source/Developer/TargetPlatform/Private/TargetPlatformManagerModule.cpp b/Engine/Source/Developer/TargetPlatform/Private/TargetPlatformManagerModule.cpp
--- a/Engine/Source/Developer/TargetPlatform/Private/TargetPlatformManagerModule.cpp
+++ b/Engine/Source/Developer/TargetPlatform/Private/TargetPlatformManagerModule.cpp
@@ -28,13 +28,9 @@
 
 DEFINE_LOG_CATEGORY_STATIC(LogTargetPlatformManager, Log, All);
 
+// AutoSDKs needs the extra DDPI info
+#define AUTOSDKS_ENABLED DDPI_HAS_EXTENDED_PLATFORMINFO_DATA
 
-// autosdks only function properly on windows right now.
-#if !IS_MONOLITHIC && (PLATFORM_WINDOWS)
-	#define AUTOSDKS_ENABLED 1
-#else
-	#define AUTOSDKS_ENABLED 0
-#endif
 
 static const size_t MaxPlatformCount = 64;		// In the unlikely event that someone bumps this please note that there's
 												// an implicit assumption that there won't be more than 64 unique target
@@ -45,8 +41,7 @@ static const ITargetPlatform* TargetPlatformArray[MaxPlatformCount];
 
 static int32 PlatformCounter = 0;
 
-int32 
-ITargetPlatform::AssignPlatformOrdinal(const ITargetPlatform& Platform)
+int32 ITargetPlatform::AssignPlatformOrdinal(const ITargetPlatform& Platform)
 {
 	check(PlatformCounter < MaxPlatformCount);
 
@@ -58,14 +53,25 @@ ITargetPlatform::AssignPlatformOrdinal(const ITargetPlatform& Platform)
 	return Ordinal;
 }
 
-const ITargetPlatform* 
-ITargetPlatform::GetPlatformFromOrdinal(int32 Ordinal)
+const ITargetPlatform* ITargetPlatform::GetPlatformFromOrdinal(int32 Ordinal)
 {
 	check(Ordinal < PlatformCounter);
 
 	return TargetPlatformArray[Ordinal];
 }
 
+ITargetPlatform::FOnTargetDeviceDiscovered& ITargetPlatform::OnDeviceDiscovered()
+{
+	static FOnTargetDeviceDiscovered Delegate;
+	return Delegate;
+}
+
+ITargetPlatform::FOnTargetDeviceLost& ITargetPlatform::OnDeviceLost()
+{
+	static FOnTargetDeviceLost Delegate;
+	return Delegate;
+}
+
 
 /**
  * Module for the target platform manager
@@ -103,11 +109,11 @@ public:
 
 			// we have to setup our local environment according to AutoSDKs or the ITargetPlatform's IsSDkInstalled calls may fail
 			// before we get a change to setup for a given platform.  Use the platforminfo list to avoid any kind of interdependency.
-			for (const PlatformInfo::FPlatformInfo& PlatformInfo : PlatformInfo::GetPlatformInfoArray())
+			for (auto Pair: FDataDrivenPlatformInfoRegistry::GetAllPlatformInfos())
 			{
-				if (PlatformInfo.AutoSDKPath.Len() > 0)
+				if (Pair.Value.AutoSDKPath.Len() > 0)
 				{
-					SetupAndValidateAutoSDK(PlatformInfo.AutoSDKPath);
+					SetupAndValidateAutoSDK(Pair.Value.AutoSDKPath);
 				}
 			}
 		}
@@ -634,6 +640,64 @@ protected:
 		return false;
 	}
 
+	bool InitializeSinglePlatform(const FString& PlatformName, const FDataDrivenPlatformInfo& Info)
+	{
+		// disabled?
+		if (Info.bEnabledForUse == false)
+		{
+			return false;
+		}
+
+		// there are two ways targetplatform modules are setup: a single DLL per TargetPlatform, or a DLL for the platform
+		// that returns multiple TargetPlatforms. we try single first, then full platform
+		FName PlatformModuleName = *(PlatformName + TEXT("TargetPlatform"));
+
+		ITargetPlatformModule* Module = nullptr;
+
+		if (FModuleManager::Get().ModuleExists(*PlatformModuleName.ToString()))
+		{
+			Module = FModuleManager::LoadModulePtr<ITargetPlatformModule>(PlatformModuleName);
+		}
+
+		// original logic for module loading here
+		if (Module)
+		{
+			// would like to move this check to GetActiveTargetPlatforms, but too many things cache this result
+			// this setup will become faster after TTP 341897 is complete.
+		RETRY_SETUPANDVALIDATE:
+			if (SetupAndValidateAutoSDK(Info.AutoSDKPath))
+			{
+				TArray<ITargetPlatform*> TargetPlatforms = Module->GetTargetPlatforms();
+				for (ITargetPlatform* Platform : TargetPlatforms)
+				{
+					UE_LOG(LogTargetPlatformManager, Display, TEXT("Loaded TargetPlatform '%s'"), *Platform->PlatformName());
+					Platforms.Add(Platform);
+					PlatformsByName.Add(FName(Platform->PlatformName()), Platform);
+				}
+
+				// only success path
+				return true;
+			}
+			else
+			{
+				// this hack is here because if you try and setup and validate autosdk some times it will fail because shared files are in use by another child cooker
+				static bool bIsChildCooker = FParse::Param(FCommandLine::Get(), TEXT("cookchild"));
+				if (bIsChildCooker)
+				{
+					static int Counter = 0;
+					++Counter;
+					if (Counter < 10)
+					{
+						goto RETRY_SETUPANDVALIDATE;
+					}
+				}
+				UE_LOG(LogTargetPlatformManager, Display, TEXT("Failed to SetupAndValidateAutoSDK for platform '%s'"), *PlatformName);
+			}
+		}
+
+		return false;
+	}
+
 	/** Discovers the available target platforms. */
 	void DiscoverAvailablePlatforms()
 	{
@@ -656,97 +720,29 @@ protected:
 
 		// find a set of valid target platform names (the platform DataDrivenPlatformInfo.ini file was found indicates support for the platform 
 		// exists on disk, so the TP is expected to work)
-		const TArray<PlatformInfo::FPlatformInfo>& PlatformInfos = PlatformInfo::GetPlatformInfoArray();
+		const TArray<PlatformInfo::FTargetPlatformInfo*>& PlatformInfos = PlatformInfo::GetPlatformInfoArray();
 
-		TSet<ITargetPlatformModule*> ProcessedModules;
 		FScopedSlowTask SlowTask(PlatformInfos.Num());
-		for (const PlatformInfo::FPlatformInfo& PlatInfo : PlatformInfos)
+		for (auto Pair : FDataDrivenPlatformInfoRegistry::GetAllPlatformInfos())
 		{
-			SlowTask.EnterProgressFrame(1);
-
-			// by defaulty load all PIs we can have
-			bool bLoadTargetPlatform = true;
+			const FString& PlatformName = Pair.Key;
+			const FDataDrivenPlatformInfo& Info = Pair.Value;
 
-			// disabled?
-			if (PlatInfo.bEnabledForUse == false)
-			{
-				bLoadTargetPlatform = false;
-			}
+			SlowTask.EnterProgressFrame(1);
 
 #if WITH_EDITOR
 			// if we have the editor and we are using -game
 			// only need to instantiate the current platform 
 			if (IsRunningGame())
 			{
-				if (PlatInfo.IniPlatformName != FPlatformProperties::IniPlatformName())
+				if (PlatformName != FPlatformProperties::IniPlatformName())
 				{
-					bLoadTargetPlatform = false;
+					continue;
 				}
 			}
 #endif
 
-			// now load the TP module
-			if (bLoadTargetPlatform)
-			{
-				// there are two ways targetplatform modules are setup: a single DLL per TargetPlatform, or a DLL for the platform
-				// that returns multiple TargetPlatforms. we try single first, then full platform
-				FName FullPlatformModuleName = *(PlatInfo.IniPlatformName + TEXT("TargetPlatform"));
-				FName SingleTargetPlatformModuleName = *(PlatInfo.TargetPlatformName.ToString() + TEXT("TargetPlatform"));
-				bool bFullPlatformModuleNameIsValid = !PlatInfo.IniPlatformName.IsEmpty();
-
-				ITargetPlatformModule* Module = nullptr;
-				
-				if (FModuleManager::Get().ModuleExists(*SingleTargetPlatformModuleName.ToString()))
-				{
-					Module = FModuleManager::LoadModulePtr<ITargetPlatformModule>(SingleTargetPlatformModuleName);
-				}
-				else if (bFullPlatformModuleNameIsValid && FModuleManager::Get().ModuleExists(*FullPlatformModuleName.ToString()))
-				{
-					Module = FModuleManager::LoadModulePtr<ITargetPlatformModule>(FullPlatformModuleName);
-				}
-
-				// if we have already processed this module, we can skip it!
-				if (ProcessedModules.Contains(Module))
-				{
-					continue;
-				}
-
-				// original logic for module loading here
-				if (Module)
-				{
-					ProcessedModules.Add(Module);
-
-					TArray<ITargetPlatform*> TargetPlatforms = Module->GetTargetPlatforms();
-					for (ITargetPlatform* Platform : TargetPlatforms)
-					{
-						// would like to move this check to GetActiveTargetPlatforms, but too many things cache this result
-						// this setup will become faster after TTP 341897 is complete.
-					RETRY_SETUPANDVALIDATE:
-						if (SetupAndValidateAutoSDK(Platform->GetPlatformInfo().AutoSDKPath))
-						{
-							const FString& PlatformName = Platform->PlatformName();
-							UE_LOG(LogTargetPlatformManager, Display, TEXT("Loaded TargetPlatform '%s'"), *PlatformName);
-							Platforms.Add(Platform);
-							PlatformsByName.Add(FName(PlatformName), Platform);
-						}
-						else
-						{
-							// this hack is here because if you try and setup and validate autosdk some times it will fail because shared files are in use by another child cooker
-							static bool bIsChildCooker = FParse::Param(FCommandLine::Get(), TEXT("cookchild"));
-							if (bIsChildCooker)
-							{
-								static int Counter = 0;
-								++Counter;
-								if (Counter < 10)
-								{
-									goto RETRY_SETUPANDVALIDATE;
-								}
-							}
-							UE_LOG(LogTargetPlatformManager, Display, TEXT("Failed to SetupAndValidateAutoSDK for platform '%s'"), *Platform->PlatformName());
-						}
-					}
-				}
-			}
+			InitializeSinglePlatform(PlatformName, Info);
 		}
 
 		if (!Platforms.Num())
@@ -805,7 +801,7 @@ protected:
 #if PLATFORM_WINDOWS
 		FString HostPlatform(TEXT("HostWin64"));
 #else
-#error Fill in your host platform directory
+		FString HostPlatform = FString::Printf(TEXT("Host%hs"), FPlatformProperties::IniPlatformName());
 #endif		
 
 		static const FString SDKRootEnvFar(TEXT("UE_SDKS_ROOT"));
@@ -987,6 +983,8 @@ protected:
 
 	bool SetupSDKStatus(const FString& TargetPlatforms)
 	{
+		FDataDrivenPlatformInfoRegistry::UpdateSdkStatus();
+#if 0
 		DECLARE_SCOPE_CYCLE_COUNTER( TEXT( "FTargetPlatformManagerModule::SetupSDKStatus" ), STAT_FTargetPlatformManagerModule_SetupSDKStatus, STATGROUP_TargetPlatform );
 
 		// run UBT with -validate -allplatforms and read the output
@@ -1074,9 +1072,54 @@ protected:
 				}
 			}
 		}
+#endif
 		return true;
 	}
 
+	bool UpdateAfterSDKInstall(const FString& PlatformName)
+	{
+		const FDataDrivenPlatformInfo& Info = FDataDrivenPlatformInfoRegistry::GetPlatformInfo(PlatformName);
+
+		FString AutoSDKPath = Info.AutoSDKPath;
+		FName AutoSDKName(*AutoSDKPath);
+		if (AutoSDKName != NAME_None)
+		{
+			// make sure we can re-do the AutoSDK setup
+			PlatformsSetup.Remove(AutoSDKName);
+		}
+
+		// note: this assumes, along with other Turnkey code, that there is a TargetPlatform named with the IniPlatformName
+		ITargetPlatform* TargetPlatform = FindTargetPlatform(PlatformName);
+
+		bool bTPInitialized = false;
+		// if we didn't have a TP before, discover it now, it will do everything we need)
+		if (TargetPlatform == nullptr)
+		{
+			// create the TP(s) that weren't around before due to a bad SDK
+			bTPInitialized = InitializeSinglePlatform(PlatformName, Info);
+		}
+		else
+		{
+			if (AutoSDKName != NAME_None)
+			{
+				// setup AutoSDK, and then re-initialize the TP
+				SetupAndValidateAutoSDK(AutoSDKPath);
+			}
+
+			bTPInitialized = TargetPlatform->InitializeHostPlatform();
+		}
+
+		if (bTPInitialized)
+		{
+			FDataDrivenPlatformInfoRegistry::UpdateSdkStatus();
+		}
+
+		FDataDrivenPlatformInfoRegistry::ClearDeviceStatus(PlatformName);
+		Invalidate();
+
+		return bTPInitialized;
+	}
+
 private:
 
 	void SetupEnvironmentVariables(TArray<FString> &EnvVarNames, const TArray<FString>& EnvVarValues)
diff --git a/Engine/Source/Developer/TargetPlatform/Public/Common/TargetPlatformBase.h b/Engine/Source/Developer/TargetPlatform/Public/Common/TargetPlatformBase.h
--- a/Engine/Source/Developer/TargetPlatform/Public/Common/TargetPlatformBase.h
+++ b/Engine/Source/Developer/TargetPlatform/Public/Common/TargetPlatformBase.h
@@ -31,11 +31,16 @@ public:
 		return PlatformInfo->DisplayName;
 	}
 
-	virtual const PlatformInfo::FPlatformInfo& GetPlatformInfo() const override
+	virtual const PlatformInfo::FTargetPlatformInfo& GetTargetPlatformInfo() const override
 	{
 		return *PlatformInfo;
 	}
 
+	virtual const FDataDrivenPlatformInfo& GetPlatformInfo() const override
+	{
+		return *PlatformInfo->DataDrivenPlatformInfo;
+	}
+
 	TARGETPLATFORM_API virtual bool UsesForwardShading() const override;
 
 	TARGETPLATFORM_API virtual bool UsesDBuffer() const override;
@@ -168,9 +173,16 @@ public:
 		return false; 
 	}
 
+	virtual bool InitializeHostPlatform()
+	{
+		// if the platform doesn't need anything, it's valid to do nothing
+		return true;
+	}
+
+
 protected:
 
-	FTargetPlatformBase(const PlatformInfo::FPlatformInfo *const InPlatformInfo)
+	FTargetPlatformBase(const PlatformInfo::FTargetPlatformInfo *const InPlatformInfo)
 		: PlatformInfo(InPlatformInfo)
 	{
 		checkf(PlatformInfo, TEXT("Null PlatformInfo was passed to FTargetPlatformBase. Check the static IsUsable function before creating this object. See FWindowsTargetPlatformModule::GetTargetPlatform()"));
@@ -179,7 +191,7 @@ protected:
 	}
 
 	/** Information about this platform */
-	const PlatformInfo::FPlatformInfo *PlatformInfo;
+	const PlatformInfo::FTargetPlatformInfo *PlatformInfo;
 	int32 PlatformOrdinal;
 
 private:
diff --git a/Engine/Source/Developer/TargetPlatform/Public/Interfaces/ITargetPlatform.h b/Engine/Source/Developer/TargetPlatform/Public/Interfaces/ITargetPlatform.h
--- a/Engine/Source/Developer/TargetPlatform/Public/Interfaces/ITargetPlatform.h
+++ b/Engine/Source/Developer/TargetPlatform/Public/Interfaces/ITargetPlatform.h
@@ -5,10 +5,12 @@
 #include "CoreMinimal.h"
 #include "Interfaces/ITargetDevice.h"
 
+struct FDataDrivenPlatformInfo;
+
 namespace PlatformInfo
 {
 	// Forward declare type from DesktopPlatform rather than add an include dependency to everything using ITargetPlatform
-	struct FPlatformInfo;
+	struct FTargetPlatformInfo;
 }
 
 class IDeviceManagerCustomPlatformWidgetCreator;
@@ -208,9 +210,14 @@ public:
 	virtual bool RequiresTempTarget(bool bProjectHasCode, EBuildConfiguration Configuration, bool bRequiresAssetNativization, FText& OutReason) const = 0;
 
 	/**
-	 * Returns the information about this platform
+	 * Returns the information about this target platform
+	 */
+	virtual const PlatformInfo::FTargetPlatformInfo& GetTargetPlatformInfo() const = 0;
+
+	/**
+	 * Returns the information about the platform as a whole
 	 */
-	virtual const PlatformInfo::FPlatformInfo& GetPlatformInfo() const = 0;
+	virtual const FDataDrivenPlatformInfo& GetPlatformInfo() const = 0;
 
 	/**
 	 * Gets the platform's INI name (so an offline tool can load the INI for the given target platform).
@@ -608,19 +615,24 @@ public:
 	 */
 	virtual bool CopyFileToTarget(const FString& DeviceId, const FString& HostFilename, const FString& TargetFilename, const TMap<FString,FString>& CustomPlatformData) = 0;
 
+	/**
+	 * Initializes the host platform to support target devices (may be called multiple times after an SDK is installed while running)
+	 */
+	virtual bool InitializeHostPlatform() = 0;
+
 public:
 
 	/**
 	 * Gets an event delegate that is executed when a new target device has been discovered.
 	 */
 	DECLARE_EVENT_OneParam(ITargetPlatform, FOnTargetDeviceDiscovered, ITargetDeviceRef /*DiscoveredDevice*/);
-	virtual FOnTargetDeviceDiscovered& OnDeviceDiscovered() = 0;
+	static TARGETPLATFORM_API FOnTargetDeviceDiscovered& OnDeviceDiscovered();
 
 	/**
 	 * Gets an event delegate that is executed when a target device has been lost, i.e. disconnected or timed out.
 	 */
 	DECLARE_EVENT_OneParam(ITargetPlatform, FOnTargetDeviceLost, ITargetDeviceRef /*LostDevice*/);
-	virtual FOnTargetDeviceLost& OnDeviceLost() = 0;
+	static TARGETPLATFORM_API FOnTargetDeviceLost& OnDeviceLost();
 
 public:
 
diff --git a/Engine/Source/Developer/TargetPlatform/Public/Interfaces/ITargetPlatformManagerModule.h b/Engine/Source/Developer/TargetPlatform/Public/Interfaces/ITargetPlatformManagerModule.h
--- a/Engine/Source/Developer/TargetPlatform/Public/Interfaces/ITargetPlatformManagerModule.h
+++ b/Engine/Source/Developer/TargetPlatform/Public/Interfaces/ITargetPlatformManagerModule.h
@@ -169,6 +169,11 @@ public:
 	 */
 	virtual bool UpdatePlatformEnvironment(const FString& PlatformName, TArray<FString> &Keys, TArray<FString> &Values) = 0;
 
+	/**
+	 * After installing an SDK with Turnkey, this will refresh the TargetPlatform, find devices, etc
+	 */
+	virtual bool UpdateAfterSDKInstall(const FString& TargetPlatformName) = 0;
+
 public:
 
 	/** Virtual destructor. */
diff --git a/Engine/Source/Developer/Windows/WindowsTargetPlatform/Private/GenericWindowsTargetPlatform.h b/Engine/Source/Developer/Windows/WindowsTargetPlatform/Private/GenericWindowsTargetPlatform.h
--- a/Engine/Source/Developer/Windows/WindowsTargetPlatform/Private/GenericWindowsTargetPlatform.h
+++ b/Engine/Source/Developer/Windows/WindowsTargetPlatform/Private/GenericWindowsTargetPlatform.h
@@ -473,18 +473,6 @@ public:
 		return TProperties::GetVariantPriority();
 	}
 
-	DECLARE_DERIVED_EVENT(TGenericWindowsTargetPlatform, ITargetPlatform::FOnTargetDeviceDiscovered, FOnTargetDeviceDiscovered);
-	virtual FOnTargetDeviceDiscovered& OnDeviceDiscovered( ) override
-	{
-		return DeviceDiscoveredEvent;
-	}
-
-	DECLARE_DERIVED_EVENT(TGenericWindowsTargetPlatform, ITargetPlatform::FOnTargetDeviceLost, FOnTargetDeviceLost);
-	virtual FOnTargetDeviceLost& OnDeviceLost( ) override
-	{
-		return DeviceLostEvent;
-	}
-
 	virtual bool UsesRayTracing() const override
 	{
 		bool bEnableRayTracing = false;
@@ -520,13 +508,6 @@ private:
 	bool bSupportCompressedVolumeTexture;
 #endif // WITH_ENGINE
 
-private:
-
-	// Holds an event delegate that is executed when a new target device has been discovered.
-	FOnTargetDeviceDiscovered DeviceDiscoveredEvent;
-
-	// Holds an event delegate that is executed when a target device has been lost, i.e. disconnected or timed out.
-	FOnTargetDeviceLost DeviceLostEvent;
 };
 
 #undef LOCTEXT_NAMESPACE
diff --git a/Engine/Source/Editor/DeviceProfileEditor/Private/DetailsPanel/SDeviceProfileDetailsPanel.cpp b/Engine/Source/Editor/DeviceProfileEditor/Private/DetailsPanel/SDeviceProfileDetailsPanel.cpp
--- a/Engine/Source/Editor/DeviceProfileEditor/Private/DetailsPanel/SDeviceProfileDetailsPanel.cpp
+++ b/Engine/Source/Editor/DeviceProfileEditor/Private/DetailsPanel/SDeviceProfileDetailsPanel.cpp
@@ -100,7 +100,7 @@ void SDeviceProfileDetailsPanel::RefreshUI()
 		TArray<ITargetPlatform*> TargetPlatforms = GetTargetPlatformManager()->GetTargetPlatforms();
 		if (TargetPlatforms.Num())
 		{
-			DeviceProfileTypeIcon = FEditorStyle::GetBrush(TargetPlatforms[0]->GetPlatformInfo().GetIconStyleName(PlatformInfo::EPlatformIconSize::Normal));
+			DeviceProfileTypeIcon = FEditorStyle::GetBrush(TargetPlatforms[0]->GetPlatformInfo().GetIconStyleName(EPlatformIconSize::Normal));
 		}
 
 		SettingsView->SetObject(&*ViewingProfile);
diff --git a/Engine/Source/Editor/DeviceProfileServices/Private/DeviceProfileServicesUIManager.cpp b/Engine/Source/Editor/DeviceProfileServices/Private/DeviceProfileServicesUIManager.cpp
--- a/Engine/Source/Editor/DeviceProfileServices/Private/DeviceProfileServicesUIManager.cpp
+++ b/Engine/Source/Editor/DeviceProfileServices/Private/DeviceProfileServicesUIManager.cpp
@@ -87,7 +87,7 @@ void FDeviceProfileServicesUIManager::CreatePlatformMap()
 	for (int32 Index = 0; Index < Platforms.Num(); ++Index)
 	{
 		PlatformList.Add(MakeShareable(new FString(Platforms[Index]->IniPlatformName())));
-		DeviceTypeToIconMap.Add( Platforms[Index]->PlatformName(), Platforms[Index]->GetPlatformInfo().GetIconStyleName(PlatformInfo::EPlatformIconSize::Normal) );
+		DeviceTypeToIconMap.Add( Platforms[Index]->PlatformName(), Platforms[Index]->GetPlatformInfo().GetIconStyleName(EPlatformIconSize::Normal) );
 	}
 }
 
diff --git a/Engine/Source/Editor/EditorStyle/Private/SlateEditorStyle.cpp b/Engine/Source/Editor/EditorStyle/Private/SlateEditorStyle.cpp
--- a/Engine/Source/Editor/EditorStyle/Private/SlateEditorStyle.cpp
+++ b/Engine/Source/Editor/EditorStyle/Private/SlateEditorStyle.cpp
@@ -8079,29 +8079,26 @@ void FSlateEditorStyle::FStyle::SetupAutomationStyles()
 
 		Set( "Launcher.Platform.Warning", new IMAGE_BRUSH( "Icons/alert", Icon24x24) );
 
-#if (WITH_EDITOR || (IS_PROGRAM && PLATFORM_DESKTOP))
-		Set( "Launcher.Platform.AllPlatforms", new IMAGE_BRUSH( "Launcher/All_Platforms_24x", Icon24x24) );
-		Set( "Launcher.Platform.AllPlatforms.Large", new IMAGE_BRUSH( "Launcher/All_Platforms_128x", Icon64x64) );
-		Set( "Launcher.Platform.AllPlatforms.XLarge", new IMAGE_BRUSH( "Launcher/All_Platforms_128x", Icon128x128) );
+#if DDPI_HAS_EXTENDED_PLATFORMINFO_DATA
 
-		for(const PlatformInfo::FPlatformInfo& PlatformInfo : PlatformInfo::GetPlatformInfoArray())
+		for (auto Pair : FDataDrivenPlatformInfoRegistry::GetAllPlatformInfos())
 		{
-			// @todo platplug: Add in .Small icons
+			const FDataDrivenPlatformInfo& PlatformInfo = Pair.Value;
 
 			// some platforms may specify a "rooted" path in the platform extensions directory, so look for that case here, and use a different path for the brush
-			FString NormalIconPath = PlatformInfo.GetIconPath(PlatformInfo::EPlatformIconSize::Normal);
+			FString NormalIconPath = PlatformInfo.GetIconPath(EPlatformIconSize::Normal);
 			if (NormalIconPath.StartsWith(TEXT("/Platforms/")))
 			{
-				#define PLATFORM_IMAGE_BRUSH( PlatformPath, ... ) FSlateImageBrush( PlatformPath.Replace(TEXT("/Platforms/"), *FPaths::EnginePlatformExtensionsDir()) + TEXT(".png") , __VA_ARGS__ )
-				Set(PlatformInfo.GetIconStyleName(PlatformInfo::EPlatformIconSize::Normal), new PLATFORM_IMAGE_BRUSH(NormalIconPath, Icon24x24));
-				Set(PlatformInfo.GetIconStyleName(PlatformInfo::EPlatformIconSize::Large), new PLATFORM_IMAGE_BRUSH(PlatformInfo.GetIconPath(PlatformInfo::EPlatformIconSize::Large), Icon64x64));
-				Set(PlatformInfo.GetIconStyleName(PlatformInfo::EPlatformIconSize::XLarge), new PLATFORM_IMAGE_BRUSH(PlatformInfo.GetIconPath(PlatformInfo::EPlatformIconSize::XLarge), Icon128x128));
+#define PLATFORM_IMAGE_BRUSH( PlatformPath, ... ) FSlateImageBrush( PlatformPath.Replace(TEXT("/Platforms/"), *FPaths::EnginePlatformExtensionsDir()) + TEXT(".png") , __VA_ARGS__ )
+				Set(PlatformInfo.GetIconStyleName(EPlatformIconSize::Normal), new PLATFORM_IMAGE_BRUSH(NormalIconPath, Icon24x24));
+				Set(PlatformInfo.GetIconStyleName(EPlatformIconSize::Large), new PLATFORM_IMAGE_BRUSH(PlatformInfo.GetIconPath(EPlatformIconSize::Large), Icon64x64));
+				Set(PlatformInfo.GetIconStyleName(EPlatformIconSize::XLarge), new PLATFORM_IMAGE_BRUSH(PlatformInfo.GetIconPath(EPlatformIconSize::XLarge), Icon128x128));
 			}
 			else
 			{
-				Set(PlatformInfo.GetIconStyleName(PlatformInfo::EPlatformIconSize::Normal), new IMAGE_BRUSH(*NormalIconPath, Icon24x24));
-				Set(PlatformInfo.GetIconStyleName(PlatformInfo::EPlatformIconSize::Large), new IMAGE_BRUSH(*PlatformInfo.GetIconPath(PlatformInfo::EPlatformIconSize::Large), Icon64x64));
-				Set(PlatformInfo.GetIconStyleName(PlatformInfo::EPlatformIconSize::XLarge), new IMAGE_BRUSH(*PlatformInfo.GetIconPath(PlatformInfo::EPlatformIconSize::XLarge), Icon128x128));
+				Set(PlatformInfo.GetIconStyleName(EPlatformIconSize::Normal), new IMAGE_BRUSH(*NormalIconPath, Icon24x24));
+				Set(PlatformInfo.GetIconStyleName(EPlatformIconSize::Large), new IMAGE_BRUSH(*PlatformInfo.GetIconPath(EPlatformIconSize::Large), Icon64x64));
+				Set(PlatformInfo.GetIconStyleName(EPlatformIconSize::XLarge), new IMAGE_BRUSH(*PlatformInfo.GetIconPath(EPlatformIconSize::XLarge), Icon128x128));
 			}
 		}
 #endif
diff --git a/Engine/Source/Editor/EditorStyle/Private/StarshipStyle.cpp b/Engine/Source/Editor/EditorStyle/Private/StarshipStyle.cpp
--- a/Engine/Source/Editor/EditorStyle/Private/StarshipStyle.cpp
+++ b/Engine/Source/Editor/EditorStyle/Private/StarshipStyle.cpp
@@ -7358,29 +7358,26 @@ void FStarshipEditorStyle::FStyle::SetupAutomationStyles()
 
 		Set( "Launcher.Platform.Warning", new IMAGE_BRUSH( "Icons/alert", Icon24x24) );
 
-#if (WITH_EDITOR || (IS_PROGRAM && PLATFORM_DESKTOP))
-		Set( "Launcher.Platform.AllPlatforms", new IMAGE_BRUSH( "Launcher/All_Platforms_24x", Icon24x24) );
-		Set( "Launcher.Platform.AllPlatforms.Large", new IMAGE_BRUSH( "Launcher/All_Platforms_128x", Icon64x64) );
-		Set( "Launcher.Platform.AllPlatforms.XLarge", new IMAGE_BRUSH( "Launcher/All_Platforms_128x", Icon128x128) );
+#if DDPI_HAS_EXTENDED_PLATFORMINFO_DATA
 
-		for(const PlatformInfo::FPlatformInfo& PlatformInfo : PlatformInfo::GetPlatformInfoArray())
+		for (auto Pair : FDataDrivenPlatformInfoRegistry::GetAllPlatformInfos())
 		{
-			// @todo platplug: Add in .Small icons
+			const FDataDrivenPlatformInfo& PlatformInfo = Pair.Value;
 
 			// some platforms may specify a "rooted" path in the platform extensions directory, so look for that case here, and use a different path for the brush
-			FString NormalIconPath = PlatformInfo.GetIconPath(PlatformInfo::EPlatformIconSize::Normal);
+			FString NormalIconPath = PlatformInfo.GetIconPath(EPlatformIconSize::Normal);
 			if (NormalIconPath.StartsWith(TEXT("/Platforms/")))
 			{
 				#define PLATFORM_IMAGE_BRUSH( PlatformPath, ... ) FSlateImageBrush( PlatformPath.Replace(TEXT("/Platforms/"), *FPaths::EnginePlatformExtensionsDir()) + TEXT(".png") , __VA_ARGS__ )
-				Set(PlatformInfo.GetIconStyleName(PlatformInfo::EPlatformIconSize::Normal), new PLATFORM_IMAGE_BRUSH(NormalIconPath, Icon24x24));
-				Set(PlatformInfo.GetIconStyleName(PlatformInfo::EPlatformIconSize::Large), new PLATFORM_IMAGE_BRUSH(PlatformInfo.GetIconPath(PlatformInfo::EPlatformIconSize::Large), Icon64x64));
-				Set(PlatformInfo.GetIconStyleName(PlatformInfo::EPlatformIconSize::XLarge), new PLATFORM_IMAGE_BRUSH(PlatformInfo.GetIconPath(PlatformInfo::EPlatformIconSize::XLarge), Icon128x128));
+				Set(PlatformInfo.GetIconStyleName(EPlatformIconSize::Normal), new PLATFORM_IMAGE_BRUSH(NormalIconPath, Icon24x24));
+				Set(PlatformInfo.GetIconStyleName(EPlatformIconSize::Large), new PLATFORM_IMAGE_BRUSH(PlatformInfo.GetIconPath(EPlatformIconSize::Large), Icon64x64));
+				Set(PlatformInfo.GetIconStyleName(EPlatformIconSize::XLarge), new PLATFORM_IMAGE_BRUSH(PlatformInfo.GetIconPath(EPlatformIconSize::XLarge), Icon128x128));
 			}
 			else
 			{
-				Set(PlatformInfo.GetIconStyleName(PlatformInfo::EPlatformIconSize::Normal), new IMAGE_BRUSH(*NormalIconPath, Icon24x24));
-				Set(PlatformInfo.GetIconStyleName(PlatformInfo::EPlatformIconSize::Large), new IMAGE_BRUSH(*PlatformInfo.GetIconPath(PlatformInfo::EPlatformIconSize::Large), Icon64x64));
-				Set(PlatformInfo.GetIconStyleName(PlatformInfo::EPlatformIconSize::XLarge), new IMAGE_BRUSH(*PlatformInfo.GetIconPath(PlatformInfo::EPlatformIconSize::XLarge), Icon128x128));
+				Set(PlatformInfo.GetIconStyleName(EPlatformIconSize::Normal), new IMAGE_BRUSH(*NormalIconPath, Icon24x24));
+				Set(PlatformInfo.GetIconStyleName(EPlatformIconSize::Large), new IMAGE_BRUSH(*PlatformInfo.GetIconPath(EPlatformIconSize::Large), Icon64x64));
+				Set(PlatformInfo.GetIconStyleName(EPlatformIconSize::XLarge), new IMAGE_BRUSH(*PlatformInfo.GetIconPath(EPlatformIconSize::XLarge), Icon128x128));
 			}
 		}
 #endif
diff --git a/Engine/Source/Editor/GameProjectGeneration/Private/SProjectBrowser.cpp b/Engine/Source/Editor/GameProjectGeneration/Private/SProjectBrowser.cpp
--- a/Engine/Source/Editor/GameProjectGeneration/Private/SProjectBrowser.cpp
+++ b/Engine/Source/Editor/GameProjectGeneration/Private/SProjectBrowser.cpp
@@ -569,7 +569,7 @@ TSharedRef<SToolTip> SProjectBrowser::MakeProjectToolTip( TSharedPtr<FProjectIte
 	TSharedRef<SHorizontalBox> TargetPlatformIconsBox = SNew(SHorizontalBox);
 	for(const FName& PlatformName : ProjectItem->TargetPlatforms)
 	{
-		const PlatformInfo::FPlatformInfo* const PlatformInfo = PlatformInfo::FindPlatformInfo(PlatformName);
+		const PlatformInfo::FTargetPlatformInfo* const PlatformInfo = PlatformInfo::FindPlatformInfo(PlatformName);
 		check(PlatformInfo);
 
 		TargetPlatformIconsBox->AddSlot()
@@ -583,7 +583,7 @@ TSharedRef<SToolTip> SProjectBrowser::MakeProjectToolTip( TSharedPtr<FProjectIte
 			.HeightOverride(20)
 			[
 				SNew(SImage)
-				.Image(FEditorStyle::GetBrush(PlatformInfo->GetIconStyleName(PlatformInfo::EPlatformIconSize::Normal)))
+				.Image(FEditorStyle::GetBrush(PlatformInfo->GetIconStyleName(EPlatformIconSize::Normal)))
 			]
 		];
 	}
@@ -800,11 +800,11 @@ static TSharedPtr<FProjectItem> CreateProjectItem(const FString& ProjectFilename
 
 			// Work out which platforms this project is targeting
 			TArray<FName> TargetPlatforms;
-			for (const PlatformInfo::FPlatformInfo& PlatformInfo : PlatformInfo::GetPlatformInfoArray())
+			for (const PlatformInfo::FTargetPlatformInfo* PlatformInfo : PlatformInfo::GetVanillaPlatformInfoArray())
 			{
-				if (PlatformInfo.IsVanilla() && PlatformInfo.PlatformType == EBuildTargetType::Game && ProjectStatus.IsTargetPlatformSupported(PlatformInfo.PlatformInfoName))
+				if (PlatformInfo->PlatformType == EBuildTargetType::Game && ProjectStatus.IsTargetPlatformSupported(PlatformInfo->PlatformInfoName))
 				{
-					TargetPlatforms.Add(PlatformInfo.PlatformInfoName);
+					TargetPlatforms.Add(PlatformInfo->PlatformInfoName);
 				}
 			}
 			TargetPlatforms.Sort(FNameLexicalLess());
diff --git a/Engine/Source/Editor/MainFrame/Private/Frame/MainFrameActions.cpp b/Engine/Source/Editor/MainFrame/Private/Frame/MainFrameActions.cpp
--- a/Engine/Source/Editor/MainFrame/Private/Frame/MainFrameActions.cpp
+++ b/Engine/Source/Editor/MainFrame/Private/Frame/MainFrameActions.cpp
@@ -375,10 +375,10 @@ FString GetCookingOptionalParams()
 
 void FMainFrameActionCallbacks::CookContent(const FName InPlatformInfoName)
 {
-	const PlatformInfo::FPlatformInfo* const PlatformInfo = PlatformInfo::FindPlatformInfo(InPlatformInfoName);
+	const PlatformInfo::FTargetPlatformInfo* const PlatformInfo = PlatformInfo::FindPlatformInfo(InPlatformInfoName);
 	check(PlatformInfo);
 
-	if (FInstalledPlatformInfo::Get().IsPlatformMissingRequiredFile(PlatformInfo->BinaryFolderName))
+	if (FInstalledPlatformInfo::Get().IsPlatformMissingRequiredFile(PlatformInfo->DataDrivenPlatformInfo->UBTPlatformString))
 	{
 		if (!FInstalledPlatformInfo::OpenInstallerOptions())
 		{
@@ -389,15 +389,16 @@ void FMainFrameActionCallbacks::CookContent(const FName InPlatformInfoName)
 
 	FString OptionalParams;
 
-	if (!FModuleManager::LoadModuleChecked<IProjectTargetPlatformEditorModule>("ProjectTargetPlatformEditor").ShowUnsupportedTargetWarning(PlatformInfo->VanillaPlatformName))
+	if (!FModuleManager::LoadModuleChecked<IProjectTargetPlatformEditorModule>("ProjectTargetPlatformEditor").ShowUnsupportedTargetWarning(PlatformInfo->VanillaInfo->PlatformInfoName))
 	{
 		return;
 	}
 
-	if (PlatformInfo->SDKStatus == PlatformInfo::EPlatformSDKStatus::NotInstalled)
+	// @todo turnkey: run Turnkey here instead of showing SDK!
+	if (PlatformInfo->DataDrivenPlatformInfo->GetSdkStatus() != DDPIPlatformSdkStatus::Valid)
 	{
 		IMainFrameModule& MainFrameModule = FModuleManager::GetModuleChecked<IMainFrameModule>(TEXT("MainFrame"));
-		MainFrameModule.BroadcastMainFrameSDKNotInstalled(PlatformInfo->TargetPlatformName.ToString(), PlatformInfo->SDKTutorial);
+		MainFrameModule.BroadcastMainFrameSDKNotInstalled(PlatformInfo->TargetPlatformName.ToString(), PlatformInfo->DataDrivenPlatformInfo->SDKTutorial);
 		return;
 	}
 
@@ -476,10 +477,10 @@ void FMainFrameActionCallbacks::PackageProject( const FName InPlatformInfoName )
 	FGameProjectGenerationModule& GameProjectModule = FModuleManager::LoadModuleChecked<FGameProjectGenerationModule>(TEXT("GameProjectGeneration"));
 	bool bProjectHasCode = GameProjectModule.Get().ProjectHasCodeFiles();
 	
-	const PlatformInfo::FPlatformInfo* const PlatformInfo = PlatformInfo::FindPlatformInfo(InPlatformInfoName);
+	const PlatformInfo::FTargetPlatformInfo* const PlatformInfo = PlatformInfo::FindPlatformInfo(InPlatformInfoName);
 	check(PlatformInfo);
 
-	if (FInstalledPlatformInfo::Get().IsPlatformMissingRequiredFile(PlatformInfo->BinaryFolderName))
+	if (FInstalledPlatformInfo::Get().IsPlatformMissingRequiredFile(PlatformInfo->DataDrivenPlatformInfo->UBTPlatformString))
 	{
 		if (!FInstalledPlatformInfo::OpenInstallerOptions())
 		{
@@ -494,10 +495,11 @@ void FMainFrameActionCallbacks::PackageProject( const FName InPlatformInfoName )
 		return;
 	}
 
-	if (PlatformInfo->SDKStatus == PlatformInfo::EPlatformSDKStatus::NotInstalled || (bProjectHasCode && PlatformInfo->bUsesHostCompiler && !FSourceCodeNavigation::IsCompilerAvailable()))
+	// @todo turnkey : run turnkey
+	if (PlatformInfo->DataDrivenPlatformInfo->GetSdkStatus() != DDPIPlatformSdkStatus::Valid || (bProjectHasCode && PlatformInfo->DataDrivenPlatformInfo->bUsesHostCompiler && !FSourceCodeNavigation::IsCompilerAvailable()))
 	{
 		IMainFrameModule& MainFrameModule = FModuleManager::GetModuleChecked<IMainFrameModule>(TEXT("MainFrame"));
-		MainFrameModule.BroadcastMainFrameSDKNotInstalled(PlatformInfo->TargetPlatformName.ToString(), PlatformInfo->SDKTutorial);
+		MainFrameModule.BroadcastMainFrameSDKNotInstalled(PlatformInfo->TargetPlatformName.ToString(), PlatformInfo->DataDrivenPlatformInfo->SDKTutorial);
 		TArray<FAnalyticsEventAttribute> ParamArray;
 		ParamArray.Add(FAnalyticsEventAttribute(TEXT("Time"), 0.0));
 		FEditorAnalytics::ReportEvent(TEXT("Editor.Package.Failed"), PlatformInfo->TargetPlatformName.ToString(), bProjectHasCode, EAnalyticsErrorCodes::SDKNotFound, ParamArray);
@@ -612,7 +614,7 @@ void FMainFrameActionCallbacks::PackageProject( const FName InPlatformInfoName )
 		}
 	}
 
-	if (!FModuleManager::LoadModuleChecked<IProjectTargetPlatformEditorModule>("ProjectTargetPlatformEditor").ShowUnsupportedTargetWarning(PlatformInfo->VanillaPlatformName))
+	if (!FModuleManager::LoadModuleChecked<IProjectTargetPlatformEditorModule>("ProjectTargetPlatformEditor").ShowUnsupportedTargetWarning(PlatformInfo->VanillaInfo->PlatformInfoName))
 	{
 		return;
 	}
@@ -703,7 +705,7 @@ void FMainFrameActionCallbacks::PackageProject( const FName InPlatformInfoName )
 		OptionalParams += TEXT(" -manifests");
 	}
 
-	bool bTargetPlatformCanUseCrashReporter = PlatformInfo->bTargetPlatformCanUseCrashReporter;
+	bool bTargetPlatformCanUseCrashReporter = PlatformInfo->DataDrivenPlatformInfo->bCanUseCrashReporter;
 	if (bTargetPlatformCanUseCrashReporter && PlatformInfo->TargetPlatformName == FName("WindowsNoEditor") && PlatformInfo->PlatformFlavor == TEXT("Win32"))
 	{
 		FString MinumumSupportedWindowsOS;
@@ -781,10 +783,10 @@ void FMainFrameActionCallbacks::PackageProject( const FName InPlatformInfoName )
 			}
 
 			// Check if the receipt is for a matching promoted target
-			FString PlatformName = Platform->GetPlatformInfo().UBTTargetId.ToString();
+			FString UBTPlatformName = Platform->GetTargetPlatformInfo().DataDrivenPlatformInfo->UBTPlatformString;
 
 			extern LAUNCHERSERVICES_API bool HasPromotedTarget(const TCHAR* BaseDir, const TCHAR* TargetName, const TCHAR* Platform, EBuildConfiguration Configuration, const TCHAR* Architecture);
-			if (HasPromotedTarget(*BaseDir, *TargetName, *PlatformName, ConfigurationInfo.Configuration, nullptr))
+			if (HasPromotedTarget(*BaseDir, *TargetName, *UBTPlatformName, ConfigurationInfo.Configuration, nullptr))
 			{
 				bBuild = false;
 			}
@@ -829,7 +831,7 @@ void FMainFrameActionCallbacks::PackageProject( const FName InPlatformInfoName )
 		OptionalParams += FString::Printf(TEXT(" -target=%s -clientconfig=%s"), *Target->Name, LexToString(ConfigurationInfo.Configuration));
 	}
 
-	FString TurnkeyCommandline = FString::Printf(TEXT("Turnkey -command=VerifySdk -platform=%s -UpdateIfNeeded -EditorIO"), *PlatformInfo->IniPlatformName);
+	FString TurnkeyCommandline = FString::Printf(TEXT("Turnkey -command=VerifySdk -platform=%s -UpdateIfNeeded -EditorIO"), *PlatformInfo->IniPlatformName.ToString());
 
 	FString ProjectPath = FPaths::IsProjectFilePathSet() ? FPaths::ConvertRelativePathToFull(FPaths::GetProjectFilePath()) : FPaths::RootDir() / FApp::GetProjectName() / FApp::GetProjectName() + TEXT(".uproject");
 	FString CommandLine = FString::Printf(TEXT("-ScriptsForProject=\"%s\" %s BuildCookRun %s%s -nop4 -project=\"%s\" -cook -stage -archive -archivedirectory=\"%s\" -package -ue4exe=\"%s\" %s -utf8output"),
diff --git a/Engine/Source/Editor/MainFrame/Private/Menus/CookContentMenu.h b/Engine/Source/Editor/MainFrame/Private/Menus/CookContentMenu.h
--- a/Engine/Source/Editor/MainFrame/Private/Menus/CookContentMenu.h
+++ b/Engine/Source/Editor/MainFrame/Private/Menus/CookContentMenu.h
@@ -55,7 +55,7 @@ public:
 			check(VanillaPlatform.PlatformInfo->IsVanilla());
 
 			// Only care about game targets
-			if (VanillaPlatform.PlatformInfo->PlatformType != EBuildTargetType::Game || !VanillaPlatform.PlatformInfo->bEnabledForUse || !FInstalledPlatformInfo::Get().CanDisplayPlatform(VanillaPlatform.PlatformInfo->BinaryFolderName, ProjectType))
+			if (VanillaPlatform.PlatformInfo->PlatformType != EBuildTargetType::Game || !VanillaPlatform.PlatformInfo->bEnabledForUse || !FInstalledPlatformInfo::Get().CanDisplayPlatform(VanillaPlatform.PlatformInfo->UBTPlatformString, ProjectType))
 			{
 				continue;
 			}
@@ -90,12 +90,12 @@ protected:
 	 * @param MenuBuilder The builder for the menu that owns this menu.
 	 * @param Platform The target platform we allow packaging for
 	 */
-	static void AddPlatformToMenu(FMenuBuilder& MenuBuilder, const PlatformInfo::FPlatformInfo& PlatformInfo)
+	static void AddPlatformToMenu(FMenuBuilder& MenuBuilder, const PlatformInfo::FTargetPlatformInfo& PlatformInfo)
 	{
 		EProjectType ProjectType = FGameProjectGenerationModule::Get().ProjectHasCodeFiles() ? EProjectType::Code : EProjectType::Content;
 
 		// don't add sub-platforms that can't be displayed in an installed build
-		if (!FInstalledPlatformInfo::Get().CanDisplayPlatform(PlatformInfo.BinaryFolderName, ProjectType))
+		if (!FInstalledPlatformInfo::Get().CanDisplayPlatform(PlatformInfo.UBTPlatformString, ProjectType))
 		{
 			return;
 		}
@@ -135,9 +135,9 @@ protected:
 	 * @param MenuBuilder The builder for the menu that owns this menu.
 	 * @param SubPlatformInfos The Sub-platform information
 	 */
-	static void AddPlatformSubPlatformsToMenu(FMenuBuilder& MenuBuilder, TArray<const PlatformInfo::FPlatformInfo*> SubPlatformInfos)
+	static void AddPlatformSubPlatformsToMenu(FMenuBuilder& MenuBuilder, TArray<const PlatformInfo::FTargetPlatformInfo*> SubPlatformInfos)
 	{
-		for (const PlatformInfo::FPlatformInfo* SubPlatformInfo : SubPlatformInfos)
+		for (const PlatformInfo::FTargetPlatformInfo* SubPlatformInfo : SubPlatformInfos)
 		{
 			AddPlatformToMenu(MenuBuilder, *SubPlatformInfo);
 		}
diff --git a/Engine/Source/Editor/MainFrame/Private/Menus/PackageProjectMenu.h b/Engine/Source/Editor/MainFrame/Private/Menus/PackageProjectMenu.h
--- a/Engine/Source/Editor/MainFrame/Private/Menus/PackageProjectMenu.h
+++ b/Engine/Source/Editor/MainFrame/Private/Menus/PackageProjectMenu.h
@@ -36,58 +36,52 @@ public:
 	static void MakeMenu( FMenuBuilder& MenuBuilder )
 	{
 		TArray<FName> AllPlatformSubMenus;
-		const TArray<FString>& ConfidentalPlatforms = FDataDrivenPlatformInfoRegistry::GetConfidentialPlatforms();
 
-		TArray<PlatformInfo::FVanillaPlatformEntry> VanillaPlatforms = PlatformInfo::BuildPlatformHierarchy(PlatformInfo::EPlatformFilter::All);
+		TArray<PlatformInfo::FTargetPlatformInfo*> VanillaPlatforms = PlatformInfo::GetVanillaPlatformInfoArray();
 		if (!VanillaPlatforms.Num())
 		{
 			return;
 		}
 
-		VanillaPlatforms.Sort([](const PlatformInfo::FVanillaPlatformEntry& One, const PlatformInfo::FVanillaPlatformEntry& Two) -> bool
+		Algo::Sort(VanillaPlatforms, [](const PlatformInfo::FTargetPlatformInfo* One, const PlatformInfo::FTargetPlatformInfo* Two) -> bool
 		{
-			return One.PlatformInfo->DisplayName.CompareTo(Two.PlatformInfo->DisplayName) < 0;
+			return One->DisplayName.CompareTo(Two->DisplayName) < 0;
 		});
 
 		IProjectTargetPlatformEditorModule& ProjectTargetPlatformEditorModule = FModuleManager::LoadModuleChecked<IProjectTargetPlatformEditorModule>("ProjectTargetPlatformEditor");
 		EProjectType ProjectType = FGameProjectGenerationModule::Get().ProjectHasCodeFiles() ? EProjectType::Code : EProjectType::Content;
 
 		// Build up a menu from the tree of platforms
-		for (const PlatformInfo::FVanillaPlatformEntry& VanillaPlatform : VanillaPlatforms)
+		for (const PlatformInfo::FTargetPlatformInfo* VanillaPlatform : VanillaPlatforms)
 		{
-			check(VanillaPlatform.PlatformInfo->IsVanilla());
+			check(VanillaPlatform->IsVanilla());
 
 			// Only care about game targets
-			if (VanillaPlatform.PlatformInfo->PlatformType != EBuildTargetType::Game || !VanillaPlatform.PlatformInfo->bEnabledForUse || !FInstalledPlatformInfo::Get().CanDisplayPlatform(VanillaPlatform.PlatformInfo->BinaryFolderName, ProjectType))
-			{
-				continue;
-			}
-
-			// Make sure we're able to run this platform
-			if (VanillaPlatform.PlatformInfo->bIsConfidential && !ConfidentalPlatforms.Contains(VanillaPlatform.PlatformInfo->IniPlatformName))
+			if (VanillaPlatform->PlatformType != EBuildTargetType::Game || !VanillaPlatform->DataDrivenPlatformInfo->bEnabledForUse || !FInstalledPlatformInfo::Get().CanDisplayPlatform(VanillaPlatform->DataDrivenPlatformInfo->UBTPlatformString, ProjectType))
 			{
 				continue;
 			}
 
 			// Check if this platform has a submenu entry
-			if (VanillaPlatform.PlatformInfo->PlatformSubMenu != NAME_None)
+			if (VanillaPlatform->DataDrivenPlatformInfo->PlatformSubMenu != NAME_None)
 			{
-				TArray<const PlatformInfo::FPlatformInfo*> SubMenuEntries;
-				const FName& PlatformSubMenu = VanillaPlatform.PlatformInfo->PlatformSubMenu;
+				TArray<const PlatformInfo::FTargetPlatformInfo*> SubMenuEntries;
+				const FName& PlatformSubMenu = VanillaPlatform->DataDrivenPlatformInfo->PlatformSubMenu;
 
 				// Check if we've already added this submenu
 				if (AllPlatformSubMenus.Find(PlatformSubMenu) != INDEX_NONE)
+				{
 					continue;
+				}
+
 				AllPlatformSubMenus.Add(PlatformSubMenu);
 
 				// Go through all vanilla platforms looking for matching submenus
-				for (const PlatformInfo::FVanillaPlatformEntry& SubMenuVanillaPlatform : VanillaPlatforms)
+				for (PlatformInfo::FTargetPlatformInfo* SubMenuVanillaPlatform : VanillaPlatforms)
 				{
-					const PlatformInfo::FPlatformInfo* PlatformInfo = SubMenuVanillaPlatform.PlatformInfo;
-
-					if ((PlatformInfo->PlatformType == EBuildTargetType::Game) && (PlatformInfo->PlatformSubMenu == PlatformSubMenu))
+					if (SubMenuVanillaPlatform->PlatformType == EBuildTargetType::Game && SubMenuVanillaPlatform->DataDrivenPlatformInfo->PlatformSubMenu == PlatformSubMenu)
 					{
-						SubMenuEntries.Add(PlatformInfo);
+						SubMenuEntries.Add(SubMenuVanillaPlatform);
 					}
 				}
 
@@ -96,23 +90,23 @@ public:
 					const FText DisplayName = FText::FromName(PlatformSubMenu);
 
 					MenuBuilder.AddSubMenu(
-							ProjectTargetPlatformEditorModule.MakePlatformMenuItemWidget(*VanillaPlatform.PlatformInfo, false, DisplayName), 
+							ProjectTargetPlatformEditorModule.MakePlatformMenuItemWidget(*VanillaPlatform, false, DisplayName), 
 							FNewMenuDelegate::CreateStatic(&FPackageProjectMenu::AddPlatformSubPlatformsToMenu, SubMenuEntries),
 							false
 							);
 				}
 			}
-			else if (VanillaPlatform.PlatformFlavors.Num())
+			else if (VanillaPlatform->Flavors.Num())
 			{
 				MenuBuilder.AddSubMenu(
-					ProjectTargetPlatformEditorModule.MakePlatformMenuItemWidget(*VanillaPlatform.PlatformInfo), 
-					FNewMenuDelegate::CreateStatic(&FPackageProjectMenu::AddPlatformSubPlatformsToMenu, VanillaPlatform.PlatformFlavors),
+					ProjectTargetPlatformEditorModule.MakePlatformMenuItemWidget(*VanillaPlatform), 
+					FNewMenuDelegate::CreateStatic(&FPackageProjectMenu::AddPlatformSubPlatformsToMenu, VanillaPlatform->Flavors),
 					false
 					);
 			}
 			else
 			{
-				AddPlatformToMenu(MenuBuilder, *VanillaPlatform.PlatformInfo);
+				AddPlatformToMenu(MenuBuilder, *VanillaPlatform);
 			}
 		}
 
@@ -150,12 +144,12 @@ protected:
 	 * @param MenuBuilder The builder for the menu that owns this menu.
 	 * @param Platform The target platform we allow packaging for
 	 */
-	static void AddPlatformToMenu(FMenuBuilder& MenuBuilder, const PlatformInfo::FPlatformInfo& PlatformInfo)
+	static void AddPlatformToMenu(FMenuBuilder& MenuBuilder, const PlatformInfo::FTargetPlatformInfo& PlatformInfo)
 	{
 		EProjectType ProjectType = FGameProjectGenerationModule::Get().ProjectHasCodeFiles() ? EProjectType::Code : EProjectType::Content;
 
 		// don't add sub-platforms that can't be displayed in an installed build
-		if (!FInstalledPlatformInfo::Get().CanDisplayPlatform(PlatformInfo.BinaryFolderName, ProjectType))
+		if (!FInstalledPlatformInfo::Get().CanDisplayPlatform(PlatformInfo.DataDrivenPlatformInfo->UBTPlatformString, ProjectType))
 		{
 			return;
 		}
@@ -173,7 +167,7 @@ protected:
 		FText Tooltip = FText::Format(LOCTEXT("PackageGameForPlatformTooltip", "Build, cook and package your game for the {DisplayName} platform"), TooltipArguments);
 
 		FProjectStatus ProjectStatus;
-		if (IProjectManager::Get().QueryStatusForCurrentProject(ProjectStatus) && !ProjectStatus.IsTargetPlatformSupported(PlatformInfo.VanillaPlatformName))
+		if (IProjectManager::Get().QueryStatusForCurrentProject(ProjectStatus) && !ProjectStatus.IsTargetPlatformSupported(PlatformInfo.VanillaInfo->PlatformInfoName))
 		{
 			FText TooltipLine2 = FText::Format(LOCTEXT("PackageUnsupportedPlatformWarning", "{DisplayName} is not listed as a target platform for this project, so may not run as expected."), TooltipArguments);
 			Tooltip = FText::Format(FText::FromString(TEXT("{0}\n\n{1}")), Tooltip, TooltipLine2);
@@ -195,9 +189,9 @@ protected:
 	 * @param MenuBuilderThe builder for the menu that owns this menu.
 	 * @param SubPlatformInfos The Sub-platform information
 	 */
-	static void AddPlatformSubPlatformsToMenu(FMenuBuilder& MenuBuilder, TArray<const PlatformInfo::FPlatformInfo*> SubPlatformInfos)
+	static void AddPlatformSubPlatformsToMenu(FMenuBuilder& MenuBuilder, TArray<const PlatformInfo::FTargetPlatformInfo*> SubPlatformInfos)
 	{
-		for (const PlatformInfo::FPlatformInfo* SubPlatformInfo : SubPlatformInfos)
+		for (const PlatformInfo::FTargetPlatformInfo* SubPlatformInfo : SubPlatformInfos)
 		{
 			if (SubPlatformInfo->PlatformType != EBuildTargetType::Game)
 			{
diff --git a/Engine/Source/Editor/ProjectTargetPlatformEditor/Private/ProjectTargetPlatformEditorModule.cpp b/Engine/Source/Editor/ProjectTargetPlatformEditor/Private/ProjectTargetPlatformEditorModule.cpp
--- a/Engine/Source/Editor/ProjectTargetPlatformEditor/Private/ProjectTargetPlatformEditorModule.cpp
+++ b/Engine/Source/Editor/ProjectTargetPlatformEditor/Private/ProjectTargetPlatformEditorModule.cpp
@@ -69,7 +69,7 @@ public:
 			);
 	}
 
-	virtual TSharedRef<SWidget> MakePlatformMenuItemWidget(const PlatformInfo::FPlatformInfo& PlatformInfo, const bool bForCheckBox = false, const FText& DisplayNameOverride = FText()) const override
+	virtual TSharedRef<SWidget> MakePlatformMenuItemWidget(const PlatformInfo::FTargetPlatformInfo& PlatformInfo, const bool bForCheckBox = false, const FText& DisplayNameOverride = FText()) const override
 	{
 		struct Local
 		{
@@ -80,6 +80,8 @@ public:
 			}
 		};
 
+		const float MenuIconSize = FCoreStyle::Get().GetFloat("Menu.MenuIconSize", nullptr, 16.f);
+
 		return 
 			SNew(SHorizontalBox)
 			+SHorizontalBox::Slot()
@@ -94,24 +96,24 @@ public:
 				.VAlign(VAlign_Center)
 				[
 					SNew(SBox)
-					.WidthOverride(MultiBoxConstants::MenuIconSize)
-					.HeightOverride(MultiBoxConstants::MenuIconSize)
+					.WidthOverride(MenuIconSize)
+					.HeightOverride(MenuIconSize)
 					[
 						SNew(SImage)
-						.Image(FEditorStyle::GetBrush(PlatformInfo.GetIconStyleName(PlatformInfo::EPlatformIconSize::Normal)))
+						.Image(FEditorStyle::GetBrush(PlatformInfo.GetIconStyleName(EPlatformIconSize::Normal)))
 					]
 				]
 				+SOverlay::Slot()
-				.Padding(FMargin(MultiBoxConstants::MenuIconSize * 0.5f, 0, 0, 0))
+				.Padding(FMargin(MenuIconSize * 0.5f, 0, 0, 0))
 				.HAlign(HAlign_Left)
 				.VAlign(VAlign_Bottom)
 				[
 					SNew(SBox)
-					.WidthOverride(MultiBoxConstants::MenuIconSize)
-					.HeightOverride(MultiBoxConstants::MenuIconSize)
+					.WidthOverride(MenuIconSize)
+					.HeightOverride(MenuIconSize)
 					[
 						SNew(SImage)
-						.Visibility(TAttribute<EVisibility>::Create(TAttribute<EVisibility>::FGetter::CreateStatic(&Local::IsUnsupportedPlatformWarningVisible, PlatformInfo.VanillaPlatformName)))
+						.Visibility(TAttribute<EVisibility>::Create(TAttribute<EVisibility>::FGetter::CreateStatic(&Local::IsUnsupportedPlatformWarningVisible, PlatformInfo.VanillaInfo->PlatformInfoName)))
 						.Image(FEditorStyle::GetBrush("Launcher.Platform.Warning"))
 					]
 				]
@@ -129,12 +131,12 @@ public:
 
 	virtual bool ShowUnsupportedTargetWarning(const FName PlatformName) const override
 	{
-		const PlatformInfo::FPlatformInfo* const PlatformInfo = PlatformInfo::FindPlatformInfo(PlatformName);
+		const PlatformInfo::FTargetPlatformInfo* const PlatformInfo = PlatformInfo::FindPlatformInfo(PlatformName);
 		check(PlatformInfo);
 
 		// Don't show the warning during automation testing; the dlg is modal and blocks
 		FProjectStatus ProjectStatus;
-		if(!GIsAutomationTesting && IProjectManager::Get().QueryStatusForCurrentProject(ProjectStatus) && !ProjectStatus.IsTargetPlatformSupported(PlatformInfo->VanillaPlatformName))
+		if(!GIsAutomationTesting && IProjectManager::Get().QueryStatusForCurrentProject(ProjectStatus) && !ProjectStatus.IsTargetPlatformSupported(PlatformInfo->VanillaInfo->PlatformInfoName))
 		{
 			FFormatNamedArguments Args;
 			Args.Add(TEXT("DisplayName"), PlatformInfo->DisplayName);
diff --git a/Engine/Source/Editor/ProjectTargetPlatformEditor/Private/Widgets/SProjectTargetPlatformSettings.cpp b/Engine/Source/Editor/ProjectTargetPlatformEditor/Private/Widgets/SProjectTargetPlatformSettings.cpp
--- a/Engine/Source/Editor/ProjectTargetPlatformEditor/Private/Widgets/SProjectTargetPlatformSettings.cpp
+++ b/Engine/Source/Editor/ProjectTargetPlatformEditor/Private/Widgets/SProjectTargetPlatformSettings.cpp
@@ -22,27 +22,27 @@ void SProjectTargetPlatformSettings::Construct(const FArguments& InArgs)
 {
 	// Create and sort a list of vanilla platforms that are game targets (sort by display name)
 	// We show all of the platforms regardless of whether we have an SDK installed for them or not
-	for(const PlatformInfo::FPlatformInfo& PlatformInfo : PlatformInfo::GetPlatformInfoArray())
+	for(const PlatformInfo::FTargetPlatformInfo* PlatformInfo : PlatformInfo::GetPlatformInfoArray())
 	{
-		if(PlatformInfo.IsVanilla() && PlatformInfo.PlatformType == EBuildTargetType::Game)
+		if(PlatformInfo->IsVanilla() && PlatformInfo->PlatformType == EBuildTargetType::Game)
 		{
-			AvailablePlatforms.Add(&PlatformInfo);
+			AvailablePlatforms.Add(PlatformInfo);
 		}
 	}
 
-	AvailablePlatforms.Sort([](const PlatformInfo::FPlatformInfo& One, const PlatformInfo::FPlatformInfo& Two) -> bool
+	AvailablePlatforms.Sort([](const PlatformInfo::FTargetPlatformInfo& One, const PlatformInfo::FTargetPlatformInfo& Two) -> bool
 	{
 		return One.DisplayName.CompareTo(Two.DisplayName) < 0;
 	});
 
 	// Generate a widget for each platform
 	TSharedRef<SVerticalBox> PlatformsListBox = SNew(SVerticalBox);
-	for(const PlatformInfo::FPlatformInfo* AvailablePlatform : AvailablePlatforms)
+	for(const PlatformInfo::FTargetPlatformInfo* AvailablePlatform : AvailablePlatforms)
 	{
 		PlatformsListBox->AddSlot()
 		.AutoHeight()
 		[
-			MakePlatformRow(AvailablePlatform->DisplayName, AvailablePlatform->PlatformInfoName, FEditorStyle::GetBrush(AvailablePlatform->GetIconStyleName(PlatformInfo::EPlatformIconSize::Normal)))
+			MakePlatformRow(AvailablePlatform->DisplayName, AvailablePlatform->PlatformInfoName, FEditorStyle::GetBrush(AvailablePlatform->GetIconStyleName(EPlatformIconSize::Normal)))
 		];
 	}
 
@@ -187,7 +187,7 @@ void SProjectTargetPlatformSettings::HandlePlatformCheckBoxStateChanged(ECheckBo
 		else
 		{
 			// We've deselected "All Platforms", so manually select every available platform
-			for(const PlatformInfo::FPlatformInfo* AvailablePlatform : AvailablePlatforms)
+			for(const PlatformInfo::FTargetPlatformInfo* AvailablePlatform : AvailablePlatforms)
 			{
 				FGameProjectGenerationModule::Get().UpdateSupportedTargetPlatforms(AvailablePlatform->TargetPlatformName, true);
 			}
diff --git a/Engine/Source/Editor/ProjectTargetPlatformEditor/Private/Widgets/SProjectTargetPlatformSettings.h b/Engine/Source/Editor/ProjectTargetPlatformEditor/Private/Widgets/SProjectTargetPlatformSettings.h
--- a/Engine/Source/Editor/ProjectTargetPlatformEditor/Private/Widgets/SProjectTargetPlatformSettings.h
+++ b/Engine/Source/Editor/ProjectTargetPlatformEditor/Private/Widgets/SProjectTargetPlatformSettings.h
@@ -34,5 +34,5 @@ private:
 	/** Handle the "enabled" checkbox state being changed for this platform */
 	void HandlePlatformCheckBoxStateChanged(ECheckBoxState InState, const FName PlatformName) const;
 
-	TArray<const PlatformInfo::FPlatformInfo*> AvailablePlatforms;
+	TArray<const PlatformInfo::FTargetPlatformInfo*> AvailablePlatforms;
 };
diff --git a/Engine/Source/Editor/ProjectTargetPlatformEditor/Public/Interfaces/IProjectTargetPlatformEditorModule.h b/Engine/Source/Editor/ProjectTargetPlatformEditor/Public/Interfaces/IProjectTargetPlatformEditorModule.h
--- a/Engine/Source/Editor/ProjectTargetPlatformEditor/Public/Interfaces/IProjectTargetPlatformEditorModule.h
+++ b/Engine/Source/Editor/ProjectTargetPlatformEditor/Public/Interfaces/IProjectTargetPlatformEditorModule.h
@@ -49,7 +49,7 @@ public:
 	 * @param bForCheckBox - true if the widget is for a checkbox menu item, false if it's for any other menu item (this affects the padding)
 	 * @param DisplayNameOverride - An alternate name to use for the platform, otherwise Platform->DisplayName() will be used
 	 */
-	virtual TSharedRef<SWidget> MakePlatformMenuItemWidget(const PlatformInfo::FPlatformInfo& PlatformInfo, const bool bForCheckBox = false, const FText& DisplayNameOverride = FText()) const = 0;
+	virtual TSharedRef<SWidget> MakePlatformMenuItemWidget(const PlatformInfo::FTargetPlatformInfo& PlatformInfo, const bool bForCheckBox = false, const FText& DisplayNameOverride = FText()) const = 0;
 
 	/**
 	 * Check to see if the given platform is on the list of supported targets, and show a warning message if it's not, allowing the user to continue or cancel
diff --git a/Engine/Source/Editor/PropertyEditor/Private/SPerPlatformPropertiesWidget.cpp b/Engine/Source/Editor/PropertyEditor/Private/SPerPlatformPropertiesWidget.cpp
--- a/Engine/Source/Editor/PropertyEditor/Private/SPerPlatformPropertiesWidget.cpp
+++ b/Engine/Source/Editor/PropertyEditor/Private/SPerPlatformPropertiesWidget.cpp
@@ -62,7 +62,7 @@ void SPerPlatformPropertiesWidget::ConstructChildren()
 		const TArray<FString> Filters = { TEXT("NoEditor"), TEXT("Client"), TEXT("Server") };
 		for (const FName& PlatformName : VanillaPlatformNameArray)
 		{
-			const PlatformInfo::FPlatformInfo* PlatformInfo = PlatformInfo::FindPlatformInfo(PlatformName);
+			const PlatformInfo::FTargetPlatformInfo* PlatformInfo = PlatformInfo::FindPlatformInfo(PlatformName);
 			FString PlatformNameString = PlatformName.ToString();
 			for ( const FString& Filter : Filters)
 			{
@@ -79,7 +79,7 @@ void SPerPlatformPropertiesWidget::ConstructChildren()
 			if (PlatformNameString.Len() && !PlatformOverrides.Contains(FilteredName))
 			{	
 				BasePlatformNameArray.AddUnique(FilteredName);
-				GroupToPlatform.AddUnique(PlatformInfo->PlatformGroupName, FilteredName);
+				GroupToPlatform.AddUnique(PlatformInfo->DataDrivenPlatformInfo->PlatformGroupName, FilteredName);
 			}
 		}
 
diff --git a/Engine/Source/Editor/UnrealEd/Classes/Editor/EditorEngine.h b/Engine/Source/Editor/UnrealEd/Classes/Editor/EditorEngine.h
--- a/Engine/Source/Editor/UnrealEd/Classes/Editor/EditorEngine.h
+++ b/Engine/Source/Editor/UnrealEd/Classes/Editor/EditorEngine.h
@@ -805,7 +805,7 @@ public:
 	virtual void SetMapBuildCancelled(bool InCancelled) override { /* Intentionally empty. */ }
 	virtual void HandleNetworkFailure(UWorld *World, UNetDriver *NetDriver, ENetworkFailure::Type FailureType, const FString& ErrorString) override;
 	virtual ERHIFeatureLevel::Type GetDefaultWorldFeatureLevel() const override { return DefaultWorldFeatureLevel; }
-	virtual bool GetPreviewPlatformName(FName& PlatformGroupName, FName& VanillaPlatformName) const override;
+	virtual bool GetPreviewPlatformName(FName& PlatformName) const override;
 
 protected:
 	virtual void InitializeObjectReferences() override;
diff --git a/Engine/Source/Editor/UnrealEd/Private/CookOnTheFlyServer.cpp b/Engine/Source/Editor/UnrealEd/Private/CookOnTheFlyServer.cpp
--- a/Engine/Source/Editor/UnrealEd/Private/CookOnTheFlyServer.cpp
+++ b/Engine/Source/Editor/UnrealEd/Private/CookOnTheFlyServer.cpp
@@ -6653,7 +6653,7 @@ void UCookOnTheFlyServer::StartCookByTheBook( const FCookByTheBookStartupOptions
 		for (const ITargetPlatform* Platform : TargetPlatforms)
 		{
 			FString UBTPlatformName;
-			Platform->GetPlatformInfo().UBTTargetId.ToString(UBTPlatformName);
+			Platform->GetPlatformInfo().UBTPlatformName.ToString(UBTPlatformName);
 			UBTPlatformStrings.Emplace(MoveTemp(UBTPlatformName));
 		}
 
@@ -6721,7 +6721,7 @@ void UCookOnTheFlyServer::StartCookByTheBook( const FCookByTheBookStartupOptions
 		// if we are cooking dlc we must be based on a release version cook
 		check( !BasedOnReleaseVersion.IsEmpty() );
 
-		for ( const ITargetPlatform* TargetPlatform: TargetPlatforms )
+		for ( const ITargetPlatform* TargetPlatform : TargetPlatforms )
 		{
 			FString PlatformNameString = TargetPlatform->PlatformName();
 			FName PlatformName(*PlatformNameString);
@@ -6738,13 +6738,11 @@ void UCookOnTheFlyServer::StartCookByTheBook( const FCookByTheBookStartupOptions
 
 			if (!bSucceeded)
 			{
-				using namespace PlatformInfo;
 				// Check all possible flavors 
 				// For example release version could be cooked as Android_ASTC flavor, but DLC can be made as Android_ETC2
-				FVanillaPlatformEntry VanillaPlatfromEntry = BuildPlatformHierarchy(PlatformName, EPlatformFilter::CookFlavor);
-				for (const FPlatformInfo* PlatformFlaworInfo : VanillaPlatfromEntry.PlatformFlavors)
+				for (const PlatformInfo::FTargetPlatformInfo* PlatformFlavorInfo : TargetPlatform->GetTargetPlatformInfo().Flavors)
 				{
-					OriginalSandboxRegistryFilename = GetBasedOnReleaseVersionAssetRegistryPath(BasedOnReleaseVersion, PlatformFlaworInfo->PlatformInfoName.ToString()) / GetAssetRegistryFilename();
+					OriginalSandboxRegistryFilename = GetBasedOnReleaseVersionAssetRegistryPath(BasedOnReleaseVersion, PlatformFlavorInfo->PlatformInfoName.ToString()) / GetAssetRegistryFilename();
 					bSucceeded = GetAllPackageFilenamesFromAssetRegistry(OriginalSandboxRegistryFilename, PackageList);
 					if (bSucceeded)
 					{
diff --git a/Engine/Source/Editor/UnrealEd/Private/EditorEngine.cpp b/Engine/Source/Editor/UnrealEd/Private/EditorEngine.cpp
--- a/Engine/Source/Editor/UnrealEd/Private/EditorEngine.cpp
+++ b/Engine/Source/Editor/UnrealEd/Private/EditorEngine.cpp
@@ -7675,14 +7675,12 @@ void UEditorEngine::SaveEditorFeatureLevel()
 	Settings->PostEditChange();
 }
 
-bool UEditorEngine::GetPreviewPlatformName(FName& PlatformGroupName, FName& VanillaPlatformName) const
+bool UEditorEngine::GetPreviewPlatformName(FName& PlatformName) const
 {
-	FName PlatformName = PreviewPlatform.GetEffectivePreviewPlatformName();
-	const PlatformInfo::FPlatformInfo* PlatInfo;
-	if (PlatformName != NAME_None && (PlatInfo = PlatformInfo::FindPlatformInfo(PlatformName)) != nullptr)
+	FName PreviewPlatformName = PreviewPlatform.GetEffectivePreviewPlatformName();
+	if (PreviewPlatformName != NAME_None)
 	{
-		VanillaPlatformName = PlatInfo->VanillaPlatformName;
-		PlatformGroupName = PlatInfo->PlatformGroupName;
+		PlatformName = PreviewPlatformName;
 		return true;
 	}
 
diff --git a/Engine/Source/Editor/UnrealEd/Private/Kismet2/DebuggerCommands.cpp b/Engine/Source/Editor/UnrealEd/Private/Kismet2/DebuggerCommands.cpp
--- a/Engine/Source/Editor/UnrealEd/Private/Kismet2/DebuggerCommands.cpp
+++ b/Engine/Source/Editor/UnrealEd/Private/Kismet2/DebuggerCommands.cpp
@@ -996,7 +996,7 @@ TSharedRef< SWidget > FPlayWorldCommands::GeneratePlayMenuContent(TSharedRef<FUI
  * Create an All_<platform>_devices_on_<host> submenu
  * can be extended to any othe All <Platform> aggregate proxy
 */
-static void MakeAllDevicesSubMenu(FMenuBuilder& InMenuBuilder, const PlatformInfo::FPlatformInfo* InPlatformInfo, const TSharedPtr<ITargetDeviceProxy> DeviceProxy)
+static void MakeAllDevicesSubMenu(FMenuBuilder& InMenuBuilder, const PlatformInfo::FTargetPlatformInfo* InPlatformInfo, const TSharedPtr<ITargetDeviceProxy> DeviceProxy)
 {
 	ITargetDeviceServicesModule* TargetDeviceServicesModule = static_cast<ITargetDeviceServicesModule*>(FModuleManager::Get().LoadModule(TEXT("TargetDeviceServices")));
 	IProjectTargetPlatformEditorModule& ProjectTargetPlatformEditorModule = FModuleManager::LoadModuleChecked<IProjectTargetPlatformEditorModule>("ProjectTargetPlatformEditor");
@@ -1008,9 +1008,9 @@ static void MakeAllDevicesSubMenu(FMenuBuilder& InMenuBuilder, const PlatformInf
 		FName Variant = *It;
 
 		// for an aggregate (All_<platform>_devices_on_<host>) proxy, allow only the "Android_<texture_compression>" variants
-		const PlatformInfo::FPlatformInfo* platformInfo = PlatformInfo::FindPlatformInfo(Variant);
-		if (DeviceProxy->IsAggregated() && platformInfo != NULL &&
-			(Variant == platformInfo->VanillaPlatformName || platformInfo->PlatformType != EBuildTargetType::Game))
+		const PlatformInfo::FTargetPlatformInfo* PlatformInfo = PlatformInfo::FindPlatformInfo(Variant);
+		if (DeviceProxy->IsAggregated() && PlatformInfo != NULL &&
+			(Variant == PlatformInfo->VanillaInfo->PlatformInfoName || PlatformInfo->PlatformType != EBuildTargetType::Game))
 		{
 			continue;
 		}
@@ -1087,9 +1087,9 @@ static void TurnkeyInstallSdk(FString PlatformName, bool bPreferFull, bool bForc
 
 	FText TaskName = LOCTEXT("InstallingSdk", "Installing Sdk");
 	IUATHelperModule::Get().CreateUatTask(CommandLine, FText::FromString(PlatformName), TaskName, TaskName, FEditorStyle::GetBrush(TEXT("MainFrame.PackageProject")),
-		[](FString, double)
+		[PlatformName](FString, double)
 	{
-		AsyncTask(ENamedThreads::GameThread, []()
+		AsyncTask(ENamedThreads::GameThread, [PlatformName]()
 		{
 
 			// read in env var changes
@@ -1124,16 +1124,19 @@ static void TurnkeyInstallSdk(FString PlatformName, bool bPreferFull, bool bForc
 			}
 
 			// update the Sdk status
-			FDataDrivenPlatformInfoRegistry::UpdateSdkStatus();
+//			FDataDrivenPlatformInfoRegistry::UpdateSdkStatus();
+			GetTargetPlatformManager()->UpdateAfterSDKInstall(PlatformName);
+			RenderUtilsInit();
 
-			// show restart dialog
-			FModuleManager::GetModuleChecked<ISettingsEditorModule>("SettingsEditor").OnApplicationRestartRequired();
+
+ 			// show restart dialog
+ 			FModuleManager::GetModuleChecked<ISettingsEditorModule>("SettingsEditor").OnApplicationRestartRequired();
 		});
 	}
 	);
 }
 
-static TAttribute<FText> MakeSdkStatusAttribute(const PlatformInfo::FPlatformInfo* PlatformInfo, TSharedPtr< ITargetDeviceProxy> DeviceProxy)
+static TAttribute<FText> MakeSdkStatusAttribute(const PlatformInfo::FTargetPlatformInfo* PlatformInfo, TSharedPtr< ITargetDeviceProxy> DeviceProxy)
 {
 // 	FFormatNamedArguments LabelArguments;
 // 	LabelArguments.Add(TEXT("DeviceName"), FText::FromString(DeviceProxy->GetName()));
@@ -1156,12 +1159,13 @@ static TAttribute<FText> MakeSdkStatusAttribute(const PlatformInfo::FPlatformInf
 
 
 	FText DisplayString = DeviceProxy ? FText::FromString(DeviceProxy->GetName()) : PlatformInfo->DisplayName;
+	FString DeviceId = DeviceProxy ? DeviceProxy->GetTargetDeviceId(NAME_None) : FString();
 
-	return TAttribute<FText>::Create(TAttribute<FText>::FGetter::CreateLambda([PlatformInfo, DisplayString]()
+	return TAttribute<FText>::Create(TAttribute<FText>::FGetter::CreateLambda([IniPlatformName=PlatformInfo->IniPlatformName, DisplayString, DeviceId]()
 	{
 		// get the status, or Unknown if it's not there
-		const FDataDrivenPlatformInfoRegistry::FPlatformInfo& Info = FDataDrivenPlatformInfoRegistry::GetPlatformInfo(PlatformInfo->IniPlatformName);
-		DDPIPlatformSdkStatus Status = Info.SdkStatus;
+		const FDataDrivenPlatformInfo& Info = FDataDrivenPlatformInfoRegistry::GetPlatformInfo(IniPlatformName);
+		DDPIPlatformSdkStatus Status = DeviceId.Len() ? Info.GetStatusForDeviceId(DeviceId) : Info.GetSdkStatus();
 
 		// @todo turnkey: Have premade FText's by SdkStatus for speed
 		FText Desc = FText::FromString(
@@ -1178,14 +1182,14 @@ static TAttribute<FText> MakeSdkStatusAttribute(const PlatformInfo::FPlatformInf
 }
 
 
-static void MakeTurnkeyPlatformMenu(FMenuBuilder& MenuBuilder, const PlatformInfo::FPlatformInfo* PlatformInfo, ITargetDeviceServicesModule* TargetDeviceServicesModule)
+static void MakeTurnkeyPlatformMenu(FMenuBuilder& MenuBuilder, const PlatformInfo::FTargetPlatformInfo* PlatformInfo, ITargetDeviceServicesModule* TargetDeviceServicesModule)
 {
 	MenuBuilder.BeginSection("AllDevices", LOCTEXT("TurnkeySection_AllDevices", "All Devices"));
 
 	TArray<TSharedPtr<ITargetDeviceProxy>> DeviceProxies;
-	TargetDeviceServicesModule->GetDeviceProxyManager()->GetAllProxies(*PlatformInfo->IniPlatformName, DeviceProxies);
+	TargetDeviceServicesModule->GetDeviceProxyManager()->GetAllProxies(PlatformInfo->IniPlatformName, DeviceProxies);
 
-	FString PlatformName = PlatformInfo->UBTTargetId.ToString();
+	FString PlatformName = PlatformInfo->DataDrivenPlatformInfo->UBTPlatformString;
 	for (const TSharedPtr<ITargetDeviceProxy> Proxy : DeviceProxies)
 	{
 		FString DeviceName = Proxy->GetName();
@@ -1214,9 +1218,9 @@ static void MakeTurnkeyPlatformMenu(FMenuBuilder& MenuBuilder, const PlatformInf
 
 	MenuBuilder.BeginSection("AllDevices", LOCTEXT("TurnkeySection_Sdks", "Sdk Managment"));
 
-	const FDataDrivenPlatformInfoRegistry::FPlatformInfo& Info = FDataDrivenPlatformInfoRegistry::GetPlatformInfo(PlatformInfo->IniPlatformName);
+	const FDataDrivenPlatformInfo& Info = FDataDrivenPlatformInfoRegistry::GetPlatformInfo(PlatformInfo->IniPlatformName);
 	const TCHAR* NoDevice = nullptr;
-	if (Info.SdkStatus == DDPIPlatformSdkStatus::OutOfDate)
+	if (Info.GetSdkStatus() == DDPIPlatformSdkStatus::OutOfDate)
 	{
 		MenuBuilder.AddMenuEntry(
 			LOCTEXT("Turnkey_InstallSdkMinimal", "Update Sdk (Prefer Minimal)"),
@@ -1232,7 +1236,7 @@ static void MakeTurnkeyPlatformMenu(FMenuBuilder& MenuBuilder, const PlatformInf
 			FExecuteAction::CreateStatic(TurnkeyInstallSdk, PlatformName, true, false, NoDevice)
 		);
 	}
-	else if (Info.SdkStatus == DDPIPlatformSdkStatus::Valid)
+	else if (Info.GetSdkStatus() == DDPIPlatformSdkStatus::Valid)
 	{
 		MenuBuilder.AddMenuEntry(
 			LOCTEXT("Turnkey_InstallSdkMinimal", "Force Reinstall Sdk (Prefer Minimal)"),
@@ -1266,7 +1270,7 @@ static void MakeTurnkeyPlatformMenu(FMenuBuilder& MenuBuilder, const PlatformInf
 	}
 }
 
-static void GenerateDeviceProxyMenuParams(TSharedPtr<ITargetDeviceProxy> DeviceProxy, const PlatformInfo::FPlatformInfo* EditorPlatformInfo, FName PlatformName, FUIAction& OutAction, FText& OutTooltip)
+static void GenerateDeviceProxyMenuParams(TSharedPtr<ITargetDeviceProxy> DeviceProxy, const PlatformInfo::FTargetPlatformInfo* EditorPlatformInfo, FName PlatformName, FUIAction& OutAction, FText& OutTooltip)
 {
 	// 	// create an All_<platform>_devices_on_<host> submenu
 	// 	if (DeviceProxy->IsAggregated())
@@ -1314,13 +1318,6 @@ TSharedRef< SWidget > FPlayWorldCommands::GenerateTurnkeyMenuContent(TSharedRef<
 	const bool bShouldCloseWindowAfterMenuSelection = true;
 	FMenuBuilder MenuBuilder(bShouldCloseWindowAfterMenuSelection, InCommandList);
 
-	TArray<PlatformInfo::FVanillaPlatformEntry> VanillaPlatforms = PlatformInfo::BuildPlatformHierarchy(PlatformInfo::EPlatformFilter::All);
-
-	VanillaPlatforms.Sort([](const PlatformInfo::FVanillaPlatformEntry& One, const PlatformInfo::FVanillaPlatformEntry& Two) -> bool
-	{
-		return One.PlatformInfo->DisplayName.CompareTo(Two.PlatformInfo->DisplayName) < 0;
-	});
-
 	// shared devices section
 	ITargetDeviceServicesModule* TargetDeviceServicesModule = static_cast<ITargetDeviceServicesModule*>(FModuleManager::Get().LoadModule(TEXT("TargetDeviceServices")));
 	IProjectTargetPlatformEditorModule& ProjectTargetPlatformEditorModule = FModuleManager::LoadModuleChecked<IProjectTargetPlatformEditorModule>("ProjectTargetPlatformEditor");
@@ -1337,10 +1334,10 @@ TSharedRef< SWidget > FPlayWorldCommands::GenerateTurnkeyMenuContent(TSharedRef<
 			}
 
 			FName PlatformName(*Pair.Key);
-			const FDataDrivenPlatformInfoRegistry::FPlatformInfo& Info = Pair.Value;
+			const FDataDrivenPlatformInfo& Info = Pair.Value;
 
 			// find a matching editor-facing FPlatformInfo object
-			const PlatformInfo::FPlatformInfo* EditorPlatformInfo = PlatformInfo::FindVanillaPlatformInfo(PlatformName);
+			const PlatformInfo::FTargetPlatformInfo* EditorPlatformInfo = PlatformInfo::FindVanillaPlatformInfo(PlatformName);
 			check(EditorPlatformInfo != nullptr);
 
 			// look for devices for all platforms, even if the platform isn't installed - Turnkey can install Sdk after selecting LaunchOn
@@ -1363,7 +1360,7 @@ TSharedRef< SWidget > FPlayWorldCommands::GenerateTurnkeyMenuContent(TSharedRef<
 					MenuBuilder.AddMenuEntry(
 						MakeSdkStatusAttribute(EditorPlatformInfo, DeviceProxies[0]),
 						Tooltip,
-						FSlateIcon(FEditorStyle::GetStyleSetName(), EditorPlatformInfo->GetIconStyleName(PlatformInfo::EPlatformIconSize::Normal)),
+						FSlateIcon(FEditorStyle::GetStyleSetName(), EditorPlatformInfo->GetIconStyleName(EPlatformIconSize::Normal)),
 						Action,
 						NAME_None,
 						EUserInterfaceActionType::Button
@@ -1387,7 +1384,7 @@ TSharedRef< SWidget > FPlayWorldCommands::GenerateTurnkeyMenuContent(TSharedRef<
 								SubMenuBuilder.AddMenuEntry(
 									MakeSdkStatusAttribute(EditorPlatformInfo, Proxy),
 									SubTooltip,
-									FSlateIcon(FEditorStyle::GetStyleSetName(), EditorPlatformInfo->GetIconStyleName(PlatformInfo::EPlatformIconSize::Normal)),
+									FSlateIcon(FEditorStyle::GetStyleSetName(), EditorPlatformInfo->GetIconStyleName(EPlatformIconSize::Normal)),
 									SubAction,
 									NAME_None,
 									EUserInterfaceActionType::Button
@@ -1398,7 +1395,7 @@ TSharedRef< SWidget > FPlayWorldCommands::GenerateTurnkeyMenuContent(TSharedRef<
 						NAME_None,
 						EUserInterfaceActionType::Check,
 						false,
-						FSlateIcon(FEditorStyle::GetStyleSetName(), EditorPlatformInfo->GetIconStyleName(PlatformInfo::EPlatformIconSize::Normal)),
+						FSlateIcon(FEditorStyle::GetStyleSetName(), EditorPlatformInfo->GetIconStyleName(EPlatformIconSize::Normal)),
 						true
 					);
 				}
@@ -1407,7 +1404,7 @@ TSharedRef< SWidget > FPlayWorldCommands::GenerateTurnkeyMenuContent(TSharedRef<
 				for (const TSharedPtr<ITargetDeviceProxy> Proxy : DeviceProxies)
 				{
 					FString DeviceId = Proxy->GetTargetDeviceId(NAME_None);
-					if (FDataDrivenPlatformInfoRegistry::GetPlatformInfo(EditorPlatformInfo->IniPlatformName).GetStatusForDeviceId(DeviceId) == DDPIPlatformSdkStatus::Unknown)
+					if (EditorPlatformInfo->DataDrivenPlatformInfo->GetStatusForDeviceId(DeviceId) == DDPIPlatformSdkStatus::Unknown)
 					{
 						DeviceIdsToQuery.Add(DeviceId);
 					}
@@ -1439,8 +1436,8 @@ TSharedRef< SWidget > FPlayWorldCommands::GenerateTurnkeyMenuContent(TSharedRef<
 		}
 
 		FName PlatformName(*Pair.Key);
-		const FDataDrivenPlatformInfoRegistry::FPlatformInfo& Info = Pair.Value;
-		const PlatformInfo::FPlatformInfo* EditorPlatformInfo = PlatformInfo::FindVanillaPlatformInfo(PlatformName);
+		const FDataDrivenPlatformInfo& Info = Pair.Value;
+		const PlatformInfo::FTargetPlatformInfo* EditorPlatformInfo = PlatformInfo::FindVanillaPlatformInfo(PlatformName);
 
 		FString DeviceNameToCheckForStatus;
 		MenuBuilder.AddSubMenu(
@@ -1448,7 +1445,7 @@ TSharedRef< SWidget > FPlayWorldCommands::GenerateTurnkeyMenuContent(TSharedRef<
 			FText::FromString(PlatformName.ToString()),
 			FNewMenuDelegate::CreateStatic(&MakeTurnkeyPlatformMenu, EditorPlatformInfo, TargetDeviceServicesModule),
 			false, 
-			FSlateIcon(FEditorStyle::GetStyleSetName(), EditorPlatformInfo->GetIconStyleName(PlatformInfo::EPlatformIconSize::Normal)),
+			FSlateIcon(FEditorStyle::GetStyleSetName(), EditorPlatformInfo->GetIconStyleName(EPlatformIconSize::Normal)),
 			true
 		);
 	}
@@ -1551,11 +1548,11 @@ TSharedRef< SWidget > FPlayWorldCommands::GenerateTurnkeyMenuContent(TSharedRef<
 
 void PopulateLaunchMenu(UToolMenu* Menu)
 {
-	TArray<PlatformInfo::FVanillaPlatformEntry> VanillaPlatforms = PlatformInfo::BuildPlatformHierarchy(PlatformInfo::EPlatformFilter::All);
+	TArray<PlatformInfo::FTargetPlatformInfo*> VanillaPlatforms = PlatformInfo::GetVanillaPlatformInfoArray();
 
-	VanillaPlatforms.Sort([](const PlatformInfo::FVanillaPlatformEntry& One, const PlatformInfo::FVanillaPlatformEntry& Two) -> bool
+	Algo::Sort(VanillaPlatforms, [](const PlatformInfo::FTargetPlatformInfo* One, const PlatformInfo::FTargetPlatformInfo* Two) -> bool
 	{
-		return One.PlatformInfo->DisplayName.CompareTo(Two.PlatformInfo->DisplayName) < 0;
+		return One->DisplayName.CompareTo(Two->DisplayName) < 0;
 	});
 
 	// shared devices section
@@ -1571,33 +1568,33 @@ void PopulateLaunchMenu(UToolMenu* Menu)
 	PlatformsToCheckFlavorsFor.Add(TEXT("Android"));
 	PlatformsToCheckFlavorsFor.Add(TEXT("IOS"));
 	TArray<FName> PlatformsWithNoDevices;
-	TArray<PlatformInfo::FPlatformInfo> PlatformsToAddInstallLinksFor;
+	TArray<const PlatformInfo::FTargetPlatformInfo*> PlatformsToAddInstallLinksFor;
 	EProjectType ProjectType = FGameProjectGenerationModule::Get().ProjectHasCodeFiles() ? EProjectType::Code : EProjectType::Content;
 
 	{
 		FToolMenuSection& Section = Menu->AddSection("LevelEditorLaunchDevices", LOCTEXT("LaunchButtonDevicesSection", "Devices"));
-		for (const PlatformInfo::FVanillaPlatformEntry& VanillaPlatform : VanillaPlatforms)
+		for (const PlatformInfo::FTargetPlatformInfo* VanillaPlatform : VanillaPlatforms)
 		{
 			// for the Editor we are only interested in launching standalone games
-			if (VanillaPlatform.PlatformInfo->PlatformType != EBuildTargetType::Game || !VanillaPlatform.PlatformInfo->bEnabledForUse || !FInstalledPlatformInfo::Get().CanDisplayPlatform(VanillaPlatform.PlatformInfo->BinaryFolderName, ProjectType))
+			if (VanillaPlatform->PlatformType != EBuildTargetType::Game || !VanillaPlatform->DataDrivenPlatformInfo->bEnabledForUse || !FInstalledPlatformInfo::Get().CanDisplayPlatform(VanillaPlatform->DataDrivenPlatformInfo->UBTPlatformString, ProjectType))
 			{
 				continue;
 			}
 
-			if (VanillaPlatform.PlatformInfo->SDKStatus == PlatformInfo::EPlatformSDKStatus::Installed)
+			if (VanillaPlatform->DataDrivenPlatformInfo->GetSdkStatus() == DDPIPlatformSdkStatus::Valid)
 			{
 				// for each platform...
 				TArray<TSharedPtr<ITargetDeviceProxy>> DeviceProxies;
 				// the list of proxies include the "Al_Android" entry
-				TargetDeviceServicesModule->GetDeviceProxyManager()->GetAllProxies(VanillaPlatform.PlatformInfo->VanillaPlatformName, DeviceProxies);
+				TargetDeviceServicesModule->GetDeviceProxyManager()->GetAllProxies(VanillaPlatform->PlatformInfoName, DeviceProxies);
 
 				// if this platform had no devices, but we want to show an extra option if not installed right
 				if (DeviceProxies.Num() == 0)
 				{
-					if (PlatformsWithNoDevices.Find(VanillaPlatform.PlatformInfo->VanillaPlatformName) == INDEX_NONE)
+					if (PlatformsWithNoDevices.Find(VanillaPlatform->PlatformInfoName) == INDEX_NONE)
 					{
 						// add an entry with a no devices found
-						PlatformsWithNoDevices.Add(VanillaPlatform.PlatformInfo->VanillaPlatformName);
+						PlatformsWithNoDevices.Add(VanillaPlatform->PlatformInfoName);
 					}
 				}
 				else
@@ -1611,13 +1608,13 @@ void PopulateLaunchMenu(UToolMenu* Menu)
 						if (DeviceProxy->IsAggregated())
 						{
 							FString AggregateDevicedName(FString::Printf(TEXT("  %s"), *DeviceProxy->GetName())); //align with the other menu entries
-							FSlateIcon AggregateDeviceIcon(FEditorStyle::GetStyleSetName(), VanillaPlatform.PlatformInfo->GetIconStyleName(PlatformInfo::EPlatformIconSize::Normal));
+							FSlateIcon AggregateDeviceIcon(FEditorStyle::GetStyleSetName(), VanillaPlatform->GetIconStyleName(EPlatformIconSize::Normal));
 
 							Section.AddSubMenu(
 								NAME_None,
 								FText::FromString(AggregateDevicedName),
 								FText::FromString(AggregateDevicedName),
-								FNewMenuDelegate::CreateStatic(&MakeAllDevicesSubMenu, VanillaPlatform.PlatformInfo, DeviceProxy),
+								FNewMenuDelegate::CreateStatic(&MakeAllDevicesSubMenu, VanillaPlatform, DeviceProxy),
 								false, AggregateDeviceIcon, true
 							);
 							continue;
@@ -1652,7 +1649,7 @@ void PopulateLaunchMenu(UToolMenu* Menu)
 						// ... generate tooltip text
 						FFormatNamedArguments TooltipArguments;
 						TooltipArguments.Add(TEXT("DeviceID"), FText::FromString(DeviceProxy->GetName()));
-						TooltipArguments.Add(TEXT("DisplayName"), VanillaPlatform.PlatformInfo->DisplayName);
+						TooltipArguments.Add(TEXT("DisplayName"), VanillaPlatform->DisplayName);
 						FText Tooltip = FText::Format(LOCTEXT("LaunchDeviceToolTipText_ThisDevice", "Launch the game on this {DisplayName} device ({DeviceID})"), TooltipArguments);
 						if (!DeviceProxy->IsAuthorized())
 						{
@@ -1660,7 +1657,7 @@ void PopulateLaunchMenu(UToolMenu* Menu)
 						}
 
 						FProjectStatus ProjectStatus;
-						if (IProjectManager::Get().QueryStatusForCurrentProject(ProjectStatus) && !ProjectStatus.IsTargetPlatformSupported(VanillaPlatform.PlatformInfo->VanillaPlatformName))
+						if (IProjectManager::Get().QueryStatusForCurrentProject(ProjectStatus) && !ProjectStatus.IsTargetPlatformSupported(VanillaPlatform->PlatformInfoName))
 						{
 							FText TooltipLine2 = FText::Format(LOCTEXT("LaunchDevicePlatformWarning", "{DisplayName} is not listed as a target platform for this project, so may not run as expected."), TooltipArguments);
 							Tooltip = FText::Format(FText::FromString(TEXT("{0}\n\n{1}")), Tooltip, TooltipLine2);
@@ -1670,7 +1667,7 @@ void PopulateLaunchMenu(UToolMenu* Menu)
 						FToolMenuEntry& Entry = Section.AddEntry(FToolMenuEntry::InitMenuEntry(
 							NAME_None,
 							LaunchDeviceAction,
-							ProjectTargetPlatformEditorModule.MakePlatformMenuItemWidget(*VanillaPlatform.PlatformInfo, true, Label)
+							ProjectTargetPlatformEditorModule.MakePlatformMenuItemWidget(*VanillaPlatform, true, Label)
 						));
 						Entry.ToolTip = Tooltip;
 						Entry.UserInterfaceActionType = EUserInterfaceActionType::Check;
@@ -1680,9 +1677,9 @@ void PopulateLaunchMenu(UToolMenu* Menu)
 			else
 			{
 				// if the platform wasn't installed, we'll add a menu item later (we never care about code in this case, since we don't compile)
-				if (PlatformsToMaybeInstallLinksFor.Find(VanillaPlatform.PlatformInfo->VanillaPlatformName.ToString()) != INDEX_NONE)
+				if (PlatformsToMaybeInstallLinksFor.Find(VanillaPlatform->PlatformInfoName.ToString()) != INDEX_NONE)
 				{
-					PlatformsToAddInstallLinksFor.Add(*(VanillaPlatform.PlatformInfo));
+					PlatformsToAddInstallLinksFor.Add(VanillaPlatform);
 				}
 			}
 		}
@@ -1752,7 +1749,7 @@ void PopulateLaunchMenu(UToolMenu* Menu)
 			FToolMenuSection& Section = Menu->AddSection("NoDevices");
 			for (int32 PlatformIndex = 0; PlatformIndex < PlatformsWithNoDevices.Num(); PlatformIndex++)
 			{
-				const PlatformInfo::FPlatformInfo* PlatformInfo = PlatformInfo::FindVanillaPlatformInfo(PlatformsWithNoDevices[PlatformIndex]);
+				const PlatformInfo::FTargetPlatformInfo* PlatformInfo = PlatformInfo::FindVanillaPlatformInfo(PlatformsWithNoDevices[PlatformIndex]);
 
 				// ... generate display label...
 				FFormatNamedArguments LabelArguments;
@@ -1800,14 +1797,12 @@ void PopulateLaunchMenu(UToolMenu* Menu)
 	{
 		{
 			FToolMenuSection& Section = Menu->AddSection("SDKUninstalledTutorials");
-			for (int32 PlatformIndex = 0; PlatformIndex < PlatformsToAddInstallLinksFor.Num(); PlatformIndex++)
+			for (const PlatformInfo::FTargetPlatformInfo* Platform : PlatformsToAddInstallLinksFor)
 			{
-				const PlatformInfo::FPlatformInfo& Platform = PlatformsToAddInstallLinksFor[PlatformIndex];
-
-				FUIAction Action(FExecuteAction::CreateStatic(&FInternalPlayWorldCommandCallbacks::HandleShowSDKTutorial, Platform.DisplayName.ToString(), Platform.SDKTutorial));
+				FUIAction Action(FExecuteAction::CreateStatic(&FInternalPlayWorldCommandCallbacks::HandleShowSDKTutorial, Platform->DisplayName.ToString(), Platform->DataDrivenPlatformInfo->SDKTutorial));
 
 				FFormatNamedArguments LabelArguments;
-				LabelArguments.Add(TEXT("PlatformName"), Platform.DisplayName);
+				LabelArguments.Add(TEXT("PlatformName"), Platform->DisplayName);
 				FText Label = FText::Format(LOCTEXT("LaunchPlatformLabel", "{PlatformName} Support"), LabelArguments);
 
 
@@ -2642,16 +2637,16 @@ bool FInternalPlayWorldCommandCallbacks::IsReadyToLaunchOnDevice(FString DeviceI
 	DeviceId.FindChar(TEXT('@'), Index);
 	FString PlatformName = DeviceId.Left(Index);
 
-	const PlatformInfo::FPlatformInfo* const PlatformInfo = PlatformInfo::FindPlatformInfo(FName(*PlatformName));
+	const PlatformInfo::FTargetPlatformInfo* const PlatformInfo = PlatformInfo::FindPlatformInfo(FName(*PlatformName));
 	checkf(PlatformInfo, TEXT("Unable to find PlatformInfo for %s"), *PlatformName);
 
 	FGameProjectGenerationModule& GameProjectModule = FModuleManager::LoadModuleChecked<FGameProjectGenerationModule>(TEXT("GameProjectGeneration"));
 	bool bHasCode = GameProjectModule.Get().ProjectHasCodeFiles();
 
-	if (PlatformInfo->SDKStatus == PlatformInfo::EPlatformSDKStatus::NotInstalled)
+	if (PlatformInfo->DataDrivenPlatformInfo->GetSdkStatus() != DDPIPlatformSdkStatus::Valid)
 	{
 		IMainFrameModule& MainFrameModule = FModuleManager::GetModuleChecked<IMainFrameModule>(TEXT("MainFrame"));
-		MainFrameModule.BroadcastMainFrameSDKNotInstalled(PlatformInfo->TargetPlatformName.ToString(), PlatformInfo->SDKTutorial);
+		MainFrameModule.BroadcastMainFrameSDKNotInstalled(PlatformInfo->TargetPlatformName.ToString(), PlatformInfo->DataDrivenPlatformInfo->SDKTutorial);
 		TArray<FAnalyticsEventAttribute> ParamArray;
 		ParamArray.Add(FAnalyticsEventAttribute(TEXT("Time"), 0.0));
 		FEditorAnalytics::ReportEvent(TEXT("Editor.LaunchOn.Failed"), PlatformInfo->TargetPlatformName.ToString(), bHasCode, EAnalyticsErrorCodes::SDKNotFound, ParamArray);
@@ -2770,31 +2765,54 @@ bool FInternalPlayWorldCommandCallbacks::IsReadyToLaunchOnDevice(FString DeviceI
 	else
 	{
 		IMainFrameModule& MainFrameModule = FModuleManager::GetModuleChecked<IMainFrameModule>(TEXT("MainFrame"));
-		MainFrameModule.BroadcastMainFrameSDKNotInstalled(PlatformInfo->TargetPlatformName.ToString(), PlatformInfo->SDKTutorial);
+		MainFrameModule.BroadcastMainFrameSDKNotInstalled(PlatformInfo->TargetPlatformName.ToString(), PlatformInfo->DataDrivenPlatformInfo->SDKTutorial);
 		return false;
 	}
 
 	return true;
 }
 
-void FInternalPlayWorldCommandCallbacks::HandleLaunchOnDeviceActionExecute(FString DeviceId, FString DeviceName, bool bUseTurnkey)
+static void PrepareLaunchOn(FString DeviceId, FString DeviceName)
 {
-	if (IsReadyToLaunchOnDevice(DeviceId))
-	{
-		ULevelEditorPlaySettings* PlaySettings = GetMutableDefault<ULevelEditorPlaySettings>();
+	ULevelEditorPlaySettings* PlaySettings = GetMutableDefault<ULevelEditorPlaySettings>();
 
-		PlaySettings->LastExecutedLaunchModeType = LaunchMode_OnDevice;
-		PlaySettings->LastExecutedLaunchDevice = DeviceId;
-		PlaySettings->LastExecutedLaunchName = DeviceName;
+	PlaySettings->LastExecutedLaunchModeType = LaunchMode_OnDevice;
+	PlaySettings->LastExecutedLaunchDevice = DeviceId;
+	PlaySettings->LastExecutedLaunchName = DeviceName;
 
-		PlaySettings->PostEditChange();
+	PlaySettings->PostEditChange();
+
+	PlaySettings->SaveConfig();
+}
 
-		PlaySettings->SaveConfig();
+void FInternalPlayWorldCommandCallbacks::HandleLaunchOnDeviceActionExecute(FString DeviceId, FString DeviceName, bool bUseTurnkey)
+{
+	if (bUseTurnkey)
+	{
+		int32 Index = 0;
+		DeviceId.FindChar(TEXT('@'), Index);
+		FString TargetPlatformName = DeviceId.Left(Index);
+		const PlatformInfo::FTargetPlatformInfo* const PlatformInfo = PlatformInfo::FindPlatformInfo(FName(*TargetPlatformName));
+		FString IniPlatformName = PlatformInfo->IniPlatformName.ToString();
 
+		FString CommandLine = FString::Printf(TEXT("Turnkey -command=VerifySdk -UpdateIfNeeded -platform=%s -EditorIO -noturnkeyvariables -device=%s"), *IniPlatformName, *DeviceName);
+		FText TaskName = LOCTEXT("VerifyingSDK", "Verifying SDK and Device");
 
-		LaunchOnDevice(DeviceId, DeviceName, bUseTurnkey);
+		IUATHelperModule::Get().CreateUatTask(CommandLine, FText::FromString(IniPlatformName), TaskName, TaskName, FEditorStyle::GetBrush(TEXT("MainFrame.PackageProject")),
+			[DeviceId, DeviceName](FString, double)
+		{
+			PrepareLaunchOn(DeviceId, DeviceName);
+			LaunchOnDevice(DeviceId, DeviceName, false);
+		});
+	}
+	else
+	{
+		if (IsReadyToLaunchOnDevice(DeviceId))
+		{
+			PrepareLaunchOn(DeviceId, DeviceName);
+			LaunchOnDevice(DeviceId, DeviceName, false);
+		}
 	}
-
 }
 
 
@@ -3113,10 +3131,10 @@ void FInternalPlayWorldCommandCallbacks::LaunchOnDevice(const FString& DeviceId,
 	FTargetDeviceId TargetDeviceId;
 	if (FTargetDeviceId::Parse(DeviceId, TargetDeviceId))
 	{
-		const PlatformInfo::FPlatformInfo* const PlatformInfo = PlatformInfo::FindPlatformInfo(*TargetDeviceId.GetPlatformName());
+		const PlatformInfo::FTargetPlatformInfo* const PlatformInfo = PlatformInfo::FindPlatformInfo(*TargetDeviceId.GetPlatformName());
 		check(PlatformInfo);
 
-		if (FInstalledPlatformInfo::Get().IsPlatformMissingRequiredFile(PlatformInfo->BinaryFolderName))
+		if (FInstalledPlatformInfo::Get().IsPlatformMissingRequiredFile(PlatformInfo->DataDrivenPlatformInfo->UBTPlatformString))
 		{
 			if (!FInstalledPlatformInfo::OpenInstallerOptions())
 			{
diff --git a/Engine/Source/Programs/AutomationTool/Turnkey/Providers/GoogleDriveCopyProvider.cs b/Engine/Source/Programs/AutomationTool/Turnkey/Providers/GoogleDriveCopyProvider.cs
--- a/Engine/Source/Programs/AutomationTool/Turnkey/Providers/GoogleDriveCopyProvider.cs
+++ b/Engine/Source/Programs/AutomationTool/Turnkey/Providers/GoogleDriveCopyProvider.cs
@@ -163,6 +163,8 @@ namespace Turnkey
 				return null;
 			}
 
+			TurnkeyUtils.Log("Enumerating GoogleDrive spec: {0}", Operation);
+
 			Dictionary<string, Tuple<Google.Apis.Drive.v3.Data.File, List<string>>> Output = new Dictionary<string, Tuple<Google.Apis.Drive.v3.Data.File, List<string>>>();
 
 			List<string> ExpansionSet = new List<string>();
diff --git a/Engine/Source/Programs/AutomationTool/Turnkey/SdkInfo.cs b/Engine/Source/Programs/AutomationTool/Turnkey/SdkInfo.cs
--- a/Engine/Source/Programs/AutomationTool/Turnkey/SdkInfo.cs
+++ b/Engine/Source/Programs/AutomationTool/Turnkey/SdkInfo.cs
@@ -190,7 +190,7 @@ namespace Turnkey
 				List<SdkInfo> Sdks = TurnkeyManifest.GetDiscoveredSdks();
 				Sdks = Sdks.FindAll(x => x.SupportsPlatform(Platform.IniPlatformType));
 				Sdks = Sdks.FindAll(x => x.Type == Type);
-				Sdks = Sdks.FindAll(x => (x.Version == null && x.CustomSdkId != null) || (Type == SdkType.Flash ? TurnkeyUtils.IsValueValid(Platform.GetAllowedSoftwareVersions(), x.Version, Platform) : SDK.IsVersionValid(x.Version, bForAutoSDK:x.Type==SdkType.AutoSdk)));
+				Sdks = Sdks.FindAll(x => (x.Version == null && x.CustomSdkId != null) || (Type == SdkType.Flash ? TurnkeyUtils.IsValueValid(x.Version, Platform.GetAllowedSoftwareVersions(), Platform) : SDK.IsVersionValid(x.Version, bForAutoSDK:x.Type==SdkType.AutoSdk)));
 				if (DeviceType != null)
 				{
 					Sdks = Sdks.FindAll(x => TurnkeyUtils.IsValueValid(DeviceType, x.AllowedFlashDeviceTypes, Platform));
diff --git a/Engine/Source/Programs/AutomationTool/Turnkey/TurnkeyUtils.cs b/Engine/Source/Programs/AutomationTool/Turnkey/TurnkeyUtils.cs
--- a/Engine/Source/Programs/AutomationTool/Turnkey/TurnkeyUtils.cs
+++ b/Engine/Source/Programs/AutomationTool/Turnkey/TurnkeyUtils.cs
@@ -4,6 +4,7 @@ using System;
 using System.Collections;
 using System.Collections.Generic;
 using System.IO;
+using System.Linq;
 using System.Text;
 using System.Text.RegularExpressions;
 using System.Globalization;
@@ -119,7 +120,8 @@ namespace Turnkey
 				// All platforms means to install every platform with an installer
 				if (PlatformChoice == PlatformOptions.Count)
 				{
-					Platforms = PossiblePlatforms;
+					// HACK UNTIL WIN32 IS GONE
+					Platforms = PossiblePlatforms.Where(x => x != UnrealTargetPlatform.Win32).ToList();
 				}
 				else
 				{
@@ -128,7 +130,8 @@ namespace Turnkey
 			}
 			else if (PlatformString == "All")
 			{
-				Platforms = PossiblePlatforms;
+				// HACK UNTIL WIN32 IS GONE
+				Platforms = PossiblePlatforms.Where(x => x != UnrealTargetPlatform.Win32).ToList();
 			}
 			else
 			{
@@ -297,7 +300,7 @@ namespace Turnkey
 
 				// convert inputs to uint (unless already converted above by platform)
 				UInt64 ValueInt = 0;
-				if (!TryConvertToUint64(Value, out ValueInt) && !UEBuildPlatformSDK.GetSDKForPlatform(Platform.ToString()).TryConvertVersionToInt(Value, out ValueInt))
+				if (!TryConvertToUint64(Value, out ValueInt) && !UEBuildPlatformSDK.GetSDKForPlatform(Platform.PlatformType.ToString()).TryConvertVersionToInt(Value, out ValueInt))
 				{
 					TurnkeyUtils.Log("Warning: range: input value [{0}] was not an unsigned integer, and platform couldn't convert it", Value);
 					return false;
diff --git a/Engine/Source/Runtime/Core/Private/Misc/ConfigCacheIni.cpp b/Engine/Source/Runtime/Core/Private/Misc/ConfigCacheIni.cpp
--- a/Engine/Source/Runtime/Core/Private/Misc/ConfigCacheIni.cpp
+++ b/Engine/Source/Runtime/Core/Private/Misc/ConfigCacheIni.cpp
@@ -3683,7 +3683,7 @@ void FConfigFile::AddStaticLayersToHierarchy(const TCHAR* InBaseIniName, const T
 				bool bGenerateCacheKey = EnumHasAnyFlags(Layer.Flag, EConfigLayerFlags::GenerateCacheKey) && ExpansionIndex == UE_ARRAY_COUNT(GConfigExpansions) - 1;
 				checkfSlow(!(bGenerateCacheKey && bHasPlatformTag), TEXT("EConfigLayerFlags::GenerateCacheKey shouldn't have a platform tag"));
 
-				const FDataDrivenPlatformInfoRegistry::FPlatformInfo& Info = FDataDrivenPlatformInfoRegistry::GetPlatformInfo(PlatformName);
+				const FDataDrivenPlatformInfo& Info = FDataDrivenPlatformInfoRegistry::GetPlatformInfo(PlatformName);
 
 				// go over parents, and then this platform, unless there's no platform tag, then we simply want to run through the loop one time to add it to the
 				int32 NumPlatforms = bHasPlatformTag ? Info.IniParentChain.Num() + 1 : 1;
diff --git a/Engine/Source/Runtime/Core/Private/Misc/DataDrivenPlatformInfoRegistry.cpp b/Engine/Source/Runtime/Core/Private/Misc/DataDrivenPlatformInfoRegistry.cpp
--- a/Engine/Source/Runtime/Core/Private/Misc/DataDrivenPlatformInfoRegistry.cpp
+++ b/Engine/Source/Runtime/Core/Private/Misc/DataDrivenPlatformInfoRegistry.cpp
@@ -11,14 +11,15 @@
 #include "Misc/MonitoredProcess.h"
 
 
-TMap<FString, FDataDrivenPlatformInfoRegistry::FPlatformInfo> FDataDrivenPlatformInfoRegistry::DataDrivenPlatforms;
+TMap<FString, FDataDrivenPlatformInfo> FDataDrivenPlatformInfoRegistry::DataDrivenPlatforms;
 
 #if DDPI_HAS_EXTENDED_PLATFORMINFO_DATA
+// NO AUTO-KICK-OFF - let the editor request it as needed
 // delay the kick off of running UAT so we can check IsRunningCommandlet()
-FDelayedAutoRegisterHelper GPlatformInfoInit(EDelayedRegisterRunPhase::TaskGraphSystemReady, []()
-{
-	FDataDrivenPlatformInfoRegistry::UpdateSdkStatus();
-});
+// FDelayedAutoRegisterHelper GPlatformInfoInit(EDelayedRegisterRunPhase::TaskGraphSystemReady, []()
+// {
+// 	FDataDrivenPlatformInfoRegistry::UpdateSdkStatus();
+// });
 #endif
 
 static const TArray<FString>& GetDataDrivenIniFilenames()
@@ -98,7 +99,12 @@ static void DDPIIniRedirect(FString& StringData)
 static FString DDPITryRedirect(const FConfigFile& IniFile, const TCHAR* Key, bool* OutHadBang=nullptr)
 {
 	FString StringData;
-	if (IniFile.GetString(TEXT("DataDrivenPlatformInfo"), Key, StringData))
+	bool bWasFound = false;
+	if ((bWasFound = IniFile.GetString(TEXT("DataDrivenPlatformInfo"), Key, StringData)) == false)
+	{
+		bWasFound = IniFile.GetString(TEXT("DataDrivenPlatformInfo"), *FString::Printf(TEXT("%s:%s"), ANSI_TO_TCHAR(FPlatformProperties::IniPlatformName()), Key), StringData);
+	}
+	if (bWasFound)
 	{
 		if (StringData.StartsWith(TEXT("ini:")) || StringData.StartsWith(TEXT("!ini:")))
 		{
@@ -149,6 +155,17 @@ static void DDPIGetUInt(const FConfigFile& IniFile, const TCHAR* Key, uint32& Ou
 	}
 }
 
+static void DDPIGetName(const FConfigFile& IniFile, const TCHAR* Key, FName& OutName)
+{
+	FString StringData = DDPITryRedirect(IniFile, Key);
+
+	// if we ended up with a string, convert it, otherwise leave it alone
+	if (StringData.Len() > 0)
+	{
+		OutName = FName(*StringData);
+	}
+}
+
 static void DDPIGetString(const FConfigFile& IniFile, const TCHAR* Key, FString& OutString)
 {
 	FString StringData = DDPITryRedirect(IniFile, Key);
@@ -166,7 +183,7 @@ static void DDPIGetStringArray(const FConfigFile& IniFile, const TCHAR* Key, TAr
 	IniFile.GetArray(TEXT("DataDrivenPlatformInfo"), Key, OutArray);
 }
 
-static void LoadDDPIIniSettings(const FConfigFile& IniFile, FDataDrivenPlatformInfoRegistry::FPlatformInfo& Info)
+static void LoadDDPIIniSettings(const FConfigFile& IniFile, FDataDrivenPlatformInfo& Info, const FString& PlatformName)
 {
 	DDPIGetBool(IniFile, TEXT("bIsConfidential"), Info.bIsConfidential);
 	DDPIGetBool(IniFile, TEXT("bIsFakePlatform"), Info.bIsFakePlatform);
@@ -183,27 +200,44 @@ static void LoadDDPIIniSettings(const FConfigFile& IniFile, FDataDrivenPlatformI
 
 
 #if DDPI_HAS_EXTENDED_PLATFORMINFO_DATA
-	// now look in the PlatformInfo objects in the file for TP and UBT names
 
-	FName TargetPlatformKey("TargetPlatformName");
-	FName UBTNameKey("UBTTargetID");
-	for (auto& Pair : IniFile)
+	DDPIGetString(IniFile, TEXT("AutoSDKPath"), Info.AutoSDKPath);
+	DDPIGetString(IniFile, TEXT("TutorialPath"), Info.SDKTutorial);
+	DDPIGetName(IniFile, TEXT("PlatformGroupName"), Info.PlatformGroupName);
+	DDPIGetName(IniFile, TEXT("PlatformSubMenu"), Info.PlatformSubMenu);
+	
+
+	DDPIGetString(IniFile, TEXT("NormalIconPath"), Info.IconPaths.NormalPath);
+	DDPIGetString(IniFile, TEXT("LargeIconPath"), Info.IconPaths.LargePath);
+	DDPIGetString(IniFile, TEXT("XLargeIconPath"), Info.IconPaths.XLargePath);
+	if (Info.IconPaths.XLargePath == TEXT(""))
 	{
-		if (Pair.Key.StartsWith(TEXT("PlatformInfo")))
-		{
-			const FConfigValue* PlatformName = Pair.Value.Find(TargetPlatformKey);
-			if (PlatformName)
-			{
-				Info.AllTargetPlatformNames.AddUnique(PlatformName->GetValue());
-			}
-			
-			PlatformName = Pair.Value.Find(UBTNameKey);
-			if (PlatformName)
-			{
-				Info.AllUBTPlatformNames.AddUnique(PlatformName->GetValue());
-			}
-		}
+		Info.IconPaths.XLargePath = Info.IconPaths.XLargePath;
+	}
+
+	Info.IconPaths.NormalStyleName = *FString::Printf(TEXT("Launcher.Platform_%s"), *PlatformName);
+	Info.IconPaths.LargeStyleName = *FString::Printf(TEXT("Launcher.Platform_%s.Large"), *PlatformName);
+	Info.IconPaths.XLargeStyleName = *FString::Printf(TEXT("Launcher.Platform_%s.XLarge"), *PlatformName);
+
+	Info.bCanUseCrashReporter = true; // not specified means true, not false
+	DDPIGetBool(IniFile, TEXT("bCanUseCrashReporter"), Info.bCanUseCrashReporter);
+	DDPIGetBool(IniFile, TEXT("bUsesHostCompiler"), Info.bUsesHostCompiler);
+	DDPIGetBool(IniFile, TEXT("bUATClosesAfterLaunch"), Info.bUATClosesAfterLaunch);
+	DDPIGetBool(IniFile, TEXT("bIsEnabled"), Info.bEnabledForUse);
+
+	DDPIGetName(IniFile, TEXT("UBTPlatformName"), Info.UBTPlatformName);
+	// if unspecified, use the ini platform name (only Win64 breaks this)
+	if (Info.UBTPlatformName == NAME_None)
+	{
+		Info.UBTPlatformName = *PlatformName;
 	}
+	Info.UBTPlatformString = Info.UBTPlatformName.ToString();
+		
+	
+	// now that we have all targetplatforms in a single TP module per platform, just look for it (or a ShaderFormat for other tools that may want this)
+	// we could look for Platform*, but then platforms that are a substring of another one could return a false positive (Windows* would find Windows31TargetPlatform)
+	Info.bHasCompiledTargetSupport = FDataDrivenPlatformInfoRegistry::HasCompiledSupportForPlatform(PlatformName, FDataDrivenPlatformInfoRegistry::EPlatformNameType::TargetPlatform);
+
 #endif
 }
 
@@ -211,7 +245,7 @@ static void LoadDDPIIniSettings(const FConfigFile& IniFile, FDataDrivenPlatformI
 /**
 * Get the global set of data driven platform information
 */
-const TMap<FString, FDataDrivenPlatformInfoRegistry::FPlatformInfo>& FDataDrivenPlatformInfoRegistry::GetAllPlatformInfos()
+const TMap<FString, FDataDrivenPlatformInfo>& FDataDrivenPlatformInfoRegistry::GetAllPlatformInfos()
 {
 	static bool bHasSearchedForPlatforms = false;
 
@@ -234,8 +268,8 @@ const TMap<FString, FDataDrivenPlatformInfoRegistry::FPlatformInfo>& FDataDriven
 			if (IniFile.Contains(TEXT("DataDrivenPlatformInfo")))
 			{
 				// cache info
-				FDataDrivenPlatformInfoRegistry::FPlatformInfo& Info = DataDrivenPlatforms.Add(PlatformName, FDataDrivenPlatformInfoRegistry::FPlatformInfo());
-				LoadDDPIIniSettings(IniFile, Info);
+				FDataDrivenPlatformInfo& Info = DataDrivenPlatforms.Add(PlatformName, FDataDrivenPlatformInfo());
+				LoadDDPIIniSettings(IniFile, Info, PlatformName);
 
 				// get the parent to build list later
 				FString IniParent;
@@ -262,15 +296,19 @@ const TMap<FString, FDataDrivenPlatformInfoRegistry::FPlatformInfo>& FDataDriven
 #if DDPI_HAS_EXTENDED_PLATFORMINFO_DATA
 
 // some shared functionality
-static void PrepForTurnkeyReport(FString& UAT, FString& ReportFilename)
+static void PrepForTurnkeyReport(FString& Command, FString& BaseCommandline, FString& ReportFilename)
 {
-	static int ReportIndex = 0;
-
-	// get path to AutomationTool - skipping RunUAT for speed
-//		FString AutomationTool = FPaths::Combine(FPaths::EngineDir(), TEXT("Binaries/DotNET/AutomationTool.exe"));
-	UAT = FPaths::ConvertRelativePathToFull(FPaths::Combine(FPaths::EngineDir(), TEXT("Build/BatchFiles/RunUAT.bat")));
+ 	static int ReportIndex = 0;
 
+	FString LogFilename = FPaths::ConvertRelativePathToFull(FPaths::Combine(FPaths::ProjectIntermediateDir(), *FString::Printf(TEXT("TurnkeyLog_%d.log"), ReportIndex)));
 	ReportFilename = FPaths::ConvertRelativePathToFull(FPaths::Combine(FPaths::ProjectIntermediateDir(), *FString::Printf(TEXT("TurnkeyReport_%d.log"), ReportIndex++)));
+
+	Command = TEXT("{EngineDir}/Build/BatchFiles/RunuAT");
+//	Command = TEXT("{EngineDir}/Binaries/DotNET/AutomationTool.exe");
+	BaseCommandline = FString::Printf(TEXT("Turnkey -command=VerifySdk -ReportFilename=\"%s\" -log=\"%s\""), *ReportFilename, *LogFilename);
+
+	// convert into appropriate calls for the current platform
+	FPlatformProcess::ModifyCreateProcParams(Command, BaseCommandline, FGenericPlatformProcess::ECreateProcHelperFlags::AppendPlatformScriptExtension);
 }
 
 static FString ConvertToDDPIPlatform(const FString& Platform)
@@ -326,22 +364,9 @@ void FDataDrivenPlatformInfoRegistry::UpdateSdkStatus()
 	}
 
 
-	FString AutomationTool, ReportFilename;
-	PrepForTurnkeyReport(AutomationTool, ReportFilename);
-
-	// run Turnkey to get all of the platform statuses
-// 	FString ProcessURL = TEXT("{EngineDir}/Binaries/DotNET/AutomationTool.exe");
-// 	FString ProcessArguments = FString::Printf(TEXT("Turnkey -command=VerifySdk -platform=%s -ReportFilename=\"%s\""),
-// 		*FString::JoinBy(DataDrivenPlatforms, TEXT("+"), [](TPair<FString, FDataDrivenPlatformInfoRegistry::FPlatformInfo> Pair) { return ConvertToUATPlatform(Pair.Key); }),
-// 		*ReportFilename);
-
-	FString ProcessURL = TEXT("{EngineDir}/Build/BatchFiles/RunuAT");
-	FString ProcessArguments = FString::Printf(TEXT("Turnkey -command=VerifySdk -platform=%s -ReportFilename=\"%s\""),
-		*FString::JoinBy(DataDrivenPlatforms, TEXT("+"), [](TPair<FString, FDataDrivenPlatformInfoRegistry::FPlatformInfo> Pair) { return ConvertToUATPlatform(Pair.Key); }),
-		*ReportFilename);
-
-	// convert into appropriate calls for the current platform
-	FPlatformProcess::ModifyCreateProcParams(ProcessURL, ProcessArguments, FGenericPlatformProcess::ECreateProcHelperFlags::RunThroughShell | FGenericPlatformProcess::ECreateProcHelperFlags::AppendPlatformScriptExtension);
+	FString Command, BaseCommandline, ReportFilename;
+	PrepForTurnkeyReport(Command, BaseCommandline, ReportFilename);
+	FString Commandline = BaseCommandline + FString(TEXT(" -platform=")) + FString::JoinBy(DataDrivenPlatforms, TEXT("+"), [](TPair<FString, FDataDrivenPlatformInfo> Pair) { return ConvertToUATPlatform(Pair.Key); });
 
 	// reset status to unknown
 	for (auto& It : DataDrivenPlatforms)
@@ -352,7 +377,7 @@ void FDataDrivenPlatformInfoRegistry::UpdateSdkStatus()
 		It.Value.ClearDeviceStatus();
 	}
 
-	FMonitoredProcess* TurnkeyProcess = new FMonitoredProcess(ProcessURL, ProcessArguments, true, false);
+	FMonitoredProcess* TurnkeyProcess = new FMonitoredProcess(Command, Commandline, true, false);
 	TurnkeyProcess->OnCompleted().BindLambda([ReportFilename, TurnkeyProcess](int32 ExitCode)
 	{
 		AsyncTask(ENamedThreads::GameThread, [ReportFilename, TurnkeyProcess, ExitCode]()
@@ -430,7 +455,7 @@ void FDataDrivenPlatformInfoRegistry::UpdateSdkStatus()
 	TurnkeyProcess->Launch();
 }
 
-FDataDrivenPlatformInfoRegistry::FPlatformInfo& FDataDrivenPlatformInfoRegistry::DeviceIdToInfo(FString DeviceId, FString* OutDeviceName)
+FDataDrivenPlatformInfo& FDataDrivenPlatformInfoRegistry::DeviceIdToInfo(FString DeviceId, FString* OutDeviceName)
 {
 	TArray<FString> PlatformAndDevice;
 	DeviceId.ParseIntoArray(PlatformAndDevice, TEXT("@"), true);
@@ -447,14 +472,12 @@ FDataDrivenPlatformInfoRegistry::FPlatformInfo& FDataDrivenPlatformInfoRegistry:
 
 void FDataDrivenPlatformInfoRegistry::UpdateDeviceSdkStatus(TArray<FString> PlatformDeviceIds)
 {
-	FString AutomationTool, ReportFilename;
-	PrepForTurnkeyReport(AutomationTool, ReportFilename);
+	FString Command, BaseCommandline, ReportFilename;
+	PrepForTurnkeyReport(Command, BaseCommandline, ReportFilename);
 
-	// run Turnkey to get all of the platform statuses (turn the array into "Platform@Device+Platform@Device+...")
-	FString Commandline = FString::Printf(TEXT("/c %s Turnkey -command=VerifySdk -Device=\"%s\" -ReportFilename=\"%s\""),
-		*AutomationTool,
-		*FString::JoinBy(PlatformDeviceIds, TEXT("+"), [](const FString& DeviceId) { return ConvertToUATDeviceId(DeviceId); }),
-		*ReportFilename);
+	// the platform part of the Id may need to be converted to be turnkey (ie UBT) proper
+
+	FString Commandline = BaseCommandline + FString(TEXT(" -Device=")) + FString::JoinBy(PlatformDeviceIds, TEXT("+"), [](FString Id) { return ConvertToUATDeviceId(Id); });
 
 	// set status to querying
 	for (const FString& Id : PlatformDeviceIds)
@@ -462,7 +485,7 @@ void FDataDrivenPlatformInfoRegistry::UpdateDeviceSdkStatus(TArray<FString> Plat
 		DeviceIdToInfo(Id).PerDeviceStatus.Add(ConvertToDDPIDeviceId(Id), DDPIPlatformSdkStatus::Querying);
 	}
 
-	FMonitoredProcess* TurnkeyProcess = new FMonitoredProcess(TEXT("cmd.exe"), Commandline, true, false);
+	FMonitoredProcess* TurnkeyProcess = new FMonitoredProcess(Command, Commandline, true, false);
 	TurnkeyProcess->OnCompleted().BindLambda([ReportFilename, TurnkeyProcess, PlatformDeviceIds](int32 ExitCode)
 	{
 		AsyncTask(ENamedThreads::GameThread, [ReportFilename, TurnkeyProcess, PlatformDeviceIds, ExitCode]()
@@ -498,7 +521,7 @@ void FDataDrivenPlatformInfoRegistry::UpdateDeviceSdkStatus(TArray<FString> Plat
 
 			for (const FString& Id : PlatformDeviceIds)
 			{
-				FDataDrivenPlatformInfoRegistry::FPlatformInfo& Info = DeviceIdToInfo(Id);
+				FDataDrivenPlatformInfo& Info = DeviceIdToInfo(Id);
 				
 				DDPIPlatformSdkStatus& Status = Info.PerDeviceStatus[ConvertToDDPIDeviceId(Id)];
 				if (Status == DDPIPlatformSdkStatus::Querying)
@@ -518,19 +541,35 @@ void FDataDrivenPlatformInfoRegistry::UpdateDeviceSdkStatus(TArray<FString> Plat
 	TurnkeyProcess->Launch();
 }
 
+void FDataDrivenPlatformInfoRegistry::ClearDeviceStatus(const FString& PlatformName)
+{
+	if (PlatformName.Len() == 0)
+	{
+		for (auto It : DataDrivenPlatforms)
+		{
+			It.Value.ClearDeviceStatus();
+		}
+	}
+	else
+	{
+		if (DataDrivenPlatforms.Contains(PlatformName))
+		{
+			DataDrivenPlatforms[PlatformName].ClearDeviceStatus();
+		}
+	}
+}
 
-DDPIPlatformSdkStatus FDataDrivenPlatformInfoRegistry::FPlatformInfo::GetStatusForDeviceId(const FString& DeviceId) const
+DDPIPlatformSdkStatus FDataDrivenPlatformInfo::GetStatusForDeviceId(const FString& DeviceId) const
 {
 	// return the status, or Unknown if not known
 	return PerDeviceStatus.FindRef(ConvertToDDPIDeviceId(DeviceId));
 }
 
-void FDataDrivenPlatformInfoRegistry::FPlatformInfo::ClearDeviceStatus()
+void FDataDrivenPlatformInfo::ClearDeviceStatus()
 {
 	PerDeviceStatus.Empty();
 }
 
-
 #endif
 
 
@@ -544,7 +583,7 @@ const TArray<FString>& FDataDrivenPlatformInfoRegistry::GetValidPlatformDirector
 		bHasSearchedForPlatforms = true;
 
 		// look for possible platforms
-		const TMap<FString, FDataDrivenPlatformInfoRegistry::FPlatformInfo>& Infos = GetAllPlatformInfos();
+		const TMap<FString, FDataDrivenPlatformInfo>& Infos = GetAllPlatformInfos();
 		for (auto Pair : Infos)
 		{
 #if DDPI_HAS_EXTENDED_PLATFORMINFO_DATA
@@ -570,13 +609,24 @@ const TArray<FString>& FDataDrivenPlatformInfoRegistry::GetValidPlatformDirector
 }
 
 
-const FDataDrivenPlatformInfoRegistry::FPlatformInfo& FDataDrivenPlatformInfoRegistry::GetPlatformInfo(const FString& PlatformName)
+const FDataDrivenPlatformInfo& FDataDrivenPlatformInfoRegistry::GetPlatformInfo(const FString& PlatformName)
 {
-	const FPlatformInfo* Info = GetAllPlatformInfos().Find(PlatformName);
-	static FPlatformInfo Empty;
+	const FDataDrivenPlatformInfo* Info = GetAllPlatformInfos().Find(PlatformName);
+	static FDataDrivenPlatformInfo Empty;
 	return Info ? *Info : Empty;
 }
 
+const FDataDrivenPlatformInfo& FDataDrivenPlatformInfoRegistry::GetPlatformInfo(FName PlatformName)
+{
+	// @todo maybe make the key be FName instead of FString and flip these?
+	return GetPlatformInfo(PlatformName.ToString());
+}
+
+const FDataDrivenPlatformInfo& FDataDrivenPlatformInfoRegistry::GetPlatformInfo(const char* PlatformName)
+{
+	return GetPlatformInfo(FName(PlatformName));
+}
+
 
 const TArray<FString>& FDataDrivenPlatformInfoRegistry::GetConfidentialPlatforms()
 {
@@ -608,18 +658,8 @@ bool FDataDrivenPlatformInfoRegistry::HasCompiledSupportForPlatform(const FStrin
 	if (PlatformNameType == EPlatformNameType::Ini)
 	{
 		// get the DDPI info object
-		const FPlatformInfo& Info = GetPlatformInfo(PlatformName);
-
-		// look to see if any of the TPs in the Info are valid - if at least one is, we are good
-		for (const FString& TPName : Info.AllTargetPlatformNames)
-		{
-			if (HasCompiledSupportForPlatform(TPName, EPlatformNameType::TargetPlatform))
-			{
-				return true;
-			}
-		}
-		return false;
-
+		const FDataDrivenPlatformInfo& Info = GetPlatformInfo(PlatformName);
+		return Info.bHasCompiledTargetSupport;
 	}
 	else if (PlatformNameType == EPlatformNameType::UBT)
 	{
@@ -628,9 +668,8 @@ bool FDataDrivenPlatformInfoRegistry::HasCompiledSupportForPlatform(const FStrin
 		// find all the DataDrivenPlatformInfo objects and find a matching the UBT name
 		for (auto& Pair : FDataDrivenPlatformInfoRegistry::GetAllPlatformInfos())
 		{
-			// if this platform contains the UBT platform name, then check the info for it's TPs
-			// (we could be tricky and match UBT platforms with TPs just for thse UBT platforms, but that complexity does not seem needed)
-			if (Pair.Value.AllUBTPlatformNames.Contains(PlatformName))
+			// if this platform matches the UBT platform name, check it's Ini name
+			if (Pair.Value.UBTPlatformName == PlatformFName)
 			{
 				return HasCompiledSupportForPlatform(Pair.Key, EPlatformNameType::Ini);
 			}
@@ -640,10 +679,15 @@ bool FDataDrivenPlatformInfoRegistry::HasCompiledSupportForPlatform(const FStrin
 	}
 	else if (PlatformNameType == EPlatformNameType::TargetPlatform)
 	{
-		// was this TP compiled?
-		return FModuleManager::Get().ModuleExists(*FString::Printf(TEXT("%sTargetPlatform"), *PlatformName));
+		// was this TP compiled, or a shaderformat (useful for SCW if it ever calls this)
+		return 
+			FModuleManager::Get().ModuleExists(*FString::Printf(TEXT("%sTargetPlatform"), *PlatformName)) || 
+			FModuleManager::Get().ModuleExists(*FString::Printf(TEXT("%sShaderFormat"), *PlatformName));
 	}
 
 	return false;
 }
+
+static void ClearDeviceStatus(const FString& PlatformName);
+
 #endif
diff --git a/Engine/Source/Runtime/Core/Private/Serialization/MemoryImage.cpp b/Engine/Source/Runtime/Core/Private/Serialization/MemoryImage.cpp
--- a/Engine/Source/Runtime/Core/Private/Serialization/MemoryImage.cpp
+++ b/Engine/Source/Runtime/Core/Private/Serialization/MemoryImage.cpp
@@ -46,7 +46,7 @@ void FPlatformTypeLayoutParameters::InitializeForArchive(FArchive& Ar)
 
 void FPlatformTypeLayoutParameters::InitializeForPlatform(const FString& PlatformName, bool bHasEditorOnlyData)
 {
-	const FDataDrivenPlatformInfoRegistry::FPlatformInfo& PlatformInfo = FDataDrivenPlatformInfoRegistry::GetPlatformInfo(PlatformName);
+	const FDataDrivenPlatformInfo& PlatformInfo = FDataDrivenPlatformInfoRegistry::GetPlatformInfo(PlatformName);
 
 	Flags = Flag_Initialized;
 	if (bHasEditorOnlyData) Flags |= Flag_WithEditorOnly;
diff --git a/Engine/Source/Runtime/Core/Public/HAL/IPlatformFileOpenLogWrapper.h b/Engine/Source/Runtime/Core/Public/HAL/IPlatformFileOpenLogWrapper.h
--- a/Engine/Source/Runtime/Core/Public/HAL/IPlatformFileOpenLogWrapper.h
+++ b/Engine/Source/Runtime/Core/Public/HAL/IPlatformFileOpenLogWrapper.h
@@ -117,7 +117,7 @@ public:
 		}
 		else
 		{
-			if (FDataDrivenPlatformInfoRegistry::GetPlatformInfo(FPlatformProperties::PlatformName()).bIsConfidential)
+			if (FDataDrivenPlatformInfoRegistry::GetPlatformInfo(FPlatformProperties::IniPlatformName()).bIsConfidential)
 			{
 				LogFileDirectory = FPaths::Combine(FPlatformMisc::ProjectDir(), TEXT("Platforms"), StringCast<TCHAR>(FPlatformProperties::PlatformName()).Get(), TEXT("Build"), TEXT("FileOpenOrder"));
 			}
diff --git a/Engine/Source/Runtime/Core/Public/Misc/DataDrivenPlatformInfoRegistry.h b/Engine/Source/Runtime/Core/Public/Misc/DataDrivenPlatformInfoRegistry.h
--- a/Engine/Source/Runtime/Core/Public/Misc/DataDrivenPlatformInfoRegistry.h
+++ b/Engine/Source/Runtime/Core/Public/Misc/DataDrivenPlatformInfoRegistry.h
@@ -7,7 +7,7 @@
 #include "Containers/Map.h"
 #include "Containers/Array.h"
 
-#define DDPI_HAS_EXTENDED_PLATFORMINFO_DATA WITH_EDITOR && !IS_MONOLITHIC
+#define DDPI_HAS_EXTENDED_PLATFORMINFO_DATA (WITH_EDITOR || IS_PROGRAM) && !IS_MONOLITHIC
 
 #if DDPI_HAS_EXTENDED_PLATFORMINFO_DATA
 
@@ -22,67 +22,213 @@ enum class DDPIPlatformSdkStatus : uint8
 	Error,
 	FlashValid,
 	FlashOutOfDate,
+	// @todo turnkey: add AutoSdkValid and ManualSdkValid, with Valid a Combination of both
 };
 
-#endif
 
+/** Available icon sizes (see FPlatformIconPaths) */
+enum class EPlatformIconSize : uint8
+{
+	/** Normal sized icon (24x24) */
+	Normal,
+
+	/** Large sized icon (64x64) */
+	Large,
 
-struct CORE_API FDataDrivenPlatformInfoRegistry
+	/** Extra large sized icon (128x128) */
+	XLarge,
+};
+
+
+
+// enum class EPlatformSDKStatus : uint8
+// {
+// 	/** SDK status is unknown */
+// 	Unknown,
+// 
+// 	/** SDK is installed */
+// 	Installed,
+// 
+// 	/** SDK is not installed */
+// 	NotInstalled,
+// };
+
+/** Information about where to find the platform icons (for use by FEditorStyle) */
+struct FPlatformIconPaths
 {
-	// Information about a platform loaded from disk
-	struct FPlatformInfo
+	FPlatformIconPaths()
 	{
-		// cached list of ini parents
-		TArray<FString> IniParentChain;
+	}
 
-		// is this platform confidential
-		bool bIsConfidential = false;
+	FPlatformIconPaths(const FString& InIconPath)
+		: NormalPath(InIconPath)
+		, LargePath(InIconPath)
+		, XLargePath(InIconPath)
+	{
+	}
 
-		// some platforms are here just for IniParentChain needs and are not concrete platforms
-		bool bIsFakePlatform = false;
+	FPlatformIconPaths(const FString& InNormalPath, const FString& InLargePath)
+		: NormalPath(InNormalPath)
+		, LargePath(InLargePath)
+		, XLargePath(InLargePath)
+	{
+	}
 
-		// the name of the ini section to use to load audio compression settings (used at runtime and cooktime)
-		FString AudioCompressionSettingsIniSectionName;
+	FPlatformIconPaths(const FString& InNormalPath, const FString& InLargePath, const FString& InXLargePath)
+		: NormalPath(InNormalPath)
+		, LargePath(InLargePath)
+		, XLargePath(InXLargePath)
+	{
+	}
 
-		// list of additonal restricted folders
-		TArray<FString> AdditionalRestrictedFolders;
+	FName NormalStyleName;
+	FString NormalPath;
 
-		// MemoryFreezing information, matches FPlatformTypeLayoutParameters - defaults are clang, noneditor
-		uint32 Freezing_MaxFieldAlignment = 0xffffffff;
-		bool Freezing_b32Bit = false;
-		bool Freezing_bForce64BitMemoryImagePointers = false;
-		bool Freezing_bAlignBases = false;
-		bool Freezing_bWithRayTracing = false;
+	FName LargeStyleName;
+	FString LargePath;
 
-		// NOTE: add more settings here (and read them in in the LoadDDPIIniSettings() function in the .cpp)
+	FName XLargeStyleName;
+	FString XLargePath;
+};
+
+#endif
+
+// Information about a platform loaded from disk
+struct FDataDrivenPlatformInfo
+{
+	// cached list of ini parents
+	TArray<FString> IniParentChain;
+
+	// is this platform confidential
+	bool bIsConfidential = false;
+
+	// some platforms are here just for IniParentChain needs and are not concrete platforms
+	bool bIsFakePlatform = false;
+
+	// the name of the ini section to use to load audio compression settings (used at runtime and cooktime)
+	FString AudioCompressionSettingsIniSectionName;
+
+	// list of additonal restricted folders
+	TArray<FString> AdditionalRestrictedFolders;
+
+	// MemoryFreezing information, matches FPlatformTypeLayoutParameters - defaults are clang, noneditor
+	uint32 Freezing_MaxFieldAlignment = 0xffffffff;
+	bool Freezing_b32Bit = false;
+	bool Freezing_bForce64BitMemoryImagePointers = false;
+	bool Freezing_bAlignBases = false;
+	bool Freezing_bWithRayTracing = false;
+
+	// NOTE: add more settings here (and read them in in the LoadDDPIIniSettings() function in the .cpp)
 
 
 #if DDPI_HAS_EXTENDED_PLATFORMINFO_DATA
-		// list of TP names (WindowsNoEditor, etc)
-		TArray<FString> AllTargetPlatformNames;
-		// list of UBT platform names (Win32, Win64, etc)
-		TArray<FString> AllUBTPlatformNames;
 
-		// Information about the validity of using a platform, discovered via Turnkey
-		DDPIPlatformSdkStatus SdkStatus;
-		FString SdkErrorInformation;
+	// Information about the validity of using a platform, discovered via Turnkey
+private:
+	DDPIPlatformSdkStatus SdkStatus;
+public:
+	DDPIPlatformSdkStatus GetSdkStatus(bool bBlockIfQuerying = true) const
+	{
+		check(!(bBlockIfQuerying && SdkStatus == DDPIPlatformSdkStatus::Querying));
+		return SdkStatus;
+	}
+	FString SdkErrorInformation;
+
+	// setting moved from PlatformInfo::FTargetPlatformInfo
+
 
-		// Get the status of a device, or Unknown if not specified
-		CORE_API DDPIPlatformSdkStatus GetStatusForDeviceId(const FString& DeviceId) const;
-		CORE_API void ClearDeviceStatus();
+	/** Information about where to find the platform icons (for use by FEditorStyle) */
+	FPlatformIconPaths IconPaths;
+
+	/** Path under CarefullyRedist for the SDK.  FString so case sensitive platforms don't get messed up by a pre-existing FName of a different casing. */
+	FString AutoSDKPath;
+
+	/** Tutorial path for tutorial to install SDK */
+	FString SDKTutorial;
+
+	/** An identifier to group similar platforms together, such as "Mobile" and "Console". Used for Per-Platform Override Properties. */
+	FName PlatformGroupName;
+
+	/** Submenu name to group similar platforms together in menus, such as "Linux" and "LinuxAArch64".  */
+	FName PlatformSubMenu;
+
+	/** An identifier that corresponds to UBT's UnrealTargetPlatform enum (and by proxy, FGenericPlatformMisc::GetUBTPlatform()), as well as the directory Binaries are placed under */
+	FName UBTPlatformName;
+	FString UBTPlatformString;
+
+	/** Whether or not the platform can use Crash Reporter */
+	bool bCanUseCrashReporter;
+
+	/** Enabled for use */
+	bool bEnabledForUse;
+
+	/** Whether code projects for this platform require the host platform compiler to be installed. Host platforms typically have a SDK status of valid, but they can't necessarily build. */
+	bool bUsesHostCompiler;
+
+	/** Whether UAT closes immediately after launching on this platform, or if it sticks around to read output from the running process */
+	bool bUATClosesAfterLaunch;
+
+	/** Whether or not this editor/program has compiled in support for this platform (by looking for TargetPlatform style DLLs, without loading them) */
+	bool bHasCompiledTargetSupport;
+
+
+
+	// Get the status of a device, or Unknown if not specified
+	CORE_API DDPIPlatformSdkStatus GetStatusForDeviceId(const FString& DeviceId) const;
+	CORE_API void ClearDeviceStatus();
+
+
+	/** Get the icon name (for FEditorStyle) used by the given icon type for this platform */
+	FName GetIconStyleName(const EPlatformIconSize InIconSize) const
+	{
+		switch (InIconSize)
+		{
+		case EPlatformIconSize::Normal:
+			return IconPaths.NormalStyleName;
+		case EPlatformIconSize::Large:
+			return IconPaths.LargeStyleName;
+		case EPlatformIconSize::XLarge:
+			return IconPaths.XLargeStyleName;
+		default:
+			break;
+		}
+		return NAME_None;
+	}
+
+	/** Get the path to the icon on disk (for FEditorStyle) for the given icon type for this platform */
+	const FString& GetIconPath(const EPlatformIconSize InIconSize) const
+	{
+		switch (InIconSize)
+		{
+		case EPlatformIconSize::Normal:
+			return IconPaths.NormalPath;
+		case EPlatformIconSize::Large:
+			return IconPaths.LargePath;
+		case EPlatformIconSize::XLarge:
+			return IconPaths.XLargePath;
+		default:
+			break;
+		}
+		static const FString EmptyString = TEXT("");
+		return EmptyString;
+	}
 
 private:
-		friend FDataDrivenPlatformInfoRegistry;
+	friend struct FDataDrivenPlatformInfoRegistry;
 
-		// Information about the validity of each connected device (by string, discovered by Turnkey)
-		TMap<FString, DDPIPlatformSdkStatus> PerDeviceStatus;
+	// Information about the validity of each connected device (by string, discovered by Turnkey)
+	TMap<FString, DDPIPlatformSdkStatus> PerDeviceStatus;
 #endif
-	};
+};
+
+
+struct CORE_API FDataDrivenPlatformInfoRegistry
+{
 
 	/**
 	* Get the global set of data driven platform information
 	*/
-	static const TMap<FString, FDataDrivenPlatformInfoRegistry::FPlatformInfo>& GetAllPlatformInfos();
+	static const TMap<FString, FDataDrivenPlatformInfo>& GetAllPlatformInfos();
 
 	/**
 	 * Gets a set of platform names based on GetAllPlatformInfos, their AdditionalRestrictedFolders, and possibly filtered based on what editor has support compiled for
@@ -94,7 +240,9 @@ private:
 	 * Get the data driven platform info for a given platform. If the platform doesn't have any on disk,
 	 * this will return a default constructed FConfigDataDrivenPlatformInfo
 	 */
-	static const FPlatformInfo& GetPlatformInfo(const FString& PlatformName);
+	static const FDataDrivenPlatformInfo& GetPlatformInfo(const FString& PlatformName);
+	static const FDataDrivenPlatformInfo& GetPlatformInfo(FName PlatformName);
+	static const FDataDrivenPlatformInfo& GetPlatformInfo(const char* PlatformName);
 
 	/**
 	 * Gets a list of all known confidential platforms (note these are just the platforms you have access to, so, for example PS4 won't be
@@ -139,13 +287,18 @@ private:
 	 */
 	static void UpdateDeviceSdkStatus(TArray<FString> PlatformDeviceIds);
 
+	/**
+	 * Wipes out cached device status for all devices in a platform (or all platforms if PlatformName is empty)
+	 */
+	static void ClearDeviceStatus(const FString& PlatformName);
+
 private:
-	static FDataDrivenPlatformInfoRegistry::FPlatformInfo& DeviceIdToInfo(FString DeviceId, FString* OutDeviceName = nullptr);
+	static FDataDrivenPlatformInfo& DeviceIdToInfo(FString DeviceId, FString* OutDeviceName = nullptr);
 
 #endif
 
 
 private:
-	static TMap<FString, FDataDrivenPlatformInfoRegistry::FPlatformInfo> DataDrivenPlatforms;
+	static TMap<FString, FDataDrivenPlatformInfo> DataDrivenPlatforms;
 };
 
diff --git a/Engine/Source/Runtime/Engine/Classes/Engine/Engine.h b/Engine/Source/Runtime/Engine/Classes/Engine/Engine.h
--- a/Engine/Source/Runtime/Engine/Classes/Engine/Engine.h
+++ b/Engine/Source/Runtime/Engine/Classes/Engine/Engine.h
@@ -1835,8 +1835,8 @@ public:
 	virtual ERHIFeatureLevel::Type GetDefaultWorldFeatureLevel() const;
 
 #if WITH_EDITOR
-	/** Return the platform group name and vanilla platform name the current preview platform, or false if there is no preview platform. */
-	virtual bool GetPreviewPlatformName(FName& PlatformGroupName, FName& VanillaPlatformName) const;
+	/** Return the ini platform name the current preview platform, or false if there is no preview platform. */
+	virtual bool GetPreviewPlatformName(FName& PlatformName) const;
 
 	/** Editor-only event triggered when the actor list of the world has changed */
 	DECLARE_EVENT( UEngine, FLevelActorListChangedEvent );
diff --git a/Engine/Source/Runtime/Engine/Private/DeviceProfiles/DeviceProfileManager.cpp b/Engine/Source/Runtime/Engine/Private/DeviceProfiles/DeviceProfileManager.cpp
--- a/Engine/Source/Runtime/Engine/Private/DeviceProfiles/DeviceProfileManager.cpp
+++ b/Engine/Source/Runtime/Engine/Private/DeviceProfiles/DeviceProfileManager.cpp
@@ -594,7 +594,7 @@ void UDeviceProfileManager::LoadProfiles()
 				ITargetPlatform* Platform = TargetPlatforms[PlatformIndex];
 
 				// Set TextureLODSettings
-				const UTextureLODSettings* TextureLODSettingsObj = FindProfile(*Platform->GetPlatformInfo().VanillaPlatformName.ToString(), false);
+				const UTextureLODSettings* TextureLODSettingsObj = FindProfile(*Platform->GetTargetPlatformInfo().VanillaInfo->PlatformInfoName.ToString(), false);
 				Platform->RegisterTextureLODSettings(TextureLODSettingsObj);
 			}
 		}
diff --git a/Engine/Source/Runtime/Engine/Private/PerPlatformProperties.cpp b/Engine/Source/Runtime/Engine/Private/PerPlatformProperties.cpp
--- a/Engine/Source/Runtime/Engine/Private/PerPlatformProperties.cpp
+++ b/Engine/Source/Runtime/Engine/Private/PerPlatformProperties.cpp
@@ -22,8 +22,7 @@ ENGINE_API FArchive& operator<<(FArchive& Ar, TPerPlatformProperty<StructType, V
 		bCooked = true;
 		Ar << bCooked;
 		// Save out platform override if it exists and Default otherwise
-		const PlatformInfo::FPlatformInfo& PlatformInfo = Ar.CookingTarget()->GetPlatformInfo();
-		ValueType Value = Property.GetValueForPlatformIdentifiers(PlatformInfo.PlatformGroupName, PlatformInfo.VanillaPlatformName);
+		ValueType Value = Property.GetValueForPlatform(*Ar.CookingTarget()->IniPlatformName());
 		Ar << Value;
 	}
 	else
@@ -56,7 +55,7 @@ ENGINE_API void operator<<(FStructuredArchive::FSlot Slot, TPerPlatformProperty<
 		bCooked = true;
 		Record << SA_VALUE(TEXT("bCooked"), bCooked);
 		// Save out platform override if it exists and Default otherwise
-		ValueType Value = Property.GetValueForPlatformIdentifiers(UnderlyingArchive.CookingTarget()->GetPlatformInfo().PlatformGroupName);
+		ValueType Value = Property.GetValueForPlatform(*UnderlyingArchive.CookingTarget()->IniPlatformName());
 		Record << SA_VALUE(TEXT("Value"), Value);
 	}
 	else
diff --git a/Engine/Source/Runtime/Engine/Private/SkeletalMesh.cpp b/Engine/Source/Runtime/Engine/Private/SkeletalMesh.cpp
--- a/Engine/Source/Runtime/Engine/Private/SkeletalMesh.cpp
+++ b/Engine/Source/Runtime/Engine/Private/SkeletalMesh.cpp
@@ -4142,15 +4142,13 @@ void USkeletalMesh::ResetLODInfo()
 bool USkeletalMesh::GetSupportsLODStreaming(const ITargetPlatform* TargetPlatform) const
 {
 	check(TargetPlatform);
-	const FName PlatformGroupName = TargetPlatform->GetPlatformInfo().PlatformGroupName;
-	const FName VanillaPlatformName = TargetPlatform->GetPlatformInfo().VanillaPlatformName;
 	if (bOverrideLODStreamingSettings)
 	{
-		return bSupportLODStreaming.GetValueForPlatformIdentifiers(PlatformGroupName, VanillaPlatformName);
+		return bSupportLODStreaming.GetValueForPlatform(*TargetPlatform->IniPlatformName());
 	}
 	else
 	{
-		return GetDefault<URendererSettings>()->bStreamSkeletalMeshLODs.GetValueForPlatformIdentifiers(PlatformGroupName, VanillaPlatformName);
+		return GetDefault<URendererSettings>()->bStreamSkeletalMeshLODs.GetValueForPlatform(*TargetPlatform->IniPlatformName());
 	}
 }
 
@@ -4159,9 +4157,7 @@ int32 USkeletalMesh::GetMaxNumStreamedLODs(const ITargetPlatform* TargetPlatform
 	check(TargetPlatform);
 	if (bOverrideLODStreamingSettings)
 	{
-		return MaxNumStreamedLODs.GetValueForPlatformIdentifiers(
-			TargetPlatform->GetPlatformInfo().PlatformGroupName,
-			TargetPlatform->GetPlatformInfo().VanillaPlatformName);
+		return MaxNumStreamedLODs.GetValueForPlatform(*TargetPlatform->IniPlatformName());
 	}
 	else
 	{
@@ -4172,15 +4168,13 @@ int32 USkeletalMesh::GetMaxNumStreamedLODs(const ITargetPlatform* TargetPlatform
 int32 USkeletalMesh::GetMaxNumOptionalLODs(const ITargetPlatform* TargetPlatform) const
 {
 	check(TargetPlatform);
-	const FName PlatformGroupName = TargetPlatform->GetPlatformInfo().PlatformGroupName;
-	const FName VanillaPlatformName = TargetPlatform->GetPlatformInfo().VanillaPlatformName;
 	if (bOverrideLODStreamingSettings)
 	{
-		return MaxNumOptionalLODs.GetValueForPlatformIdentifiers(PlatformGroupName, VanillaPlatformName) <= 0 ? 0 : MAX_MESH_LOD_COUNT;
+		return MaxNumOptionalLODs.GetValueForPlatform(*TargetPlatform->IniPlatformName()) <= 0 ? 0 : MAX_MESH_LOD_COUNT;
 	}
 	else
 	{
-		return GetDefault<URendererSettings>()->bDiscardSkeletalMeshOptionalLODs.GetValueForPlatformIdentifiers(PlatformGroupName, VanillaPlatformName) ? 0 : MAX_MESH_LOD_COUNT;
+		return GetDefault<URendererSettings>()->bDiscardSkeletalMeshOptionalLODs.GetValueForPlatform(*TargetPlatform->IniPlatformName()) ? 0 : MAX_MESH_LOD_COUNT;
 	}
 }
 #endif
diff --git a/Engine/Source/Runtime/Engine/Private/SkeletalMeshLODRenderData.cpp b/Engine/Source/Runtime/Engine/Private/SkeletalMeshLODRenderData.cpp
--- a/Engine/Source/Runtime/Engine/Private/SkeletalMeshLODRenderData.cpp
+++ b/Engine/Source/Runtime/Engine/Private/SkeletalMeshLODRenderData.cpp
@@ -701,9 +701,7 @@ int32 FSkeletalMeshLODRenderData::GetPlatformMinLODIdx(const ITargetPlatform* Ta
 {
 #if WITH_EDITOR
 	check(TargetPlatform && SkeletalMesh);
-	const FName PlatformGroupName = TargetPlatform->GetPlatformInfo().PlatformGroupName;
-	const FName VanillaPlatformName = TargetPlatform->GetPlatformInfo().VanillaPlatformName;
-	return  SkeletalMesh->MinLod.GetValueForPlatformIdentifiers(PlatformGroupName, VanillaPlatformName);
+	return  SkeletalMesh->MinLod.GetValueForPlatform(*TargetPlatform->IniPlatformName());
 #else
 	return 0;
 #endif
@@ -722,8 +720,8 @@ uint8 FSkeletalMeshLODRenderData::GenerateClassStripFlags(FArchive& Ar, const US
 	bool bMeshDisablesMinLodStrip = false;
 	if (bIsCook)
 	{
-		MinMeshLod = OwnerMesh ? OwnerMesh->MinLod.GetValueForPlatformIdentifiers(CookTarget->GetPlatformInfo().PlatformGroupName, CookTarget->GetPlatformInfo().VanillaPlatformName) : 0;
-		bMeshDisablesMinLodStrip = OwnerMesh ? OwnerMesh->DisableBelowMinLodStripping.GetValueForPlatformIdentifiers(CookTarget->GetPlatformInfo().PlatformGroupName, CookTarget->GetPlatformInfo().VanillaPlatformName) : false;
+		MinMeshLod = OwnerMesh ? OwnerMesh->MinLod.GetValueForPlatform(*CookTarget->IniPlatformName()) : 0;
+		bMeshDisablesMinLodStrip = OwnerMesh ? OwnerMesh->DisableBelowMinLodStripping.GetValueForPlatform(*CookTarget->IniPlatformName()) : false;
 	}
 	const bool bWantToStripBelowMinLod = bIsCook && GStripSkeletalMeshLodsDuringCooking != 0 && MinMeshLod > LODIdx && !bMeshDisablesMinLodStrip;
 
diff --git a/Engine/Source/Runtime/Engine/Private/StaticMesh.cpp b/Engine/Source/Runtime/Engine/Private/StaticMesh.cpp
--- a/Engine/Source/Runtime/Engine/Private/StaticMesh.cpp
+++ b/Engine/Source/Runtime/Engine/Private/StaticMesh.cpp
@@ -268,9 +268,7 @@ int32 FStaticMeshLODResources::GetPlatformMinLODIdx(const ITargetPlatform* Targe
 {
 #if WITH_EDITOR
 	check(TargetPlatform && StaticMesh);
-	return StaticMesh->MinLOD.GetValueForPlatformIdentifiers(
-		TargetPlatform->GetPlatformInfo().PlatformGroupName,
-		TargetPlatform->GetPlatformInfo().VanillaPlatformName);
+	return StaticMesh->MinLOD.GetValueForPlatform(*TargetPlatform->IniPlatformName());
 #else
 	return 0;
 #endif
@@ -353,9 +351,7 @@ bool FStaticMeshLODResources::IsLODInlined(const ITargetPlatform* TargetPlatform
 	}
 
 	int32 MaxNumStreamedLODs = 0;
-	const int32 NumStreamedLODsOverride = StaticMesh->NumStreamedLODs.GetValueForPlatformIdentifiers(
-		TargetPlatform->GetPlatformInfo().PlatformGroupName,
-		TargetPlatform->GetPlatformInfo().VanillaPlatformName);
+	const int32 NumStreamedLODsOverride = StaticMesh->NumStreamedLODs.GetValueForPlatform(*TargetPlatform->IniPlatformName());
 	if (NumStreamedLODsOverride >= 0)
 	{
 		MaxNumStreamedLODs = NumStreamedLODsOverride;
@@ -2300,9 +2296,7 @@ static FString BuildStaticMeshDerivedDataKeySuffix(const ITargetPlatform* Target
 	KeySuffix += TEXT("MNS");
 	if (bAllowLODStreaming)
 	{
-		int32 MaxNumStreamedLODs = Mesh->NumStreamedLODs.GetValueForPlatformIdentifiers(
-			TargetPlatform->GetPlatformInfo().PlatformGroupName,
-			TargetPlatform->GetPlatformInfo().VanillaPlatformName);
+		int32 MaxNumStreamedLODs = Mesh->NumStreamedLODs.GetValueForPlatform(*TargetPlatform->IniPlatformName());
 		if (MaxNumStreamedLODs < 0)
 		{
 			MaxNumStreamedLODs = LODGroup.GetDefaultMaxNumStreamedLODs();
diff --git a/Engine/Source/Runtime/Engine/Private/TextureLODSettings.cpp b/Engine/Source/Runtime/Engine/Private/TextureLODSettings.cpp
--- a/Engine/Source/Runtime/Engine/Private/TextureLODSettings.cpp
+++ b/Engine/Source/Runtime/Engine/Private/TextureLODSettings.cpp
@@ -280,8 +280,7 @@ void UTextureLODSettings::GetDownscaleOptions(const UTexture& Texture, const ITa
 		GroupDownscaleOptions = ETextureDownscaleOptions::SimpleAverage;
 	}
 		
-	const PlatformInfo::FPlatformInfo& Info = CurrentPlatform.GetPlatformInfo();
-	Downscale = Texture.Downscale.GetValueForPlatformIdentifiers(Info.PlatformGroupName, Info.VanillaPlatformName);;
+	Downscale = Texture.Downscale.GetValueForPlatform(*CurrentPlatform.IniPlatformName());
 	if (Downscale < 1.f)
 	{
 		// value is not overriden for this texture
diff --git a/Engine/Source/Runtime/Engine/Private/UnrealEngine.cpp b/Engine/Source/Runtime/Engine/Private/UnrealEngine.cpp
--- a/Engine/Source/Runtime/Engine/Private/UnrealEngine.cpp
+++ b/Engine/Source/Runtime/Engine/Private/UnrealEngine.cpp
@@ -16277,7 +16277,7 @@ void UEngine::PostEditChangeProperty(struct FPropertyChangedEvent& PropertyChang
 	Super::PostEditChangeProperty(PropertyChangedEvent);
 }
 
-bool UEngine::GetPreviewPlatformName(FName& PlatformGroupName, FName& VanillaPlatformName) const
+bool UEngine::GetPreviewPlatformName(FName& PlatformName) const
 {
 	// Support returning "Mobile" when started with -featureleveles31 etc
 
@@ -16294,8 +16294,8 @@ bool UEngine::GetPreviewPlatformName(FName& PlatformGroupName, FName& VanillaPla
 
 	if (bCachedMobile)
 	{
-		PlatformGroupName = NAME_Mobile;
-		PlatformGroupName = NAME_None;
+		// the editor uses a real platform - this is unlikely tyo actually be called, but the code does technically handle using a group name as a platform name
+		PlatformName = NAME_Mobile;
 		return true;
 	}
 	else
diff --git a/Engine/Source/Runtime/Engine/Public/PerPlatformProperties.h b/Engine/Source/Runtime/Engine/Public/PerPlatformProperties.h
--- a/Engine/Source/Runtime/Engine/Public/PerPlatformProperties.h
+++ b/Engine/Source/Runtime/Engine/Public/PerPlatformProperties.h
@@ -12,6 +12,7 @@ PerPlatformProperties.h: Property types that can be overridden on a per-platform
 #include "Containers/Map.h"
 #include "Algo/Find.h"
 #include "Serialization/MemoryLayout.h"
+#include "Misc/DataDrivenPlatformInfoRegistry.h"
 
 #if WITH_EDITORONLY_DATA && WITH_EDITOR
 #include "Engine/Engine.h"
@@ -29,8 +30,27 @@ struct ENGINE_API TPerPlatformProperty
 	typedef _ValueType ValueType;
 
 #if WITH_EDITOR
+	/** Get the value for the given platform (using standard "ini" name, so Windows, not Win64 or WindowsNoEditor), which can be used to lookup the group */
+	_ValueType GetValueForPlatform(FName PlatformName) const
+	{
+		const _StructType* This = StaticCast<const _StructType*>(this);
+		const _ValueType* Ptr = This->PerPlatform.Find(PlatformName);
+
+		if (Ptr == nullptr)
+		{
+			const FDataDrivenPlatformInfo& Info = FDataDrivenPlatformInfoRegistry::GetPlatformInfo(PlatformName);
+			if (Info.PlatformGroupName != NAME_None)
+			{
+				Ptr = This->PerPlatform.Find(Info.PlatformGroupName);
+			}
+		}
+
+		return Ptr ? *Ptr : This->Default;
+	}
+
 	/* Return the value */
-	_ValueType GetValueForPlatformIdentifiers(FName PlatformGroupName, FName VanillaPlatformName = NAME_None) const
+	UE_DEPRECATED(5.0, "GetValueForPlatform should now be used, with a single platform name")
+		_ValueType GetValueForPlatformIdentifiers(FName PlatformGroupName, FName VanillaPlatformName = NAME_None) const
 	{
 		const _StructType* This = StaticCast<const _StructType*>(this);
 		
@@ -68,11 +88,12 @@ struct ENGINE_API TPerPlatformProperty
 	_ValueType GetValue() const
 	{
 #if WITH_EDITORONLY_DATA && WITH_EDITOR
-		FName PlatformGroupName, VanillaPlatformName;
+		FName PlatformName;
 		// Lookup the override preview platform info, if any
-		if (GEngine && GEngine->GetPreviewPlatformName(PlatformGroupName, VanillaPlatformName))
+		// @todo this doesn't set PlatformName, just a group, but GetValueForPlatform() will technically work being given a Group name instead of a platform name, so we just use it
+		if (GEngine && GEngine->GetPreviewPlatformName(PlatformName))
 		{
-			return GetValueForPlatformIdentifiers(PlatformGroupName, VanillaPlatformName);
+			return GetValueForPlatform(PlatformName);
 		}
 		else		
 #endif
