Fixing a shutdown crash related to moving editor mode manager ownership to toolkits.

#rnx

diff --git a/Engine/Source/Editor/UnrealEd/Private/EditorModeManager.cpp b/Engine/Source/Editor/UnrealEd/Private/EditorModeManager.cpp
--- a/Engine/Source/Editor/UnrealEd/Private/EditorModeManager.cpp
+++ b/Engine/Source/Editor/UnrealEd/Private/EditorModeManager.cpp
@@ -754,9 +754,13 @@ void FEditorModeTools::RemoveAllDelegateHandlers()
 		GEditor->GetEditorSubsystem<UAssetEditorSubsystem>()->OnEditorModeUnregistered().RemoveAll(this);
 	}
 
-	USelection::SelectionChangedEvent.RemoveAll(this);
-	USelection::SelectNoneEvent.RemoveAll(this);
-	USelection::SelectObjectEvent.RemoveAll(this);
+	// For now, check that UObjects are even valid, because the level editor has a global static mode tools
+	if (UObjectInitialized())
+	{
+		USelection::SelectionChangedEvent.RemoveAll(this);
+		USelection::SelectNoneEvent.RemoveAll(this);
+		USelection::SelectObjectEvent.RemoveAll(this);
+	}
 
 	OnEditorModeIDChanged().Clear();
 	OnWidgetModeChanged().Clear();
