Updated Header, follow up from 13885548

diff --git a/Engine/Source/Editor/UnrealEd/Public/EditorModeManager.h b/Engine/Source/Editor/UnrealEd/Public/EditorModeManager.h
--- a/Engine/Source/Editor/UnrealEd/Public/EditorModeManager.h
+++ b/Engine/Source/Editor/UnrealEd/Public/EditorModeManager.h
@@ -112,7 +112,6 @@ protected:
 		
 private:
 	void RebuildModeToolBar();
-	void SpawnOrUpdateModeToolbar();
 public:
 
 	/**
@@ -557,11 +556,6 @@ public:
 	/** returns true if all active EdModes are OK with an AutoSave happening now  */
 	bool CanAutoSave() const;
 
-	/*
-	* Sets the active Modes ToolBar Palette Tab to the named Palette
-	*/
-	void  InvokeToolPaletteTab(FEditorModeID InMode, FName InPaletteName);
-
 	void RemoveAllDelegateHandlers();
 
 protected:
